[Mutable] Added graph aligment options.

#nrx

diff --git a/Engine/Plugins/Experimental/Mutable/Source/CustomizableObjectEditor/Private/MuCOE/CustomizableObjectEditor.cpp b/Engine/Plugins/Experimental/Mutable/Source/CustomizableObjectEditor/Private/MuCOE/CustomizableObjectEditor.cpp
--- a/Engine/Plugins/Experimental/Mutable/Source/CustomizableObjectEditor/Private/MuCOE/CustomizableObjectEditor.cpp
+++ b/Engine/Plugins/Experimental/Mutable/Source/CustomizableObjectEditor/Private/MuCOE/CustomizableObjectEditor.cpp
@@ -565,7 +565,7 @@ TSharedRef<SDockTab> FCustomizableObjectEditor::SpawnTab_InstanceProperties( con
 
 
 /** Create new tab for the supplied graph - don't call this directly, call SExplorer->FindTabForGraph.*/
-TSharedRef<SGraphEditor> FCustomizableObjectEditor::CreateGraphEditorWidget(UEdGraph* InGraph)
+void FCustomizableObjectEditor::CreateGraphEditorWidget(UEdGraph* InGraph)
 {
 	check(InGraph != NULL);
 
@@ -592,38 +592,7 @@ TSharedRef<SGraphEditor> FCustomizableObjectEditor::CreateGraphEditorWidget(UEdG
 		Action.FCustomizableObjectSchemaAction_NewNode::PerformAction(InGraph, nullptr, FVector2D::ZeroVector, false);
 	}
 	
-	GraphEditorCommands = MakeShareable( new FUICommandList );
-	{
-		// Editing commands
-		GraphEditorCommands->MapAction( FGenericCommands::Get().Delete,
-			FExecuteAction::CreateSP( this, &FCustomizableObjectEditor::DeleteSelectedNodes ),
-			FCanExecuteAction::CreateSP( this, &FCustomizableObjectEditor::CanDeleteNodes )
-			);
-
-		GraphEditorCommands->MapAction( FGenericCommands::Get().Copy,
-			FExecuteAction::CreateSP( this, &FCustomizableObjectEditor::CopySelectedNodes ),
-			FCanExecuteAction::CreateSP( this, &FCustomizableObjectEditor::CanCopyNodes )
-			);
-
-		GraphEditorCommands->MapAction( FGenericCommands::Get().Paste,
-			FExecuteAction::CreateSP( this, &FCustomizableObjectEditor::PasteNodes ),
-			FCanExecuteAction::CreateSP( this, &FCustomizableObjectEditor::CanPasteNodes )
-			);
-
-		GraphEditorCommands->MapAction( FGenericCommands::Get().Cut,
-			FExecuteAction::CreateSP( this, &FCustomizableObjectEditor::CutSelectedNodes ),
-			FCanExecuteAction::CreateSP( this, &FCustomizableObjectEditor::CanCutNodes )
-			);
-
-		GraphEditorCommands->MapAction(FGenericCommands::Get().Duplicate,
-			FExecuteAction::CreateSP(this, &FCustomizableObjectEditor::DuplicateSelectedNodes),
-			FCanExecuteAction::CreateSP(this, &FCustomizableObjectEditor::CanDuplicateSelectedNodes)
-			);
-
-		GraphEditorCommands->MapAction(FCustomizableObjectEditorNodeContextCommands::Get().CreateComment,
-			FExecuteAction::CreateSP(this, &FCustomizableObjectEditor::CreateCommentBoxFromKey)
-		);
-	}
+	GraphEditorCommands = MakeShareable(new FUICommandList);
 
 	TSharedRef<SWidget> TitleBarWidget =
 		SNew(SHorizontalBox)
@@ -648,13 +617,66 @@ TSharedRef<SGraphEditor> FCustomizableObjectEditor::CreateGraphEditorWidget(UEdG
 	InEvents.OnTextCommitted = FOnNodeTextCommitted::CreateSP(this, &FCustomizableObjectEditor::OnNodeTitleCommitted);
 
 	// Make full graph editor
-	return SNew(SGraphEditor)
+	GraphEditor = SNew(SGraphEditor)
 		.AdditionalCommands(GraphEditorCommands)
 		.Appearance(AppearanceInfo)
 		.GraphToEdit(InGraph)
 		.GraphEvents(InEvents)
 		.TitleBar(TitleBarWidget)
 		.ShowGraphStateOverlay(false); // Removes graph state overlays (border and text) such as "SIMULATING" and "READ-ONLY"
+
+	// Editing commands
+	GraphEditorCommands->MapAction(FGenericCommands::Get().Delete,
+		FExecuteAction::CreateSP(this, &FCustomizableObjectEditor::DeleteSelectedNodes),
+		FCanExecuteAction::CreateSP(this, &FCustomizableObjectEditor::CanDeleteNodes));
+
+	GraphEditorCommands->MapAction( FGenericCommands::Get().Copy,
+		FExecuteAction::CreateSP(this, &FCustomizableObjectEditor::CopySelectedNodes),
+		FCanExecuteAction::CreateSP(this, &FCustomizableObjectEditor::CanCopyNodes));
+
+	GraphEditorCommands->MapAction(FGenericCommands::Get().Paste,
+		FExecuteAction::CreateSP(this, &FCustomizableObjectEditor::PasteNodes),
+		FCanExecuteAction::CreateSP(this, &FCustomizableObjectEditor::CanPasteNodes));
+
+	GraphEditorCommands->MapAction(FGenericCommands::Get().Cut,
+		FExecuteAction::CreateSP(this, &FCustomizableObjectEditor::CutSelectedNodes),
+		FCanExecuteAction::CreateSP(this, &FCustomizableObjectEditor::CanCutNodes));
+
+	GraphEditorCommands->MapAction(FGenericCommands::Get().Duplicate,
+		FExecuteAction::CreateSP(this, &FCustomizableObjectEditor::DuplicateSelectedNodes),
+		FCanExecuteAction::CreateSP(this, &FCustomizableObjectEditor::CanDuplicateSelectedNodes));
+
+	GraphEditorCommands->MapAction(FCustomizableObjectEditorNodeContextCommands::Get().CreateComment,
+		FExecuteAction::CreateSP(this, &FCustomizableObjectEditor::CreateCommentBoxFromKey));
+
+	// Alignment Commands
+	GraphEditorCommands->MapAction(FGraphEditorCommands::Get().AlignNodesTop,
+		FExecuteAction::CreateSP(GraphEditor.Get(), &SGraphEditor::OnAlignTop));
+
+	GraphEditorCommands->MapAction(FGraphEditorCommands::Get().AlignNodesMiddle,
+		FExecuteAction::CreateSP(GraphEditor.Get(), &SGraphEditor::OnAlignMiddle));
+
+	GraphEditorCommands->MapAction(FGraphEditorCommands::Get().AlignNodesBottom,
+		FExecuteAction::CreateSP(GraphEditor.Get(), &SGraphEditor::OnAlignBottom));
+
+	GraphEditorCommands->MapAction(FGraphEditorCommands::Get().AlignNodesLeft,
+		FExecuteAction::CreateSP(GraphEditor.Get(), &SGraphEditor::OnAlignLeft));
+
+	GraphEditorCommands->MapAction(FGraphEditorCommands::Get().AlignNodesCenter,
+		FExecuteAction::CreateSP(GraphEditor.Get(), &SGraphEditor::OnAlignCenter));
+
+	GraphEditorCommands->MapAction(FGraphEditorCommands::Get().AlignNodesRight,
+		FExecuteAction::CreateSP(GraphEditor.Get(), &SGraphEditor::OnAlignRight));
+
+	GraphEditorCommands->MapAction(FGraphEditorCommands::Get().StraightenConnections,
+		FExecuteAction::CreateSP(GraphEditor.Get(), &SGraphEditor::OnStraightenConnections));
+
+	// Distribution Commands
+	GraphEditorCommands->MapAction(FGraphEditorCommands::Get().DistributeNodesHorizontally,
+		FExecuteAction::CreateSP(GraphEditor.Get(), &SGraphEditor::OnDistributeNodesH));
+	
+	GraphEditorCommands->MapAction(FGraphEditorCommands::Get().DistributeNodesVertically,
+		FExecuteAction::CreateSP(GraphEditor.Get(), &SGraphEditor::OnDistributeNodesV));
 }
 
 
@@ -662,7 +684,7 @@ TSharedRef<SDockTab> FCustomizableObjectEditor::SpawnTab_Graph( const FSpawnTabA
 {
 	check( Args.GetTabId().TabType == GraphTabId );
 
-	GraphEditor = CreateGraphEditorWidget(CustomizableObject->Source);
+	CreateGraphEditorWidget(CustomizableObject->Source);
 
 	TSharedRef<SDockTab> DockTab = SNew(SDockTab)
 		.Label( FText::FromString( GetTabPrefix() + LOCTEXT( "SourceGraph", "Source Graph" ).ToString() ) )
diff --git a/Engine/Plugins/Experimental/Mutable/Source/CustomizableObjectEditor/Private/MuCOE/CustomizableObjectEditor.h b/Engine/Plugins/Experimental/Mutable/Source/CustomizableObjectEditor/Private/MuCOE/CustomizableObjectEditor.h
--- a/Engine/Plugins/Experimental/Mutable/Source/CustomizableObjectEditor/Private/MuCOE/CustomizableObjectEditor.h
+++ b/Engine/Plugins/Experimental/Mutable/Source/CustomizableObjectEditor/Private/MuCOE/CustomizableObjectEditor.h
@@ -310,8 +310,7 @@ private:
 	/** */
 	void OnPreviewInstanceUpdated();
 	
-	/** */
-	TSharedRef<SGraphEditor> CreateGraphEditorWidget(UEdGraph* InGraph);
+	void CreateGraphEditorWidget(UEdGraph* InGraph);
 
 	/** Copy the currently selected nodes */
 	void CopySelectedNodes();
diff --git a/Engine/Plugins/Experimental/Mutable/Source/CustomizableObjectEditor/Private/MuCOE/EdGraphSchema_CustomizableObject.cpp b/Engine/Plugins/Experimental/Mutable/Source/CustomizableObjectEditor/Private/MuCOE/EdGraphSchema_CustomizableObject.cpp
--- a/Engine/Plugins/Experimental/Mutable/Source/CustomizableObjectEditor/Private/MuCOE/EdGraphSchema_CustomizableObject.cpp
+++ b/Engine/Plugins/Experimental/Mutable/Source/CustomizableObjectEditor/Private/MuCOE/EdGraphSchema_CustomizableObject.cpp
@@ -822,84 +822,106 @@ void UEdGraphSchema_CustomizableObject::GetContextMenuActionsReconstructAllChild
 
 void UEdGraphSchema_CustomizableObject::GetContextMenuActions(UToolMenu* Menu, UGraphNodeContextMenuContext* Context) const
 {
-	if (Context && Context->Node)
+	if (!Context || !Context->Node)
 	{
-		if (!Context->Pin) // On Node right click
+		return;
+	}
+	
+	if (!Context->Pin) // On Node right click
+	{
+		if (!Context->bIsDebugging)
 		{
-			if (!Context->bIsDebugging)
+			// Node contextual actions
+			FToolMenuSection& Section = Menu->AddSection("EdGraphSchemaNodeActions", LOCTEXT("NodeActionsMenuHeader", "Node Actions"));
+			Section.AddMenuEntry(FGenericCommands::Get().Delete);
+			Section.AddMenuEntry(FGenericCommands::Get().Cut);
+			Section.AddMenuEntry(FGenericCommands::Get().Copy);
+			Section.AddMenuEntry(FGenericCommands::Get().Duplicate);
+			Section.AddMenuEntry(FGraphEditorCommands::Get().ReconstructNodes);
+			Section.AddMenuEntry(FGraphEditorCommands::Get().BreakNodeLinks);
+
+			Section.AddSubMenu("Alignment", LOCTEXT("AlignmentHeader", "Alignment"), FText(), FNewToolMenuDelegate::CreateLambda([](UToolMenu* AlignmentMenu)
 			{
-				// Node contextual actions
-				FToolMenuSection& Section = Menu->AddSection("EdGraphSchemaNodeActions", LOCTEXT("NodeActionsMenuHeader", "Node Actions"));
-				Section.AddMenuEntry(FGenericCommands::Get().Delete);
-				Section.AddMenuEntry(FGenericCommands::Get().Cut);
-				Section.AddMenuEntry(FGenericCommands::Get().Copy);
-				Section.AddMenuEntry(FGenericCommands::Get().Duplicate);
-				Section.AddMenuEntry(FGraphEditorCommands::Get().ReconstructNodes);
-				Section.AddMenuEntry(FGraphEditorCommands::Get().BreakNodeLinks);
-				
-				Section.AddSubMenu(
-					"ReconstructAllNodes",
-					LOCTEXT("ReconstructChildAllNodes", "Refresh All Child Nodes"),
-					LOCTEXT("ReconstructAllChildNodes_Tooltip", "Refresh all child nodes from the selected ones (inclusive)."),
-					FNewToolMenuDelegate::CreateUObject(this, &UEdGraphSchema_CustomizableObject::GetContextMenuActionsReconstructAllChildNodes, MakeWeakObjectPtr(Context)));
-			}
+				{
+					FToolMenuSection& SubSection = AlignmentMenu->AddSection("EdGraphSchemaAlignment", LOCTEXT("AlignHeader", "Align"));
+					SubSection.AddMenuEntry(FGraphEditorCommands::Get().AlignNodesTop);
+					SubSection.AddMenuEntry(FGraphEditorCommands::Get().AlignNodesMiddle);
+					SubSection.AddMenuEntry(FGraphEditorCommands::Get().AlignNodesBottom);
+					SubSection.AddMenuEntry(FGraphEditorCommands::Get().AlignNodesLeft);
+					SubSection.AddMenuEntry(FGraphEditorCommands::Get().AlignNodesCenter);
+					SubSection.AddMenuEntry(FGraphEditorCommands::Get().AlignNodesRight);
+					SubSection.AddMenuEntry(FGraphEditorCommands::Get().StraightenConnections);
+				}
 
-			struct SCommentUtility
+				{
+					FToolMenuSection& SubSection = AlignmentMenu->AddSection("EdGraphSchemaDistribution", LOCTEXT("DistributionHeader", "Distribution"));
+					SubSection.AddMenuEntry(FGraphEditorCommands::Get().DistributeNodesHorizontally);
+					SubSection.AddMenuEntry(FGraphEditorCommands::Get().DistributeNodesVertically);
+				}
+			}));
+			
+			Section.AddSubMenu(
+				"ReconstructAllNodes",
+				LOCTEXT("ReconstructChildAllNodes", "Refresh All Child Nodes"),
+				LOCTEXT("ReconstructAllChildNodes_Tooltip", "Refresh all child nodes from the selected ones (inclusive)."),
+				FNewToolMenuDelegate::CreateUObject(this, &UEdGraphSchema_CustomizableObject::GetContextMenuActionsReconstructAllChildNodes, MakeWeakObjectPtr(Context)));
+		}
+
+		struct SCommentUtility
+		{
+			static void CreateComment(const UEdGraphSchema_CustomizableObject* Schema, UEdGraph* Graph)
 			{
-				static void CreateComment(const UEdGraphSchema_CustomizableObject* Schema, UEdGraph* Graph)
+				if (Schema && Graph)
 				{
-					if (Schema && Graph)
-					{
-						Schema->AddComment(Graph, NULL, FVector2D::ZeroVector, true);
-					}
+					Schema->AddComment(Graph, NULL, FVector2D::ZeroVector, true);
 				}
-			};
+			}
+		};
 
-			FToolMenuSection& section = Menu->AddSection("SchemaActionComment", LOCTEXT("MultiCommentHeader", "Comment Group"));
-			section.AddMenuEntry("MultiCommentDesc", LOCTEXT("MultiCommentDesc", "Create Comment from Selection"),
-				LOCTEXT("CommentToolTip", "Create a resizable comment box around selection."),
-				FSlateIcon(),
-				FUIAction(FExecuteAction::CreateStatic(SCommentUtility::CreateComment, this, const_cast<UEdGraph*>(ToRawPtr(Context->Graph)))));
-		}
-		else // On Pin right click
+		FToolMenuSection& section = Menu->AddSection("SchemaActionComment", LOCTEXT("MultiCommentHeader", "Comment Group"));
+		section.AddMenuEntry("MultiCommentDesc", LOCTEXT("MultiCommentDesc", "Create Comment from Selection"),
+			LOCTEXT("CommentToolTip", "Create a resizable comment box around selection."),
+			FSlateIcon(),
+			FUIAction(FExecuteAction::CreateStatic(SCommentUtility::CreateComment, this, const_cast<UEdGraph*>(ToRawPtr(Context->Graph)))));
+	}
+	else // On Pin right click
+	{
+		UCustomizableObjectNodeTable* TableNode = Cast<UCustomizableObjectNodeTable>(Context->Node);
+		UEdGraphPin* TexturePin = (UEdGraphPin*)Context->Pin;
+
+		if (TableNode && TexturePin && !TexturePin->LinkedTo.Num() && (TexturePin->PinType.PinCategory == PC_Image || TexturePin->PinType.PinCategory == PC_PassThroughImage))
 		{
-			UCustomizableObjectNodeTable* TableNode = Cast<UCustomizableObjectNodeTable>(Context->Node);
-			UEdGraphPin* TexturePin = (UEdGraphPin*)Context->Pin;
+			FText ActionText = FText::Format(LOCTEXT("ChangeTexturePinMode_Label", 
+				"Set as {0} texture"), Context->Pin->PinType.PinCategory == PC_PassThroughImage ? FText::FromString("Mutable") : FText::FromString("Pass Through"));
+			
+			FText ToolTipText = FText::Format(LOCTEXT("ChangeTexturePinMode_Tooltip",
+				"Set the texture pin as {0} texture."), Context->Pin->PinType.PinCategory == PC_PassThroughImage ? FText::FromString("Mutable") : FText::FromString("Pass Through"));
 
-			if (TableNode && TexturePin && !TexturePin->LinkedTo.Num() && (TexturePin->PinType.PinCategory == PC_Image || TexturePin->PinType.PinCategory == PC_PassThroughImage))
-			{
-				FText ActionText = FText::Format(LOCTEXT("ChangeTexturePinMode_Label", 
-					"Set as {0} texture"), Context->Pin->PinType.PinCategory == PC_PassThroughImage ? FText::FromString("Mutable") : FText::FromString("Pass Through"));
-				
-				FText ToolTipText = FText::Format(LOCTEXT("ChangeTexturePinMode_Tooltip",
-					"Set the texture pin as {0} texture."), Context->Pin->PinType.PinCategory == PC_PassThroughImage ? FText::FromString("Mutable") : FText::FromString("Pass Through"));
+			FToolMenuSection& Section = Menu->FindOrAddSection("EdGraphSchemaPinActions");
+			Section.InitSection("EdGraphSchemaPinActions", LOCTEXT("PinActionsMenuHeader", "Pin Actions"), FToolMenuInsert());
 
-				FToolMenuSection& Section = Menu->FindOrAddSection("EdGraphSchemaPinActions");
-				Section.InitSection("EdGraphSchemaPinActions", LOCTEXT("PinActionsMenuHeader", "Pin Actions"), FToolMenuInsert());
+			Section.AddMenuEntry
+			(
+				"ChangeTexturePinMode", ActionText, ToolTipText, FSlateIcon(),
+				FUIAction(FExecuteAction::CreateLambda([TableNode, TexturePin]()
+					{
+						TableNode->ChangeImagePinMode(TexturePin);
+					}))
+			);
 
+			if (!TableNode->IsImagePinDefault(TexturePin))
+			{
 				Section.AddMenuEntry
 				(
-					"ChangeTexturePinMode", ActionText, ToolTipText, FSlateIcon(),
+					"SetTexturePinModeDefault",
+					LOCTEXT("SetTexturePinModeDefault_Label","Set pin as default."),
+					LOCTEXT("SetTexturePinModeDefault_Tooltip","Set the selected texture pin to use the default node mode."),
+					FSlateIcon(),
 					FUIAction(FExecuteAction::CreateLambda([TableNode, TexturePin]()
 						{
-							TableNode->ChangeImagePinMode(TexturePin);
+							TableNode->ChangeImagePinMode(TexturePin, true);
 						}))
 				);
-
-				if (!TableNode->IsImagePinDefault(TexturePin))
-				{
-					Section.AddMenuEntry
-					(
-						"SetTexturePinModeDefault",
-						LOCTEXT("SetTexturePinModeDefault_Label","Set pin as default."),
-						LOCTEXT("SetTexturePinModeDefault_Tooltip","Set the selected texture pin to use the default node mode."),
-						FSlateIcon(),
-						FUIAction(FExecuteAction::CreateLambda([TableNode, TexturePin]()
-							{
-								TableNode->ChangeImagePinMode(TexturePin, true);
-							}))
-					);
-				}
 			}
 		}
 	}
