[Mutable] Mesh memory managment.

- Raster image mesh reuse between stages now uses the heap as re-caching was problematic with the new memory managment.

#rnx

diff --git a/Engine/Plugins/Experimental/Mutable/Source/CustomizableObjectEditor/Private/MuCOE/MutableMeshPreviewUtils.cpp b/Engine/Plugins/Experimental/Mutable/Source/CustomizableObjectEditor/Private/MuCOE/MutableMeshPreviewUtils.cpp
--- a/Engine/Plugins/Experimental/Mutable/Source/CustomizableObjectEditor/Private/MuCOE/MutableMeshPreviewUtils.cpp
+++ b/Engine/Plugins/Experimental/Mutable/Source/CustomizableObjectEditor/Private/MuCOE/MutableMeshPreviewUtils.cpp
@@ -171,16 +171,14 @@ namespace MutableMeshPreviewUtils
 			// Skeletal mesh is ready to be used as an argument for the process of reformatting the Mutable Mesh we want to preview
 
 			// 3) Call mu::MeshFormat to get the mesh with the wanted buffer structure to have the data of the original mesh
-			mu::MeshPtrConst ResultMutableMesh = mu::MeshFormat(
-				InMutableMesh.get(),
-				FormattedMutableMesh.get(),
-				false,
-				true,
-				true,
-				false,
-				false );
-			
-			return ResultMutableMesh;
+			mu::Ptr<mu::Mesh> FormatResult = new mu::Mesh();
+
+			bool bOutSuccess = false;
+			mu::MeshFormat(FormatResult.get(), InMutableMesh.get(), FormattedMutableMesh.get(),
+				false, true, true, false, false,
+				bOutSuccess);
+		
+			return FormatResult;
 		}
 
 		/** Method designed to scan the vertex buffers of the provided mutable mesh and tell the caller if the buffers are
diff --git a/Engine/Plugins/Experimental/Mutable/Source/MutableRuntime/Private/MuR/CodeRunner.cpp b/Engine/Plugins/Experimental/Mutable/Source/MutableRuntime/Private/MuR/CodeRunner.cpp
--- a/Engine/Plugins/Experimental/Mutable/Source/MutableRuntime/Private/MuR/CodeRunner.cpp
+++ b/Engine/Plugins/Experimental/Mutable/Source/MutableRuntime/Private/MuR/CodeRunner.cpp
@@ -365,7 +365,7 @@ namespace mu
             case DT_STRING:		StoreString( cat, LoadString( rat ) ); break;
             case DT_COLOUR:		StoreColor( cat, LoadColor( rat ) ); break;
             case DT_PROJECTOR:  StoreProjector( cat, LoadProjector(rat) ); break;
-            case DT_MESH:       StoreMesh( cat, LoadMesh(rat) ); break;
+			case DT_MESH:       StoreMesh( cat, LoadMesh(rat) ); break;
             case DT_IMAGE:      StoreImage( cat, LoadImage(rat) ); break;
             case DT_LAYOUT:     StoreLayout( cat, LoadLayout(rat) ); break;
             case DT_INSTANCE:   StoreInstance( cat, LoadInstance(rat) ); break;
@@ -906,28 +906,31 @@ namespace mu
             FProgram& program = pModel->GetPrivate()->m_program;
 
             // Assume the ROM has been loaded previously
-            check(program.m_constantMeshes[cat].Value )
+            check(program.m_constantMeshes[cat].Value)
+
+            Ptr<const Mesh> SourceConst;
+            program.GetConstant(cat, SourceConst);
 
-            MeshPtrConst pSourceConst;
-            program.GetConstant( cat, pSourceConst );
-            MeshPtr Source = pSourceConst->Clone();
+			check(SourceConst);
+			Ptr<Mesh> Source = CreateMesh(SourceConst->GetDataSize());
+			Source->CopyFrom(*SourceConst);
 
             // Set the separate skeleton if necessary
-            if (args.skeleton>=0)
+            if (args.skeleton >= 0)
             {
-                check( program.m_constantSkeletons.Num()>size_t(args.skeleton)  );
+                check(program.m_constantSkeletons.Num() > size_t(args.skeleton));
                 Ptr<const Skeleton> pSkeleton = program.m_constantSkeletons[args.skeleton];
                 Source->SetSkeleton(pSkeleton);
             }
 
 			if (args.physicsBody >= 0)
 			{
-                check( program.m_constantPhysicsBodies.Num()>size_t(args.physicsBody)  );
+                check(program.m_constantPhysicsBodies.Num() > size_t(args.physicsBody));
                 Ptr<const PhysicsBody> pPhysicsBody = program.m_constantPhysicsBodies[args.physicsBody];
                 Source->SetPhysicsBody(pPhysicsBody);
 			}
 
-            StoreMesh( item, Source );
+            StoreMesh(item, Source);
 			//UE_LOG(LogMutableCore, Log, TEXT("Set mesh constant %d."), item.At);
             break;
         }
@@ -941,7 +944,7 @@ namespace mu
 
 			int32 MipsToSkip = item.ExecutionOptions;
             Ptr<const Image> Source;
-            program.GetConstant( cat, Source, MipsToSkip);
+            program.GetConstant(cat, Source, MipsToSkip);
 
 			// Assume the ROM has been loaded previously in a task generated at IssueOp
 			check(Source);
@@ -983,37 +986,45 @@ namespace mu
 		MUTABLE_CPUPROFILER_SCOPE(RunCode_Mesh);
 
 		OP_TYPE type = pModel->GetPrivate()->m_program.GetOpType(item.At);
+
         switch (type)
         {
 
         case OP_TYPE::ME_APPLYLAYOUT:
         {
 			OP::MeshApplyLayoutArgs args = pModel->GetPrivate()->m_program.GetOpArgs<OP::MeshApplyLayoutArgs>(item.At);
-            switch (item.Stage)
-            {
-            case 0:
-                    AddOp( FScheduledOp( item.At, item, 1),
-                           FScheduledOp( args.mesh, item),
-                           FScheduledOp( args.layout, item) );
-                break;
-
+			switch (item.Stage)
+			{
+			case 0:
+			{
+				AddOp(FScheduledOp(item.At, item, 1),
+					FScheduledOp(args.mesh, item),
+					FScheduledOp(args.layout, item));
+				break;
+			}
             case 1:
             {
            		MUTABLE_CPUPROFILER_SCOPE(ME_APPLYLAYOUT)
             		
-                Ptr<const Mesh> pBase = LoadMesh( FCacheAddress(args.mesh,item) );
+                Ptr<const Mesh> pBase = LoadMesh(FCacheAddress(args.mesh, item));
 
-                MeshPtr pApplied;
                 if (pBase)
                 {
-                    pApplied = pBase->Clone();
+					Ptr<Mesh> Result = CloneOrTakeOver(pBase);
 
-                    Ptr<const Layout> pLayout = LoadLayout( FCacheAddress(args.layout,item) );
+                    Ptr<const Layout> pLayout = LoadLayout(FCacheAddress(args.layout, item));
                     int texCoordsSet = args.channel;
-                    MeshApplyLayout( pApplied.get(), pLayout.get(), texCoordsSet );
+
+                    MeshApplyLayout(Result.get(), pLayout.get(), texCoordsSet);
+					
+					StoreMesh(item, Result);
                 }
+				else
+				{
+					StoreMesh(item, nullptr);
+				}
+
 
-                StoreMesh( item, pApplied );
                 break;
             }
 
@@ -1029,26 +1040,29 @@ namespace mu
 			const uint8* data = pModel->GetPrivate()->m_program.GetOpArgsPointer(item.At);
 
 			OP::ADDRESS BaseAt = 0;
-			FMemory::Memcpy(&BaseAt, data, sizeof(OP::ADDRESS)); data += sizeof(OP::ADDRESS);
+			FMemory::Memcpy(&BaseAt, data, sizeof(OP::ADDRESS)); 
+			data += sizeof(OP::ADDRESS);
 
 			OP::ADDRESS TargetAt = 0;
-			FMemory::Memcpy(&TargetAt, data, sizeof(OP::ADDRESS)); data += sizeof(OP::ADDRESS);
+			FMemory::Memcpy(&TargetAt, data, sizeof(OP::ADDRESS)); 
+			data += sizeof(OP::ADDRESS);
 
             switch (item.Stage)
             {
             case 0:
-                if (BaseAt && TargetAt)
-                {
-                    AddOp( FScheduledOp( item.At, item, 1),
-                            FScheduledOp( BaseAt, item),
-                            FScheduledOp( TargetAt, item) );
-                }
-                else
-                {
-                    StoreMesh( item, nullptr );
-                }
-                break;
-
+			{
+				if (BaseAt && TargetAt)
+				{
+					AddOp(FScheduledOp(item.At, item, 1),
+						FScheduledOp(BaseAt, item),
+						FScheduledOp(TargetAt, item));
+				}
+				else
+				{
+					StoreMesh(item, nullptr);
+				}
+				break;
+			}
             case 1:
             {
        	        MUTABLE_CPUPROFILER_SCOPE(ME_DIFFERENCE)
@@ -1060,27 +1074,39 @@ namespace mu
 				TArray<int32, TInlineAllocator<8>> SemanticIndices;
 
 				uint8 bIgnoreTextureCoords = 0;
-				FMemory::Memcpy(&bIgnoreTextureCoords, data, sizeof(uint8)); data += sizeof(uint8);
+				FMemory::Memcpy(&bIgnoreTextureCoords, data, sizeof(uint8)); 
+				data += sizeof(uint8);
 
 				uint8 NumChannels = 0;
-				FMemory::Memcpy(&NumChannels, data, sizeof(uint8)); data += sizeof(uint8);
+				FMemory::Memcpy(&NumChannels, data, sizeof(uint8)); 
+				data += sizeof(uint8);
 
-                for ( uint8 i=0; i< NumChannels; ++i )
+                for (uint8 i = 0; i < NumChannels; ++i)
                 {
 					uint8 Semantic = 0;
-					FMemory::Memcpy(&Semantic, data, sizeof(uint8)); data += sizeof(uint8);
+					FMemory::Memcpy(&Semantic, data, sizeof(uint8)); 
+					data += sizeof(uint8);
+					
 					uint8 SemanticIndex = 0;
-					FMemory::Memcpy(&SemanticIndex, data, sizeof(uint8)); data += sizeof(uint8);
+					FMemory::Memcpy(&SemanticIndex, data, sizeof(uint8)); 
+					data += sizeof(uint8);
 
 					Semantics.Add(MESH_BUFFER_SEMANTIC(Semantic));
 					SemanticIndices.Add(SemanticIndex);
                 }
 
-                MeshPtr pResult = MeshDifference( pBase.get(), pTarget.get(),
-                                          NumChannels, Semantics.GetData(), SemanticIndices.GetData(),
-                                          bIgnoreTextureCoords!=0 );
+				Ptr<Mesh> Result = CreateMesh();
+				bool bOutSuccess = false;
+                MeshDifference(Result.get(), pBase.get(), pTarget.get(),
+                               NumChannels, Semantics.GetData(), SemanticIndices.GetData(),
+                               bIgnoreTextureCoords != 0, bOutSuccess);
+			
+				check(bOutSuccess);
+
+				Release(pBase);
+				Release(pTarget);
 
-                StoreMesh( item, pResult );
+                StoreMesh(item, Result);
                 break;
             }
 
@@ -1096,25 +1122,28 @@ namespace mu
 			const uint8* data = pModel->GetPrivate()->m_program.GetOpArgsPointer(item.At);
 
 			OP::ADDRESS FactorAt = 0;
-			FMemory::Memcpy(&FactorAt, data, sizeof(OP::ADDRESS)); data += sizeof(OP::ADDRESS);
+			FMemory::Memcpy(&FactorAt, data, sizeof(OP::ADDRESS)); 
+			data += sizeof(OP::ADDRESS);
 			
 			OP::ADDRESS BaseAt = 0;
-			FMemory::Memcpy(&BaseAt, data, sizeof(OP::ADDRESS)); data += sizeof(OP::ADDRESS);
+			FMemory::Memcpy(&BaseAt, data, sizeof(OP::ADDRESS)); 
+			data += sizeof(OP::ADDRESS);
 
 			OP::ADDRESS TargetAt = 0;
-			FMemory::Memcpy(&TargetAt, data, sizeof(OP::ADDRESS)); data += sizeof(OP::ADDRESS);
+			FMemory::Memcpy(&TargetAt, data, sizeof(OP::ADDRESS)); 
+			data += sizeof(OP::ADDRESS);
 
 			switch (item.Stage)
             {
             case 0:
                 if (BaseAt)
                 {
-                    AddOp( FScheduledOp(item.At, item, 1),
-                           FScheduledOp(FactorAt, item) );
+                    AddOp(FScheduledOp(item.At, item, 1),
+                           FScheduledOp(FactorAt, item));
                 }
                 else
                 {
-                    StoreMesh( item, nullptr );
+                    StoreMesh(item, nullptr);
                 }
                 break;
 
@@ -1122,7 +1151,7 @@ namespace mu
             {
                 MUTABLE_CPUPROFILER_SCOPE(ME_MORPH_1)
 
-                float Factor = LoadScalar( FCacheAddress(FactorAt,item) );
+                float Factor = LoadScalar(FCacheAddress(FactorAt, item));
 
                 // Factor goes from -1 to 1 across all targets. [0 - 1] represents positive morphs, while [-1, 0) represent negative morphs.
 				Factor = FMath::Clamp(Factor, -1.0f, 1.0f); // Is the factor not in range [-1, 1], it will index a non existing morph.
@@ -1132,18 +1161,18 @@ namespace mu
 				uint32 dataAddress = uint32(m_heapData.Add(HeapData));
 
                 // No morph
-				if (Factor < UE_SMALL_NUMBER && Factor > -UE_SMALL_NUMBER )
+				if (FMath::IsNearlyZero(Factor))
                 {                        
-                    AddOp( FScheduledOp( item.At, item, 2, dataAddress),
-                            FScheduledOp(BaseAt, item) );
+                    AddOp(FScheduledOp(item.At, item, 2, dataAddress),
+						FScheduledOp(BaseAt, item));
                 }
                 // The Morph, partial or full
                 else
                 {
                     // We will need the base again
-                    AddOp( FScheduledOp( item.At, item, 2, dataAddress),
-                            FScheduledOp(BaseAt, item),
-                            FScheduledOp(TargetAt, item) );
+                    AddOp(FScheduledOp(item.At, item, 2, dataAddress),
+						FScheduledOp(BaseAt, item),
+						FScheduledOp(TargetAt, item));
                 }
 
                 break;
@@ -1153,29 +1182,53 @@ namespace mu
             {
        		    MUTABLE_CPUPROFILER_SCOPE(ME_MORPH_2)
 
-                Ptr<const Mesh> pBase = LoadMesh( FCacheAddress(BaseAt,item) );
+                Ptr<const Mesh> pBase = LoadMesh(FCacheAddress(BaseAt, item));
 
                 // Factor from 0 to 1 between the two targets
-                const FScheduledOpData& HeapData = m_heapData[ (size_t)item.CustomState ];
+                const FScheduledOpData& HeapData = m_heapData[(size_t)item.CustomState];
                 float Factor = HeapData.Interpolate.Bifactor;
 
-                MeshPtrConst pResult;
-
                 if (pBase)
                 {
 					// No morph
-					if (Factor < UE_SMALL_NUMBER && Factor > -UE_SMALL_NUMBER )
+					if (FMath::IsNearlyZero(Factor))
                     {
-						pResult = pBase;
+						StoreMesh(item, pBase);
                     }
 					// The Morph, partial or full
                     else 
                     {
-                        Ptr<const Mesh> pMorph = LoadMesh( FCacheAddress(TargetAt,item) );
-						pResult = MeshMorph(pBase.get(), pMorph.get(), Factor);
+                        Ptr<const Mesh> pMorph = LoadMesh(FCacheAddress(TargetAt,item));
+						
+						if (pMorph)
+						{
+							Ptr<Mesh> Result = CreateMesh(pBase->GetDataSize());
+							bool bOutSuccess = false;
+							MeshMorph(Result.get(), pBase.get(), pMorph.get(), Factor, bOutSuccess);
+
+							Release(pMorph);
+
+							if (!bOutSuccess)
+							{
+								Release(Result);
+								StoreMesh(item, pBase);
+							}
+							else
+							{
+								Release(pBase);
+								StoreMesh(item, Result);
+							}
+						}
+						else
+						{
+							StoreMesh(item, pBase);
+						}
                     }
                 }
-                StoreMesh( item, pResult );
+				else
+				{
+					StoreMesh(item, nullptr);
+				}
 
                 break;
             }
@@ -1193,49 +1246,66 @@ namespace mu
             switch (item.Stage)
             {
             case 0:
-                    AddOp( FScheduledOp( item.At, item, 1),
-                           FScheduledOp( args.base, item),
-                           FScheduledOp( args.added, item) );
-                break;
-
+			{
+				AddOp(FScheduledOp(item.At, item, 1),
+					FScheduledOp(args.base, item),
+					FScheduledOp(args.added, item));
+				break;
+			}
             case 1:
             {
 				MUTABLE_CPUPROFILER_SCOPE(ME_MERGE_1)
 
-                MeshPtrConst pA = LoadMesh( FCacheAddress(args.base,item) );
-                MeshPtrConst pB = LoadMesh( FCacheAddress(args.added,item) );
-
-                MeshPtr pResult;
+                Ptr<const Mesh> pA = LoadMesh(FCacheAddress(args.base, item));
+                Ptr<const Mesh> pB = LoadMesh(FCacheAddress(args.added, item));
 
                 if (pA && pB && pA->GetVertexCount() && pB->GetVertexCount())
                 {
-                    pResult = MeshMerge( pA.get(), pB.get(), !args.newSurfaceID);
+					FMeshMergeScratchMeshes Scratch;
+					Scratch.FirstReformat = CreateMesh();
+					Scratch.SecondReformat = CreateMesh();
+
+					Ptr<Mesh> Result = CreateMesh(pA->GetDataSize() + pB->GetDataSize());
+
+					MeshMerge(Result.get(), pA.get(), pB.get(), !args.newSurfaceID, Scratch);
+
+					Release(Scratch.FirstReformat);
+					Release(Scratch.SecondReformat);
 
                     if (args.newSurfaceID)
                     {
 						check(pB->GetSurfaceCount() == 1);
-						pResult->m_surfaces.Last().m_id = args.newSurfaceID;
+						Result->m_surfaces.Last().m_id = args.newSurfaceID;
                     }
+
+					Release(pA);
+					Release(pB);
+					StoreMesh(item, Result);
                 }
                 else if (pA && pA->GetVertexCount())
                 {
-					pResult = pA->Clone();
+					Release(pB);
+					StoreMesh(item, pA);
                 }
                 else if (pB && pB->GetVertexCount())
                 {
-                    pResult = pB->Clone();
-                    check(pResult->GetSurfaceCount()==1);
-                    if (pResult->GetSurfaceCount()>0 && args.newSurfaceID)
+					Ptr<Mesh> Result = CloneOrTakeOver(pB);
+
+                    check(Result->GetSurfaceCount() == 1);
+
+                    if (Result->GetSurfaceCount() > 0 && args.newSurfaceID)
                     {
-                        pResult->m_surfaces.Last().m_id=args.newSurfaceID;
+                        Result->m_surfaces.Last().m_id = args.newSurfaceID;
                     }
+
+					Release(pA);
+					StoreMesh(item, Result);
                 }
                 else
                 {
-                    pResult = new Mesh();
+					StoreMesh(item, CreateMesh());
                 }
 
-                StoreMesh( item, pResult );
                 break;
             }
 
@@ -1252,18 +1322,19 @@ namespace mu
             switch (item.Stage)
             {
             case 0:
-                if ( args.base )
-                {
-                    AddOp( FScheduledOp( item.At, item, 1),
-                           FScheduledOp( args.base, item),
-                           FScheduledOp( args.factor, item ) );
-                }
-                else
-                {
-                    StoreMesh( item, nullptr );
-                }
-                break;
-
+			{
+				if (args.base)
+				{
+					AddOp(FScheduledOp(item.At, item, 1),
+						FScheduledOp(args.base, item),
+						FScheduledOp(args.factor, item));
+				}
+				else
+				{
+					StoreMesh(item, nullptr);
+				}
+				break;
+			}
             case 1:
             {
             	MUTABLE_CPUPROFILER_SCOPE(ME_INTERPOLATE_1)
@@ -1277,7 +1348,7 @@ namespace mu
                 }
 
                 // Factor from 0 to 1 across all targets
-                float factor = LoadScalar( FCacheAddress(args.factor,item) );
+                float factor = LoadScalar(FCacheAddress(args.factor, item));
 
                 float delta = 1.0f/(count-1);
                 int min = (int)floorf( factor/delta );
@@ -1293,54 +1364,48 @@ namespace mu
 				uint32 dataAddress = uint32(m_heapData.Num());
 
                 // Just the first of the targets
-                if ( bifactor < UE_SMALL_NUMBER )
+                if (bifactor < UE_SMALL_NUMBER)
                 {
-                    if( min==0 )
+                    if (min == 0)
                     {
                         // Just the base
-                            Ptr<const Mesh> pBase = LoadMesh( FCacheAddress(args.base,item) );
-                            MeshPtr pResult;
-                            if (pBase)
-                            {
-                                pResult = pBase->Clone();
-                            }
-                            StoreMesh( item, pResult );
-                        }
+						Ptr<const Mesh> pBase = LoadMesh(FCacheAddress(args.base, item));
+						StoreMesh(item, pBase);
+					}
                     else
                     {
                         // Base with one full morph
                         m_heapData.Add(data);
-                            AddOp( FScheduledOp( item.At, item, 2, dataAddress),
-                                   FScheduledOp( args.base, item),
-                                   FScheduledOp( args.targets[min-1], item) );
-                        }
-                    }
+						AddOp(FScheduledOp(item.At, item, 2, dataAddress),
+							FScheduledOp(args.base, item),
+							FScheduledOp(args.targets[min-1], item));
+					}
+				}
                 // Just the second of the targets
-                else if ( bifactor > 1.0f-UE_SMALL_NUMBER )
+                else if (bifactor > 1.0f-UE_SMALL_NUMBER)
                 {
                     m_heapData.Add(data);
-                        AddOp( FScheduledOp( item.At, item, 2, dataAddress),
-                               FScheduledOp( args.base, item),
-                               FScheduledOp( args.targets[max-1], item) );
-                    }
+					AddOp(FScheduledOp(item.At, item, 2, dataAddress),
+						FScheduledOp(args.base, item),
+						FScheduledOp(args.targets[max-1], item));
+				}
                 // Mix the first target on the base
-                else if ( min==0 )
+                else if (min == 0)
                 {
                     m_heapData.Add(data);
-                        AddOp( FScheduledOp( item.At, item, 2, dataAddress),
-                               FScheduledOp( args.base, item),
-                               FScheduledOp( args.targets[0], item)
-                               );
-                    }
+					AddOp(FScheduledOp(item.At, item, 2, dataAddress),
+						FScheduledOp(args.base, item),
+						FScheduledOp(args.targets[0], item));
+				}
                 // Mix two targets on the base
                 else
                 {
                     m_heapData.Add(data);
-                        AddOp( FScheduledOp( item.At, item, 2, dataAddress),
-                               FScheduledOp( args.base, item),
-                               FScheduledOp( args.targets[min-1], item),
-                               FScheduledOp( args.targets[max-1], item) );
-                    }
+					AddOp(FScheduledOp(item.At, item, 2, dataAddress),
+						FScheduledOp(args.base, item),
+						FScheduledOp(args.targets[min-1], item),
+						FScheduledOp(args.targets[max-1], item));
+				}
 
                 break;
             }
@@ -1364,74 +1429,140 @@ namespace mu
                 int min = data.Interpolate.Min;
                 int max = data.Interpolate.Max;
 
-                Ptr<const Mesh> pBase = LoadMesh( FCacheAddress(args.base, item) );
-
-                MeshPtr pResult;
+                Ptr<const Mesh> pBase = LoadMesh(FCacheAddress(args.base, item));
 
                 if (pBase)
                 {
                     // Just the first of the targets
-                    if ( bifactor < UE_SMALL_NUMBER )
+                    if (bifactor < UE_SMALL_NUMBER)
                     {
-                        if( min==0 )
+                        if (min == 0)
                         {
                             // Just the base. It should have been dealt with in the previous stage.
-                            check( false );
+                            check(false);
                         }
                         else
                         {
                             // Base with one full morph
-                            Ptr<const Mesh> pMorph = LoadMesh( FCacheAddress(args.targets[min-1],item) );
-                            pResult = MeshMorph( pBase.get(), pMorph.get() );
+                            Ptr<const Mesh> pMorph = LoadMesh(FCacheAddress(args.targets[min-1], item));
+							
+							Ptr<Mesh> Result = CreateMesh(pBase->GetDataSize());
+
+							bool bOutSuccess = false;
+                            MeshMorph(Result.get(), pBase.get(), pMorph.get(), bOutSuccess);
+						
+							Release(pMorph);
+
+							if (!bOutSuccess)
+							{
+								Release(Result);
+								StoreMesh(item, pBase);
+							}
+							else
+							{
+								Release(pBase);
+								StoreMesh(item, Result);
+							}
                         }
                     }
                     // Just the second of the targets
-                    else if ( bifactor > 1.0f-UE_SMALL_NUMBER )
+                    else if (bifactor > 1.0f-UE_SMALL_NUMBER)
                     {
-                        check( max>0 );
-                        Ptr<const Mesh> pMorph = LoadMesh( FCacheAddress(args.targets[max-1],item) );
+                        check(max > 0);
+                        Ptr<const Mesh> pMorph = LoadMesh(FCacheAddress(args.targets[max-1], item));
 
                         if (pMorph)
                         {
-                            pResult = MeshMorph( pBase.get(), pMorph.get() );
+							Ptr<Mesh> Result = CreateMesh(pBase->GetDataSize());
+							
+							bool bOutSuccess = false;
+                            MeshMorph(Result.get(), pBase.get(), pMorph.get(), bOutSuccess);
+
+							Release(pMorph);
+							if (!bOutSuccess)
+							{
+								Release(Result);
+								StoreMesh(item, pBase);
+							}
+							else
+							{
+								Release(pBase);
+								StoreMesh(item, Result);
+							}
+
                         }
                         else
                         {
-                            pResult = pBase->Clone();
+							StoreMesh(item, pBase);
                         }
                     }
                     // Mix the first target on the base
-                    else if ( min==0 )
+                    else if (min == 0)
                     {
-                        Ptr<const Mesh> pMorph = LoadMesh( FCacheAddress( args.targets[0], item ) );
+                        Ptr<const Mesh> pMorph = LoadMesh(FCacheAddress(args.targets[0], item));
                         if (pMorph)
                         {
-                            pResult = MeshMorph( pBase.get(), pMorph.get(), bifactor );
+							Ptr<Mesh> Result = CreateMesh(pBase->GetDataSize());
+
+							bool bOutSuccess = false;
+                            MeshMorph(Result.get(), pBase.get(), pMorph.get(), bifactor, bOutSuccess);
+
+							Release(pMorph);
+
+							if (!bOutSuccess)
+							{
+								Release(Result);
+								StoreMesh(item, pBase);
+							}
+							else
+							{
+								Release(pBase);
+								StoreMesh(item, Result);
+							}
                         }
                         else
                         {
-                            pResult = pBase->Clone();
+							StoreMesh(item, pBase);
                         }
-
                     }
                     // Mix two targets on the base
                     else
                     {
-                        Ptr<const Mesh> pMin = LoadMesh( FCacheAddress(args.targets[min-1],item) );
-                        Ptr<const Mesh> pMax = LoadMesh( FCacheAddress(args.targets[max-1],item) );
+                        Ptr<const Mesh> pMin = LoadMesh(FCacheAddress(args.targets[min-1], item));
+                        Ptr<const Mesh> pMax = LoadMesh(FCacheAddress(args.targets[max-1], item));
 
                         if (pMin && pMax)
                         {
-                            pResult = MeshMorph2( pBase.get(), pMin.get(), pMax.get(), bifactor );
+							Ptr<Mesh> Result = CreateMesh(pBase->GetDataSize());
+
+							bool bOutSuccess = false;
+                            MeshMorph2(Result.get(), pBase.get(), pMin.get(), pMax.get(), bifactor, bOutSuccess);
+
+							Release(pMin);
+							Release(pMax);
+
+							if (!bOutSuccess)
+							{
+								Release(Result);
+								StoreMesh(item, pBase);
+							}
+							else
+							{
+								Release(pBase);
+								StoreMesh(item, Result);
+							}
                         }
                         else
                         {
-                            pResult = pBase->Clone();
+							StoreMesh(item, pBase);
                         }
                     }
                 }
+				else
+				{
+					StoreMesh(item, pBase);
+				}
 
-                StoreMesh( item, pResult );
                 break;
             }
 
@@ -1448,25 +1579,45 @@ namespace mu
             switch (item.Stage)
             {
             case 0:
-                    AddOp( FScheduledOp( item.At, item, 1),
-                           FScheduledOp( args.source, item),
-                           FScheduledOp( args.clip, item) );
-                break;
-
+			{
+				AddOp(FScheduledOp(item.At, item, 1),
+					FScheduledOp(args.source, item),
+					FScheduledOp(args.clip, item));
+				break;
+			}
             case 1:
             {
             	MUTABLE_CPUPROFILER_SCOPE(ME_MASKCLIPMESH_1)
             		
-                Ptr<const Mesh> Source = LoadMesh( FCacheAddress(args.source,item) );
-                Ptr<const Mesh> pClip = LoadMesh( FCacheAddress(args.clip,item) );
+                Ptr<const Mesh> Source = LoadMesh(FCacheAddress(args.source, item));
+                Ptr<const Mesh> pClip = LoadMesh(FCacheAddress(args.clip, item));
 
                 // Only if both are valid.
-                MeshPtr pResult;
                 if (Source.get() && pClip.get())
                 {
-                    pResult = MeshMaskClipMesh(Source.get(), pClip.get());
+					Ptr<Mesh> Result = CreateMesh();
+
+					bool bOutSuccess = false;
+                    MeshMaskClipMesh(Result.get(), Source.get(), pClip.get(), bOutSuccess);
+					
+					Release(Source);
+					Release(pClip);
+					if (!bOutSuccess)
+					{
+						Release(Result);
+						StoreMesh(item, nullptr);
+					}
+					else
+					{
+						StoreMesh(item, Result);
+					}
                 }
-                StoreMesh( item, pResult );
+				else
+				{
+					Release(Source);
+					Release(pClip);
+					StoreMesh(item, nullptr);
+				}
 
                 break;
             }
@@ -1484,25 +1635,46 @@ namespace mu
             switch (item.Stage)
             {
             case 0:
-                    AddOp( FScheduledOp( item.At, item, 1),
-                           FScheduledOp( args.source, item),
-                           FScheduledOp( args.fragment, item) );
-                break;
-
+			{
+				AddOp(FScheduledOp(item.At, item, 1),
+					FScheduledOp(args.source, item),
+					FScheduledOp(args.fragment, item));
+				break;
+			}
             case 1:
             {
            		MUTABLE_CPUPROFILER_SCOPE(ME_MASKDIFF_1)
             		
-                Ptr<const Mesh> Source = LoadMesh( FCacheAddress(args.source,item) );
-                Ptr<const Mesh> pClip = LoadMesh( FCacheAddress(args.fragment,item) );
+                Ptr<const Mesh> Source = LoadMesh(FCacheAddress(args.source, item));
+                Ptr<const Mesh> pClip = LoadMesh(FCacheAddress(args.fragment, item));
 
                 // Only if both are valid.
-                MeshPtr pResult;
                 if (Source.get() && pClip.get())
                 {
-                    pResult = MeshMaskDiff(Source.get(), pClip.get());
+					Ptr<Mesh> Result = CreateMesh();
+
+					bool bOutSuccess = false;
+                    MeshMaskDiff(Result.get(), Source.get(), pClip.get(), bOutSuccess);
+
+					Release(Source);
+					Release(pClip);
+
+					if (!bOutSuccess)
+					{
+						Release(Result);
+						StoreMesh(item, nullptr);
+					}
+					else
+					{
+						StoreMesh(item, Result);
+					}
                 }
-                StoreMesh( item, pResult );
+				else
+				{
+					Release(Source);
+					Release(pClip);
+					StoreMesh(item, nullptr);
+				}
 
                 break;
             }
@@ -1520,42 +1692,74 @@ namespace mu
             switch (item.Stage)
             {
             case 0:
-                if (args.source && args.format)
-                {
-                         AddOp( FScheduledOp( item.At, item, 1),
-                                FScheduledOp( args.source, item),
-                                FScheduledOp( args.format, item));
-                    }
-                    else
-                    {
-                    StoreMesh( item, nullptr );
-                }
-                break;
-
+			{
+				if (args.source && args.format)
+				{
+					AddOp(FScheduledOp(item.At, item, 1),
+						FScheduledOp(args.source, item),
+						FScheduledOp(args.format, item));
+				}
+				else
+				{
+					StoreMesh(item, nullptr);
+				}
+				break;
+			}
             case 1:
             {
             	MUTABLE_CPUPROFILER_SCOPE(ME_FORMAT_1)
             		
-                Ptr<const Mesh> Source = LoadMesh( FCacheAddress(args.source,item) );
-                Ptr<const Mesh> pFormat = LoadMesh( FCacheAddress(args.format,item) );
-
-                uint8 flags = args.buffers;
-                MeshPtr pResult;
-                pResult = MeshFormat( Source.get(),
-                                      pFormat.get(),
-                                      true,
-                                      (flags & OP::MeshFormatArgs::BT_VERTEX) != 0,
-                                      (flags & OP::MeshFormatArgs::BT_INDEX) != 0,
-                                      (flags & OP::MeshFormatArgs::BT_FACE) != 0,
-                                      (flags & OP::MeshFormatArgs::BT_IGNORE_MISSING) != 0
-                                      );
-
-                if (flags & OP::MeshFormatArgs::BT_RESETBUFFERINDICES)
-                {
-                    pResult->ResetBufferIndices();
-                }
+                Ptr<const Mesh> Source = LoadMesh(FCacheAddress(args.source,item));
+                Ptr<const Mesh> pFormat = LoadMesh(FCacheAddress(args.format,item));
 
-                StoreMesh( item, pResult );
+				if (Source)
+				{
+					uint8 flags = args.buffers;
+					if (!pFormat && !(flags & OP::MeshFormatArgs::BT_RESETBUFFERINDICES))
+					{
+						StoreMesh(item, Source);
+					}
+					else if (!pFormat)
+					{
+						Ptr<Mesh> Result = CloneOrTakeOver(Source);
+
+						if (flags & OP::MeshFormatArgs::BT_RESETBUFFERINDICES)
+						{
+							Result->ResetBufferIndices();
+						}
+
+						StoreMesh(item, Result);
+					}
+					else
+					{
+						Ptr<Mesh> Result = CreateMesh();
+
+						bool bOutSuccess = false;
+						MeshFormat(Result.get(), Source.get(), pFormat.get(),
+							true,
+							(flags & OP::MeshFormatArgs::BT_VERTEX) != 0,
+							(flags & OP::MeshFormatArgs::BT_INDEX) != 0,
+							(flags & OP::MeshFormatArgs::BT_FACE) != 0,
+							(flags & OP::MeshFormatArgs::BT_IGNORE_MISSING) != 0,
+							bOutSuccess);
+
+						check(bOutSuccess);
+
+						if (flags & OP::MeshFormatArgs::BT_RESETBUFFERINDICES)
+						{
+							Result->ResetBufferIndices();
+						}
+
+						Release(Source);
+						Release(pFormat);
+						StoreMesh(item, Result);
+					}
+				}
+				else
+				{
+					Release(pFormat);
+					StoreMesh(item, nullptr);
+				}
                 break;
             }
 
@@ -1585,38 +1789,49 @@ namespace mu
             switch (item.Stage)
             {
             case 0:
-                if (source)
-                {
-                        AddOp( FScheduledOp( item.At, item, 1),
-                               FScheduledOp( source, item) );
-                }
-                else
-                {
-                    StoreMesh( item, nullptr );
-                }
-                break;
-
+			{
+				if (source)
+				{
+					AddOp(FScheduledOp(item.At, item, 1),
+						FScheduledOp(source, item));
+				}
+				else
+				{
+					StoreMesh(item, nullptr);
+				}
+				break;
+			}
             case 1:
             {
 				MUTABLE_CPUPROFILER_SCOPE(ME_EXTRACTLAYOUTBLOCK_1)
 
-                Ptr<const Mesh> Source = LoadMesh( FCacheAddress(source,item) );
-				MeshPtr pResult;
+                Ptr<const Mesh> Source = LoadMesh(FCacheAddress(source, item));
+
+                // Access with memcpy necessary for unaligned arm issues.
+                uint32 blocks[1024];
+				FMemory::Memcpy(blocks, data, sizeof(uint32)*FMath::Min(1024,int(blockCount)));
 
 				if (Source)
 				{
-					// Access with memcpy necessary for unaligned arm issues.
-					uint32 blocks[1024];
-					FMemory::Memcpy(blocks, data, sizeof(uint32)* FMath::Min(1024, int(blockCount)));
+					Ptr<Mesh> Result = CreateMesh();
+					bool bOutSuccess;
+					MeshExtractLayoutBlock(Result.get(), Source.get(), layout, blockCount, blocks, bOutSuccess);
 
-					pResult = MeshExtractLayoutBlock(Source.get(),
-						layout,
-						blockCount,
-						blocks);
+					if (!bOutSuccess)
+					{
+						Release(Result);
+						StoreMesh(item, Source);
+					}
+					else
+					{
+						Release(Source);
+						StoreMesh(item, Result);
+					}
+				}
+				else
+				{
+					StoreMesh(item, nullptr);
 				}
-
-				StoreMesh(item, pResult);
-
                 break;
             }
 
@@ -1633,28 +1848,47 @@ namespace mu
             switch (item.Stage)
             {
             case 0:
-                if (args.source)
-                {
-                        AddOp( FScheduledOp( item.At, item, 1),
-                               FScheduledOp( args.source, item) );
-                    }
-                    else
-                    {
-                    StoreMesh( item, nullptr );
-                }
-                break;
+			{
+				if (args.source)
+				{
+					AddOp(FScheduledOp(item.At, item, 1),
+						FScheduledOp(args.source, item));
+				}
+				else
+				{
+					StoreMesh(item, nullptr);
+				}
 
+				break;
+			}
             case 1:
             {
 				MUTABLE_CPUPROFILER_SCOPE(ME_EXTRACTFACEGROUP_1)
 
-                Ptr<const Mesh> Source = LoadMesh( FCacheAddress(args.source,item) );
+                Ptr<const Mesh> Source = LoadMesh(FCacheAddress(args.source, item));
+				if (Source)
+				{
+					Ptr<Mesh> Result = CreateMesh();
 
-                MeshPtr pResult;
-                pResult = MeshExtractFaceGroup( Source.get(),
-                                                args.group );
+					bool bOutSuccess = false;
+					MeshExtractFaceGroup(Result.get(), Source.get(), args.group, bOutSuccess);
 
-                StoreMesh( item, pResult );
+					Release(Source);
+					if (!bOutSuccess)
+					{
+						check(false);
+						Release(Result);
+						StoreMesh(item, nullptr);
+					}
+					else
+					{
+						StoreMesh(item, Result);
+					}
+				}
+				else
+				{
+					StoreMesh(item, nullptr);
+				}
                 break;
             }
 
@@ -1671,29 +1905,41 @@ namespace mu
             switch (item.Stage)
             {
             case 0:
-                if (args.source)
-                {
-                        AddOp( FScheduledOp( item.At, item, 1),
-                               FScheduledOp( args.source, item) );
-                    }
-                    else
-                    {
-                    StoreMesh( item, nullptr );
-                }
-                break;
-
+			{
+				if (args.source)
+				{
+					AddOp(FScheduledOp(item.At, item, 1),
+						FScheduledOp(args.source, item));
+				}
+				else
+				{
+					StoreMesh(item, nullptr);
+				}
+				break;
+			}
             case 1:
             {
            		MUTABLE_CPUPROFILER_SCOPE(ME_TRANSFORM_1)
             		
-                Ptr<const Mesh> Source = LoadMesh( FCacheAddress(args.source,item) );
+                Ptr<const Mesh> Source = LoadMesh(FCacheAddress(args.source,item));
 
-                const mat4f& mat = pModel->GetPrivate()->m_program.
-                    m_constantMatrices[args.matrix];
+                const mat4f& mat = pModel->GetPrivate()->m_program.m_constantMatrices[args.matrix];
 
-                MeshPtr pResult = MeshTransform(Source.get(), ToUnreal(mat) );
+				Ptr<Mesh> Result = CreateMesh(Source ? Source->GetDataSize() : 0);
 
-                StoreMesh( item, pResult );
+				bool bOutSuccess = false;
+                MeshTransform(Result.get(), Source.get(), ToUnreal(mat), bOutSuccess);
+
+				if (!bOutSuccess)
+				{
+					Release(Result);
+					StoreMesh(item, Source);
+				}
+				else
+				{
+					Release(Source);
+					StoreMesh(item, Result);
+				}
                 break;
             }
 
@@ -1710,61 +1956,103 @@ namespace mu
             switch (item.Stage)
             {
             case 0:
-                if (args.source)
-                {
-                        AddOp( FScheduledOp( item.At, item, 1),
-                               FScheduledOp( args.source, item) );
-                    }
-                    else
-                    {
-                    StoreMesh( item, nullptr );
-                }
-                break;
-
+			{
+				if (args.source)
+				{
+					AddOp(FScheduledOp(item.At, item, 1),
+						FScheduledOp(args.source, item));
+				}
+				else
+				{
+					StoreMesh(item, nullptr);
+				}
+				break;
+			}
             case 1:
             {
            		MUTABLE_CPUPROFILER_SCOPE(ME_CLIPMORPHPLANE_1)
             		
-                Ptr<const Mesh> Source = LoadMesh( FCacheAddress(args.source,item) );
+                Ptr<const Mesh> Source = LoadMesh(FCacheAddress(args.source, item));
 
-                MeshPtr pResult;
-
-                check( args.morphShape < (uint32)pModel->GetPrivate()->m_program.m_constantShapes.Num() );
+                check(args.morphShape < (uint32)pModel->GetPrivate()->m_program.m_constantShapes.Num());
 
                 // Should be an ellipse
-                const FShape& morphShape = pModel->GetPrivate()->m_program.
-                    m_constantShapes[args.morphShape];
+                const FShape& morphShape = pModel->GetPrivate()->m_program.m_constantShapes[args.morphShape];
 
                 const mu::vec3f& origin = morphShape.position;
                 const mu::vec3f& normal = morphShape.up;
 
                 if (args.vertexSelectionType == OP::MeshClipMorphPlaneArgs::VS_SHAPE)
                 {
-                    check( args.vertexSelectionShapeOrBone < (uint32)pModel->GetPrivate()->m_program.m_constantShapes.Num() );
+                    check(args.vertexSelectionShapeOrBone < (uint32)pModel->GetPrivate()->m_program.m_constantShapes.Num());
 
                     // Should be None or an axis aligned box
                     const FShape& selectionShape = pModel->GetPrivate()->m_program.m_constantShapes[args.vertexSelectionShapeOrBone];
-                    pResult = MeshClipMorphPlane(Source.get(), origin, normal, args.dist, args.factor, morphShape.size[0], morphShape.size[1], morphShape.size[2], selectionShape);
+
+					Ptr<Mesh> Result = CreateMesh(Source->GetDataSize());
+
+					bool bOutSuccess = false;
+					MeshClipMorphPlane(Result.get(), Source.get(), origin, normal, args.dist, args.factor, morphShape.size[0], morphShape.size[1], morphShape.size[2], selectionShape, bOutSuccess, mu::string(), -1);
+					
+					if (!bOutSuccess)
+					{
+						Release(Result);
+						StoreMesh(item, Source);
+					}
+					else
+					{
+						Release(Source);
+						StoreMesh(item, Result);
+					}
                 }
 
-                else if (args.vertexSelectionType == OP::MeshClipMorphPlaneArgs::VS_BONE_HIERARCHY)
-                {
-                    check( args.vertexSelectionShapeOrBone < (uint32)pModel->GetPrivate()->m_program.m_constantStrings.Num() );
+				else if (args.vertexSelectionType == OP::MeshClipMorphPlaneArgs::VS_BONE_HIERARCHY)
+				{
+					check(args.vertexSelectionShapeOrBone < (uint32)pModel->GetPrivate()->m_program.m_constantStrings.Num());
 
-                    FShape selectionShape;
-                    selectionShape.type = (uint8)FShape::Type::None;
-                    const string& selectionBone = pModel->GetPrivate()->m_program.m_constantStrings[args.vertexSelectionShapeOrBone];
-					pResult = MeshClipMorphPlane(Source.get(), origin, normal, args.dist, args.factor, morphShape.size[0], morphShape.size[1], morphShape.size[2], selectionShape, selectionBone, args.maxBoneRadius);
+					FShape selectionShape;
+					selectionShape.type = (uint8)FShape::Type::None;
+					const string& selectionBone = pModel->GetPrivate()->m_program.m_constantStrings[args.vertexSelectionShapeOrBone];
+
+					Ptr<Mesh> Result = CreateMesh(Source->GetDataSize());
+
+					bool bOutSuccess = false;
+					MeshClipMorphPlane(Result.get(), Source.get(), origin, normal, args.dist, args.factor, morphShape.size[0], morphShape.size[1], morphShape.size[2], selectionShape, bOutSuccess, selectionBone, args.maxBoneRadius);
+
+					if (!bOutSuccess)
+					{
+						Release(Result);
+						StoreMesh(item, Source);
+					}
+					else
+					{
+						Release(Source);
+						StoreMesh(item, Result);
+					}
                 }
                 else
                 {
                     // No vertex selection
                     FShape selectionShape;
                     selectionShape.type = (uint8)FShape::Type::None;
-                    pResult = MeshClipMorphPlane(Source.get(), origin, normal, args.dist, args.factor, morphShape.size[0], morphShape.size[1], morphShape.size[2], selectionShape);
+
+					Ptr<Mesh> Result = CreateMesh(Source->GetDataSize());
+
+					bool bOutSuccess = false;
+					MeshClipMorphPlane(Result.get(), Source.get(), origin, normal, args.dist, args.factor, morphShape.size[0], morphShape.size[1], morphShape.size[2], selectionShape, bOutSuccess, mu::string(), -1.0f);
+
+					if (!bOutSuccess)
+					{
+						Release(Result);
+						StoreMesh(item, Source);
+					}
+					else
+					{
+						Release(Source);
+						StoreMesh(item, Result);
+					}
                 }
 
-                StoreMesh( item, pResult );
                 break;
             }
 
@@ -1782,37 +2070,53 @@ namespace mu
             switch (item.Stage)
             {
             case 0:
-                if (args.source)
-                {
-                        AddOp( FScheduledOp( item.At, item, 1),
-                               FScheduledOp( args.source, item),
-                               FScheduledOp( args.clipMesh, item) );
-                    }
-                    else
-                    {
-                    StoreMesh( item, nullptr );
-                }
-                break;
+			{
+				if (args.source)
+				{
+					AddOp(FScheduledOp(item.At, item, 1),
+						FScheduledOp(args.source, item),
+						FScheduledOp(args.clipMesh, item));
+				}
+				else
+				{
+					StoreMesh(item, nullptr);
+				}
 
+				break;
+			}
             case 1:
             {
 				MUTABLE_CPUPROFILER_SCOPE(ME_CLIPWITHMESH_1)
 
-                Ptr<const Mesh> Source = LoadMesh( FCacheAddress(args.source,item) );
-                Ptr<const Mesh> pClip = LoadMesh( FCacheAddress(args.clipMesh,item) );
+                Ptr<const Mesh> Source = LoadMesh(FCacheAddress(args.source, item));
+                Ptr<const Mesh> pClip = LoadMesh(FCacheAddress(args.clipMesh, item));
 
                 // Only if both are valid.
-                MeshPtr pResult;
                 if (Source && pClip)
                 {
-                    pResult = MeshClipWithMesh(Source.get(), pClip.get());
+					Ptr<Mesh> Result = CreateMesh(Source->GetDataSize());
+
+					bool bOutSuccess = false;
+                    MeshClipWithMesh(Result.get(), Source.get(), pClip.get(), bOutSuccess);
+
+					Release(pClip);
+					if (!bOutSuccess)
+					{
+						Release(Result);
+						StoreMesh(item, Source);
+					}
+					else
+					{
+						Release(Source);
+						StoreMesh(item, Result);
+					}
                 }
-                else if (Source)
+                else
                 {
-                    pResult = Source->Clone();
+					Release(pClip);
+					StoreMesh(item, Source);
                 }
 
-                StoreMesh( item, pResult );
                 break;
             }
 
@@ -1828,36 +2132,52 @@ namespace mu
 			switch (item.Stage)
 			{
 			case 0:
+			{
 				if (args.mesh)
 				{
-						AddOp(FScheduledOp(item.At, item, 1),
-							FScheduledOp(args.mesh, item),
-							FScheduledOp(args.clipShape, item));
-					}
-					else
-					{
+					AddOp(FScheduledOp(item.At, item, 1),
+						FScheduledOp(args.mesh, item),
+						FScheduledOp(args.clipShape, item));
+				}
+				else
+				{
 					StoreMesh(item, nullptr);
 				}
 				break;
-			
+			}
 			case 1:
 			{
 				MUTABLE_CPUPROFILER_SCOPE(ME_CLIPDEFORM_1)
 
 				Ptr<const Mesh> BaseMesh = LoadMesh(FCacheAddress(args.mesh, item));
 				Ptr<const Mesh> ClipShape = LoadMesh(FCacheAddress(args.clipShape, item));
-				Ptr<Mesh> pResult;
 
-				if ( BaseMesh && ClipShape )
+				if (BaseMesh && ClipShape)
 				{
-					pResult = MeshClipDeform(BaseMesh.get(), ClipShape.get(), args.clipWeightThreshold);
+					Ptr<Mesh> Result = CreateMesh(BaseMesh->GetDataSize());
+
+					bool bOutSuccess = false;
+					MeshClipDeform(Result.get(), BaseMesh.get(), ClipShape.get(), args.clipWeightThreshold, bOutSuccess);
+
+					Release(ClipShape);
+
+					if (!bOutSuccess)
+					{
+						Release(Result);
+						StoreMesh(item, BaseMesh);
+					}
+					else
+					{
+						Release(BaseMesh);
+						StoreMesh(item, Result);
+					}
 				}
-				else if( BaseMesh )
+				else
 				{
-					pResult = BaseMesh->Clone();
+					Release(ClipShape);
+					StoreMesh(item, BaseMesh);
 				}
 
-				StoreMesh(item, pResult);
 				break;
 			}
 
@@ -1874,37 +2194,47 @@ namespace mu
             switch (item.Stage)
             {
             case 0:
-                if (args.source)
-                {
-                        AddOp( FScheduledOp( item.At, item, 1),
-                               FScheduledOp( args.source, item),
-                               FScheduledOp( args.reference, item) );
-                    }
-                    else
-                    {
-                    StoreMesh( item, nullptr );
-                }
-                break;
-
+			{
+				if (args.source)
+				{
+					AddOp(FScheduledOp(item.At, item, 1),
+						FScheduledOp(args.source, item),
+						FScheduledOp(args.reference, item));
+				}
+				else
+				{
+					StoreMesh(item, nullptr);
+				}
+				break;
+			}
             case 1:
             {
     			MUTABLE_CPUPROFILER_SCOPE(ME_REMAPINDICES_1)
 
-                Ptr<const Mesh> Source = LoadMesh( FCacheAddress(args.source,item) );
-                Ptr<const Mesh> pReference = LoadMesh( FCacheAddress(args.reference,item) );
+                Ptr<const Mesh> Source = LoadMesh(FCacheAddress(args.source,item));
+                Ptr<const Mesh> pReference = LoadMesh(FCacheAddress(args.reference,item));
 
                 // Only if both are valid.
-                MeshPtr pResult;
-                if (Source && pReference)
-                {
-                    pResult = MeshRemapIndices(Source.get(), pReference.get());
-                }
-                else if (Source)
-                {
-                    pResult = Source->Clone();
+                MeshPtr Result = CreateMesh(Source->GetDataSize());
+                
+				bool bOutSuccess = false;
+				if (Source && pReference)
+                {	
+                    MeshRemapIndices(Result.get(), Source.get(), pReference.get(), bOutSuccess);
                 }
+			
+				Release(pReference);
+				if (!bOutSuccess)
+				{
+					Release(Result);
+					StoreMesh(item, Source);
+				}
+				else
+				{	
+					Release(Source);
+					StoreMesh(item, Result);
+				}
 
-                StoreMesh( item, pResult );
                 break;
             }
 
@@ -1922,37 +2252,52 @@ namespace mu
             switch (item.Stage)
             {
             case 0:
-                if (args.base)
-                {
-                        AddOp( FScheduledOp( item.At, item, 1),
-                               FScheduledOp( args.base, item),
-                               FScheduledOp( args.pose, item) );
-                    }
-                    else
-                    {
-                    StoreMesh( item, nullptr );
-                }
-                break;
-
+			{
+				if (args.base)
+				{
+					AddOp(FScheduledOp(item.At, item, 1),
+						FScheduledOp(args.base, item),
+						FScheduledOp(args.pose, item));
+				}
+				else
+				{
+					StoreMesh(item, nullptr);
+				}
+				break;
+			}
             case 1:
             {
           		MUTABLE_CPUPROFILER_SCOPE(ME_APPLYPOSE_1)
 
-                Ptr<const Mesh> pBase = LoadMesh( FCacheAddress(args.base,item) );
-                Ptr<const Mesh> pPose = LoadMesh( FCacheAddress(args.pose,item) );
+                Ptr<const Mesh> pBase = LoadMesh(FCacheAddress(args.base, item));
+                Ptr<const Mesh> pPose = LoadMesh(FCacheAddress(args.pose, item));
 
                 // Only if both are valid.
-                MeshPtr pResult;
                 if (pBase && pPose)
                 {
-                    pResult = MeshApplyPose(pBase.get(), pPose.get());
+					Ptr<Mesh> Result = CreateMesh(pBase->GetSkeleton() ? pBase->GetDataSize() : 0);
+
+					bool bOutSuccess = false;
+					MeshApplyPose(Result.get(), pBase.get(), pPose.get(), bOutSuccess);
+
+					Release(pPose);
+					if (!bOutSuccess)
+					{
+						Release(Result);
+						StoreMesh(item, pBase);
+					}
+					else
+					{
+						Release(pBase);
+						StoreMesh(item, Result);
+					}
                 }
-                else if (pBase)
+                else
                 {
-                    pResult = pBase->Clone();
+					Release(pPose);
+					StoreMesh(item, pBase);
                 }
 
-                StoreMesh( item, pResult );
                 break;
             }
 
@@ -1970,20 +2315,21 @@ namespace mu
 			switch (item.Stage)
 			{
 			case 0:
+			{
 				if (args.meshA)
 				{
-						AddOp(FScheduledOp(item.At, item, 1),
-							FScheduledOp(args.meshA, item),
-							FScheduledOp(args.meshB, item),
-							FScheduledOp(args.scalarA, item),
-							FScheduledOp(args.scalarB, item));
-					}
-					else
-					{
+					AddOp(FScheduledOp(item.At, item, 1),
+						FScheduledOp(args.meshA, item),
+						FScheduledOp(args.meshB, item),
+						FScheduledOp(args.scalarA, item),
+						FScheduledOp(args.scalarB, item));
+				}
+				else
+				{
 					StoreMesh(item, nullptr);
 				}
 				break;
-
+			}
 			case 1:
 			{
 				MUTABLE_CPUPROFILER_SCOPE(ME_GEOMETRYOPERATION_1)
@@ -1993,9 +2339,24 @@ namespace mu
 				float ScalarA = LoadScalar(FCacheAddress(args.scalarA, item));
 				float ScalarB = LoadScalar(FCacheAddress(args.scalarB, item));
 
-				MeshPtr pResult = MeshGeometryOperation(MeshA.get(),MeshB.get(),ScalarA,ScalarB);
+				Ptr<Mesh> Result = CreateMesh(MeshA->GetDataSize());
+
+				bool bOutSuccess = false;
+				MeshGeometryOperation(Result.get(), MeshA.get(), MeshB.get(), ScalarA, ScalarB, bOutSuccess);
+
+				Release(MeshA);
+				Release(MeshB);
+
+				if (!bOutSuccess)
+				{
+					Release(Result);
+					StoreMesh(item, nullptr);
+				}
+				else
+				{
+					StoreMesh(item, Result);
+				}
 
-				StoreMesh(item, pResult);
 				break;
 			}
 
@@ -2015,25 +2376,25 @@ namespace mu
 			switch (item.Stage)
 			{
 			case 0:
+			{
 				if (Args.mesh)
 				{
 					AddOp(FScheduledOp(item.At, item, 1),
-							FScheduledOp(Args.mesh, item),
-							FScheduledOp(Args.shape, item));
+						FScheduledOp(Args.mesh, item),
+						FScheduledOp(Args.shape, item));
 				}
 				else
 				{
 					StoreMesh(item, nullptr);
 				}
 				break;
-
+			}
 			case 1:
 			{
 				MUTABLE_CPUPROFILER_SCOPE(ME_BINDSHAPE_1)
 				Ptr<const Mesh> BaseMesh = LoadMesh(FCacheAddress(Args.mesh, item));
 				Ptr<const Mesh> Shape = LoadMesh(FCacheAddress(Args.shape, item));
 				
-				
 				EShapeBindingMethod BindingMethod = static_cast<EShapeBindingMethod>(Args.bindingMethod); 
 
 				if (BindingMethod == EShapeBindingMethod::ReshapeClosestProject)
@@ -2072,33 +2433,54 @@ namespace mu
 
 					const EMeshBindShapeFlags BindFlags = static_cast<EMeshBindShapeFlags>(Args.flags);
 
-					// -----------
-					Ptr<Mesh> BindMesh = 
-						MeshBindShapeReshape( BaseMesh.get(), Shape.get(), BonesToDeform, PhysicsToDeform, BindFlags);
-					
-					// null bind mesh indicates nothing has bond so the base mesh can be reused.
-					if (BindMesh)
+					Ptr<Mesh> BindMeshResult = CreateMesh();
+
+					bool bOutSuccess = false;
+					MeshBindShapeReshape(BindMeshResult.get(), BaseMesh.get(), Shape.get(), BonesToDeform, PhysicsToDeform, BindFlags, bOutSuccess);
+				
+					Release(Shape);
+					// not success indicates nothing has bond so the base mesh can be reused.
+					if (!bOutSuccess)
+					{
+						Release(BindMeshResult);
+						StoreMesh(item, BaseMesh);
+					}
+					else
 					{
 						if (!EnumHasAnyFlags(BindFlags, EMeshBindShapeFlags::ReshapeVertices))
 						{
-							Ptr<Mesh> BindMeshNoVerts = BaseMesh->Clone();
-							BindMeshNoVerts->m_AdditionalBuffers = MoveTemp(BindMesh->m_AdditionalBuffers);
-							StoreMesh(item, BindMeshNoVerts);
+							Ptr<Mesh> BindMeshNoVertsResult = CloneOrTakeOver(BaseMesh);
+							BindMeshNoVertsResult->m_AdditionalBuffers = MoveTemp(BindMeshResult->m_AdditionalBuffers);
+
+							Release(BaseMesh);
+							Release(BindMeshResult);
+							StoreMesh(item, BindMeshNoVertsResult);
 						}
 						else
 						{
-							StoreMesh(item, BindMesh);
+							Release(BaseMesh);
+							StoreMesh(item, BindMeshResult);
 						}
 					}
-					else
-					{
-						StoreMesh(item, BaseMesh);
-					}	
 				}	
 				else
 				{
-					Ptr<Mesh> ResultPtr = MeshBindShapeClipDeform( BaseMesh.get(), Shape.get(), BindingMethod );
-					StoreMesh(item, ResultPtr);
+					Ptr<Mesh> Result = CreateMesh(BaseMesh->GetDataSize());
+
+					bool bOutSuccess = false;
+					MeshBindShapeClipDeform(Result.get(), BaseMesh.get(), Shape.get(), BindingMethod, bOutSuccess);
+
+					Release(Shape);
+					if (!bOutSuccess)
+					{
+						Release(Result);
+						StoreMesh(item, BaseMesh);
+					}
+					else
+					{
+						Release(BaseMesh);
+						StoreMesh(item, Result);
+					}
 				}
 
 				break;
@@ -2118,18 +2500,19 @@ namespace mu
 			switch (item.Stage)
 			{
 			case 0:
+			{
 				if (args.mesh)
 				{
-						AddOp(FScheduledOp(item.At, item, 1),
-							FScheduledOp(args.mesh, item),
-							FScheduledOp(args.shape, item));
+					AddOp(FScheduledOp(item.At, item, 1),
+						FScheduledOp(args.mesh, item),
+						FScheduledOp(args.shape, item));
 				}
 				else
 				{
 					StoreMesh(item, nullptr);
 				}
 				break;
-
+			}
 			case 1:
 			{
 				MUTABLE_CPUPROFILER_SCOPE(ME_APPLYSHAPE_1)
@@ -2138,36 +2521,48 @@ namespace mu
 				Ptr<const Mesh> Shape = LoadMesh(FCacheAddress(args.shape, item));
 
 				const EMeshBindShapeFlags ReshapeFlags = static_cast<EMeshBindShapeFlags>(args.flags);
-				Ptr<Mesh> ReshapedMesh = MeshApplyShape(BaseMesh.get(), Shape.get(), ReshapeFlags);
-
 				const bool bReshapeVertices = EnumHasAnyFlags(ReshapeFlags, EMeshBindShapeFlags::ReshapeVertices);
 
-				if (!ReshapedMesh)
+				Ptr<Mesh> ReshapedMeshResult = CreateMesh(BaseMesh->GetDataSize());
+
+				bool bOutSuccess = false;
+				MeshApplyShape(ReshapedMeshResult.get(), BaseMesh.get(), Shape.get(), ReshapeFlags, bOutSuccess);
+
+				Release(Shape);
+				
+				if (!bOutSuccess)
 				{
+					Release(ReshapedMeshResult);
 					StoreMesh(item, BaseMesh);
 				}
-				else if (!bReshapeVertices)
-				{
-					// Clone without Skeleton, Physics or Poses 
-					EMeshCloneFlags CloneFlags = ~(
-							EMeshCloneFlags::WithSkeleton          | 
-							EMeshCloneFlags::WithPhysicsBody       |
-							EMeshCloneFlags::WithAdditionalPhysics |
-							EMeshCloneFlags::WithPoses);
- 
-					Ptr<Mesh> Result = BaseMesh->Clone(CloneFlags);
-					Result->SetSkeleton(ReshapedMesh->GetSkeleton().get());
-					Result->SetPhysicsBody(ReshapedMesh->GetPhysicsBody().get());
-					Result->AdditionalPhysicsBodies = ReshapedMesh->AdditionalPhysicsBodies;
-					Result->BonePoses = ReshapedMesh->BonePoses;
-
-					StoreMesh(item, Result);
-				}
 				else
 				{
-					StoreMesh(item, ReshapedMesh);
+					if (!bReshapeVertices)
+					{
+						// Clone without Skeleton, Physics or Poses 
+						EMeshCopyFlags CopyFlags = ~(
+							EMeshCopyFlags::WithSkeleton |
+							EMeshCopyFlags::WithPhysicsBody |
+							EMeshCopyFlags::WithAdditionalPhysics |
+							EMeshCopyFlags::WithPoses);
+
+						Ptr<Mesh> NoVerticesReshpedMesh = CloneOrTakeOver(BaseMesh);
+
+						NoVerticesReshpedMesh->SetSkeleton(ReshapedMeshResult->GetSkeleton().get());
+						NoVerticesReshpedMesh->SetPhysicsBody(ReshapedMeshResult->GetPhysicsBody().get());
+						NoVerticesReshpedMesh->AdditionalPhysicsBodies = ReshapedMeshResult->AdditionalPhysicsBodies;
+						NoVerticesReshpedMesh->BonePoses = ReshapedMeshResult->BonePoses;
+
+						Release(BaseMesh);
+						Release(ReshapedMeshResult);
+						StoreMesh(item, NoVerticesReshpedMesh);
+					}
+					else
+					{
+						Release(BaseMesh);
+						StoreMesh(item, ReshapedMeshResult);
+					}
 				}
-
 				break;
 			}
 
@@ -2188,8 +2583,8 @@ namespace mu
 				if (Args.Morph)
 				{
 					AddOp(FScheduledOp(item.At, item, 1), 
-						  FScheduledOp(Args.Morph, item), 
-						  FScheduledOp(Args.Reshape, item));
+						FScheduledOp(Args.Morph, item),
+						FScheduledOp(Args.Reshape, item));
 				}
 				else 
 				{
@@ -2206,23 +2601,26 @@ namespace mu
 
 				if (ReshapeMesh)
 				{
-					// Clone without Skeleton, Physics or Poses 
-					EMeshCloneFlags CloneFlags = ~(
-							EMeshCloneFlags::WithSkeleton    | 
-							EMeshCloneFlags::WithPhysicsBody | 
-							EMeshCloneFlags::WithPoses);
+					// Copy without Skeleton, Physics or Poses 
+					EMeshCopyFlags CopyFlags = ~(
+							EMeshCopyFlags::WithSkeleton    | 
+							EMeshCopyFlags::WithPhysicsBody | 
+							EMeshCopyFlags::WithPoses);
 
-					Ptr<Mesh> ResultPtr = MorphedMesh->Clone(CloneFlags);
-					ResultPtr->SetSkeleton(ReshapeMesh->GetSkeleton().get());
-					ResultPtr->SetPhysicsBody(ReshapeMesh->GetPhysicsBody().get());
-					ResultPtr->BonePoses = ReshapeMesh->BonePoses;
+					Ptr<Mesh> Result = CreateMesh(MorphedMesh->GetDataSize());
+					Result->CopyFrom(*MorphedMesh, CopyFlags);
 
-					StoreMesh(item, ResultPtr);
+					Result->SetSkeleton(ReshapeMesh->GetSkeleton().get());
+					Result->SetPhysicsBody(ReshapeMesh->GetPhysicsBody().get());
+					Result->BonePoses = ReshapeMesh->BonePoses;
+
+					Release(MorphedMesh);
+					Release(ReshapeMesh);
+					StoreMesh(item, Result);
 				}
 				else
 				{
-					Ptr<Mesh> ResultPtr = MorphedMesh->Clone();
-					StoreMesh(item, ResultPtr);
+					StoreMesh(item, MorphedMesh);
 				}
 
 				break;
@@ -2241,57 +2639,73 @@ namespace mu
             switch (item.Stage)
             {
             case 0:
-                if (args.source)
-                {
-                        AddOp( FScheduledOp( item.At, item, 1),
-                               FScheduledOp( args.source, item),
-                               FScheduledOp( args.skeleton, item) );
-                    }
-                    else
-                    {
-                    StoreMesh( item, nullptr );
-                }
-                break;
-
+			{
+				if (args.source)
+				{
+					AddOp(FScheduledOp(item.At, item, 1),
+						FScheduledOp(args.source, item),
+						FScheduledOp(args.skeleton, item));
+				}
+				else
+				{
+					StoreMesh(item, nullptr);
+				}
+				break;
+			}
             case 1:
             {
             	MUTABLE_CPUPROFILER_SCOPE(ME_SETSKELETON_1)
             		
-                Ptr<const Mesh> Source = LoadMesh( FCacheAddress(args.source,item) );
-                Ptr<const Mesh> pSkeleton = LoadMesh( FCacheAddress(args.skeleton,item) );
+                Ptr<const Mesh> Source = LoadMesh(FCacheAddress(args.source, item));
+                Ptr<const Mesh> pSkeleton = LoadMesh(FCacheAddress(args.skeleton, item));
 
                 // Only if both are valid.
-                MeshPtr pResult;
                 if (Source && pSkeleton)
                 {
-                    if ( Source->GetSkeleton()
-                         &&
-                         !Source->GetSkeleton()->m_bones.IsEmpty() )
+                    if (Source->GetSkeleton() && !Source->GetSkeleton()->m_bones.IsEmpty())
                     {
                         // For some reason we already have bone data, so we can't just overwrite it
                         // or the skinning may break. This may happen because of a problem in the
                         // optimiser that needs investigation.
                         // \TODO Be defensive, for now.
-                        UE_LOG(LogMutableCore,Warning, TEXT("Performing a MeshRemapSkeleton, instead of MeshSetSkeletonData because source mesh already has some skeleton."));
-                        pResult = MeshRemapSkeleton( Source.get(), pSkeleton->GetSkeleton().get() );
-                        if (!pResult)
+                        UE_LOG(LogMutableCore, Warning, TEXT("Performing a MeshRemapSkeleton, instead of MeshSetSkeletonData because source mesh already has some skeleton."));
+
+						Ptr<Mesh> Result = CreateMesh(Source->GetDataSize());
+
+						bool bOutSuccess = false;
+                        MeshRemapSkeleton(Result.get(), Source.get(), pSkeleton->GetSkeleton().get(), bOutSuccess);
+
+						Release(pSkeleton);
+
+                        if (!bOutSuccess)
                         {
-                            pResult = Source->Clone();
+							Release(Result);
+							StoreMesh(item, Source);
                         }
+						else
+						{
+							//Result->GetPrivate()->CheckIntegrity();
+							Release(Source);
+							StoreMesh(item, Result);
+						}
                     }
                     else
                     {
-                        pResult = Source->Clone();
-                        pResult->SetSkeleton(pSkeleton->GetSkeleton().get());
+						Ptr<Mesh> Result = CloneOrTakeOver(Source);
+
+                        Result->SetSkeleton(pSkeleton->GetSkeleton().get());
+
+						//Result->GetPrivate()->CheckIntegrity();
+						Release(pSkeleton);
+						StoreMesh(item, Result);
                     }
-                    //pResult->GetPrivate()->CheckIntegrity();
                 }
-                else if (Source)
+                else
                 {
-                    pResult = Source->Clone();
+					Release(pSkeleton);
+					StoreMesh(item, Source);
                 }
 
-                StoreMesh( item, pResult );
                 break;
             }
 
@@ -2311,23 +2725,29 @@ namespace mu
             const uint8* data = pModel->GetPrivate()->m_program.GetOpArgsPointer(item.At);
 
             OP::ADDRESS source;
-            FMemory::Memcpy(&source,data,sizeof(OP::ADDRESS)); data += sizeof(OP::ADDRESS);
+            FMemory::Memcpy(&source,data,sizeof(OP::ADDRESS)); 
+			data += sizeof(OP::ADDRESS);
 
             TArray<FScheduledOp> conditions;
 			TArray<OP::ADDRESS> masks;
 
             uint16 removes;
-			FMemory::Memcpy(&removes,data,sizeof(uint16)); data += sizeof(uint16);
+			FMemory::Memcpy(&removes,data,sizeof(uint16)); 
+			data += sizeof(uint16);
 
             for( uint16 r=0; r<removes; ++r)
             {
                 OP::ADDRESS condition;
-				FMemory::Memcpy(&condition,data,sizeof(OP::ADDRESS)); data += sizeof(OP::ADDRESS);
-                conditions.Emplace( condition, item );
+				FMemory::Memcpy(&condition,data,sizeof(OP::ADDRESS)); 
+				data += sizeof(OP::ADDRESS);
+                
+				conditions.Emplace(condition, item);
 
                 OP::ADDRESS mask;
-				FMemory::Memcpy(&mask,data,sizeof(OP::ADDRESS)); data += sizeof(OP::ADDRESS);
-                masks.Add( mask );
+				FMemory::Memcpy(&mask,data,sizeof(OP::ADDRESS)); 
+				data += sizeof(OP::ADDRESS);
+
+                masks.Add(mask);
             }
 
 
@@ -2335,17 +2755,18 @@ namespace mu
             switch (item.Stage)
             {
             case 0:
-                if (source)
-                {
-                    // Request the conditions
-                    AddOp( FScheduledOp( item.At, item, 1), conditions );
-                }
-                else
-                {
-                    StoreMesh( item, nullptr );
-                }
-                break;
-
+			{
+				if (source)
+				{
+					// Request the conditions
+					AddOp(FScheduledOp(item.At, item, 1), conditions);
+				}
+				else
+				{
+					StoreMesh(item, nullptr);
+				}
+				break;
+			}
             case 1:
             {
         		MUTABLE_CPUPROFILER_SCOPE(ME_REMOVEMASK_1)
@@ -2360,19 +2781,19 @@ namespace mu
                     bool value = true;
                     if (conditions[r].At)
                     {
-                        value = LoadBool( FCacheAddress(conditions[r].At, item) );
+                        value = LoadBool(FCacheAddress(conditions[r].At, item));
                     }
 
                     if (value)
                     {
-                        deps.Emplace( masks[r], item );
+                        deps.Emplace(masks[r], item);
                     }
                 }
 
                 if (source)
                 {
-                        AddOp( FScheduledOp( item.At, item, 2), deps );
-                    }
+					AddOp(FScheduledOp(item.At, item, 2), deps);
+				}
                 break;
             }
 
@@ -2381,30 +2802,57 @@ namespace mu
             	MUTABLE_CPUPROFILER_SCOPE(ME_REMOVEMASK_2)
             	
                 // \todo: single remove operation with all masks?
-                Ptr<const Mesh> Source = LoadMesh( FCacheAddress(source,item) );
+                Ptr<const Mesh> Source = LoadMesh(FCacheAddress(source, item));
 
-                MeshPtrConst pResult = Source;
+				if (Source)
+				{
+					Ptr<Mesh> Result = CreateMesh(Source->GetDataSize());
+					Result->CopyFrom(*Source);
 
-                for( size_t r=0; pResult && r<conditions.Num(); ++r )
-                {
-                    // If there is no expression, we'll assume true.
-                    bool value = true;
-                    if (conditions[r].At)
-                    {
-                        value = LoadBool( FCacheAddress(conditions[r].At, item) );
-                    }
+					Release(Source);
 
-                    if (value)
-                    {
-                        Ptr<const Mesh> Mask = LoadMesh( FCacheAddress(masks[r],item) );
-                        if (Mask)
-                        {
-                            pResult = MeshRemoveMask(pResult.get(), Mask.get());
-                        }
-                    }
-                }
+					for (int32 r = 0; r < conditions.Num(); ++r)
+					{
+						// If there is no expression, we'll assume true.
+						bool value = true;
+						if (conditions[r].At)
+						{
+							value = LoadBool(FCacheAddress(conditions[r].At, item));
+						}
+
+						if (value)
+						{
+							Ptr<const Mesh> Mask = LoadMesh(FCacheAddress(masks[r], item));
+							if (Mask)
+							{
+								//MeshRemoveMask will make a copy of Result, try to make room for it. 
+								Ptr<Mesh> IterResult = CreateMesh(Result->GetDataSize());
+
+								bool bOutSuccess = false;
+								MeshRemoveMask(IterResult.get(), Result.get(), Mask.get(), bOutSuccess);
+
+								Release(Mask);
+
+								if (!bOutSuccess)
+								{
+									Release(IterResult);
+								}
+								else
+								{
+									Swap(Result, IterResult);
+									Release(IterResult);
+								}
+							}
+						}
+					}
+
+					StoreMesh(item, Result);
+				}
+				else
+				{
+					StoreMesh(item, nullptr);
+				}
 
-                StoreMesh( item, pResult );
                 break;
             }
 
@@ -2421,34 +2869,52 @@ namespace mu
             switch (item.Stage)
             {
             case 0:
-                if (args.mesh)
-                {
-                        AddOp( FScheduledOp( item.At, item, 1),
-                               FScheduledOp( args.mesh, item),
-                               FScheduledOp( args.projector, item));
-                }
-                else
-                {
-                    StoreMesh( item, nullptr );
-                }
-                break;
-
+			{
+				if (args.mesh)
+				{
+					AddOp(FScheduledOp(item.At, item, 1),
+						FScheduledOp(args.mesh, item),
+						FScheduledOp(args.projector, item));
+				}
+				else
+				{
+					StoreMesh(item, nullptr);
+				}
+				break;
+			}
             case 1:
             {
 				MUTABLE_CPUPROFILER_SCOPE(ME_PROJECT_1)
 
-                Ptr<const Mesh> pMesh = LoadMesh( FCacheAddress(args.mesh,item) );
-                FProjector Projector = LoadProjector( FCacheAddress(args.projector,item) );
+                Ptr<const Mesh> pMesh = LoadMesh(FCacheAddress(args.mesh,item));
+                const FProjector Projector = LoadProjector(FCacheAddress(args.projector, item));
 
                 // Only if both are valid.
-                MeshPtr pResult;
-                if (pMesh
-                    && pMesh.get()
-                    && pMesh.get()->GetVertexBuffers().GetBufferCount() > 0 )
+                if (pMesh && pMesh->GetVertexBuffers().GetBufferCount() > 0)
                 {
-                    pResult = MeshProject(pMesh.get(), Projector);
+					Ptr<Mesh> Result = CreateMesh();
+
+					bool bOutSuccess = false;
+					MeshProject(Result.get(), pMesh.get(), Projector, bOutSuccess);
+
+					if (!bOutSuccess)
+					{
+						Release(Result);
+						StoreMesh(item, pMesh);
+					}
+					else
+					{	
+//						Result->GetPrivate()->CheckIntegrity();
+						Release(pMesh);
+						StoreMesh(item, Result);
+					}
                 }
-                StoreMesh( item, pResult );
+				else
+				{
+					Release(pMesh);
+					StoreMesh(item, nullptr);
+				}
+
                 break;
             }
 
@@ -2465,6 +2931,7 @@ namespace mu
 			switch (item.Stage)
 			{
 			case 0:
+			{
 				if (args.source)
 				{
 					AddOp(FScheduledOp(item.At, item, 1), FScheduledOp(args.source, item));
@@ -2474,16 +2941,29 @@ namespace mu
 					StoreMesh(item, nullptr);
 				}
 				break;
-
+			}
 			case 1:
 			{
 				MUTABLE_CPUPROFILER_SCOPE(ME_OPTIMIZESKINNING_1)
 
 				Ptr<const Mesh> Source = LoadMesh(FCacheAddress(args.source, item));
 
-				MeshPtrConst pResult = MeshOptimizeSkinning(Source.get());
+				Ptr<Mesh> Result = CreateMesh();
+
+				bool bOutSuccess = false;
+				MeshOptimizeSkinning(Result.get(), Source.get(), bOutSuccess);
+
+				if (!bOutSuccess)
+				{
+					Release(Result);
+					StoreMesh(item, Source);
+				}
+				else
+				{
+					Release(Source);
+					StoreMesh(item, Result);
+				}
 
-				StoreMesh(item, pResult ? pResult : Source);
 				break;
 			}
 
@@ -3934,6 +4414,7 @@ namespace mu
 					if (pMesh)
 					{
 						ImageRasterMesh(pMesh.get(), ResultImage.get(), args.blockIndex, CropMin, UncroppedSize);
+						Release(pMesh);
 					}
 
 					// Stop execution.
@@ -3978,18 +4459,13 @@ namespace mu
 						Data.RasterMesh.Mip = static_cast<uint8>(FMath::FloorToInt32(Data.RasterMesh.MipValue));
 					}
 				}
-	
-				// Set mesh and projector back to the cache to prevent possible code re-execution.
-				// this needs to be done before the operation is re-scheduled otherwise the execution may deadlock.	
-				// TODO: Review if this is a good way of reusing data between stages.
-				StoreMesh(FScheduledOp::FromOpAndOptions(args.mesh, item, 0), pMesh);
-				StoreProjector(FScheduledOp::FromOpAndOptions(args.projector, item, 0), Projector);
-				
+		
+				const int32 DataHeapAddress = m_heapData.Add(Data);
 
-				int32 DataHeapAddress = m_heapData.Add(Data);
+				// pMesh is need again in the next stage, store it in the heap.
+				m_heapData[DataHeapAddress].Resource = const_cast<Mesh*>(pMesh.get());
 
 				AddOp(FScheduledOp(item.At, item, 2, DataHeapAddress),
-					FScheduledOp::FromOpAndOptions(args.mesh, item, 0),
 					FScheduledOp::FromOpAndOptions(args.projector, item, 0),
 					FScheduledOp::FromOpAndOptions(args.image, item, Data.RasterMesh.Mip),
 					FScheduledOp(args.mask, item),
@@ -4010,7 +4486,11 @@ namespace mu
 					break;
 				}
 
-                Ptr<const Mesh> pMesh = LoadMesh(FScheduledOp::FromOpAndOptions(args.mesh, item, 0));
+				FScheduledOpData& Data = m_heapData[item.CustomState];
+
+				// Unsafe downcast, should be fine as it is known to be a Mesh.
+				Ptr<const Mesh> pMesh = static_cast<Mesh*>(Data.Resource.get());
+				Data.Resource = nullptr;
 
 				if (!pMesh)
 				{
@@ -4041,8 +4521,6 @@ namespace mu
 					--MipsToDrop;
 				}
 
-				FScheduledOpData Data = m_heapData[item.CustomState];
-
 				// Raster with projection
 				Ptr<const Image> Source = LoadImage(FCacheAddress(args.image, item.ExecutionIndex, Data.RasterMesh.Mip));
 
@@ -4177,6 +4655,7 @@ namespace mu
 					}
 				}
 
+				Release(pMesh);
 				Release(Source);
 				Release(Mask);
 				StoreImage(item, pNew);
@@ -5318,6 +5797,7 @@ namespace mu
 
 				Ptr<const Layout> Result = LayoutFromMesh_RemoveBlocks(Mesh.get(), args.LayoutIndex);
 
+				Release(Mesh);
 				StoreLayout(item, Result);
 				break;
 			}
diff --git a/Engine/Plugins/Experimental/Mutable/Source/MutableRuntime/Private/MuR/CodeRunner.h b/Engine/Plugins/Experimental/Mutable/Source/MutableRuntime/Private/MuR/CodeRunner.h
--- a/Engine/Plugins/Experimental/Mutable/Source/MutableRuntime/Private/MuR/CodeRunner.h
+++ b/Engine/Plugins/Experimental/Mutable/Source/MutableRuntime/Private/MuR/CodeRunner.h
@@ -479,7 +479,7 @@ namespace  mu
 
 		void StoreMesh(const FCacheAddress& To, Ptr<const Mesh> Resource)
 		{
-			m_pSystem->WorkingMemoryManager.CurrentInstanceCache->SetMesh(To, Resource);
+			m_pSystem->WorkingMemoryManager.StoreMesh(To, Resource);
 		}
 
 		void StoreImage(const FCacheAddress& To, Ptr<const Image> Resource)
@@ -531,6 +531,25 @@ namespace  mu
 			return m_pSystem->WorkingMemoryManager.Release(Resource);
 		}
 
+		UE_NODISCARD Ptr<Mesh> CreateMesh(int32 BudgetReserveSize = 0)
+		{
+			return m_pSystem->WorkingMemoryManager.CreateMesh(BudgetReserveSize);
+		}
+
+		UE_NODISCARD Ptr<Mesh> CloneOrTakeOver(Ptr<const Mesh>& Ref)
+		{
+			return m_pSystem->WorkingMemoryManager.CloneOrTakeOver(Ref);
+		}
+
+		void Release(Ptr<const Mesh>& Resource)
+		{
+			m_pSystem->WorkingMemoryManager.Release(Resource);
+		}
+
+		void Release(Ptr<Mesh>& Resource)
+		{
+			m_pSystem->WorkingMemoryManager.Release(Resource);
+		}
 
 	// TODO: protect this.
 	public:
diff --git a/Engine/Plugins/Experimental/Mutable/Source/MutableRuntime/Private/MuR/Mesh.cpp b/Engine/Plugins/Experimental/Mutable/Source/MutableRuntime/Private/MuR/Mesh.cpp
--- a/Engine/Plugins/Experimental/Mutable/Source/MutableRuntime/Private/MuR/Mesh.cpp
+++ b/Engine/Plugins/Experimental/Mutable/Source/MutableRuntime/Private/MuR/Mesh.cpp
@@ -22,7 +22,6 @@ void Mesh::Serialise( const Mesh* p, OutputArchive& arch )
     arch << *p;
 }
 
-
 //---------------------------------------------------------------------------------------------
 MeshPtr Mesh::StaticUnserialise( InputArchive& arch )
 {
@@ -81,7 +80,7 @@ MeshPtr Mesh::Clone() const
     return pResult;
 }
 
-MeshPtr Mesh::Clone(EMeshCloneFlags Flags) const
+MeshPtr Mesh::Clone(EMeshCopyFlags Flags) const
 {
     //MUTABLE_CPUPROFILER_SCOPE(MeshClone);
 	LLM_SCOPE_BYNAME(TEXT("MutableRuntime"));
@@ -91,55 +90,55 @@ MeshPtr Mesh::Clone(EMeshCloneFlags Flags) const
     pResult->m_internalId = m_internalId;
 	pResult->m_staticFormatFlags = m_staticFormatFlags;
 
-	if (EnumHasAnyFlags(Flags, EMeshCloneFlags::WithSurfaces))
+	if (EnumHasAnyFlags(Flags, EMeshCopyFlags::WithSurfaces))
 	{
 		pResult->m_surfaces = m_surfaces;
 	}
 
-	if (EnumHasAnyFlags(Flags, EMeshCloneFlags::WithSkeleton))
+	if (EnumHasAnyFlags(Flags, EMeshCopyFlags::WithSkeleton))
 	{
 		pResult->m_pSkeleton = m_pSkeleton;
 	}
 
-	if (EnumHasAnyFlags(Flags, EMeshCloneFlags::WithPhysicsBody))
+	if (EnumHasAnyFlags(Flags, EMeshCopyFlags::WithPhysicsBody))
 	{
 		pResult->m_pPhysicsBody = m_pPhysicsBody;
 	}
 
-	if (EnumHasAnyFlags(Flags, EMeshCloneFlags::WithFaceGroups))
+	if (EnumHasAnyFlags(Flags, EMeshCopyFlags::WithFaceGroups))
 	{
 		pResult->m_faceGroups = m_faceGroups;
 	}
 
-	if (EnumHasAnyFlags(Flags, EMeshCloneFlags::WithTags))
+	if (EnumHasAnyFlags(Flags, EMeshCopyFlags::WithTags))
 	{
 		pResult->m_tags = m_tags;
 	}
 
     // Clone the main buffers
-	if (EnumHasAnyFlags(Flags, EMeshCloneFlags::WithVertexBuffers))
+	if (EnumHasAnyFlags(Flags, EMeshCopyFlags::WithVertexBuffers))
     {
 		pResult->m_VertexBuffers = m_VertexBuffers;
 	}
 
-	if (EnumHasAnyFlags(Flags, EMeshCloneFlags::WithIndexBuffers))
+	if (EnumHasAnyFlags(Flags, EMeshCopyFlags::WithIndexBuffers))
     {
 		pResult->m_IndexBuffers = m_IndexBuffers;
 	}
 
-	if (EnumHasAnyFlags(Flags, EMeshCloneFlags::WithFaceBuffers))
+	if (EnumHasAnyFlags(Flags, EMeshCopyFlags::WithFaceBuffers))
     {
 		pResult->m_FaceBuffers = m_FaceBuffers;
 	}
 
 	// Clone additional buffers
-	if (EnumHasAnyFlags(Flags, EMeshCloneFlags::WithAdditionalBuffers))
+	if (EnumHasAnyFlags(Flags, EMeshCopyFlags::WithAdditionalBuffers))
 	{
 		pResult->m_AdditionalBuffers = m_AdditionalBuffers;
 	}
 
     // Clone the layout	
-	if (EnumHasAnyFlags(Flags, EMeshCloneFlags::WithLayouts))
+	if (EnumHasAnyFlags(Flags, EMeshCopyFlags::WithLayouts))
 	{
 		pResult->m_layouts = m_layouts;
 	}
@@ -149,24 +148,24 @@ MeshPtr Mesh::Clone(EMeshCloneFlags Flags) const
 	// physics body doen't need to be deep cloned either as they are also assumed to be shared.
 	
 	// Clone bone poses
-	if (EnumHasAnyFlags(Flags, EMeshCloneFlags::WithPoses))
+	if (EnumHasAnyFlags(Flags, EMeshCopyFlags::WithPoses))
 	{
 		pResult->BonePoses = BonePoses;
 	}
 
 	// Clone BoneMap
-	if (EnumHasAnyFlags(Flags, EMeshCloneFlags::WithBoneMap))
+	if (EnumHasAnyFlags(Flags, EMeshCopyFlags::WithBoneMap))
 	{
 		pResult->BoneMapIndices = BoneMapIndices;
 	}
 
 	// Clone SkeletonIDs
-	if (EnumHasAnyFlags(Flags, EMeshCloneFlags::WithSkeletonIDs))
+	if (EnumHasAnyFlags(Flags, EMeshCopyFlags::WithSkeletonIDs))
 	{
 		pResult->SkeletonIDs = SkeletonIDs;
 	}
 
-	if (EnumHasAnyFlags(Flags, EMeshCloneFlags::WithAdditionalPhysics))
+	if (EnumHasAnyFlags(Flags, EMeshCopyFlags::WithAdditionalPhysics))
 	{
 		pResult->AdditionalPhysicsBodies = AdditionalPhysicsBodies;
 	}
@@ -175,6 +174,96 @@ MeshPtr Mesh::Clone(EMeshCloneFlags Flags) const
 }
 
 
+//---------------------------------------------------------------------------------------------
+void Mesh::CopyFrom(const Mesh& From, EMeshCopyFlags Flags)
+{
+    //MUTABLE_CPUPROFILER_SCOPE(CopyFrom);
+
+    m_internalId = From.m_internalId;
+	m_staticFormatFlags = From.m_staticFormatFlags;
+
+	if (EnumHasAnyFlags(Flags, EMeshCopyFlags::WithSurfaces))
+	{
+		m_surfaces = From.m_surfaces;
+	}
+
+	if (EnumHasAnyFlags(Flags, EMeshCopyFlags::WithSkeleton))
+	{
+		m_pSkeleton = From.m_pSkeleton;
+	}
+
+	if (EnumHasAnyFlags(Flags, EMeshCopyFlags::WithPhysicsBody))
+	{
+		m_pPhysicsBody = From.m_pPhysicsBody;
+	}
+
+	if (EnumHasAnyFlags(Flags, EMeshCopyFlags::WithFaceGroups))
+	{
+		m_faceGroups = From.m_faceGroups;
+	}
+
+	if (EnumHasAnyFlags(Flags, EMeshCopyFlags::WithTags))
+	{
+		m_tags = From.m_tags;
+	}
+
+    // Copy the main buffers
+	if (EnumHasAnyFlags(Flags, EMeshCopyFlags::WithVertexBuffers))
+    {
+		m_VertexBuffers = From.m_VertexBuffers;
+	}
+
+	if (EnumHasAnyFlags(Flags, EMeshCopyFlags::WithIndexBuffers))
+    {
+		m_IndexBuffers = From.m_IndexBuffers;
+	}
+
+	if (EnumHasAnyFlags(Flags, EMeshCopyFlags::WithFaceBuffers))
+    {
+		m_FaceBuffers = From.m_FaceBuffers;
+	}
+
+	// Copy additional buffers
+	if (EnumHasAnyFlags(Flags, EMeshCopyFlags::WithAdditionalBuffers))
+	{
+		m_AdditionalBuffers = From.m_AdditionalBuffers;
+	}
+
+    // Copy the layout	
+	if (EnumHasAnyFlags(Flags, EMeshCopyFlags::WithLayouts))
+	{
+		m_layouts = From.m_layouts;
+	}
+    // The skeleton is not copied because it is not owned by this mesh and it is always assumed
+    // to be shared.
+
+	// physics body doen't need to be deep copied either as they are also assumed to be shared.
+	
+	// Copy bone poses
+	if (EnumHasAnyFlags(Flags, EMeshCopyFlags::WithPoses))
+	{
+		BonePoses = From.BonePoses;
+	}
+
+	// Copy BoneMap
+	if (EnumHasAnyFlags(Flags, EMeshCopyFlags::WithBoneMap))
+	{
+		BoneMapIndices = From.BoneMapIndices;
+	}
+
+	// Copy SkeletonIDs
+	if (EnumHasAnyFlags(Flags, EMeshCopyFlags::WithSkeletonIDs))
+	{
+		SkeletonIDs = From.SkeletonIDs;
+	}
+
+	if (EnumHasAnyFlags(Flags, EMeshCopyFlags::WithAdditionalPhysics))
+	{
+		AdditionalPhysicsBodies = From.AdditionalPhysicsBodies;
+	}
+
+}
+
 //---------------------------------------------------------------------------------------------
 uint32_t Mesh::GetId() const
 {
@@ -606,13 +695,24 @@ void Mesh::AddSkeletonID(int32 SkeletonID)
 //---------------------------------------------------------------------------------------------
 int32 Mesh::GetDataSize() const
 {
-    return sizeof(Mesh)
+	// TODO: review if other mesh fields like additional physics assets
+	// are relevant and add them to the count.
+
+	// Should be allocation sizes used for this?
+	int32 AdditionalBuffersSize = 0;
+	for (const TPair<EMeshBufferType, FMeshBufferSet>&  AdditionalBuffer : m_AdditionalBuffers)
+	{
+		AdditionalBuffersSize += AdditionalBuffer.Value.GetDataSize();
+	}
+
+	return sizeof(Mesh)
 		+ m_IndexBuffers.GetDataSize()
 		+ m_VertexBuffers.GetDataSize()
-		+ m_FaceBuffers.GetDataSize();
+		+ m_FaceBuffers.GetDataSize()
+		+ BonePoses.Num() * sizeof(FBonePose)
+		+ AdditionalBuffersSize;
 }
 
-
 //---------------------------------------------------------------------------------------------
 bool Mesh::HasCompatibleFormat( const Mesh* pOther ) const
 {
diff --git a/Engine/Plugins/Experimental/Mutable/Source/MutableRuntime/Private/MuR/MeshBufferSet.cpp b/Engine/Plugins/Experimental/Mutable/Source/MutableRuntime/Private/MuR/MeshBufferSet.cpp
--- a/Engine/Plugins/Experimental/Mutable/Source/MutableRuntime/Private/MuR/MeshBufferSet.cpp
+++ b/Engine/Plugins/Experimental/Mutable/Source/MutableRuntime/Private/MuR/MeshBufferSet.cpp
@@ -84,17 +84,23 @@ namespace mu
 
 
 	//---------------------------------------------------------------------------------------------
-	void FMeshBufferSet::SetElementCount(int32 count)
+	void FMeshBufferSet::SetElementCount(int32 Count)
 	{
-		check(count >= 0);
+		check(Count >= 0);
 		LLM_SCOPE_BYNAME(TEXT("MutableRuntime"));
 
-		m_elementCount = count;
+		// If the new size is 0, allow shrinking.
+		// TODO: Add better shrink policy or let the user decide. Denying shrinking
+		// unconditionally could mean having small meshes that use lots of memory. For now
+		// allow it if no other allocation will be done.  
+		const bool bAllowShrinking = Count == 0;
 
 		for (MESH_BUFFER& buf : m_buffers)
 		{
-			buf.m_data.SetNumUninitialized(buf.m_elementSize * count, false);
+			buf.m_data.SetNumUninitialized(buf.m_elementSize * Count, bAllowShrinking);
 		}
+		
+		m_elementCount = Count;
 	}
 
 
@@ -407,6 +413,17 @@ namespace mu
 		return res;
 	}
 
+	int32 FMeshBufferSet::GetAllocatedSize() const
+	{
+		int32 ByteCount = 0;
+
+		for (int32 BufferIndex = 0; BufferIndex < m_buffers.Num(); ++BufferIndex)
+		{
+			ByteCount += m_buffers[BufferIndex].m_data.GetAllocatedSize();
+		}
+
+		return ByteCount;
+	}
 
 	//-----------------------------------------------------------------------------------------
 	void FMeshBufferSet::CopyElement(uint32_t fromIndex, uint32_t toIndex)
diff --git a/Engine/Plugins/Experimental/Mutable/Source/MutableRuntime/Private/MuR/OpImageProject.cpp b/Engine/Plugins/Experimental/Mutable/Source/MutableRuntime/Private/MuR/OpImageProject.cpp
--- a/Engine/Plugins/Experimental/Mutable/Source/MutableRuntime/Private/MuR/OpImageProject.cpp
+++ b/Engine/Plugins/Experimental/Mutable/Source/MutableRuntime/Private/MuR/OpImageProject.cpp
@@ -2682,11 +2682,12 @@ void MeshProject_Optimised_Wrapping( const Mesh* pMesh,
 
 
 //-------------------------------------------------------------------------------------------------
-MeshPtr MeshProject_Optimised( const Mesh* pMesh,
-                               const FProjector& projector )
+void MeshProject_Optimised(Mesh* Result, const Mesh* pMesh, const FProjector& projector, bool& bOutSuccess)
 {
 	MUTABLE_CPUPROFILER_SCOPE(MeshProject_Optimised);
 
+	bOutSuccess = true;
+
     FVector3f projectorPosition,projectorDirection,projectorSide,projectorUp;
     projectorPosition = FVector3f( projector.position[0], projector.position[1], projector.position[2] );
     projector.GetDirectionSideUp( projectorDirection, projectorSide, projectorUp );
@@ -2695,7 +2696,6 @@ MeshPtr MeshProject_Optimised( const Mesh* pMesh,
     // Create with worse case, shrink later.
     int vertexCount = pMesh->GetVertexCount();
     int indexCount = pMesh->GetIndexCount();
-    MeshPtr pResult;
     int currentVertex = 0;
     int currentIndex = 0;
 
@@ -2708,41 +2708,41 @@ MeshPtr MeshProject_Optimised( const Mesh* pMesh,
     case PROJECTOR_TYPE::PLANAR:
     case PROJECTOR_TYPE::CYLINDRICAL:
     {
-        pResult = CreateMeshOptimisedForProjection( layout );
-        pResult->GetVertexBuffers().SetElementCount(vertexCount);
-        pResult->GetIndexBuffers().SetElementCount(indexCount);
-		uint32* pResultIndices = reinterpret_cast<uint32*>( pResult->GetIndexBuffers().GetBufferData(0) );
-		OPTIMISED_VERTEX* pResultVertices = reinterpret_cast<OPTIMISED_VERTEX*>( pResult->GetVertexBuffers().GetBufferData(0) );
+        CreateMeshOptimisedForProjection(Result, layout);
+        Result->GetVertexBuffers().SetElementCount(vertexCount);
+        Result->GetIndexBuffers().SetElementCount(indexCount);
+		uint32* pResultIndices = reinterpret_cast<uint32*>(Result->GetIndexBuffers().GetBufferData(0));
+		OPTIMISED_VERTEX* pResultVertices = reinterpret_cast<OPTIMISED_VERTEX*>(Result->GetVertexBuffers().GetBufferData(0));
 
         // Get the vertices
-        check( pMesh->GetVertexBuffers().GetElementSize(0)==sizeof(OPTIMISED_VERTEX) );
-		const OPTIMISED_VERTEX* pVertices = reinterpret_cast<const OPTIMISED_VERTEX*>( pMesh->GetVertexBuffers().GetBufferData(0) );
+        check(pMesh->GetVertexBuffers().GetElementSize(0)==sizeof(OPTIMISED_VERTEX));
+		const OPTIMISED_VERTEX* pVertices = reinterpret_cast<const OPTIMISED_VERTEX*>(pMesh->GetVertexBuffers().GetBufferData(0));
 
         // Get the indices
-        check( pMesh->GetIndexBuffers().GetElementSize(0)==4 );
-		const uint32* pIndices = reinterpret_cast<const uint32*>( pMesh->GetIndexBuffers().GetBufferData(0) );
+        check(pMesh->GetIndexBuffers().GetElementSize(0) == 4);
+		const uint32* pIndices = reinterpret_cast<const uint32*>(pMesh->GetIndexBuffers().GetBufferData(0));
         int faceCount = pMesh->GetFaceCount();
 
         if (projector.type==PROJECTOR_TYPE::PLANAR)
         {
-            MeshProject_Optimised_Planar( pVertices, vertexCount,
+            MeshProject_Optimised_Planar(pVertices, vertexCount,
                                           pIndices, faceCount,
                                           projectorPosition, projectorDirection,
                                           projectorSide, projectorUp,
                                           projectorScale,
                                           pResultVertices, currentVertex,
-                                          pResultIndices, currentIndex  );
+                                          pResultIndices, currentIndex);
         }
 
         else
         {
-            MeshProject_Optimised_Cylindrical( pVertices, vertexCount,
-                                               pIndices, faceCount,
-                                               projectorPosition, projectorDirection,
-                                               projectorSide, projectorUp,
-                                               projectorScale,
-                                               pResultVertices, currentVertex,
-                                               pResultIndices, currentIndex  );
+            MeshProject_Optimised_Cylindrical(pVertices, vertexCount,
+                                              pIndices, faceCount,
+                                              projectorPosition, projectorDirection,
+                                              projectorSide, projectorUp,
+                                              projectorScale,
+                                              pResultVertices, currentVertex,
+                                              pResultIndices, currentIndex);
         }
 
         break;
@@ -2750,18 +2750,18 @@ MeshPtr MeshProject_Optimised( const Mesh* pMesh,
 
     case PROJECTOR_TYPE::WRAPPING:
 	{
-        pResult = CreateMeshOptimisedForWrappingProjection( layout );
-        pResult->GetVertexBuffers().SetElementCount(vertexCount);
-        pResult->GetIndexBuffers().SetElementCount(indexCount);
-		uint32* pResultIndices = reinterpret_cast<uint32*>( pResult->GetIndexBuffers().GetBufferData(0) );
-		OPTIMISED_VERTEX_WRAPPING* pResultVertices = reinterpret_cast<OPTIMISED_VERTEX_WRAPPING*>( pResult->GetVertexBuffers().GetBufferData(0) );
-
-        MeshProject_Optimised_Wrapping( pMesh,
-                                        projectorPosition, projectorDirection,
-                                        projectorSide, projectorUp,
-                                        projectorScale,
-                                        pResultVertices, currentVertex,
-                                        pResultIndices, currentIndex  );
+        CreateMeshOptimisedForWrappingProjection(Result, layout);
+        Result->GetVertexBuffers().SetElementCount(vertexCount);
+        Result->GetIndexBuffers().SetElementCount(indexCount);
+		uint32* pResultIndices = reinterpret_cast<uint32*>(Result->GetIndexBuffers().GetBufferData(0));
+		OPTIMISED_VERTEX_WRAPPING* pResultVertices = reinterpret_cast<OPTIMISED_VERTEX_WRAPPING*>(Result->GetVertexBuffers().GetBufferData(0));
+
+        MeshProject_Optimised_Wrapping(pMesh,
+                                       projectorPosition, projectorDirection,
+                                       projectorSide, projectorUp,
+                                       projectorScale,
+                                       pResultVertices, currentVertex,
+                                       pResultIndices, currentIndex);
 
         break;
 	}
@@ -2769,16 +2769,15 @@ MeshPtr MeshProject_Optimised( const Mesh* pMesh,
     default:
         // Projector type not implemented.
         check(false);
+		bOutSuccess = false;
         break;
 
     }
 
     // Shrink result mesh
-    pResult->GetVertexBuffers().SetElementCount(currentVertex);
-    pResult->GetIndexBuffers().SetElementCount(currentIndex);
-    pResult->GetFaceBuffers().SetElementCount(currentIndex/3);
-
-    return pResult;
+    Result->GetVertexBuffers().SetElementCount(currentVertex);
+    Result->GetIndexBuffers().SetElementCount(currentIndex);
+    Result->GetFaceBuffers().SetElementCount(currentIndex/3);
 }
 #ifdef DEBUG_PROJECTION
 UE_ENABLE_OPTIMIZATION
@@ -2787,47 +2786,41 @@ UE_ENABLE_OPTIMIZATION
 
 
 //-------------------------------------------------------------------------------------------------
-MeshPtr MeshProject( const Mesh* pMesh,
-                         const FProjector& projector )
+void MeshProject(Mesh* Result, const Mesh* pMesh, const FProjector& projector, bool& bOutSuccess)
 {
 	MUTABLE_CPUPROFILER_SCOPE(MeshProject);
 
-    MeshPtr pResult;
-
-    if ( pMesh->m_staticFormatFlags & (1<<SMF_PROJECT) )
+    if (pMesh->m_staticFormatFlags & (1<<SMF_PROJECT))
     {
         // Mesh-optimised version
-        pResult = MeshProject_Optimised( pMesh, projector );
+        MeshProject_Optimised(Result, pMesh, projector, bOutSuccess);
     }
-    else if ( pMesh->m_staticFormatFlags & (1<<SMF_PROJECTWRAPPING) )
+    else if (pMesh->m_staticFormatFlags & (1<<SMF_PROJECTWRAPPING))
     {
         // Mesh-optimised version for wrapping projectors
         // \todo: make sure the projector is a wrapping projector
-        pResult = MeshProject_Optimised( pMesh, projector );
+        MeshProject_Optimised(Result, pMesh, projector, bOutSuccess);
     }
     else
     {
         check(false);
     }
 
-	if (pResult)
+	if (bOutSuccess)
 	{
-		pResult->m_surfaces.SetNum(0);
-		pResult->EnsureSurfaceData();
-		pResult->ResetStaticFormatFlags();
+		Result->m_surfaces.SetNum(0);
+		Result->EnsureSurfaceData();
+		Result->ResetStaticFormatFlags();
 	}
-
-    return pResult;
 }
 
 
 
 //-------------------------------------------------------------------------------------------------
-MeshPtr CreateMeshOptimisedForProjection( int layout )
+void CreateMeshOptimisedForProjection(Mesh* Result, int layout)
 {
-    MeshPtr pFormatMesh = new Mesh();
-    pFormatMesh->GetVertexBuffers().SetBufferCount( 1 );
-    pFormatMesh->GetIndexBuffers().SetBufferCount( 1 );
+    Result->GetVertexBuffers().SetBufferCount( 1 );
+    Result->GetIndexBuffers().SetBufferCount( 1 );
 
     MESH_BUFFER_SEMANTIC semantics[3] =	{ MBS_TEXCOORDS,	MBS_POSITION,	MBS_NORMAL };
     int semanticIndices[3] =			{ 0,				0,				0 };
@@ -2835,7 +2828,7 @@ MeshPtr CreateMeshOptimisedForProjection( int layout )
     int componentCounts[3] =			{ 2,				3,				3 };
     int offsets[3] =					{ 0,				8,				20 };
     semanticIndices[0] = layout;
-    pFormatMesh->GetVertexBuffers().SetBuffer
+    Result->GetVertexBuffers().SetBuffer
             ( 0, 32, 3,
               semantics, semanticIndices,
               formats, componentCounts,
@@ -2846,22 +2839,19 @@ MeshPtr CreateMeshOptimisedForProjection( int layout )
     MESH_BUFFER_FORMAT iformats[1] =		{ MBF_UINT32 };
     int icomponentCounts[1] =				{ 1 };
     int ioffsets[1] =						{ 0 };
-    pFormatMesh->GetIndexBuffers().SetBuffer
+    Result->GetIndexBuffers().SetBuffer
             ( 0, 4, 1,
               isemantics, isemanticIndices,
               iformats, icomponentCounts,
               ioffsets );
-
-    return pFormatMesh;
 }
 
 
 //-------------------------------------------------------------------------------------------------
-MeshPtr CreateMeshOptimisedForWrappingProjection( int layout )
+void CreateMeshOptimisedForWrappingProjection(Mesh* Result, int layout)
 {
-    MeshPtr pFormatMesh = new Mesh();
-    pFormatMesh->GetVertexBuffers().SetBufferCount( 1 );
-    pFormatMesh->GetIndexBuffers().SetBufferCount( 1 );
+    Result->GetVertexBuffers().SetBufferCount( 1 );
+    Result->GetIndexBuffers().SetBufferCount( 1 );
 
     MESH_BUFFER_SEMANTIC semantics[4] =	{ MBS_TEXCOORDS,	MBS_POSITION,	MBS_NORMAL,     MBS_LAYOUTBLOCK };
     int semanticIndices[4] =			{ 0,				0,				0,              0 };
@@ -2870,7 +2860,7 @@ MeshPtr CreateMeshOptimisedForWrappingProjection( int layout )
     int offsets[4] =					{ 0,				8,				20,             32 };
     semanticIndices[0] = layout;
     semanticIndices[3] = layout;
-    pFormatMesh->GetVertexBuffers().SetBuffer
+    Result->GetVertexBuffers().SetBuffer
             ( 0, 36, 4,
               semantics, semanticIndices,
               formats, componentCounts,
@@ -2881,13 +2871,11 @@ MeshPtr CreateMeshOptimisedForWrappingProjection( int layout )
     MESH_BUFFER_FORMAT iformats[1] =		{ MBF_UINT32 };
     int icomponentCounts[1] =				{ 1 };
     int ioffsets[1] =						{ 0 };
-    pFormatMesh->GetIndexBuffers().SetBuffer
+    Result->GetIndexBuffers().SetBuffer
             ( 0, 4, 1,
               isemantics, isemanticIndices,
               iformats, icomponentCounts,
               ioffsets );
-
-    return pFormatMesh;
 }
 
 }
diff --git a/Engine/Plugins/Experimental/Mutable/Source/MutableRuntime/Private/MuR/OpImageProject.h b/Engine/Plugins/Experimental/Mutable/Source/MutableRuntime/Private/MuR/OpImageProject.h
--- a/Engine/Plugins/Experimental/Mutable/Source/MutableRuntime/Private/MuR/OpImageProject.h
+++ b/Engine/Plugins/Experimental/Mutable/Source/MutableRuntime/Private/MuR/OpImageProject.h
@@ -47,10 +47,9 @@ namespace mu
 	extern float ComputeProjectedFootprintBestMip(
 			const Mesh* pMesh, const FProjector& Projector, const FVector2f& TargetSize, const FVector2f& SourceSize);
 
-    extern MeshPtr MeshProject( const Mesh* pMesh,
-                                const FProjector& projector );
+    extern void MeshProject(Mesh* Result, const Mesh* pMesh, const FProjector& projector, bool& bOutSuccess);
 
-	MUTABLERUNTIME_API extern MeshPtr CreateMeshOptimisedForProjection( int layout );
-	MUTABLERUNTIME_API extern MeshPtr CreateMeshOptimisedForWrappingProjection( int layout );
+	MUTABLERUNTIME_API extern void CreateMeshOptimisedForProjection(Mesh* Result, int layout);
+	MUTABLERUNTIME_API extern void CreateMeshOptimisedForWrappingProjection(Mesh* Result, int layout);
 
 }
diff --git a/Engine/Plugins/Experimental/Mutable/Source/MutableRuntime/Private/MuR/OpMeshApplyPose.h b/Engine/Plugins/Experimental/Mutable/Source/MutableRuntime/Private/MuR/OpMeshApplyPose.h
--- a/Engine/Plugins/Experimental/Mutable/Source/MutableRuntime/Private/MuR/OpMeshApplyPose.h
+++ b/Engine/Plugins/Experimental/Mutable/Source/MutableRuntime/Private/MuR/OpMeshApplyPose.h
@@ -15,16 +15,17 @@ namespace mu
     //---------------------------------------------------------------------------------------------
     //! Reference version
     //---------------------------------------------------------------------------------------------
-    inline MeshPtr MeshApplyPose(const Mesh* pBase, const Mesh* pPose)
+    inline void MeshApplyPose(Mesh* Result, const Mesh* pBase, const Mesh* pPose, bool& bOutSuccess)
     {
         MUTABLE_CPUPROFILER_SCOPE(MeshApplyPose);
 
-        MeshPtr pDest = pBase->Clone();
+		bOutSuccess = true;
 
 		mu::SkeletonPtrConst pSkeleton = pBase->GetSkeleton();
 		if (!pSkeleton)
 		{
-			return pDest;
+			bOutSuccess = false;
+			return;
 		}
 
         // We assume the matrices are transforms from the binding pose bone to the new pose
@@ -53,25 +54,31 @@ namespace mu
         }
 
         // Get pointers to vertex position data
-		MeshBufferIteratorConst<MBF_FLOAT32, float, 3> itSource( pBase->m_VertexBuffers, MBS_POSITION, 0 );
-		MeshBufferIterator<MBF_FLOAT32, float, 3> itTarget( pDest->m_VertexBuffers, MBS_POSITION, 0 );
-        if (!itSource.ptr() || !itTarget.ptr())
+		MeshBufferIteratorConst<MBF_FLOAT32, float, 3> itSource(pBase->m_VertexBuffers, MBS_POSITION, 0);
+        if (!itSource.ptr())
         {
             // Formats not implemented
             check(false);
-            return pDest;
+			bOutSuccess = false;
+            return;
         }
 
         // Get pointers to skinning data
-		UntypedMeshBufferIteratorConst itBoneIndices( pBase->m_VertexBuffers, MBS_BONEINDICES, 0 );
-		UntypedMeshBufferIteratorConst itBoneWeights( pBase->m_VertexBuffers, MBS_BONEWEIGHTS, 0 );
+		UntypedMeshBufferIteratorConst itBoneIndices(pBase->m_VertexBuffers, MBS_BONEINDICES, 0);
+		UntypedMeshBufferIteratorConst itBoneWeights(pBase->m_VertexBuffers, MBS_BONEWEIGHTS, 0);
         if (!itBoneIndices.ptr() || !itBoneWeights.ptr())
         {
             // No skinning data
             check(false);
-            return pDest;
+			bOutSuccess = false;
+            return;
         }
 
+		Result->CopyFrom(*pBase);
+
+		MeshBufferIterator<MBF_FLOAT32, float, 3> itTarget(Result->m_VertexBuffers, MBS_POSITION, 0);
+		check(itTarget.ptr());
+
         // Proceed
         int vertexCount = pBase->GetVertexCount();
         int weightCount = itBoneIndices.GetComponents();
@@ -120,9 +127,6 @@ namespace mu
             ++itBoneIndices;
             ++itBoneWeights;
         }
-
-
-        return pDest;
     }
 
 }
diff --git a/Engine/Plugins/Experimental/Mutable/Source/MutableRuntime/Private/MuR/OpMeshBind.h b/Engine/Plugins/Experimental/Mutable/Source/MutableRuntime/Private/MuR/OpMeshBind.h
--- a/Engine/Plugins/Experimental/Mutable/Source/MutableRuntime/Private/MuR/OpMeshBind.h
+++ b/Engine/Plugins/Experimental/Mutable/Source/MutableRuntime/Private/MuR/OpMeshBind.h
@@ -881,16 +881,20 @@ namespace mu
 	//---------------------------------------------------------------------------------------------
     //! Generate the mesh-shape binding data
     //---------------------------------------------------------------------------------------------
-    inline MeshPtr MeshBindShapeReshape(
+    inline void MeshBindShapeReshape(
+			Mesh* Result,
 			const Mesh* BaseMesh, const Mesh* ShapeMesh, 
 			const TArray<string>& BonesToDeform, const TArray<string>& PhysicsToDeform, 
-			EMeshBindShapeFlags BindFlags)
+			EMeshBindShapeFlags BindFlags,
+			bool& bOutSuccess)
     {
 		MUTABLE_CPUPROFILER_SCOPE(MeshBindShape);
-		
+		bOutSuccess = true;
+
 		if (!BaseMesh)
 		{
-			return nullptr;
+			bOutSuccess = false;
+			return;
 		}
 
 		const bool bReshapeVertices = EnumHasAnyFlags(BindFlags, EMeshBindShapeFlags::ReshapeVertices);
@@ -905,19 +909,22 @@ namespace mu
 
 		if (!bReshapeVertices && !bSkeletonModification && !bPhysicsModification)
 		{
-			return nullptr;
+			bOutSuccess = false;
+			return;
 		}
 
 		if (!ShapeMesh)
 		{
-			return BaseMesh->Clone();
+			bOutSuccess = false;
+			return;
 		}
 	
 		int32 ShapeVertexCount = ShapeMesh->GetVertexCount();
 		int ShapeTriangleCount = ShapeMesh->GetFaceCount();
 		if (!ShapeVertexCount || !ShapeTriangleCount)
 		{
-			return BaseMesh->Clone();
+			bOutSuccess = false;
+			return;
 		}	
 
 		FShapeMeshDescriptorBind ShapeMeshDescriptor;
@@ -972,23 +979,21 @@ namespace mu
 			ShapeMeshTree.Build();
 		}
 	
-		Ptr<Mesh> Result;
-		
 		// If no vertices are needed, it is assumed we only want to reshape physics or skeleton.
 		// In that case, remove everything except physics bodies, the skeleton and pose.
 		if (!bReshapeVertices)
 		{
-			constexpr EMeshCloneFlags CloneFlags = 
-					EMeshCloneFlags::WithSkeleton    | 
-					EMeshCloneFlags::WithPhysicsBody | 
-				EMeshCloneFlags::WithPoses |
-				EMeshCloneFlags::WithAdditionalPhysics;
+			constexpr EMeshCopyFlags CopyFlags = 
+				EMeshCopyFlags::WithSkeleton    | 
+				EMeshCopyFlags::WithPhysicsBody | 
+				EMeshCopyFlags::WithPoses |
+				EMeshCopyFlags::WithAdditionalPhysics;
 
-			Result = BaseMesh->Clone(CloneFlags);
+			Result->CopyFrom(*BaseMesh, CopyFlags);
 		}
 		else
 		{	
-			Result = BaseMesh->Clone();
+			Result->CopyFrom(*BaseMesh);
 		}
 
 		int32 BindingDataIndex = 0;
@@ -1015,7 +1020,7 @@ namespace mu
 			MUTABLE_CPUPROFILER_SCOPE(BindSkeleton);
 
 			TTuple<TArray<FReshapePointBindingData>, TArray<int32>> SkeletonBindingData = 
-					BindPose(Result.get(), ShapeMeshTree, BonesToDeform);
+					BindPose(Result, ShapeMeshTree, BonesToDeform);
 			const TArray<FReshapePointBindingData>& SkeletonBindDataArray = SkeletonBindingData.Get<0>();
 			const TArray<int32>& BoneIndices = SkeletonBindingData.Get<1>();
 
@@ -1064,7 +1069,7 @@ namespace mu
 			}
 
 			TTuple<TArray<FReshapePointBindingData>, TArray<int32>, TArray<int32>> PhysicsBindingData = 
-					BindPhysicsBodies(PhysicsBodiesToBind, ShapeMeshTree, Result.get(), PhysicsToDeform);
+					BindPhysicsBodies(PhysicsBodiesToBind, ShapeMeshTree, Result, PhysicsToDeform);
 
 			const TArray<FReshapePointBindingData>& PhysicsBindDataArray = PhysicsBindingData.Get<0>();
 			const TArray<int32>& DeformedBodyIndices = PhysicsBindingData.Get<1>();
@@ -1096,9 +1101,7 @@ namespace mu
 			Result->m_AdditionalBuffers.Emplace(EMeshBufferType::PhysicsBodyDeformBinding, MoveTemp(PhysicsBodyBuffer));
 			Result->m_AdditionalBuffers.Emplace(EMeshBufferType::PhysicsBodyDeformSelection, MoveTemp(PhysicsBodySelectionBuffer));
 			Result->m_AdditionalBuffers.Emplace(EMeshBufferType::PhysicsBodyDeformOffsets, MoveTemp(PhysicsBodySelectionOffsetsBuffer));
-		}
-		
-        return Result;
+		}	
     }
 
 	//---------------------------------------------------------------------------------------------
@@ -1322,25 +1325,30 @@ namespace mu
 		return BindData;
 	}
 
-    inline MeshPtr MeshBindShapeClipDeform(const Mesh* BaseMesh, const Mesh* ShapeMesh, EShapeBindingMethod BindingMethod)
+    inline void MeshBindShapeClipDeform(Mesh* Result, const Mesh* BaseMesh, const Mesh* ShapeMesh, EShapeBindingMethod BindingMethod, bool& bOutSuccess)
     {
 		MUTABLE_CPUPROFILER_SCOPE(MeshBindShapeClipDeform);
-		
+	
+		bOutSuccess = true;
+
 		if (!BaseMesh)
 		{
-			return nullptr;
+			bOutSuccess = false;
+			return;
 		}
 
 		if (!ShapeMesh)
 		{
-			return BaseMesh->Clone();
+			bOutSuccess = false;
+			return;
 		}
 	
 		int32 ShapeVertexCount = ShapeMesh->GetVertexCount();
 		int ShapeTriangleCount = ShapeMesh->GetFaceCount();
 		if (!ShapeVertexCount || !ShapeTriangleCount)
 		{
-			return BaseMesh->Clone();
+			bOutSuccess = false;
+			return;
 		}	
 
 		FShapeMeshDescriptorBind ShapeMeshDescriptor;
@@ -1395,7 +1403,7 @@ namespace mu
 			ShapeMeshTree.Build();
 		}
 	
-		Ptr<Mesh> Result = BaseMesh->Clone();
+		Result->CopyFrom(*BaseMesh);
 		TArray<FClipDeformVertexBindingData> VerticesBindData = BindVerticesClipDeform(BaseMesh, ShapeMeshTree, BindingMethod);
 			
 		// Add the binding information to the mesh
@@ -1412,6 +1420,5 @@ namespace mu
 		VB.SetBuffer(NewBufferIndex, sizeof(FClipDeformVertexBindingData), BufDesc.Channels, BufDesc.Semantics, BufDesc.SemanticIndices, BufDesc.Formats, BufDesc.Components, BufDesc.Offsets);
 		FMemory::Memcpy(VB.GetBufferData(NewBufferIndex), VerticesBindData.GetData(), VerticesBindData.Num()*sizeof(FClipDeformVertexBindingData));
 		
-        return Result;
     }
 }
diff --git a/Engine/Plugins/Experimental/Mutable/Source/MutableRuntime/Private/MuR/OpMeshClipDeform.h b/Engine/Plugins/Experimental/Mutable/Source/MutableRuntime/Private/MuR/OpMeshClipDeform.h
--- a/Engine/Plugins/Experimental/Mutable/Source/MutableRuntime/Private/MuR/OpMeshClipDeform.h
+++ b/Engine/Plugins/Experimental/Mutable/Source/MutableRuntime/Private/MuR/OpMeshClipDeform.h
@@ -99,17 +99,19 @@ namespace mu
 		return true;
 	}
 
-	inline MeshPtr MeshClipDeform(const Mesh* BaseMesh, const Mesh* ShapeMesh, const float ClipWeightThreshold )
+	inline void MeshClipDeform(Mesh* Result, const Mesh* BaseMesh, const Mesh* ShapeMesh, const float ClipWeightThreshold, bool& bOutSuccess )
 	{
 		MUTABLE_CPUPROFILER_SCOPE(ClipDeform);
+		bOutSuccess = true;
 
 		if (!BaseMesh)
 		{
-			return nullptr;
+			bOutSuccess = false;
+			return;
 		}
 
-		constexpr EMeshCloneFlags MeshCloneFlags = ~EMeshCloneFlags::WithVertexBuffers;
-		MeshPtr Result = BaseMesh->Clone(MeshCloneFlags);
+		constexpr EMeshCopyFlags MeshCopyFlags = ~EMeshCopyFlags::WithVertexBuffers;
+		Result->CopyFrom(*BaseMesh, MeshCopyFlags);
 	
 		int BarycentricDataBuffer = 0;
 		int BarycentricDataChannel = 0;
@@ -132,12 +134,12 @@ namespace mu
 
 		if (BarycentricDataBuffer < 0)
 		{
-			return Result;
+			return; 
 		}
 
 		if (!ShapeMesh)
 		{
-			return Result;
+			return; 
 		}
 
 		// \TODO: More checks
@@ -151,7 +153,7 @@ namespace mu
 		int ShapeTriangleCount = ShapeMesh->GetFaceCount();
 		if (!ShapeVertexCount || !ShapeTriangleCount)
 		{
-			return Result;
+			return; 
 		}
 
 		TArray<FVector3f> ShapePositions;
@@ -412,7 +414,7 @@ namespace mu
 
 			if ( NumRemainingVerts == NumOriginalVerts )
 			{
-				return Result;
+				return;
 			}
 		
 			// We are guaranteed to find a removed vertex.
@@ -452,6 +454,5 @@ namespace mu
 
 			Result->GetVertexBuffers().SetElementCount( DestEnd );
 		}
-		return Result;
 	}
 }
diff --git a/Engine/Plugins/Experimental/Mutable/Source/MutableRuntime/Private/MuR/OpMeshClipMorphPlane.h b/Engine/Plugins/Experimental/Mutable/Source/MutableRuntime/Private/MuR/OpMeshClipMorphPlane.h
--- a/Engine/Plugins/Experimental/Mutable/Source/MutableRuntime/Private/MuR/OpMeshClipMorphPlane.h
+++ b/Engine/Plugins/Experimental/Mutable/Source/MutableRuntime/Private/MuR/OpMeshClipMorphPlane.h
@@ -69,9 +69,10 @@ namespace mu
 	//---------------------------------------------------------------------------------------------
 	//! Reference version
 	//---------------------------------------------------------------------------------------------
-	inline MeshPtr MeshClipMorphPlane(const Mesh* pBase, const vec3f& origin, const vec3f& normal, float dist, float factor, float radius,
-		float radius2, float angle, const FShape& selectionShape, const mu::string& boneName = mu::string(), float vertexSelectionBoneMaxRadius = -1.f)
+	inline void MeshClipMorphPlane(Mesh* Result, const Mesh* pBase, const vec3f& origin, const vec3f& normal, float dist, float factor, float radius,
+		float radius2, float angle, const FShape& selectionShape, bool& bOutSuccess, const mu::string& boneName = mu::string(), float vertexSelectionBoneMaxRadius = -1.f)
 	{
+		bOutSuccess = true;
 		//float radius = 8.f;
 		//float radius2 = 4.f;
 		//float factor = 1.f;
@@ -87,13 +88,12 @@ namespace mu
 		vec3f origin_radius_vector = cross(normal, aux_base);		// PERPENDICULAR VECTOR TO THE PLANE normal and aux base
 		check(fabs(dot(normal, origin_radius_vector)) < 0.05f);
 
-		MeshPtr pDest = pBase->Clone();
-
         uint32 vcount = pBase->GetVertexBuffers().GetElementCount();
 
 		if (!vcount)
 		{
-			return pDest;
+			bOutSuccess = false;
+			return;
 		}
 
 		TArray<bool> bone_is_affected;
@@ -275,11 +275,13 @@ namespace mu
 			}
 		}
 
+		Result->CopyFrom(*pBase);
+
 		// TODO: Replace with an array of bools?
         TArray<bool> RemovedVertices;
-		RemovedVertices.Init(false,pDest->GetVertexCount());
+		RemovedVertices.Init(false, Result->GetVertexCount());
 
-		const FMeshBufferSet& MBSPriv = pDest->GetVertexBuffers();
+		const FMeshBufferSet& MBSPriv = Result->GetVertexBuffers();
 		for (int32 b = 0; b < MBSPriv.m_buffers.Num(); ++b)
 		{
 			for (int32 c = 0; c < MBSPriv.m_buffers[b].m_channels.Num(); ++c)
@@ -287,7 +289,7 @@ namespace mu
 				MESH_BUFFER_SEMANTIC sem = MBSPriv.m_buffers[b].m_channels[c].m_semantic;
 				int semIndex = MBSPriv.m_buffers[b].m_channels[c].m_semanticIndex;
 
-				UntypedMeshBufferIterator it(pDest->GetVertexBuffers(), sem, semIndex);
+				UntypedMeshBufferIterator it(Result->GetVertexBuffers(), sem, semIndex);
 
 				switch (sem)
 				{
@@ -367,11 +369,11 @@ namespace mu
 		}
 
 		// Now remove all the faces from the result mesh that have all vertices removed
-		UntypedMeshBufferIteratorConst itBase(pDest->GetIndexBuffers(), MBS_VERTEXINDEX);
-		UntypedMeshBufferIterator itDest(pDest->GetIndexBuffers(), MBS_VERTEXINDEX);
-		int aFaceCount = pDest->GetFaceCount();
+		UntypedMeshBufferIteratorConst itBase(Result->GetIndexBuffers(), MBS_VERTEXINDEX);
+		UntypedMeshBufferIterator itDest(Result->GetIndexBuffers(), MBS_VERTEXINDEX);
+		int aFaceCount = Result->GetFaceCount();
 
-		UntypedMeshBufferIteratorConst ito(pDest->GetIndexBuffers(), MBS_VERTEXINDEX);
+		UntypedMeshBufferIteratorConst ito(Result->GetIndexBuffers(), MBS_VERTEXINDEX);
 		for (int f = 0; f < aFaceCount; ++f)
 		{
             vec3<uint32> ov;
@@ -397,21 +399,19 @@ namespace mu
 		std::size_t removedIndices = itBase - itDest;
 		check(removedIndices % 3 == 0);
 
-		pDest->GetFaceBuffers().SetElementCount(aFaceCount - (int)removedIndices / 3);
-		pDest->GetIndexBuffers().SetElementCount(aFaceCount * 3 - (int)removedIndices);
+		Result->GetFaceBuffers().SetElementCount(aFaceCount - (int)removedIndices / 3);
+		Result->GetIndexBuffers().SetElementCount(aFaceCount * 3 - (int)removedIndices);
 
 		// TODO: Should redo/reorder the face buffer before SetElementCount since some deleted faces could be left and some remaining faces deleted.
 
         // Fix the surface data if present.
-        if (pDest->m_surfaces.Num())
+        if (Result->m_surfaces.Num())
         {
             // We assume there will be only one.
-            check(pDest->m_surfaces.Num()==1);
+            check(Result->m_surfaces.Num()==1);
 
-            pDest->m_surfaces[0].m_indexCount -= (int32)removedIndices;
+            Result->m_surfaces[0].m_indexCount -= (int32)removedIndices;
         }
-
-		return pDest;
 	}
 
 }
diff --git a/Engine/Plugins/Experimental/Mutable/Source/MutableRuntime/Private/MuR/OpMeshClipWithMesh.h b/Engine/Plugins/Experimental/Mutable/Source/MutableRuntime/Private/MuR/OpMeshClipWithMesh.h
--- a/Engine/Plugins/Experimental/Mutable/Source/MutableRuntime/Private/MuR/OpMeshClipWithMesh.h
+++ b/Engine/Plugins/Experimental/Mutable/Source/MutableRuntime/Private/MuR/OpMeshClipWithMesh.h
@@ -594,27 +594,29 @@ namespace mu
     //---------------------------------------------------------------------------------------------
     //! Reference version
     //---------------------------------------------------------------------------------------------
-    inline MeshPtr MeshClipWithMesh_Reference(const Mesh* pBase, const Mesh* pClipMesh)
+    inline void MeshClipWithMesh_Reference(Mesh* Result, const Mesh* pBase, const Mesh* pClipMesh, bool& bOutSuccess)
     {
         MUTABLE_CPUPROFILER_SCOPE(MeshClipWithMesh_Reference);
         
-		MeshPtr pDest = pBase->Clone();
+		bOutSuccess = true;
 
         uint32 vcount = pClipMesh->GetVertexBuffers().GetElementCount();
         if (!vcount)
         {
-            return pDest; // Since there's nothing to clip against return a copy of the original base mesh
+			bOutSuccess = false;
+			return; // Since there's nothing to clip against return a copy of the original base mesh
         }
 
         TArray<uint8> vertex_in_clip_mesh;  // Stores whether each vertex in the original mesh in in the clip mesh volume
-        MeshClipMeshClassifyVertices( vertex_in_clip_mesh, pBase, pClipMesh );
+        MeshClipMeshClassifyVertices(vertex_in_clip_mesh, pBase, pClipMesh);
 
+		Result->CopyFrom(*pBase);
         // Now remove all the faces from the result mesh that have all the vertices outside the clip volume
-        UntypedMeshBufferIteratorConst itBase(pDest->GetIndexBuffers(), MBS_VERTEXINDEX);
-        UntypedMeshBufferIterator itDest(pDest->GetIndexBuffers(), MBS_VERTEXINDEX);
-        int aFaceCount = pDest->GetFaceCount();
+        UntypedMeshBufferIteratorConst itBase(Result->GetIndexBuffers(), MBS_VERTEXINDEX);
+        UntypedMeshBufferIterator itDest(Result->GetIndexBuffers(), MBS_VERTEXINDEX);
+        int aFaceCount = Result->GetFaceCount();
 
-        UntypedMeshBufferIteratorConst ito(pDest->GetIndexBuffers(), MBS_VERTEXINDEX);
+        UntypedMeshBufferIteratorConst ito(Result->GetIndexBuffers(), MBS_VERTEXINDEX);
         for (int f = 0; f < aFaceCount; ++f)
         {
             vec3<uint32> ov;
@@ -640,48 +642,48 @@ namespace mu
         SIZE_T removedIndices = itBase - itDest;
         check(removedIndices % 3 == 0);
 
-        pDest->GetFaceBuffers().SetElementCount(aFaceCount - (int)removedIndices / 3);
-        pDest->GetIndexBuffers().SetElementCount(aFaceCount * 3 - (int)removedIndices);
+        Result->GetFaceBuffers().SetElementCount(aFaceCount - (int)removedIndices / 3);
+        Result->GetIndexBuffers().SetElementCount(aFaceCount * 3 - (int)removedIndices);
 
         // TODO: Should redo/reorder the face buffer before SetElementCount since some deleted faces could be left and some remaining faces deleted.
 
-        //if (pDest->GetFaceBuffers().GetElementCount() == 0)
+        //if (Result->GetFaceBuffers().GetElementCount() == 0)
         //{
-        //	pDest = pBase->Clone(); // If all faces have been discarded, return a copy of the unmodified mesh because unreal doesn't like empty meshes
+        //	Result->CopyFrom(pBase); // If all faces have been discarded, return a copy of the unmodified mesh because unreal doesn't like empty meshes
         //}
 
         // [jordi] Remove unused vertices. This is necessary to avoid returning a mesh with vertices and no faces, which screws some
         // engines like Unreal Engine 4.
-        MeshRemoveUnusedVertices( pDest.get() );
-
-        return pDest;
+        MeshRemoveUnusedVertices(Result);
     }
 
 
     //---------------------------------------------------------------------------------------------
     //! CoreGeometry version
     //---------------------------------------------------------------------------------------------
-    inline MeshPtr MeshClipWithMesh(const Mesh* pBase, const Mesh* pClipMesh)
+    inline void MeshClipWithMesh(Mesh* Result, const Mesh* pBase, const Mesh* pClipMesh, bool& bOutSuccess)
     {
         MUTABLE_CPUPROFILER_SCOPE(MeshClipWithMesh);
-
-        MeshPtr pDest = pBase->Clone();
+		bOutSuccess = true;
 
         uint32 VCount = pClipMesh->GetVertexBuffers().GetElementCount();
         if (!VCount)
         {
-            return pDest; // Since there's nothing to clip against return a copy of the original base mesh
+			bOutSuccess = false;
+			return; // OutSuccess false indicates the pBase can be reused in this case. 
         }
 
+		Result->CopyFrom(*pBase);
+
         TArray<uint8> VertexInClipMesh;  // Stores whether each vertex in the original mesh in in the clip mesh volume
-        MeshClipMeshClassifyVertices( VertexInClipMesh, pBase, pClipMesh );
+        MeshClipMeshClassifyVertices(VertexInClipMesh, pBase, pClipMesh);
 
         // Now remove all the faces from the result mesh that have all the vertices outside the clip volume
-        UntypedMeshBufferIteratorConst ItBase(pDest->GetIndexBuffers(), MBS_VERTEXINDEX);
-        UntypedMeshBufferIterator ItDest(pDest->GetIndexBuffers(), MBS_VERTEXINDEX);
-        int32 AFaceCount = pDest->GetFaceCount();
+        UntypedMeshBufferIteratorConst ItBase(Result->GetIndexBuffers(), MBS_VERTEXINDEX);
+        UntypedMeshBufferIterator ItDest(Result->GetIndexBuffers(), MBS_VERTEXINDEX);
+        int32 AFaceCount = Result->GetFaceCount();
 
-        UntypedMeshBufferIteratorConst Ito(pDest->GetIndexBuffers(), MBS_VERTEXINDEX);
+        UntypedMeshBufferIteratorConst Ito(Result->GetIndexBuffers(), MBS_VERTEXINDEX);
         for (int32 F = 0; F < AFaceCount; ++F)
         {
 			uint32 OV[3] = {0, 0, 0};
@@ -708,28 +710,26 @@ namespace mu
         SIZE_T RemovedIndices = ItBase - ItDest;
         check(RemovedIndices % 3 == 0);
 
-        pDest->GetFaceBuffers().SetElementCount(AFaceCount - (int32)RemovedIndices / 3);
-        pDest->GetIndexBuffers().SetElementCount(AFaceCount * 3 - (int32)RemovedIndices);
+        Result->GetFaceBuffers().SetElementCount(AFaceCount - (int32)RemovedIndices / 3);
+        Result->GetIndexBuffers().SetElementCount(AFaceCount * 3 - (int32)RemovedIndices);
 
         // TODO: Should redo/reorder the face buffer before SetElementCount since some deleted faces could be left and some remaining faces deleted.
 
-        //if (pDest->GetFaceBuffers().GetElementCount() == 0)
+        //if (Result->GetFaceBuffers().GetElementCount() == 0)
         //{
-        //	pDest = pBase->Clone(); // If all faces have been discarded, return a copy of the unmodified mesh because unreal doesn't like empty meshes
+        //	Result->CopyFrom(pBase); // If all faces have been discarded, return a copy of the unmodified mesh because unreal doesn't like empty meshes
         //}
 
         // [jordi] Remove unused vertices. This is necessary to avoid returning a mesh with vertices and no faces, which screws some
         // engines like Unreal Engine 4.
-        MeshRemoveUnusedVertices( pDest.get() );
-
-        return pDest;
+        MeshRemoveUnusedVertices(Result);
     }
 
 
     //---------------------------------------------------------------------------------------------
     //!
     //---------------------------------------------------------------------------------------------
-    inline MeshPtr CreateMask( MeshPtrConst pBase, const TArray<uint8>& excludedVertices )
+    inline void CreateMask(Mesh* Result, const Mesh* pBase, const TArray<uint8>& excludedVertices)
     {
 
         int maskVertexCount = 0;
@@ -738,12 +738,11 @@ namespace mu
             if (!b) ++maskVertexCount;
         }
 
-        MeshPtr pMask = new Mesh();
 
         // Create the vertex buffer
         {
-            pMask->GetVertexBuffers().SetElementCount( maskVertexCount );
-            pMask->GetVertexBuffers().SetBufferCount( 1 );
+            Result->GetVertexBuffers().SetElementCount( maskVertexCount );
+            Result->GetVertexBuffers().SetBufferCount( 1 );
 
 			// Vertex index channel
 			MESH_BUFFER_SEMANTIC semantic = MBS_VERTEXINDEX;
@@ -752,7 +751,7 @@ namespace mu
 			int components = 1;
 			int offsets = 0;
 
-            pMask->GetVertexBuffers().SetBuffer
+            Result->GetVertexBuffers().SetBuffer
                 (
                     0,
                     4,
@@ -765,7 +764,7 @@ namespace mu
                 );
         }
 
-        MeshBufferIterator<MBF_UINT32,uint32,1> itMask( pMask->GetVertexBuffers(), MBS_VERTEXINDEX, 0 );
+        MeshBufferIterator<MBF_UINT32,uint32,1> itMask( Result->GetVertexBuffers(), MBS_VERTEXINDEX, 0 );
         UntypedMeshBufferIteratorConst itBase( pBase->GetVertexBuffers(), MBS_VERTEXINDEX, 0 );
 
         for (size_t v = 0; v<excludedVertices.Num(); ++v)
@@ -779,21 +778,23 @@ namespace mu
             ++itBase;
         }
 
-        return pMask;
     }
 
 
     //---------------------------------------------------------------------------------------------
     //! Generate a mask mesh with the faces of the base mesh inside the clip mesh.
     //---------------------------------------------------------------------------------------------
-    inline MeshPtr MeshMaskClipMesh(const Mesh* pBase, const Mesh* pClipMesh )
+    inline void MeshMaskClipMesh(Mesh* Result, const Mesh* pBase, const Mesh* pClipMesh, bool& bOutSuccess)
     {
         MUTABLE_CPUPROFILER_SCOPE(MeshMaskClipMesh);
 
+		bOutSuccess = true;
+
         uint32 VCount = pClipMesh->GetVertexBuffers().GetElementCount();
         if (!VCount)
         {
-            return nullptr; // Since there's nothing to clip against return a copy of the original base mesh
+			bOutSuccess = false;
+			return;
         }
 
         TArray<uint8> VertexInClipMesh;  // Stores whether each vertex in the original mesh in in the clip mesh volume
@@ -826,20 +827,20 @@ namespace mu
             }
         }
 
-        return CreateMask( pBase, VertexWithFaceNotClipped );
+        CreateMask(Result, pBase, VertexWithFaceNotClipped);
     }
 
     //---------------------------------------------------------------------------------------------
     //! Generate a mask mesh with the faces of the base mesh inside the clip mesh.
     //---------------------------------------------------------------------------------------------
-    inline MeshPtr MeshMaskClipMesh_Reference(const Mesh* pBase, const Mesh* pClipMesh )
+    inline void MeshMaskClipMesh_Reference(Mesh* Result, const Mesh* pBase, const Mesh* pClipMesh )
     {
         MUTABLE_CPUPROFILER_SCOPE(MeshMaskClipMesh_Reference);
 
         uint32 vcount = pClipMesh->GetVertexBuffers().GetElementCount();
         if (!vcount)
         {
-            return nullptr; // Since there's nothing to clip against return a copy of the original base mesh
+            return; // Since there's nothing to clip against return a copy of the original base mesh
         }
 
         TArray<uint8> vertex_in_clip_mesh;  // Stores whether each vertex in the original mesh in in the clip mesh volume
@@ -871,21 +872,22 @@ namespace mu
             }
         }
 
-
-        return CreateMask( pBase, vertex_with_face_not_clipped );
+		CreateMask(Result, pBase, vertex_with_face_not_clipped );
     }
 
     //---------------------------------------------------------------------------------------------
     //! Generate a mask mesh with the faces of the base mesh matching the fragment.
     //---------------------------------------------------------------------------------------------
-    inline MeshPtr MeshMaskDiff(const Mesh* pBase, const Mesh* pFragment )
+    inline void MeshMaskDiff(Mesh* Result, const Mesh* pBase, const Mesh* pFragment, bool& bOutSuccess)
     {
         MUTABLE_CPUPROFILER_SCOPE(MeshMaskDiff);
+		bOutSuccess = true;
 
         uint32 vcount = pFragment->GetVertexBuffers().GetElementCount();
         if (!vcount)
         {
-            return nullptr;
+			bOutSuccess = false;
+            return;
         }
 
         int sourceFaceCount = pBase->GetFaceCount();
@@ -1017,7 +1019,7 @@ namespace mu
             }
         }
 
-        return CreateMask( pBase, vertex_with_face_not_clipped );
+		CreateMask(Result, pBase, vertex_with_face_not_clipped);
     }
 
 }
diff --git a/Engine/Plugins/Experimental/Mutable/Source/MutableRuntime/Private/MuR/OpMeshDifference.h b/Engine/Plugins/Experimental/Mutable/Source/MutableRuntime/Private/MuR/OpMeshDifference.h
--- a/Engine/Plugins/Experimental/Mutable/Source/MutableRuntime/Private/MuR/OpMeshDifference.h
+++ b/Engine/Plugins/Experimental/Mutable/Source/MutableRuntime/Private/MuR/OpMeshDifference.h
@@ -21,19 +21,22 @@ struct MORPH_VERTEX
 	//! Create a diff from the mesh vertices. The meshes must have the same amount of vertices.
 	//! If the channel list is empty all the channels will be compared.
 	//---------------------------------------------------------------------------------------------
-    inline MeshPtr MeshDifference( const Mesh* pBase, const Mesh* pTarget,
+    inline void MeshDifference( Mesh* Result, const Mesh* pBase, const Mesh* pTarget,
                                    int numChannels,
                                    const MESH_BUFFER_SEMANTIC* semantics,
                                    const int* semanticIndices,
-                                   bool ignoreTexCoords )
+                                   bool ignoreTexCoords, bool& bOutSuccess)
 
 	{
+		bOutSuccess = true;
+
         if (!pBase || !pTarget)
         {
-            return nullptr;
+			bOutSuccess = false;
+            return;
         }
 
-		bool correct =
+		bool bIsCorrect =
                 pBase
                 &&
 				( pBase->GetVertexBuffers().GetElementCount()
@@ -44,12 +47,13 @@ struct MORPH_VERTEX
                 ( pBase->GetVertexBuffers().GetElementCount()>0 )
                 ;
 
-		MeshPtr pDest = new Mesh();				
-
-		if (!correct)
-            return pDest;
+		if (!bIsCorrect)
+		{
+			bOutSuccess = true; // Use the provided empty mesh as result.
+			return ;
+		}
 
-		uint32_t vcount = pBase->GetVertexBuffers().GetElementCount();		
+		uint32 vcount = pBase->GetVertexBuffers().GetElementCount();		
 
 		// If no channels were specified, get them all
 		TArray<MESH_BUFFER_SEMANTIC> allSemantics;		
@@ -127,8 +131,8 @@ struct MORPH_VERTEX
 
 		// Create the morph mesh
 		{
-			pDest->GetVertexBuffers().SetElementCount( differentVertexCount );			
-			pDest->GetVertexBuffers().SetBufferCount( 1 );								
+			Result->GetVertexBuffers().SetElementCount( differentVertexCount );			
+			Result->GetVertexBuffers().SetBufferCount( 1 );								
 
 			TArray<MESH_BUFFER_SEMANTIC> semantic;
 			TArray<int> semanticIndex;
@@ -155,7 +159,7 @@ struct MORPH_VERTEX
 				offset += 16;				
 			}
 
-			pDest->GetVertexBuffers().SetBuffer
+			Result->GetVertexBuffers().SetBuffer
 				(
 					0,
 					offset,
@@ -172,20 +176,20 @@ struct MORPH_VERTEX
 			check(itBaseId.GetElementSize());
 
 			// Set the data
-			uint8_t* pDestData = pDest->GetVertexBuffers().GetBufferData( 0 );
+			uint8_t* ResultData = Result->GetVertexBuffers().GetBufferData( 0 );
 			for (size_t v=0; v<vcount; ++v)
 			{
 				if ( isVertexDifferent[v] )
 				{
 					// index
-					*((uint32_t*)pDestData) = itBaseId.GetAsUINT32();			
-					pDestData += 4;
+					*((uint32*)ResultData) = itBaseId.GetAsUINT32();			
+					ResultData += 4;
 
 					// all channels
 					for ( int c=0; c<numChannels; ++c )
 					{
-						FMemory::Memcpy(pDestData, &deltas[v * numChannels + c], 16);
-						pDestData += 4*4;											
+						FMemory::Memcpy(ResultData, &deltas[v * numChannels + c], 16);
+						ResultData += 4*4;											
 					}
 				}
 
@@ -196,10 +200,8 @@ struct MORPH_VERTEX
 
         // Rebuild surface data.
         // \todo: For now make single surfaced
-        pDest->m_surfaces.SetNum(0);
-        pDest->EnsureSurfaceData();
-
-		return pDest;
+        Result->m_surfaces.SetNum(0);
+        Result->EnsureSurfaceData();
 	}
 
 
diff --git a/Engine/Plugins/Experimental/Mutable/Source/MutableRuntime/Private/MuR/OpMeshExtractLayoutBlock.h b/Engine/Plugins/Experimental/Mutable/Source/MutableRuntime/Private/MuR/OpMeshExtractLayoutBlock.h
--- a/Engine/Plugins/Experimental/Mutable/Source/MutableRuntime/Private/MuR/OpMeshExtractLayoutBlock.h
+++ b/Engine/Plugins/Experimental/Mutable/Source/MutableRuntime/Private/MuR/OpMeshExtractLayoutBlock.h
@@ -113,16 +113,19 @@ namespace mu
 
 
 	//---------------------------------------------------------------------------------------------
-    inline MeshPtr MeshExtractLayoutBlock( const Mesh* pSource,
-                                           uint32_t layout,
+    inline void MeshExtractLayoutBlock(Mesh* Result, const Mesh* pSource,
+                                           uint32 layout,
                                            uint16 blockCount,
-                                           const uint32_t* pExtractBlocks )
+                                           const uint32* pExtractBlocks, bool& bOutSuccess)
 	{
+		check(pSource);
+		bOutSuccess = true;
+		
 		// TODO: Optimise
-		MeshPtr pResult = pSource->Clone();
+		Result->CopyFrom(*pSource);
 
         // TODO
-        pResult->SetFaceGroupCount( 0 );
+        Result->SetFaceGroupCount( 0 );
 
 		UntypedMeshBufferIteratorConst itBlocks( pSource->GetVertexBuffers(),
 												 MBS_LAYOUTBLOCK, layout );
@@ -164,26 +167,28 @@ namespace mu
                 check( false );
             }
 
-            MeshExtractFromVertices( pSource, pResult.get(), oldToNew, newToOld );
+            MeshExtractFromVertices(pSource, Result, oldToNew, newToOld);
         }
 
-        pResult->m_surfaces.Empty();
-        pResult->EnsureSurfaceData();
-
-        return pResult;
+        Result->m_surfaces.Empty();
+        Result->EnsureSurfaceData();
 	}
 
 
     //---------------------------------------------------------------------------------------------
-    inline MeshPtr MeshExtractFaceGroup( const Mesh* pSource, int group )
+    inline void MeshExtractFaceGroup(Mesh* Result, const Mesh* pSource, int group, bool& bOutSuccess)
     {
-        if ( group<0 || group>=pSource->GetFaceGroupCount() )
+		check(pSource);
+		bOutSuccess = true;
+
+        if (group < 0 || group >= pSource->GetFaceGroupCount())
         {
-            return new Mesh;
+			bOutSuccess = true; // use empty mesh as retun value.
+            return;
         }
 
         // TODO: Optimise
-        MeshPtr pResult = pSource->Clone();
+		Result->CopyFrom(*pSource);
 
         int resultVertices = 0;
 		TArray<int> oldToNew;
@@ -207,12 +212,11 @@ namespace mu
             }
         }
 
-        MeshExtractFromVertices( pSource, pResult.get(), oldToNew, newToOld );
+        MeshExtractFromVertices(pSource, Result, oldToNew, newToOld);
 
-        pResult->m_surfaces.Empty();
-        pResult->EnsureSurfaceData();
+        Result->m_surfaces.Empty();
+        Result->EnsureSurfaceData();
 
-        return pResult;
     }
 
 }
diff --git a/Engine/Plugins/Experimental/Mutable/Source/MutableRuntime/Private/MuR/OpMeshFormat.cpp b/Engine/Plugins/Experimental/Mutable/Source/MutableRuntime/Private/MuR/OpMeshFormat.cpp
--- a/Engine/Plugins/Experimental/Mutable/Source/MutableRuntime/Private/MuR/OpMeshFormat.cpp
+++ b/Engine/Plugins/Experimental/Mutable/Source/MutableRuntime/Private/MuR/OpMeshFormat.cpp
@@ -537,25 +537,39 @@ namespace mu
 
 
 	//-------------------------------------------------------------------------------------------------
-	MeshPtr MeshFormat
+	void MeshFormat
 	(
+		Mesh* Result, 
 		const Mesh* pPureSource,
 		const Mesh* pFormat,
 		bool keepSystemBuffers,
 		bool formatVertices,
 		bool formatIndices,
 		bool formatFaces,
-		bool ignoreMissingChannels
+		bool ignoreMissingChannels,
+		bool& bOutSuccess
 	)
 	{
 		MUTABLE_CPUPROFILER_SCOPE(MeshFormat);
+		bOutSuccess = true;
+
+		if (!pPureSource) 
+		{
+			check(false);
+			bOutSuccess = false;	
+			return;
+		}
 
-		if (!pPureSource) { return nullptr; }
-		if (!pFormat) { return pPureSource->Clone(); }
+		if (!pFormat)
+		{
+			check(false);
+			bOutSuccess = false;
+			return;
+		}
 
-		MeshPtrConst pSource = pPureSource;
+		Ptr<const Mesh> pSource = pPureSource;
 
-		MeshPtr pResult = pFormat->Clone();
+		Result->CopyFrom(*pFormat);
 
 		// Make sure that the bone indices will fit in this format, or extend it.
 		if (formatVertices)
@@ -574,7 +588,7 @@ namespace mu
 					{
 						int resultBuf = 0;
 						int resultChan = 0;
-						FMeshBufferSet& formBuffs = pResult->GetVertexBuffers();
+						FMeshBufferSet& formBuffs = Result->GetVertexBuffers();
 						formBuffs.FindChannel(MBS_BONEINDICES, Channel.m_semanticIndex, &resultBuf, &resultChan);
 						if (resultBuf >= 0)
 						{
@@ -626,64 +640,61 @@ namespace mu
 		if (formatVertices)
 		{
 
-			FormatBufferSet(pSource->GetVertexBuffers(), pResult->GetVertexBuffers(),
+			FormatBufferSet(pSource->GetVertexBuffers(), Result->GetVertexBuffers(),
 				keepSystemBuffers, ignoreMissingChannels, true);
 		}
 		else
 		{
-			pResult->m_VertexBuffers = pSource->GetVertexBuffers();
+			Result->m_VertexBuffers = pSource->GetVertexBuffers();
 		}
 
 		if (formatIndices)
 		{
-			FormatBufferSet(pSource->GetIndexBuffers(), pResult->GetIndexBuffers(), keepSystemBuffers,
+			FormatBufferSet(pSource->GetIndexBuffers(), Result->GetIndexBuffers(), keepSystemBuffers,
 				ignoreMissingChannels, false);
 		}
 		else
 		{
-			pResult->m_IndexBuffers = pSource->GetIndexBuffers();
+			Result->m_IndexBuffers = pSource->GetIndexBuffers();
 		}
 
 		if (formatFaces)
 		{
 			FormatBufferSet(pSource->GetFaceBuffers(),
-				pResult->GetFaceBuffers(),
+				Result->GetFaceBuffers(),
 				keepSystemBuffers,
 				ignoreMissingChannels,
 				false);
 		}
 		else
 		{
-			pResult->m_FaceBuffers = pSource->GetFaceBuffers();
+			Result->m_FaceBuffers = pSource->GetFaceBuffers();
 		}
 
 		// Copy the rest of the data
-		pResult->SetSkeleton(pSource->GetSkeleton());
-		pResult->SetPhysicsBody(pSource->GetPhysicsBody());
+		Result->SetSkeleton(pSource->GetSkeleton());
+		Result->SetPhysicsBody(pSource->GetPhysicsBody());
 
-		pResult->m_layouts.Empty();
+		Result->m_layouts.Empty();
 		for (const Ptr<const Layout>& Layout : pSource->m_layouts)
 		{
-			pResult->m_layouts.Add(Layout->Clone());
+			Result->m_layouts.Add(Layout->Clone());
 		}
 
-		pResult->m_tags = pSource->m_tags;
+		Result->m_tags = pSource->m_tags;
 
-		pResult->m_AdditionalBuffers = pSource->m_AdditionalBuffers;
+		Result->m_AdditionalBuffers = pSource->m_AdditionalBuffers;
 
-		pResult->BonePoses = pSource->BonePoses;
-		pResult->BoneMapIndices = pSource->BoneMapIndices;
+		Result->BonePoses = pSource->BonePoses;
+		Result->BoneMapIndices = pSource->BoneMapIndices;
 
-		pResult->SkeletonIDs = pSource->SkeletonIDs;
+		Result->SkeletonIDs = pSource->SkeletonIDs;
 
 		// A shallow copy is done here, it should not be a problem.
-		pResult->AdditionalPhysicsBodies = pSource->AdditionalPhysicsBodies;
-
-		pResult->ResetStaticFormatFlags();
-		pResult->EnsureSurfaceData();
-
-		return pResult;
+		Result->AdditionalPhysicsBodies = pSource->AdditionalPhysicsBodies;
 
+		Result->ResetStaticFormatFlags();
+		Result->EnsureSurfaceData();
 	}
 
 }
diff --git a/Engine/Plugins/Experimental/Mutable/Source/MutableRuntime/Private/MuR/OpMeshFormat.h b/Engine/Plugins/Experimental/Mutable/Source/MutableRuntime/Private/MuR/OpMeshFormat.h
--- a/Engine/Plugins/Experimental/Mutable/Source/MutableRuntime/Private/MuR/OpMeshFormat.h
+++ b/Engine/Plugins/Experimental/Mutable/Source/MutableRuntime/Private/MuR/OpMeshFormat.h
@@ -15,15 +15,17 @@ class FMeshBufferSet;
 	//! \param keepSystemBuffers Will keep the internal system buffers even if they are not in the
 	//! original format. If they are, they will be duplicated, so be careful.
 	//---------------------------------------------------------------------------------------------
-	MUTABLERUNTIME_API extern MeshPtr MeshFormat
+	MUTABLERUNTIME_API extern void MeshFormat
 		(
+			Mesh* Result,
 			const Mesh* pSource,
 			const Mesh* pFormat,
 			bool keepSystemBuffers,
 			bool formatVertices,
 			bool formatIndices,
 			bool formatFaces,
-			bool ignoreMissingChannels
+			bool ignoreMissingChannels,
+			bool& bOutSuccess
 		);
 
 
diff --git a/Engine/Plugins/Experimental/Mutable/Source/MutableRuntime/Private/MuR/OpMeshGeometryOperation.h b/Engine/Plugins/Experimental/Mutable/Source/MutableRuntime/Private/MuR/OpMeshGeometryOperation.h
--- a/Engine/Plugins/Experimental/Mutable/Source/MutableRuntime/Private/MuR/OpMeshGeometryOperation.h
+++ b/Engine/Plugins/Experimental/Mutable/Source/MutableRuntime/Private/MuR/OpMeshGeometryOperation.h
@@ -105,13 +105,15 @@ namespace mu
     //---------------------------------------------------------------------------------------------
     //! Apply an example geometric operation (twist)
     //---------------------------------------------------------------------------------------------
-    inline MeshPtr MeshGeometryOperation(const Mesh* MeshA, const Mesh* MeshB, float ScalarA, float ScalarB)
+    inline void MeshGeometryOperation(Mesh* Result, const Mesh* MeshA, const Mesh* MeshB, float ScalarA, float ScalarB, bool& bOutSuccess)
     {
 		MUTABLE_CPUPROFILER_SCOPE(MeshGeometryOperation);
+		bOutSuccess = true;
 
 		if (!MeshA)
 		{
-			return nullptr;
+			bOutSuccess = false;
+			return;
 		}
 
         // Generate the FDynamicMesh
@@ -139,10 +141,8 @@ namespace mu
 		//}
 
 		// Update the Mutable mesh from the FDynamicMesh
-		Ptr<Mesh> Result = MeshA->Clone();
-		//UpdateMutableMesh(Result.get(), DynMeshResult.Get(), VertexMutableToDyn);
-
-        return Result;
+		Result->CopyFrom(*MeshA);
+		//UpdateMutableMesh(Result, DynMeshResult.Get(), VertexMutableToDyn);
     }
 
 }
diff --git a/Engine/Plugins/Experimental/Mutable/Source/MutableRuntime/Private/MuR/OpMeshMerge.h b/Engine/Plugins/Experimental/Mutable/Source/MutableRuntime/Private/MuR/OpMeshMerge.h
--- a/Engine/Plugins/Experimental/Mutable/Source/MutableRuntime/Private/MuR/OpMeshMerge.h
+++ b/Engine/Plugins/Experimental/Mutable/Source/MutableRuntime/Private/MuR/OpMeshMerge.h
@@ -12,20 +12,24 @@
 namespace mu
 {
 
+	struct FMeshMergeScratchMeshes
+	{
+		Ptr<Mesh> FirstReformat;
+		Ptr<Mesh> SecondReformat;
+	};
 
 	//---------------------------------------------------------------------------------------------
 	//! Merge two meshes into one new mesh
 	//---------------------------------------------------------------------------------------------
-	inline MeshPtr MeshMerge(const Mesh* pFirst, const Mesh* pSecond, bool bMergeSurfaces)
+	inline void MeshMerge(Mesh* Result, const Mesh* pFirst, const Mesh* pSecond, bool bMergeSurfaces, 
+		FMeshMergeScratchMeshes& ScratchMeshes)
 	{
 		MUTABLE_CPUPROFILER_SCOPE(MeshMerge);
 
-		MeshPtr pResult = new Mesh;
-
 		// Should never happen, but fixes static analysis warnings.
 		if (!(pFirst && pSecond))
 		{
-			return pResult;
+			return;
 		}
 
 		// Indices
@@ -36,13 +40,13 @@ namespace mu
 
 			const int32 FirstCount = pFirst->GetIndexBuffers().GetElementCount();
 			const int32 SecondCount = pSecond->GetIndexBuffers().GetElementCount();
-			pResult->GetIndexBuffers().SetElementCount(FirstCount + SecondCount);
+			Result->GetIndexBuffers().SetElementCount(FirstCount + SecondCount);
 
 			check(pFirst->GetIndexBuffers().GetBufferCount() <= 1);
 			check(pSecond->GetIndexBuffers().GetBufferCount() <= 1);
-			pResult->GetIndexBuffers().SetBufferCount(1);
+			Result->GetIndexBuffers().SetBufferCount(1);
 
-			MESH_BUFFER& ResultIndexBuffer = pResult->GetIndexBuffers().m_buffers[0];
+			MESH_BUFFER& ResultIndexBuffer = Result->GetIndexBuffers().m_buffers[0];
 
 			const MESH_BUFFER& FirstIndexBuffer = pFirst->GetIndexBuffers().m_buffers[0];
 			const MESH_BUFFER& SecondIndexBuffer = pSecond->GetIndexBuffers().m_buffers[0];
@@ -330,15 +334,15 @@ namespace mu
 
 			const int32 FirstCount = pFirst->GetFaceBuffers().GetElementCount();
 			const int32 SecondCount = pSecond->GetFaceBuffers().GetElementCount();
-			pResult->GetFaceBuffers().SetElementCount(FirstCount + SecondCount);
-			pResult->GetFaceBuffers().SetBufferCount(pFirst->GetFaceBuffers().GetBufferCount());
+			Result->GetFaceBuffers().SetElementCount(FirstCount + SecondCount);
+			Result->GetFaceBuffers().SetBufferCount(pFirst->GetFaceBuffers().GetBufferCount());
 
 			// Merge only the buffers present in the first mesh
-			for (int32 b = 0; b < pResult->GetFaceBuffers().GetBufferCount(); ++b)
+			for (int32 b = 0; b < Result->GetFaceBuffers().GetBufferCount(); ++b)
 			{
 				const MESH_BUFFER& first = pFirst->GetFaceBuffers().m_buffers[b];
 
-				MESH_BUFFER& result = pResult->GetFaceBuffers().m_buffers[b];
+				MESH_BUFFER& result = Result->GetFaceBuffers().m_buffers[b];
 				result.m_channels = first.m_channels;
 				result.m_elementSize = first.m_elementSize;
 				result.m_data.SetNum(result.m_elementSize * (FirstCount + SecondCount));
@@ -394,7 +398,7 @@ namespace mu
 		{
 			MUTABLE_CPUPROFILER_SCOPE(Layouts);
 
-			pResult->m_layouts.SetNum(pFirst->m_layouts.Num());
+			Result->m_layouts.SetNum(pFirst->m_layouts.Num());
 			for (int i = 0; i < pFirst->m_layouts.Num(); ++i)
 			{
 				const Layout* pF = pFirst->m_layouts[i].get();
@@ -407,7 +411,7 @@ namespace mu
 					pR->m_blocks.Append(pS->m_blocks);
 				}
 
-				pResult->m_layouts[i] = pR;
+				Result->m_layouts[i] = pR;
 			}
 		}
 
@@ -416,11 +420,11 @@ namespace mu
 		//---------------------------
 
 		// Add SkeletonIDs
-		pResult->SkeletonIDs = pFirst->SkeletonIDs;
+		Result->SkeletonIDs = pFirst->SkeletonIDs;
 
 		for (const int32 SkeletonID : pSecond->SkeletonIDs)
 		{
-			pResult->SkeletonIDs.AddUnique(SkeletonID);
+			Result->SkeletonIDs.AddUnique(SkeletonID);
 		}
 
 		// Do they have the same skeleton?
@@ -446,7 +450,7 @@ namespace mu
 			const int32 NumBonesSecond = pSecondSkeleton ? pSecondSkeleton->GetBoneCount() : 0;
 
 			pResultSkeleton = pFirstSkeleton ? pFirstSkeleton->Clone() : new Skeleton;
-			pResult->SetSkeleton(pResultSkeleton);
+			Result->SetSkeleton(pResultSkeleton);
 
 			SecondToResultBoneIndices.SetNumUninitialized(NumBonesSecond);
 
@@ -487,7 +491,7 @@ namespace mu
 		}
 		else
 		{
-			pResult->SetSkeleton(pFirst->GetSkeleton());
+			Result->SetSkeleton(pFirst->GetSkeleton());
 		}
 
 
@@ -506,7 +510,7 @@ namespace mu
 			MUTABLE_CPUPROFILER_SCOPE(Surfaces);
 			
 			const int32 NumFirstBoneMapIndices = pFirst->BoneMapIndices.Num();
-			pResult->BoneMapIndices = pFirst->BoneMapIndices;
+			Result->BoneMapIndices = pFirst->BoneMapIndices;
 
 			if (bMergeSurfaces)
 			{
@@ -517,47 +521,47 @@ namespace mu
 				{
 					const int32 SecondBoneIndex = pSecond->BoneMapIndices[SecondBoneMapIndex];
 					const int32 BoneIndex = bMergeSkeletons ? SecondToResultBoneIndices[SecondBoneIndex] : SecondBoneIndex;
-					const int32 BoneMapIndex = pResult->BoneMapIndices.AddUnique(BoneIndex);
+					const int32 BoneMapIndex = Result->BoneMapIndices.AddUnique(BoneIndex);
 					RemappedBoneMapIndices[SecondBoneMapIndex] = BoneMapIndex;
 
 					bRemapBoneIndices = bRemapBoneIndices || BoneMapIndex != SecondBoneMapIndex;
 				}
 
-				MESH_SURFACE& NewSurface = pResult->m_surfaces.AddDefaulted_GetRef();
+				MESH_SURFACE& NewSurface = Result->m_surfaces.AddDefaulted_GetRef();
 				NewSurface.m_vertexCount = pFirst->GetVertexCount() + pSecond->GetVertexCount();
 				NewSurface.m_indexCount = pFirst->GetIndexCount() + pSecond->GetIndexCount();
-				NewSurface.BoneMapCount = pResult->BoneMapIndices.Num();
+				NewSurface.BoneMapCount = Result->BoneMapIndices.Num();
 			}
 			else
 			{
 				// Add the bonemaps of the second mesh with the fixed BoneIndices
 				const int32 NumBoneMapIndices = NumFirstBoneMapIndices + NumSecondBoneMapIndices;
-				pResult->BoneMapIndices.Reserve(NumBoneMapIndices);
+				Result->BoneMapIndices.Reserve(NumBoneMapIndices);
 
 				for (const uint16& BoneIndex : pSecond->BoneMapIndices)
 				{
-					pResult->BoneMapIndices.Add(bMergeSkeletons ? SecondToResultBoneIndices[BoneIndex] : BoneIndex);
+					Result->BoneMapIndices.Add(bMergeSkeletons ? SecondToResultBoneIndices[BoneIndex] : BoneIndex);
 				}
 
 				// Add pFirst surfaces
-				pResult->m_surfaces = pFirst->m_surfaces;
+				Result->m_surfaces = pFirst->m_surfaces;
 
 				const int32 FirstVertexIndex = pFirst->GetVertexCount();
 				const int32 FirstIndexIndex = pFirst->GetIndexCount();
 
 				check(pSecond->m_surfaces.Num() == 1);
-				MESH_SURFACE& NewSurface = pResult->m_surfaces.Add_GetRef(pSecond->m_surfaces[0]);
+				MESH_SURFACE& NewSurface = Result->m_surfaces.Add_GetRef(pSecond->m_surfaces[0]);
 				NewSurface.m_firstVertex += FirstVertexIndex;
 				NewSurface.m_firstIndex += FirstIndexIndex;
 				NewSurface.BoneMapIndex += NumFirstBoneMapIndices;
 			}
 
-			for (const MESH_SURFACE& Surface : pResult->m_surfaces)
+			for (const MESH_SURFACE& Surface : Result->m_surfaces)
 			{
 				MaxBoneMapIndices = FMath::Max(MaxBoneMapIndices, Surface.BoneMapCount);
 			}
 
-			pResult->BoneMapIndices.Shrink();
+			Result->BoneMapIndices.Shrink();
 		}
 
 
@@ -566,19 +570,19 @@ namespace mu
 		{
 			MUTABLE_CPUPROFILER_SCOPE(Pose);
 
-			pResult->BonePoses.Reserve(pResult->GetSkeleton()->GetBoneCount());
+			Result->BonePoses.Reserve(Result->GetSkeleton()->GetBoneCount());
 
 			// Copy poses from the first mesh
-			pResult->BonePoses = pFirst->BonePoses;
+			Result->BonePoses = pFirst->BonePoses;
 
 			// Add or override bone poses
 			for (const Mesh::FBonePose& SecondBonePose : pSecond->BonePoses)
 			{
-				const int32 ResultBoneIndex = pResult->FindBonePose(SecondBonePose.BoneName.c_str());
+				const int32 ResultBoneIndex = Result->FindBonePose(SecondBonePose.BoneName.c_str());
 
 				if (ResultBoneIndex != INDEX_NONE)
 				{
-					Mesh::FBonePose& ResultBonePose = pResult->BonePoses[ResultBoneIndex];
+					Mesh::FBonePose& ResultBonePose = Result->BonePoses[ResultBoneIndex];
 
 					// TODO: Not sure how to tune this priority, review it.
 					// For now use a similar strategy as before. 
@@ -598,11 +602,11 @@ namespace mu
 				}
 				else
 				{
-					pResult->BonePoses.Add(SecondBonePose);
+					Result->BonePoses.Add(SecondBonePose);
 				}
 			}
 
-			pResult->BonePoses.Shrink();
+			Result->BonePoses.Shrink();
 		}
 
 
@@ -708,7 +712,7 @@ namespace mu
 			if (SharedResultPhysicsBody.Get<1>())
 			{
 				// Only one or non of the meshes has physics, share the result.
-				pResult->SetPhysicsBody(SharedResultPhysicsBody.Get<0>());
+				Result->SetPhysicsBody(SharedResultPhysicsBody.Get<0>());
 			}
 			else
 			{
@@ -720,14 +724,14 @@ namespace mu
 					AppendPhysicsBodiesUnique(*MergedResultPhysicsBody, *pSecond->GetPhysicsBody()) ||
 					pFirst->GetPhysicsBody()->bBodiesModified || pSecond->GetPhysicsBody()->bBodiesModified;
 
-				pResult->SetPhysicsBody(MergedResultPhysicsBody);
+				Result->SetPhysicsBody(MergedResultPhysicsBody);
 			}
 
 			// Additional physics bodies.
 			const int32 MaxAdditionalPhysicsResultNum = pFirst->AdditionalPhysicsBodies.Num() + pSecond->AdditionalPhysicsBodies.Num();
 
-			pResult->AdditionalPhysicsBodies.Reserve(MaxAdditionalPhysicsResultNum);
-			pResult->AdditionalPhysicsBodies.Append(pFirst->AdditionalPhysicsBodies);
+			Result->AdditionalPhysicsBodies.Reserve(MaxAdditionalPhysicsResultNum);
+			Result->AdditionalPhysicsBodies.Append(pFirst->AdditionalPhysicsBodies);
 			
 			// Not very many additional bodies expected, do a quadratic search to have unique bodies based on external id.
 			const int32 NumSecondAdditionalBodies = pSecond->AdditionalPhysicsBodies.Num();
@@ -742,7 +746,7 @@ namespace mu
 				// may need to be contemplated at some point.
 				if (!Found)
 				{
-					pResult->AdditionalPhysicsBodies.Add(pSecond->AdditionalPhysicsBodies[Index]);
+					Result->AdditionalPhysicsBodies.Add(pSecond->AdditionalPhysicsBodies[Index]);
 				}
 			}
 		}
@@ -811,10 +815,10 @@ namespace mu
 				FastPath = FastPath && !bChangeBoneIndicesFormat;
 			}
 
-			MeshPtrConst pVFirst;
-			MeshPtrConst pVSecond;
-
-			if ( !FastPath )
+			Ptr<const Mesh> pVFirst;
+			Ptr<const Mesh> pVSecond;
+			
+			if (!FastPath)
 			{
                 MUTABLE_CPUPROFILER_SCOPE(SlowPath);
 
@@ -822,11 +826,11 @@ namespace mu
 
 				// Expand component counts in vertex channels of the format mesh
 				int32 vbcount = pFirst->GetVertexBuffers().m_buffers.Num();
-				pResult->GetVertexBuffers().SetBufferCount( (int)vbcount );
+				Result->GetVertexBuffers().SetBufferCount( (int)vbcount );
 
 				for ( int32 vb = 0; vb<vbcount; ++vb )
 				{
-					MESH_BUFFER& result = pResult->GetVertexBuffers().m_buffers[vb];
+					MESH_BUFFER& result = Result->GetVertexBuffers().m_buffers[vb];
 					const MESH_BUFFER& first = pFirst->GetVertexBuffers().m_buffers[vb];
 
 					result.m_channels = first.m_channels;
@@ -906,7 +910,7 @@ namespace mu
 
                     if (someChannel && allNewChannels)
                     {
-                        pResult->GetVertexBuffers().m_buffers.Add(buf);
+                        Result->GetVertexBuffers().m_buffers.Add(buf);
                     }
                 }
 
@@ -914,7 +918,7 @@ namespace mu
 				if(bChangeBoneIndicesFormat)
 				{
 					// Iterate all vertex buffers and update the format
-					FMeshBufferSet& VertexBuffers = pResult->GetVertexBuffers();
+					FMeshBufferSet& VertexBuffers = Result->GetVertexBuffers();
 					for (int32 VertexBufferIndex = 0; VertexBufferIndex < VertexBuffers.m_buffers.Num(); ++VertexBufferIndex)
 					{
 						MESH_BUFFER& result = VertexBuffers.m_buffers[VertexBufferIndex];
@@ -942,30 +946,45 @@ namespace mu
 				}
 
                 // Allocate vertices
-                pResult->GetVertexBuffers().SetElementCount( FirstCount + SecondCount );
-
+                Result->GetVertexBuffers().SetElementCount( FirstCount + SecondCount );
+				
 				// Convert the source meshes to the new format
-                if (pFirst->GetVertexBuffers().HasSameFormat(pResult->GetVertexBuffers()))
+                if (pFirst->GetVertexBuffers().HasSameFormat(Result->GetVertexBuffers()))
                 {
                     pVFirst = pFirst;
                 }
                 else
                 {
-                    pVFirst = MeshFormat( pFirst, pResult.get(), false, true, false, false, false );
+					check(ScratchMeshes.FirstReformat);
+
+					bool bOutSuccess = false;
+					MeshFormat(ScratchMeshes.FirstReformat.get(), pFirst, Result, false, true, false, false, false, bOutSuccess);
+					
+					if (bOutSuccess)
+					{
+						pVFirst = ScratchMeshes.FirstReformat;
+					}
                 }
 
 
-                if (pSecond->GetVertexBuffers().HasSameFormat(pResult->GetVertexBuffers()))
+                if (pSecond->GetVertexBuffers().HasSameFormat(Result->GetVertexBuffers()))
                 {
                     pVSecond = pSecond;
                 }
                 else
                 {
-                    pVSecond = MeshFormat( pSecond, pResult.get(), false, true, false, false, false );
+					check(ScratchMeshes.SecondReformat);
+
+					bool bOutSuccess = false;
+                    MeshFormat(ScratchMeshes.SecondReformat.get(), pSecond, Result, false, true, false, false, false, bOutSuccess);
+
+					if (bOutSuccess)
+					{
+						pVSecond = ScratchMeshes.SecondReformat;
+					}
                 }
 
-				check( pVFirst->GetVertexBuffers().HasSameFormat
-                                ( pVSecond->GetVertexBuffers() ) );
+				check(pVFirst->GetVertexBuffers().HasSameFormat(pVSecond->GetVertexBuffers()));
 			}
 			else
 			{
@@ -975,19 +994,19 @@ namespace mu
 				pVSecond = pSecond;
 			}
 
-
-			pResult->m_VertexBuffers = pVFirst->m_VertexBuffers;
-			pResult->GetVertexBuffers().SetElementCount( FirstCount + SecondCount );
+			Result->m_VertexBuffers = pVFirst->m_VertexBuffers;
+			Result->GetVertexBuffers().SetElementCount( FirstCount + SecondCount );
 
             // first copy all the vertex data
 			{
-				for ( int32 vb = 0; vb<pResult->GetVertexBuffers().m_buffers.Num(); ++vb )
+				for (int32 vb = 0; vb < Result->GetVertexBuffers().m_buffers.Num(); ++vb)
 				{
-					MESH_BUFFER& result = pResult->GetVertexBuffers().m_buffers[vb];
+					MESH_BUFFER& result = Result->GetVertexBuffers().m_buffers[vb];
 					const MESH_BUFFER& second = pVSecond->GetVertexBuffers().m_buffers[vb];
-					if ( SecondCount )
+
+					if (SecondCount)
 					{
-						int elemSize = pResult->GetVertexBuffers().GetElementSize((int)vb);
+						int elemSize = Result->GetVertexBuffers().GetElementSize((int)vb);
 						int firstSize = FirstCount * elemSize;
 						int secondSize = SecondCount * elemSize;
 
@@ -1006,7 +1025,7 @@ namespace mu
 				check(!RemappedBoneMapIndices.IsEmpty())
 
 				// Iterate all vertex buffers and update the format
-				FMeshBufferSet& VertexBuffers = pResult->GetVertexBuffers();
+				FMeshBufferSet& VertexBuffers = Result->GetVertexBuffers();
 				for (int32 VertexBufferIndex = 0; VertexBufferIndex < VertexBuffers.m_buffers.Num(); ++VertexBufferIndex)
 				{
 					MESH_BUFFER& ResultBuffer = VertexBuffers.m_buffers[VertexBufferIndex];
@@ -1099,28 +1118,12 @@ namespace mu
 		}
 
 		// Tags
-		pResult->m_tags = pFirst->m_tags;
+		Result->m_tags = pFirst->m_tags;
 
-		for (auto& secondTag : pSecond->m_tags)
+		for (const string& SecondTag : pSecond->m_tags)
 		{
-			bool repeated = false;
-
-			for (auto& firstTag : pFirst->m_tags)
-			{
-				if (firstTag == secondTag)
-				{
-					repeated = true;
-					break;
-				}				
-			}
-
-			if (!repeated)
-			{
-				pResult->m_tags.Add(secondTag);
-			}
+			Result->m_tags.AddUnique(SecondTag);
 		}
-
-		return pResult;
 	}
 	
 
@@ -1165,15 +1168,14 @@ namespace mu
     //---------------------------------------------------------------------------------------------
     //! Return null if there is no need to remap the mesh
     //---------------------------------------------------------------------------------------------
-    inline MeshPtr MeshRemapSkeleton( const Mesh* SourceMesh, const Skeleton* Skeleton )
+    inline void MeshRemapSkeleton(Mesh* Result, const Mesh* SourceMesh, const Skeleton* Skeleton, bool& bOutSuccess)
     {
-        if (SourceMesh->GetVertexCount()==0
-             ||
-             !SourceMesh->GetSkeleton()
-             ||
-			SourceMesh->GetSkeleton()->GetBoneCount()==0 )
+		bOutSuccess = true;
+
+        if (SourceMesh->GetVertexCount() == 0 || !SourceMesh->GetSkeleton() || SourceMesh->GetSkeleton()->GetBoneCount() == 0)
         {
-            return nullptr;
+			bOutSuccess = false;
+            return;
         }
 
         mu::SkeletonPtrConst SourceSkeleton = SourceMesh->GetSkeleton();
@@ -1198,14 +1200,14 @@ namespace mu
 
         if (!bBonesRemapped)
         {
-            return nullptr;
+			bOutSuccess = false;
+            return;
         }
 
-		MeshPtr pResult = SourceMesh->Clone();
-		pResult->SetSkeleton(Skeleton);
-		pResult->SetBoneMap(RemappedBoneMap);
+		Result->CopyFrom(*SourceMesh);
 
-        return pResult;
+		Result->SetSkeleton(Skeleton);
+		Result->SetBoneMap(RemappedBoneMap);
     }
 	
 }
diff --git a/Engine/Plugins/Experimental/Mutable/Source/MutableRuntime/Private/MuR/OpMeshMorph.h b/Engine/Plugins/Experimental/Mutable/Source/MutableRuntime/Private/MuR/OpMeshMorph.h
--- a/Engine/Plugins/Experimental/Mutable/Source/MutableRuntime/Private/MuR/OpMeshMorph.h
+++ b/Engine/Plugins/Experimental/Mutable/Source/MutableRuntime/Private/MuR/OpMeshMorph.h
@@ -14,11 +14,16 @@ namespace mu
 	//---------------------------------------------------------------------------------------------
 	//! Optimized linear factor version for morphing 2 targets
 	//---------------------------------------------------------------------------------------------
-    inline MeshPtr MeshMorph2( const Mesh* pBase, const Mesh* pMin, const Mesh* pMax, const float factor )
+    inline void MeshMorph2(Mesh* Result, const Mesh* pBase, const Mesh* pMin, const Mesh* pMax, const float factor, bool& bOutSuccess)
     {
         MUTABLE_CPUPROFILER_SCOPE(MeshMorph2);
+		bOutSuccess = true;
 
-		if (!pBase) return nullptr;
+		if (!pBase)
+		{
+			bOutSuccess = true; // return an empty mesh.
+			return;
+		}
 
         const auto ApplyMorph = []
             ( auto BaseIdIter, const TArray< UntypedMeshBufferIterator >& baseChannelsIters, const int32 baseSize,
@@ -102,8 +107,8 @@ namespace mu
             }
         };
 
-		MeshPtr pDest = pBase->Clone();
 
+		
 		// Number of vertices to modify
 		const uint32 minCount = pMin ? pMin->GetVertexBuffers().GetElementCount() : 0;
 		const uint32 maxCount = pMax ? pMax->GetVertexBuffers().GetElementCount() : 0;
@@ -112,9 +117,12 @@ namespace mu
 
 		if ( baseCount == 0 || (minCount + maxCount) == 0)
 		{
-			return pDest;
+			bOutSuccess = false; // Use the passed pBase as result.
+			return;
 		}
 
+		Result->CopyFrom(*pBase);
+
 		if (refTarget)
 		{
 			uint32 ccount = refTarget->GetVertexBuffers().GetBufferChannelCount(0);
@@ -135,7 +143,7 @@ namespace mu
 				MESH_BUFFER_SEMANTIC sem = MBSPriv.m_buffers[0].m_channels[c].m_semantic;
 				int semIndex = MBSPriv.m_buffers[0].m_channels[c].m_semanticIndex;
 				
-				itBaseChannels[c] = UntypedMeshBufferIterator(pDest->GetVertexBuffers(), sem, semIndex);
+				itBaseChannels[c] = UntypedMeshBufferIterator(Result->GetVertexBuffers(), sem, semIndex);
 				if (minCount > 0)
 				{
 					itMinChannels[c] = UntypedMeshBufferIteratorConst(pMin->GetVertexBuffers(), sem, semIndex);
@@ -159,24 +167,22 @@ namespace mu
 				ApplyMorph(itBaseId, itBaseChannels, baseCount, itMaxId, itMaxChannels, maxCount, factor);
 			}
 		}
-
-        return pDest;
     }
 
 	//---------------------------------------------------------------------------------------------
 	//! \TODO Optimized linear factor version
 	//---------------------------------------------------------------------------------------------
-	inline MeshPtr MeshMorph(const Mesh* pBase, const Mesh* pMorph, float factor)
+	inline void MeshMorph(Mesh* Result, const Mesh* pBase, const Mesh* pMorph, float factor, bool& bOutSuccess)
 	{
-		return MeshMorph2( pBase, nullptr, pMorph, factor );
+		MeshMorph2(Result, pBase, nullptr, pMorph, factor, bOutSuccess);
 	}
 
     //---------------------------------------------------------------------------------------------
     //! \TODO Optimized Factor-less version
     //---------------------------------------------------------------------------------------------
-    inline MeshPtr MeshMorph( const Mesh* pBase, const Mesh* pMorph )
+	inline void MeshMorph(Mesh* Result, const Mesh* pBase, const Mesh* pMorph, bool& bOutSuccess)
     {
         // Trust the compiler to remove the factor
-        return MeshMorph( pBase, pMorph, 1.0f );
+		MeshMorph(Result, pBase, pMorph, 1.0f, bOutSuccess);
     }
 }
diff --git a/Engine/Plugins/Experimental/Mutable/Source/MutableRuntime/Private/MuR/OpMeshOptimizeSkinning.h b/Engine/Plugins/Experimental/Mutable/Source/MutableRuntime/Private/MuR/OpMeshOptimizeSkinning.h
--- a/Engine/Plugins/Experimental/Mutable/Source/MutableRuntime/Private/MuR/OpMeshOptimizeSkinning.h
+++ b/Engine/Plugins/Experimental/Mutable/Source/MutableRuntime/Private/MuR/OpMeshOptimizeSkinning.h
@@ -14,13 +14,16 @@ namespace mu
 	//---------------------------------------------------------------------------------------------
 	//!
 	//---------------------------------------------------------------------------------------------
-	inline MeshPtr MeshOptimizeSkinning(const Mesh* InMesh)
+	inline void MeshOptimizeSkinning(Mesh* Result, const Mesh* InMesh, bool& bOutSuccess)
 	{
 		MUTABLE_CPUPROFILER_SCOPE(MeshOptimizeSkinning);
 
+		bOutSuccess = true;
+
 		if (!InMesh)
 		{
-			return nullptr;
+			bOutSuccess = false;
+			return;
 		}
 
 		uint32 MaxBoneMapIndex = 0;
@@ -32,7 +35,8 @@ namespace mu
 		// We can't optimize the skinning if the mesh requires 16 bit bone indices 
 		if (MaxBoneMapIndex > MAX_uint8)
 		{
-			return nullptr;
+			bOutSuccess = false;
+			return;
 		}
 
 		bool bRequiresFormatChange = false;
@@ -60,7 +64,8 @@ namespace mu
 		// Not sure if bRequiresFormatChange will ever be true.
 		if (!bRequiresFormatChange)
 		{
-			return nullptr;
+			bOutSuccess = false;
+			return;
 		}
 
 		MUTABLE_CPUPROFILER_SCOPE(MeshOptimizeSkinning_Format);
@@ -72,7 +77,8 @@ namespace mu
 		const int32 ElementCount = InMeshVertexBuffers.GetElementCount();
 
 		// Clone mesh without VertexBuffers, they will be copied manually.
-		MeshPtr Result = InMesh->Clone(~(EMeshCloneFlags::WithVertexBuffers));
+		constexpr EMeshCopyFlags CopyFlags = ~EMeshCopyFlags::WithVertexBuffers;
+		Result->CopyFrom(*InMesh, CopyFlags);
 
 		mu::FMeshBufferSet& VertexBuffers = Result->GetVertexBuffers();
 
@@ -180,7 +186,5 @@ namespace mu
 			}
 
 		}
-
-		return Result;
 	}
 }
diff --git a/Engine/Plugins/Experimental/Mutable/Source/MutableRuntime/Private/MuR/OpMeshRemapIndices.h b/Engine/Plugins/Experimental/Mutable/Source/MutableRuntime/Private/MuR/OpMeshRemapIndices.h
--- a/Engine/Plugins/Experimental/Mutable/Source/MutableRuntime/Private/MuR/OpMeshRemapIndices.h
+++ b/Engine/Plugins/Experimental/Mutable/Source/MutableRuntime/Private/MuR/OpMeshRemapIndices.h
@@ -14,28 +14,39 @@ namespace mu
     //---------------------------------------------------------------------------------------------
     //!
     //---------------------------------------------------------------------------------------------
-    inline MeshPtr MeshRemapIndices( const Mesh* pBase, const Mesh* pReference )
+    inline void MeshRemapIndices(Mesh* Result, const Mesh* pBase, const Mesh* pReference, bool& bOutSuccess)
     {
         MUTABLE_CPUPROFILER_SCOPE(MeshRemapIndices);
 
-        if (!pBase) return nullptr;
+		bOutSuccess = true;
 
-        MeshPtr pDest = pBase->Clone();
+		if (!pBase)
+		{
+			bOutSuccess = false;
+			return;
+		}
 
-        if (!pReference) return pDest;
+		if (!pReference)
+		{
+			bOutSuccess = false;
+			return;
+		}
 
         // Number of vertices to modify
         uint32_t vcountReference = pReference->GetVertexBuffers().GetElementCount();
         uint32_t vcountBase = pBase->GetVertexBuffers().GetElementCount();
         if ( !vcountReference || !vcountBase )
         {
-            return pDest;
+			bOutSuccess = false;
+			return;
         }
 
+		Result->CopyFrom(*pBase);
+
         // Iterator of the vertex ids of the base vertices
         UntypedMeshBufferIteratorConst itReferenceId( pReference->GetVertexBuffers(), MBS_VERTEXINDEX );
 
-        FMeshBufferSet& verts = pDest->m_VertexBuffers;
+        FMeshBufferSet& verts = Result->m_VertexBuffers;
         for ( auto& buf: verts.m_buffers )
         {
             for ( auto& chan: buf.m_channels )
@@ -80,8 +91,6 @@ namespace mu
                 }
             }
         }
-
-        return pDest;
     }
 
 }
diff --git a/Engine/Plugins/Experimental/Mutable/Source/MutableRuntime/Private/MuR/OpMeshRemoveChart.h b/Engine/Plugins/Experimental/Mutable/Source/MutableRuntime/Private/MuR/OpMeshRemoveChart.h
--- a/Engine/Plugins/Experimental/Mutable/Source/MutableRuntime/Private/MuR/OpMeshRemoveChart.h
+++ b/Engine/Plugins/Experimental/Mutable/Source/MutableRuntime/Private/MuR/OpMeshRemoveChart.h
@@ -75,32 +75,33 @@ namespace mu
 
 
     //---------------------------------------------------------------------------------------------
-    inline MeshPtr MeshRemoveMask( const Mesh* pSource, const Mesh* pMask )
+    inline void MeshRemoveMask(Mesh* Result, const Mesh* pSource, const Mesh* pMask, bool& bOutSuccess)
     {
         MUTABLE_CPUPROFILER_SCOPE(MeshRemoveMask);
+		bOutSuccess = true;
 
-        MeshPtr pResult = pSource->Clone();
-
-        if (!pMask->GetVertexCount() || !pResult->GetVertexCount() ||
-			!pResult->GetIndexCount() )
+        if (!pMask->GetVertexCount() || !pSource->GetVertexCount() || !pSource->GetIndexCount())
         {
-            return pResult;
+			bOutSuccess = false;
+            return;
         }
 
+		Result->CopyFrom(*pSource);
+
         // TODO
-        pResult->SetFaceGroupCount( 0 );
+        Result->SetFaceGroupCount( 0 );
 
         MeshBufferIteratorConst<MBF_UINT32,uint32_t,1> itMaskVI( pMask->GetVertexBuffers(), MBS_VERTEXINDEX );
 
         int firstFreeVertex = 0;
 
         // For each source vertex, true if it is removed.
-        int resultVertexCount = pResult->GetVertexCount();
+        int resultVertexCount = Result->GetVertexCount();
 		uint8* removedVertices = reinterpret_cast<uint8*>(FMemory::Malloc(resultVertexCount, 16));
 		FMemory::Memzero(removedVertices,resultVertexCount);
         {
 			TArray<ID_INTERVAL> intervals;
-            ExtractVertexIndexIntervals( intervals, pResult.get() );
+            ExtractVertexIndexIntervals(intervals, Result);
 
             for ( int mv=0; mv<pMask->GetVertexBuffers().GetElementCount(); ++mv )
             {
@@ -121,16 +122,16 @@ namespace mu
         // These are indices as in the index buffer, not the absoulte vertex index as in the
         // vertexbuffer MBS_VERTEXINDEX buffers.
 		TArray<int> usedVertices;
-		usedVertices.Init( -1, pResult->GetVertexCount() );
+		usedVertices.Init( -1, Result->GetVertexCount() );
         {
             size_t removedIndices = 0;
 
-            if ( pResult->GetIndexBuffers().GetElementSize(0)==4 )
+            if ( Result->GetIndexBuffers().GetElementSize(0)==4 )
             {
-                MeshBufferIteratorConst<MBF_UINT32,uint32_t,1> itSource( pResult->GetIndexBuffers(), MBS_VERTEXINDEX );
-                MeshBufferIterator<MBF_UINT32,uint32_t,1> itDest( pResult->GetIndexBuffers(), MBS_VERTEXINDEX );
+                MeshBufferIteratorConst<MBF_UINT32,uint32_t,1> itSource( Result->GetIndexBuffers(), MBS_VERTEXINDEX );
+                MeshBufferIterator<MBF_UINT32,uint32_t,1> itDest( Result->GetIndexBuffers(), MBS_VERTEXINDEX );
 
-                int indexCount = pResult->GetIndexCount();
+                int indexCount = Result->GetIndexCount();
                 for ( int f=0; f<indexCount/3; ++f )
                 {
                     uint32_t sourceIndices[3];
@@ -168,12 +169,12 @@ namespace mu
                 removedIndices = itSource - itDest;
             }
 
-            else if ( pResult->GetIndexBuffers().GetElementSize(0)==2 )
+            else if ( Result->GetIndexBuffers().GetElementSize(0)==2 )
             {
-                MeshBufferIteratorConst<MBF_UINT16,uint16,1> itSource( pResult->GetIndexBuffers(), MBS_VERTEXINDEX );
-                MeshBufferIterator<MBF_UINT16,uint16,1> itDest( pResult->GetIndexBuffers(), MBS_VERTEXINDEX );
+                MeshBufferIteratorConst<MBF_UINT16,uint16,1> itSource( Result->GetIndexBuffers(), MBS_VERTEXINDEX );
+                MeshBufferIterator<MBF_UINT16,uint16,1> itDest( Result->GetIndexBuffers(), MBS_VERTEXINDEX );
 
-                int indexCount = pResult->GetIndexCount();
+                int indexCount = Result->GetIndexCount();
                 for ( int f=0; f<indexCount/3; ++f )
                 {
                     uint16 sourceIndices[3];
@@ -219,23 +220,23 @@ namespace mu
 
             check( removedIndices%3==0 );
 
-            int faceCount = pResult->GetFaceCount();
-            pResult->GetFaceBuffers().SetElementCount( faceCount-(int)removedIndices/3 );
-            pResult->GetIndexBuffers().SetElementCount( faceCount*3-(int)removedIndices );
+            int faceCount = Result->GetFaceCount();
+            Result->GetFaceBuffers().SetElementCount( faceCount-(int)removedIndices/3 );
+            Result->GetIndexBuffers().SetElementCount( faceCount*3-(int)removedIndices );
         }
 
         FMemory::Free( removedVertices );
 
         // Rebuild the vertex buffers
-        for ( int b=0; b<pResult->GetVertexBuffers().GetBufferCount(); ++b )
+        for ( int b=0; b<Result->GetVertexBuffers().GetBufferCount(); ++b )
         {
-            int elemSize = pResult->GetVertexBuffers().GetElementSize( b );
+            int elemSize = Result->GetVertexBuffers().GetElementSize( b );
             const uint8_t* pSourceData = pSource->GetVertexBuffers().GetBufferData( b );
-            uint8_t* pData = pResult->GetVertexBuffers().GetBufferData( b );
-            for ( int v=0; v<pResult->GetVertexCount(); ++v )
+            uint8_t* pData = Result->GetVertexBuffers().GetBufferData( b );
+            for ( int v=0; v<Result->GetVertexCount(); ++v )
             {
                 int span = 0;
-                for ( int s=0; v+s<pResult->GetVertexCount(); ++s )
+                for ( int s=0; v+s<Result->GetVertexCount(); ++s )
                 {
                     if ( usedVertices[v+s]>=0 )
                     {
@@ -269,14 +270,11 @@ namespace mu
                 }
             }
         }
-        pResult->GetVertexBuffers().SetElementCount( firstFreeVertex );
+        Result->GetVertexBuffers().SetElementCount( firstFreeVertex );
 
         // Rebuild surface data.
         // \todo: For now make single surfaced
-        pResult->m_surfaces.Empty();
-        pResult->EnsureSurfaceData();
-
-        return pResult;
+        Result->m_surfaces.Empty();
+        Result->EnsureSurfaceData();
     }
-
 }
diff --git a/Engine/Plugins/Experimental/Mutable/Source/MutableRuntime/Private/MuR/OpMeshReshape.h b/Engine/Plugins/Experimental/Mutable/Source/MutableRuntime/Private/MuR/OpMeshReshape.h
--- a/Engine/Plugins/Experimental/Mutable/Source/MutableRuntime/Private/MuR/OpMeshReshape.h
+++ b/Engine/Plugins/Experimental/Mutable/Source/MutableRuntime/Private/MuR/OpMeshReshape.h
@@ -704,13 +704,16 @@ namespace mu
 	//! Rebuild the (previously bound) mesh data for a new shape.
 	//! Proof-of-concept implementation.
 	//---------------------------------------------------------------------------------------------
-	inline MeshPtr MeshApplyShape(const Mesh* BaseMesh, const Mesh* ShapeMesh, EMeshBindShapeFlags BindFlags)
+	inline void MeshApplyShape(Mesh* Result, const Mesh* BaseMesh, const Mesh* ShapeMesh, EMeshBindShapeFlags BindFlags, bool& bOutSuccess)
 	{
 		MUTABLE_CPUPROFILER_SCOPE(MeshApplyReshape);
 
+		bOutSuccess = true;
+
 		if (!BaseMesh)
 		{
-			return nullptr;
+			bOutSuccess = false;
+			return;
 		}
 		
 		const bool bReshapeVertices = EnumHasAnyFlags(BindFlags, EMeshBindShapeFlags::ReshapeVertices);
@@ -723,7 +726,8 @@ namespace mu
 
 		if (!bReshapeVertices && !bSkeletonModification && !bPhysicsModification)
 		{
-			return nullptr;
+			bOutSuccess = false;
+			return;
 		}
 	
 		// \TODO: Multiple binding data support
@@ -735,9 +739,9 @@ namespace mu
 		const FMeshBufferSet& VB = BaseMesh->GetVertexBuffers();
 		VB.FindChannel(MBS_BARYCENTRICCOORDS, BindingDataIndex, &BarycentricDataBuffer, &BarycentricDataChannel);
 		
-		// Clone Without VertexBuffers or AdditionalBuffers
-		constexpr EMeshCloneFlags CloneFlags = ~(EMeshCloneFlags::WithVertexBuffers | EMeshCloneFlags::WithAdditionalBuffers);
-		MeshPtr Result = BaseMesh->Clone(CloneFlags);
+		// Copy Without VertexBuffers or AdditionalBuffers
+		constexpr EMeshCopyFlags CopyFlags = ~(EMeshCopyFlags::WithVertexBuffers | EMeshCopyFlags::WithAdditionalBuffers);
+		Result->CopyFrom(*BaseMesh, CopyFlags);
 	
 		FMeshBufferSet& ResultBuffers = Result->GetVertexBuffers();
 
@@ -773,14 +777,14 @@ namespace mu
 		}
 		if (!ShapeMesh)
 		{
-			return Result;
+			return;
 		}
 
 		int32 ShapeVertexCount = ShapeMesh->GetVertexCount();
 		int32 ShapeTriangleCount = ShapeMesh->GetFaceCount();
 		if (!ShapeVertexCount || !ShapeTriangleCount)
-		{
-			return Result;
+		{	
+			return;
 		}
 
 		// Generate the temp vertex query data for the shape
@@ -842,7 +846,7 @@ namespace mu
 					(const FReshapeVertexBindingData*)VB.GetBufferData(BarycentricDataBuffer),
 					VB.GetElementCount());
 
-			ApplyToVertices(Result.get(), VerticesBindingData, ShapeDescriptor);
+			ApplyToVertices(Result, VerticesBindingData, ShapeDescriptor);
 		}
 	
 		if (bReshapeSkeleton)
@@ -880,7 +884,7 @@ namespace mu
 				TArrayView<const int32> BoneIndices( 
 						(const int32*)SkeletonBindBuffer->GetBufferData(1), SkeletonBindBuffer->GetElementCount());
 
-				ApplyToPose(Result.get(), SkeletonBindingData, BoneIndices, ShapeDescriptor);
+				ApplyToPose(Result, SkeletonBindingData, BoneIndices, ShapeDescriptor);
 			}
 		}
 
@@ -939,7 +943,5 @@ namespace mu
 				ApplyToAllPhysicsBodies(*Result, *BaseMesh, BindingData, UsedIndices, Offsets, ShapeDescriptor);
 			}
 		}
-		
-		return Result;
 	}
 }
diff --git a/Engine/Plugins/Experimental/Mutable/Source/MutableRuntime/Private/MuR/OpMeshTransform.h b/Engine/Plugins/Experimental/Mutable/Source/MutableRuntime/Private/MuR/OpMeshTransform.h
--- a/Engine/Plugins/Experimental/Mutable/Source/MutableRuntime/Private/MuR/OpMeshTransform.h
+++ b/Engine/Plugins/Experimental/Mutable/Source/MutableRuntime/Private/MuR/OpMeshTransform.h
@@ -13,20 +13,24 @@ namespace mu
     //---------------------------------------------------------------------------------------------
     //! Reference version
     //---------------------------------------------------------------------------------------------
-    inline MeshPtr MeshTransform( const Mesh* pBase, const FMatrix44f& transform )
+    inline void MeshTransform(Mesh* Result, const Mesh* pBase, const FMatrix44f& transform, bool& bOutSuccess)
 	{
-		MeshPtr pDest = pBase->Clone();
+		bOutSuccess = true;
 
         uint32_t vcount = pBase->GetVertexBuffers().GetElementCount();
 
         if ( !vcount )
 		{
-			return pDest;
+			bOutSuccess = false;
+			return;
 		}
 
+
+		Result->CopyFrom(*pBase);
+
 		FMatrix44f transformIT = transform.Inverse().GetTransposed();
 
-        const FMeshBufferSet& MBSPriv = pDest->GetVertexBuffers();
+        const FMeshBufferSet& MBSPriv = Result->GetVertexBuffers();
         for ( int32 b=0; b<MBSPriv.m_buffers.Num(); ++b )
         {
 
@@ -35,7 +39,7 @@ namespace mu
                 MESH_BUFFER_SEMANTIC sem = MBSPriv.m_buffers[b].m_channels[c].m_semantic;
                 int semIndex = MBSPriv.m_buffers[b].m_channels[c].m_semanticIndex;
 
-                UntypedMeshBufferIterator it( pDest->GetVertexBuffers(), sem, semIndex );
+                UntypedMeshBufferIterator it( Result->GetVertexBuffers(), sem, semIndex );
 
                 switch ( sem )
                 {
@@ -86,9 +90,6 @@ namespace mu
                 }
             }
         }
-
-
-		return pDest;
 	}
 
 }
diff --git a/Engine/Plugins/Experimental/Mutable/Source/MutableRuntime/Private/MuR/System.cpp b/Engine/Plugins/Experimental/Mutable/Source/MutableRuntime/Private/MuR/System.cpp
--- a/Engine/Plugins/Experimental/Mutable/Source/MutableRuntime/Private/MuR/System.cpp
+++ b/Engine/Plugins/Experimental/Mutable/Source/MutableRuntime/Private/MuR/System.cpp
@@ -166,6 +166,8 @@ namespace mu
 		m_pD->WorkingMemoryManager.PooledImages.Empty();
 		m_pD->WorkingMemoryManager.CacheResources.Empty();
 		check(m_pD->WorkingMemoryManager.TempImages.IsEmpty());
+		check(m_pD->WorkingMemoryManager.TempMeshes.IsEmpty());
+
 		m_pD->WorkingMemoryManager.TrackedBudgetBytes_Pooled = 0;
 		m_pD->WorkingMemoryManager.TrackedBudgetBytes_Temp = 0;
 		m_pD->WorkingMemoryManager.TrackedBudgetBytes_Cached = 0;
@@ -363,10 +365,15 @@ namespace mu
 				mu::OP_TYPE opType = program.GetOpType(at);
 				if (GetOpDataType(opType) == DT_IMAGE)
 				{
+					// GetImageDesc may call normal execution paths where meshes are computed.
+					m_pD->WorkingMemoryManager.BeginRunnerThread();
+					
 					int8 executionOptions = 0;
 					CodeRunner Runner(m_pD->Settings, m_pD, EExecutionStrategy::MinimizeMemory, pLiveInstance->Model, pLiveInstance->OldParameters.get(), at, System::AllLODs, executionOptions, 0, FScheduledOp::EType::ImageDesc);
 					Runner.Run();
 					Runner.GetImageDescResult(OutDesc);
+
+					m_pD->WorkingMemoryManager.EndRunnerThread();
 				}
 
 				break;
@@ -479,6 +486,7 @@ namespace mu
 						}
 					}
 				}
+
 				for (Ptr<const Mesh>& Data : Instance.Cache->MeshResults)
 				{
 					if (Data)
@@ -809,10 +817,7 @@ namespace mu
 			if (!bUnrecoverableError)
 			{
 				Result = WorkingMemoryManager.LoadMesh(FCacheAddress(at, 0, 0), true);
-			}
-
-			// Debug check to see if we managed the op-hit-counts correctly
-			WorkingMemoryManager.CurrentInstanceCache->CheckHitCountsCleared();
+			}	
 		}
 
 		WorkingMemoryManager.EndRunnerThread();
@@ -918,6 +923,7 @@ namespace mu
 
 		InternalBytes += PooledImages.GetAllocatedSize();
 		InternalBytes += TempImages.GetAllocatedSize();
+		InternalBytes += TempMeshes.GetAllocatedSize();
 		InternalBytes += CacheResources.GetAllocatedSize();
 		InternalBytes += CachePerModel.GetAllocatedSize();
 		InternalBytes += LiveInstances.GetAllocatedSize();
@@ -1025,7 +1031,7 @@ namespace mu
 
 
     //---------------------------------------------------------------------------------------------
-    bool FWorkingMemoryManager::EnsureBudgetBelow( uint64 AdditionalMemory )
+	bool FWorkingMemoryManager::EnsureBudgetBelow( uint64 AdditionalMemory )
     {
 		MUTABLE_CPUPROFILER_SCOPE(EnsureBudgetBelow);
 
@@ -1043,7 +1049,7 @@ namespace mu
 			MUTABLE_CPUPROFILER_SCOPE(EnsureBudgetBelow_Measure);
 			RomBytes = TrackedBudgetBytes_Rom; // GetRomBytes();
 			PoolBytes = TrackedBudgetBytes_Pooled; // GetPooledBytes();
-			TempBytes = TrackedBudgetBytes_Temp; // GetTempBytes();
+			TempBytes = TrackedBudgetBytes_Temp + GetTempMeshesBytes(); // GetTempBytes();
 			TrackedCacheBytes = TrackedBudgetBytes_Cached; // GetTrackedCacheBytes();
 		}
 
@@ -1222,6 +1228,12 @@ namespace mu
 
 						if (Value == Removed)
 						{
+							int32 RemoveCount = CacheResources.Remove(Removed);
+							if (RemoveCount > 0)
+							{
+								TrackedBudgetBytes_Cached -= RemovedDataSize;
+							}
+
 							// \TODO: This is not very efficient, since it will search the data again.
 							Instance.Cache->SetUnused(RemIt.get_address());
 						}
@@ -1399,7 +1411,6 @@ namespace mu
         return bFinished;
     }
 
-
     //---------------------------------------------------------------------------------------------
     void FWorkingMemoryManager::MarkRomUsed( int32 romIndex, const TSharedPtr<const Model>& pModel )
     {
diff --git a/Engine/Plugins/Experimental/Mutable/Source/MutableRuntime/Private/MuR/SystemPrivate.h b/Engine/Plugins/Experimental/Mutable/Source/MutableRuntime/Private/MuR/SystemPrivate.h
--- a/Engine/Plugins/Experimental/Mutable/Source/MutableRuntime/Private/MuR/SystemPrivate.h
+++ b/Engine/Plugins/Experimental/Mutable/Source/MutableRuntime/Private/MuR/SystemPrivate.h
@@ -1136,6 +1136,7 @@ namespace mu
 
 		/** List of intermediate resources that are not soterd anywhere yet. They are still locally referenced by code. */
 		TArray<Ptr<const Image>> TempImages;
+		TArray<Ptr<const Mesh>> TempMeshes;
 
 		/** */
 		TMap<Ptr<const Resource>, int32> CacheResources;
@@ -1153,7 +1154,7 @@ namespace mu
         void MarkRomUsed( int32 RomIndex, const TSharedPtr<const Model>& );
 
 		/** */
-		Ptr<Image> CreateImage(uint32 SizeX, uint32 SizeY, uint32 Lods, EImageFormat Format, EInitializationType Init = EInitializationType::Black)
+		UE_NODISCARD Ptr<Image> CreateImage(uint32 SizeX, uint32 SizeY, uint32 Lods, EImageFormat Format, EInitializationType Init = EInitializationType::Black)
 		{
 			CheckRunnerThread();
 
@@ -1199,7 +1200,7 @@ namespace mu
 		}
 
 		/** Ref will be nulled and relesed in any case. */
-		Ptr<Image> CloneOrTakeOver(Ptr<const Image>& Resource)
+		UE_NODISCARD Ptr<Image> CloneOrTakeOver(Ptr<const Image>& Resource)
 		{
 			CheckRunnerThread();
 
@@ -1208,7 +1209,7 @@ namespace mu
 			{
 				TrackedBudgetBytes_Temp -= Resource->GetDataSize();
 			}
-
+			
 			check(!TempImages.Contains(Resource));
 			check(!PooledImages.Contains(Resource));
 
@@ -1271,7 +1272,6 @@ namespace mu
 			Resource = nullptr;
 		}
 
-
 		/** */
 		void Release(Ptr<Image>& Resource)
 		{
@@ -1311,8 +1311,79 @@ namespace mu
 			Resource = nullptr;
 		}
 
+		UE_NODISCARD Ptr<Mesh> CreateMesh(int32 BudgetReserveSize)
+		{
+			CheckRunnerThread();
+
+			EnsureBudgetBelow(BudgetReserveSize);
+
+			Ptr<Mesh> Result = new Mesh();
+
+			TempMeshes.Add(Result);
+			//TrackedBudgetBytes_Temp += Result->GetDataSize();
+
+			return Result;
+		}
+
+		UE_NODISCARD Ptr<Mesh> CloneOrTakeOver(Ptr<const Mesh>& Resource)
+		{
+			CheckRunnerThread();
+
+			const int32 ResourceDataSize = Resource->GetDataSize();
+			const int32 Removed = TempMeshes.RemoveSingle(Resource);
+			if (Removed)
+			{
+				//TrackedBudgetBytes_Temp -= ResourceDataSize;
+			}
+
+			Ptr<Mesh> Result;
+			if (!Resource->IsUnique())
+			{
+				Result = CreateMesh(ResourceDataSize); 
+				Result->CopyFrom(*Resource);
+				Release(Resource);
+			}
+			else
+			{
+				Result = Ptr<Mesh>(const_cast<Mesh*>(Resource.get()));
+				Resource = nullptr;
+			}
+
+			return Result;
+		}
+
+		void Release(Ptr<const Mesh>& Resource)
+		{
+			CheckRunnerThread();
+
+			if (!Resource)
+			{
+				return;
+			}
+
+			int32 ResourceDataSize = Resource->GetDataSize();
+			int32 Removed = TempMeshes.RemoveSingle(Resource);
+			if (Removed)
+			{
+				//TrackedBudgetBytes_Temp -= ResourceDataSize;
+			}
+
+			EnsureBudgetBelow(0);
+
+			Resource = nullptr;
+		}
+
+		void Release(Ptr<Mesh>& Resource)
+		{
+			Ptr<const Mesh> ConstPtr = Resource;
+
+			Release(ConstPtr);
+
+			Resource = nullptr;
+		}
+
 		/** */
-		Ptr<const Mesh> LoadMesh(const FCacheAddress& From, bool bTakeOwnership=false)
+		UE_NODISCARD Ptr<const Mesh> LoadMesh(const FCacheAddress& From, bool bTakeOwnership = false)
 		{
 			bool bIsLastReference = false;
 			Ptr<const Mesh> Result = CurrentInstanceCache->GetMesh(From, bIsLastReference);
@@ -1321,30 +1392,32 @@ namespace mu
 				return nullptr;
 			}
 
-			// TODO
-			// If we retrieved the last reference to this resource in "From" cache position (it could still be in other cache positions as well)
-			//if (bIsLastReference)
-			//{
-			//	int32* CountPtr = CacheResources.Find(Result);
-			//	check(CountPtr);
-			//	*CountPtr = (*CountPtr) - 1;
-			//	if (!*CountPtr)
-			//	{
-			//		CacheResources.FindAndRemoveChecked(Result);
-			//	}
-			//}
+			int32 ResultDataSize = Result->GetDataSize();
 
+			// If we retrieved the last reference to this resource in "From" cache position (it could still be in other cache positions as well)
+			if (bIsLastReference)
+			{
+				int32* CountPtr = CacheResources.Find(Result);
+				check(CountPtr);
+				*CountPtr = (*CountPtr) - 1;
+				if (!*CountPtr)
+				{
+					CacheResources.FindAndRemoveChecked(Result);
+					TrackedBudgetBytes_Cached -= ResultDataSize;
+				}
+			}
 
-			//if (!bTakeOwnership && Result->IsUnique())
-			//{
-			//	TempMesh.Add(Result);
-			//}
+			if (!bTakeOwnership && Result->IsUnique())
+			{
+				TempMeshes.Add(Result);
+				//TrackedBudgetBytes_Temp += ResultDataSize;
+			}
 
 			return Result;
 		}
 
 		/** */
-		Ptr<const Image> LoadImage(const FCacheAddress& From, bool bTakeOwnership = false)
+		UE_NODISCARD Ptr<const Image> LoadImage(const FCacheAddress& From, bool bTakeOwnership = false)
 		{
 			bool bIsLastReference = false;
 			Ptr<const Image> Result = CurrentInstanceCache->GetImage(From, bIsLastReference);
@@ -1404,19 +1477,41 @@ namespace mu
 			CurrentInstanceCache->SetImage(To, Resource);
 		}
 
+		void StoreMesh(const FCacheAddress& To, Ptr<const Mesh> Resource)
+		{
+			if (Resource)
+			{
+				int32 ResourceDataSize = Resource->GetDataSize();
+				int32 Removed = TempMeshes.RemoveSingle(Resource);
+				if (Removed)
+				{
+					//TrackedBudgetBytes_Temp -= ResourceDataSize;
+				}
+
+				int32& Count = CacheResources.FindOrAdd(Resource, 0);
+				if (!Count)
+				{
+					TrackedBudgetBytes_Cached += ResourceDataSize;
+				}
+
+				++Count;
+			}
+
+			CurrentInstanceCache->SetMesh(To, Resource);
+		}
 
 		/** Return true if the resource is not in any cache (0,1,rom). */
-		bool IsBudgetTemp(const Ptr<const Image>& Resource)
+		bool IsBudgetTemp(const Ptr<const Resource>& Resource)
 		{
 			if (!Resource)
 			{
 				return false;
 			}
+
 			bool bIsTemp = Resource->IsUnique();
 			return bIsTemp;
 		}
 
-
 		int32 GetPooledBytes() const
 		{
 			int32 Result = 0;
@@ -1430,6 +1525,16 @@ namespace mu
 			return Result;
 		}
 
+		int32 GetTempMeshesBytes() const
+		{
+			int32 Result = 0;
+			for (const Ptr<const Mesh>& Value : TempMeshes)
+			{
+				Result += Value->GetDataSize();
+			}
+
+			return Result;
+		}
 
 		int32 GetTempBytes() const
 		{
@@ -1439,7 +1544,7 @@ namespace mu
 				Result += Value->GetDataSize();
 			}
 
-			return Result;
+			return Result + GetTempMeshesBytes();
 		}
 
 
@@ -1489,7 +1594,6 @@ namespace mu
 			return Result;
 		}
 
-
 		/** Calculate the amount of bytes in data cached in the level 0 and 1 cache in all live instances. */
 		int32 GetCacheBytes() const
 		{
@@ -1550,7 +1654,6 @@ namespace mu
 			return Result;
 		}
 
-
 		/** Remove all intermediate data (big and small) from the memory except for the one that has been explicitely
 		* marked as state cache.
 		*/
@@ -1719,7 +1822,8 @@ namespace mu
 			// If this check fails it means some operation is not correctly handling resource management and didn't release
 			// a resource it created.
 			// This should be reported and reviewed, but it is not fatal. Some unnecessary memory may be used temporarily.
-			ensure(TempImages.Num()==0);
+			ensure(TempImages.Num() == 0);
+			ensure(TempMeshes.Num() == 0);
 
 #if !(UE_BUILD_SHIPPING || UE_BUILD_TEST)
 			// If this check fails it means have not set up correctly all the paths to debug threading for resource management.
diff --git a/Engine/Plugins/Experimental/Mutable/Source/MutableRuntime/Public/MuR/Mesh.h b/Engine/Plugins/Experimental/Mutable/Source/MutableRuntime/Public/MuR/Mesh.h
--- a/Engine/Plugins/Experimental/Mutable/Source/MutableRuntime/Public/MuR/Mesh.h
+++ b/Engine/Plugins/Experimental/Mutable/Source/MutableRuntime/Public/MuR/Mesh.h
@@ -117,27 +117,30 @@ namespace mu
 		ClipDeformNormalProject = 3	
 	};
 
-	enum class EMeshCloneFlags : uint32
+	enum class EMeshCopyFlags : uint32
 	{
 		None = 0,
-		WithSkeletalMesh      = 1 << 1,
-		WithSurfaces          = 1 << 2,
-		WithSkeleton          = 1 << 3,
-		WithPhysicsBody       = 1 << 4,
-		WithFaceGroups        = 1 << 5,
-		WithTags              = 1 << 6,
-		WithVertexBuffers     = 1 << 7,
-		WithIndexBuffers      = 1 << 8,
-		WithFaceBuffers       = 1 << 9,
+		WithSkeletalMesh = 1 << 1,
+		WithSurfaces = 1 << 2,
+		WithSkeleton = 1 << 3,
+		WithPhysicsBody = 1 << 4,
+		WithFaceGroups = 1 << 5,
+		WithTags = 1 << 6,
+		WithVertexBuffers = 1 << 7,
+		WithIndexBuffers = 1 << 8,
+		WithFaceBuffers = 1 << 9,
 		WithAdditionalBuffers = 1 << 10,
-		WithLayouts           = 1 << 11,
-		WithPoses			  = 1 << 12,
-		WithBoneMap			  = 1 << 13,
-		WithSkeletonIDs		  = 1 << 14,
-		WithAdditionalPhysics = 1 << 15
+		WithLayouts = 1 << 11,
+		WithPoses = 1 << 12,
+		WithBoneMap = 1 << 13,
+		WithSkeletonIDs = 1 << 14,
+		WithAdditionalPhysics = 1 << 15,
+
+		AllFlags = 0xFFFFFFFF
 	};
 	
-	ENUM_CLASS_FLAGS(EMeshCloneFlags);
+	ENUM_CLASS_FLAGS(EMeshCopyFlags);
+
 
     //! \brief Mesh object containing any number of buffers with any number of channels.
     //! The buffers can be per-index or per-vertex.
@@ -157,7 +160,10 @@ namespace mu
         MeshPtr Clone() const;
 		
 		// Clone with flags allowing to not include some parts in the cloned mesh
-		MeshPtr Clone(EMeshCloneFlags Flags) const;
+		MeshPtr Clone(EMeshCopyFlags Flags) const;
+
+		// Copy form another mesh.
+		void CopyFrom(const Mesh& From, EMeshCopyFlags Flags = EMeshCopyFlags::AllFlags);
 
         //! Serialisation
         static void Serialise( const Mesh* p, OutputArchive& arch );
diff --git a/Engine/Plugins/Experimental/Mutable/Source/MutableRuntime/Public/MuR/MeshBufferSet.h b/Engine/Plugins/Experimental/Mutable/Source/MutableRuntime/Public/MuR/MeshBufferSet.h
--- a/Engine/Plugins/Experimental/Mutable/Source/MutableRuntime/Public/MuR/MeshBufferSet.h
+++ b/Engine/Plugins/Experimental/Mutable/Source/MutableRuntime/Public/MuR/MeshBufferSet.h
@@ -356,6 +356,8 @@ namespace mu
 		//! Get the total memory size of the buffers and this struct
 		int32 GetDataSize() const;
 
+		int32 GetAllocatedSize() const;
+
 		//! Compare the mesh buffer with another one, but ignore internal data like generated
 		//! vertex indices.
 		bool IsSpecialBufferToIgnoreInSimilar(const MESH_BUFFER& b) const;
diff --git a/Engine/Plugins/Experimental/Mutable/Source/MutableTools/Private/MuT/CodeGenerator_Image.cpp b/Engine/Plugins/Experimental/Mutable/Source/MutableTools/Private/MuT/CodeGenerator_Image.cpp
--- a/Engine/Plugins/Experimental/Mutable/Source/MutableTools/Private/MuT/CodeGenerator_Image.cpp
+++ b/Engine/Plugins/Experimental/Mutable/Source/MutableTools/Private/MuT/CodeGenerator_Image.cpp
@@ -1412,8 +1412,11 @@ namespace mu
                 // will remove the faces that are not in the layout block we are generating.
                 Ptr<ASTOpConstantResource> cop = new ASTOpConstantResource();
                 cop->type = OP_TYPE::ME_CONSTANT;
-				MeshPtr pFormatMesh = CreateMeshOptimisedForWrappingProjection(node.m_layout);
-                cop->SetValue( pFormatMesh, m_compilerOptions->m_optimisationOptions.m_useDiskCache );
+
+				Ptr<Mesh> FormatMeshResult = new Mesh();
+				CreateMeshOptimisedForWrappingProjection(FormatMeshResult.get(), node.m_layout);
+
+                cop->SetValue(FormatMeshResult, m_compilerOptions->m_optimisationOptions.m_useDiskCache);
 
                 Ptr<ASTOpMeshFormat> fop = new ASTOpMeshFormat();
                 fop->Buffers = OP::MeshFormatArgs::BT_VERTEX
@@ -1441,8 +1444,11 @@ namespace mu
                 // Reformat the mesh to a more efficient format for this operation
                 Ptr<ASTOpConstantResource> cop = new ASTOpConstantResource();
                 cop->type = OP_TYPE::ME_CONSTANT;
-                MeshPtr pFormatMesh = CreateMeshOptimisedForProjection(node.m_layout);
-                cop->SetValue( pFormatMesh, m_compilerOptions->m_optimisationOptions.m_useDiskCache );
+
+				Ptr<Mesh> FormatMeshResult = new Mesh();
+                CreateMeshOptimisedForProjection(FormatMeshResult.get(), node.m_layout);
+
+                cop->SetValue(FormatMeshResult, m_compilerOptions->m_optimisationOptions.m_useDiskCache);
 
                 Ptr<ASTOpMeshFormat> fop = new ASTOpMeshFormat();
                 fop->Buffers = OP::MeshFormatArgs::BT_VERTEX
diff --git a/Engine/Plugins/Experimental/Mutable/Source/MutableTools/Private/MuT/CodeOptimiser.cpp b/Engine/Plugins/Experimental/Mutable/Source/MutableTools/Private/MuT/CodeOptimiser.cpp
--- a/Engine/Plugins/Experimental/Mutable/Source/MutableTools/Private/MuT/CodeOptimiser.cpp
+++ b/Engine/Plugins/Experimental/Mutable/Source/MutableTools/Private/MuT/CodeOptimiser.cpp
@@ -1104,25 +1104,25 @@ namespace mu
 					return true;
 				}
 
-				for( auto& rep: replacementsFound )
+				for(ADDMESH_SKELETON& Rep: replacementsFound) 
 				{
-					if ( rep.m_contributingMeshes.Contains(at) )
+					if (Rep.m_contributingMeshes.Contains(at))
 					{
 						mu::Ptr<const Mesh> pMesh =  static_cast<const Mesh*>(typedOp->GetValue().get());
 						pMesh->CheckIntegrity();
 
-						mu::Ptr<Mesh> pNewMesh = MeshRemapSkeleton( pMesh.get(),
-																rep.m_pFinalSkeleton.get());
+						Ptr<Mesh> NewMesh = new Mesh();
+						bool bOutSuccess = false;
+						MeshRemapSkeleton(NewMesh.get(), pMesh.get(), Rep.m_pFinalSkeleton.get(), bOutSuccess);
 
-						// It returns null if there is no need to remap.
-						if (pNewMesh)
+						if (bOutSuccess)
 						{
-							pNewMesh->CheckIntegrity();
+							NewMesh->CheckIntegrity();
 							mu::Ptr<ASTOpConstantResource> newOp = new ASTOpConstantResource();
 							newOp->type = OP_TYPE::ME_CONSTANT;
-							newOp->SetValue( pNewMesh, options.m_useDiskCache );
+							newOp->SetValue(NewMesh, options.m_useDiskCache);
 
-							ASTOp::Replace( at, newOp );
+							ASTOp::Replace(at, newOp);
 						}
 					}
 				}
