[mutable] Optimization for the "reuse LOD images" feature.
Image code for lower LODs is not generated.
Layout expressions are reused across LODs.
Removed the global layout strategy option from the core libraries.
#rnx

diff --git a/Engine/Plugins/Experimental/Mutable/Source/CustomizableObject/Classes/MuCO/CustomizableObject.h b/Engine/Plugins/Experimental/Mutable/Source/CustomizableObject/Classes/MuCO/CustomizableObject.h
--- a/Engine/Plugins/Experimental/Mutable/Source/CustomizableObject/Classes/MuCO/CustomizableObject.h
+++ b/Engine/Plugins/Experimental/Mutable/Source/CustomizableObject/Classes/MuCO/CustomizableObject.h
@@ -1197,10 +1197,6 @@ public:
 	UPROPERTY()
 	FCompilationOptions CompileOptions;
 
-	// 
-	UPROPERTY(EditAnywhere, Category = CompileOptions)
-	bool bDisableTextureLayoutManagement = false;
-
 	//
 	UPROPERTY(EditAnywhere, Category = CompileOptions)
 	bool bEnableRealTimeMorphTargets = false;
@@ -1449,7 +1445,7 @@ private:
 	// This is a manual version number for the binary blobs in this asset.
 	// Increasing it invalidates all the previously compiled models.
 	// Warning: If while merging code both versions have changed, take the highest+1.
-	static const int32 CurrentSupportedVersion = 388;
+	static const int32 CurrentSupportedVersion = 389;
 
 public:
 
diff --git a/Engine/Plugins/Experimental/Mutable/Source/CustomizableObjectEditor/Private/MuCOE/CustomizableObjectCompileRunnable.cpp b/Engine/Plugins/Experimental/Mutable/Source/CustomizableObjectEditor/Private/MuCOE/CustomizableObjectCompileRunnable.cpp
--- a/Engine/Plugins/Experimental/Mutable/Source/CustomizableObjectEditor/Private/MuCOE/CustomizableObjectCompileRunnable.cpp
+++ b/Engine/Plugins/Experimental/Mutable/Source/CustomizableObjectEditor/Private/MuCOE/CustomizableObjectCompileRunnable.cpp
@@ -16,9 +16,8 @@ class ITargetPlatform;
 #define LOCTEXT_NAMESPACE "CustomizableObjectEditor"
 
 
-FCustomizableObjectCompileRunnable::FCustomizableObjectCompileRunnable(mu::NodePtr Root, bool bInDisableTextureLayout)
+FCustomizableObjectCompileRunnable::FCustomizableObjectCompileRunnable(mu::Ptr<mu::Node> Root)
 	: MutableRoot(Root)
-	, bDisableTextureLayout(bInDisableTextureLayout)
 	, bThreadCompleted(false)
 	, MutableIsDisabled(false)
 {
@@ -94,10 +93,6 @@ uint32 FCustomizableObjectCompileRunnable::Run()
 	
 	CompilerOptions->SetImageTiling(Options.ImageTiling);
 
-	CompilerOptions->SetTextureLayoutStrategy(bDisableTextureLayout
-		? mu::CompilerOptions::TextureLayoutStrategy::None
-		: mu::CompilerOptions::TextureLayoutStrategy::Pack);
-
 	mu::CompilerPtr Compiler = new mu::Compiler(CompilerOptions);
 
 	UE_LOG(LogMutable, Verbose, TEXT("PROFILE: [ %16.8f ] FCustomizableObjectCompileRunnable Compile start."), FPlatformTime::Seconds());
diff --git a/Engine/Plugins/Experimental/Mutable/Source/CustomizableObjectEditor/Private/MuCOE/CustomizableObjectCompileRunnable.h b/Engine/Plugins/Experimental/Mutable/Source/CustomizableObjectEditor/Private/MuCOE/CustomizableObjectCompileRunnable.h
--- a/Engine/Plugins/Experimental/Mutable/Source/CustomizableObjectEditor/Private/MuCOE/CustomizableObjectCompileRunnable.h
+++ b/Engine/Plugins/Experimental/Mutable/Source/CustomizableObjectEditor/Private/MuCOE/CustomizableObjectCompileRunnable.h
@@ -32,13 +32,13 @@ public:
 
 private:
 
-	mu::NodePtr MutableRoot;
+	mu::Ptr<mu::Node> MutableRoot;
 	TArray<FError> ArrayWarning;
 	TArray<FError> ArrayError;
 
 public:
 
-	FCustomizableObjectCompileRunnable(mu::NodePtr Root, bool bInDisableTextureLayout);
+	FCustomizableObjectCompileRunnable(mu::Ptr<mu::Node> Root);
 
 	// FRunnable interface
 	uint32 Run() override;
@@ -58,9 +58,6 @@ public:
 
 	TSharedPtr<mu::Model, ESPMode::ThreadSafe> Model;
 
-	// Texture packing strategy
-	bool bDisableTextureLayout;
-
 	FCompilationOptions Options;
 
 	FString ErrorMsg;
diff --git a/Engine/Plugins/Experimental/Mutable/Source/CustomizableObjectEditor/Private/MuCOE/CustomizableObjectCompiler.cpp b/Engine/Plugins/Experimental/Mutable/Source/CustomizableObjectEditor/Private/MuCOE/CustomizableObjectCompiler.cpp
--- a/Engine/Plugins/Experimental/Mutable/Source/CustomizableObjectEditor/Private/MuCOE/CustomizableObjectCompiler.cpp
+++ b/Engine/Plugins/Experimental/Mutable/Source/CustomizableObjectEditor/Private/MuCOE/CustomizableObjectCompiler.cpp
@@ -125,7 +125,7 @@ void FCustomizableObjectCompiler::MutableIsDisabledCase(UCustomizableObject* Obj
 	}
 
 	UE_LOG(LogMutable, Log, TEXT("Mutable has been disabled. To reenable it, please deactivate the Disable Mutale Option in the Plugins -> Mutable option"), FPlatformTime::Seconds());
-	CompileTask = MakeShareable(new FCustomizableObjectCompileRunnable(nullptr, true));
+	CompileTask = MakeShareable(new FCustomizableObjectCompileRunnable(nullptr));
 	CompileTask->MutableIsDisabled = true;
 	LaunchMutableCompile(false);
 }
@@ -523,7 +523,6 @@ mu::NodeObjectPtr FCustomizableObjectCompiler::GenerateMutableRoot(
 	UCustomizableObjectNodeObject* ActualRoot = Root;
 	UCustomizableObject* ActualRootObject = Object;
 
-	GenerationContext.bDisableTextureLayoutManagementFlag = Object->bDisableTextureLayoutManagement;
 	ArrayAlreadyProcessedChild.Empty();
 
 	if (Root->ObjectName.IsEmpty())
@@ -546,7 +545,6 @@ mu::NodeObjectPtr FCustomizableObjectCompiler::GenerateMutableRoot(
 
 			TArray<UCustomizableObject*> VisitedObjects;
 			ActualRootObject = Root->ParentObject ? GetFullGraphRootObject(Root, VisitedObjects) : Object;
-			GenerationContext.bDisableTextureLayoutManagementFlag = ActualRootObject->bDisableTextureLayoutManagement;
 
 			if (Root->ParentObject != nullptr)
 			{
@@ -621,8 +619,6 @@ mu::NodeObjectPtr FCustomizableObjectCompiler::GenerateMutableRoot(
 
 		VisitedObjects.Empty();
 		ActualRootObject = Root->ParentObject ? GetFullGraphRootObject(Root, VisitedObjects) : Object;
-		GenerationContext.bDisableTextureLayoutManagementFlag = ActualRootObject->bDisableTextureLayoutManagement;
-
 
 	}
 
@@ -1289,7 +1285,7 @@ void FCustomizableObjectCompiler::CompileInternal(UCustomizableObject* Object, c
 			UCustomizableObjectSystem::GetInstance()->LockObject(Object);
 		}
 
-		CompileTask = MakeShareable(new FCustomizableObjectCompileRunnable(MutableRoot, GenerationContext.bDisableTextureLayoutManagementFlag));
+		CompileTask = MakeShareable(new FCustomizableObjectCompileRunnable(MutableRoot));
 		CompileTask->Options = Options;
 		State = ECustomizableObjectCompilationState::InProgress;
 
@@ -1366,7 +1362,7 @@ void FCustomizableObjectCompiler::CompileInternal(UCustomizableObject* Object, c
 }
 
 
-mu::NodePtr FCustomizableObjectCompiler::Export(UCustomizableObject* Object, const FCompilationOptions& InCompilerOptions, bool& bOutDisableLayouts)
+mu::NodePtr FCustomizableObjectCompiler::Export(UCustomizableObject* Object, const FCompilationOptions& InCompilerOptions)
 {
 	UE_LOG(LogMutable, Log, TEXT("Started Customizable Object Export %s."), *Object->GetName());
 
@@ -1409,8 +1405,6 @@ mu::NodePtr FCustomizableObjectCompiler::Export(UCustomizableObject* Object, con
 		UpdatePendingTextureConversion(false);
 	}
 
-	bOutDisableLayouts = GenerationContext.bDisableTextureLayoutManagementFlag;
-
 	return MutableRoot;
 }
 
diff --git a/Engine/Plugins/Experimental/Mutable/Source/CustomizableObjectEditor/Private/MuCOE/CustomizableObjectCompiler.h b/Engine/Plugins/Experimental/Mutable/Source/CustomizableObjectEditor/Private/MuCOE/CustomizableObjectCompiler.h
--- a/Engine/Plugins/Experimental/Mutable/Source/CustomizableObjectEditor/Private/MuCOE/CustomizableObjectCompiler.h
+++ b/Engine/Plugins/Experimental/Mutable/Source/CustomizableObjectEditor/Private/MuCOE/CustomizableObjectCompiler.h
@@ -42,7 +42,7 @@ public:
 	virtual UCustomizableObject* GetRootObject(UCustomizableObject* InObject) override;
 
 	/** Generate the Mutable Graph from the Unreal Graph. */
-	mu::NodePtr Export(UCustomizableObject* Object, const FCompilationOptions& Options, bool& bOutDisableLayouts);
+	mu::NodePtr Export(UCustomizableObject* Object, const FCompilationOptions& Options);
 
 	void CompilerLog(const FText& Message, const TArray<const UCustomizableObjectNode*>& ArrayNode, EMessageSeverity::Type MessageSeverity = EMessageSeverity::Warning, bool bAddBaseObjectInfo = true);
 	void CompilerLog(const FText& Message, const UCustomizableObjectNode* Node = nullptr, EMessageSeverity::Type MessageSeverity = EMessageSeverity::Warning, bool bAddBaseObjectInfo = true);
diff --git a/Engine/Plugins/Experimental/Mutable/Source/CustomizableObjectEditor/Private/MuCOE/GenerateMutableSource/GenerateMutableSource.h b/Engine/Plugins/Experimental/Mutable/Source/CustomizableObjectEditor/Private/MuCOE/GenerateMutableSource/GenerateMutableSource.h
--- a/Engine/Plugins/Experimental/Mutable/Source/CustomizableObjectEditor/Private/MuCOE/GenerateMutableSource/GenerateMutableSource.h
+++ b/Engine/Plugins/Experimental/Mutable/Source/CustomizableObjectEditor/Private/MuCOE/GenerateMutableSource/GenerateMutableSource.h
@@ -669,9 +669,6 @@ struct FMutableGraphGenerationContext
 	TMap<FGraphCycleKey, const UCustomizableObject*> VisitedPins;
 	const UCustomizableObject* CustomizableObjectWithCycle = nullptr;
 
-	// Texture management flag, captured from the real object root node
-	bool bDisableTextureLayoutManagementFlag = false;
-
 	// Set of all the guids of all the child CustomizableObjects in the compilation
 	TSet<FGuid> CustomizableObjectGuidsInCompilation;
 
diff --git a/Engine/Plugins/Experimental/Mutable/Source/CustomizableObjectEditor/Private/MuCOE/GenerateMutableSource/GenerateMutableSourceSurface.cpp b/Engine/Plugins/Experimental/Mutable/Source/CustomizableObjectEditor/Private/MuCOE/GenerateMutableSource/GenerateMutableSourceSurface.cpp
--- a/Engine/Plugins/Experimental/Mutable/Source/CustomizableObjectEditor/Private/MuCOE/GenerateMutableSource/GenerateMutableSourceSurface.cpp
+++ b/Engine/Plugins/Experimental/Mutable/Source/CustomizableObjectEditor/Private/MuCOE/GenerateMutableSource/GenerateMutableSourceSurface.cpp
@@ -220,9 +220,10 @@ void SetSurfaceFormat( FMutableGraphGenerationContext& GenerationContext,
 
 void UpdateSharedSurfaceId(FMutableGraphGenerationContext& GenerationContext, UCustomizableObjectNodeMaterial* NodeMaterial, mu::NodeSurfaceNewPtr InNodeSurface)
 {
+	// \TODO: Reusability is actually per-texture: if a texture is set to not use layouts (setting LayoutIndex to -1) then it can always be reused.
 	// Reuse materials between LODs when using automatics LODs, if texture layout management is disabled or if bReuseMaterials is enabled in the node material.
 	const bool bCanShareSurface = GenerationContext.CurrentAutoLODStrategy == ECustomizableObjectAutomaticLODStrategy::AutomaticFromMesh
-		&& (GenerationContext.bDisableTextureLayoutManagementFlag || NodeMaterial->bReuseMaterialBetweenLODs);
+		&& NodeMaterial->bReuseMaterialBetweenLODs;
 
 	// Set shared surface Id to reuse materials and textures between LODs if automatic LODs from mesh is being used
 	if (bCanShareSurface && InNodeSurface)
diff --git a/Engine/Plugins/Experimental/Mutable/Source/CustomizableObjectEditor/Private/MuCOE/SCustomizableObjectEditorViewport.cpp b/Engine/Plugins/Experimental/Mutable/Source/CustomizableObjectEditor/Private/MuCOE/SCustomizableObjectEditorViewport.cpp
--- a/Engine/Plugins/Experimental/Mutable/Source/CustomizableObjectEditor/Private/MuCOE/SCustomizableObjectEditorViewport.cpp
+++ b/Engine/Plugins/Experimental/Mutable/Source/CustomizableObjectEditor/Private/MuCOE/SCustomizableObjectEditorViewport.cpp
@@ -910,69 +910,30 @@ void SCustomizableObjectEditorViewportTabBody::GenerateUVMaterialOptions()
 	ArrayUVMaterialOptionString.Empty();
 	
 	int32 ComponentIndex = 0;
-
 	for (UDebugSkelMeshComponent* PreviewSkeletalMeshComponent : PreviewSkeletalMeshComponents)
 	{
 		if (PreviewSkeletalMeshComponent != nullptr && UE_MUTABLE_GETSKINNEDASSET(PreviewSkeletalMeshComponent) != nullptr && UE_MUTABLE_GETSKINNEDASSET(PreviewSkeletalMeshComponent)->GetResourceForRendering() != nullptr)
 		{
-			TMap<FString, int32> MaterialLODs;
-
-			// Add Suffix "__X" for materials with duplicated names and Suffing " LOD_X" for multiple LODs.
-			const FSkeletalMeshRenderData* MeshRes = UE_MUTABLE_GETSKINNEDASSET(PreviewSkeletalMeshComponent)->GetResourceForRendering();
 			const TArray<UMaterialInterface*> Materials = PreviewSkeletalMeshComponent->GetMaterials();
-			for (UMaterialInterface* m : Materials)
+			const FSkeletalMeshRenderData* MeshRes = UE_MUTABLE_GETSKINNEDASSET(PreviewSkeletalMeshComponent)->GetResourceForRendering();
+			for (int32 LODIndex = 0; LODIndex < MeshRes->LODRenderData.Num(); ++LODIndex)
 			{
-				const UMaterial* BaseMaterial = m->GetBaseMaterial();
-
-				FString BaseMaterialName = BaseMaterial->GetName();
-
-				int32 LODCount = 0;
-				if (int32* Count = MaterialLODs.Find(BaseMaterialName))
+				for (int32 SectionIndex = 0; SectionIndex < MeshRes->LODRenderData[LODIndex].RenderSections.Num(); ++SectionIndex)
 				{
-					if ((*Count) < MeshRes->LODRenderData.Num() - 1)
+					const FSkelMeshRenderSection& Section = MeshRes->LODRenderData[LODIndex].RenderSections[SectionIndex];
+					
+					if (!Materials.IsValidIndex(Section.MaterialIndex))
 					{
-						(*Count)++;
-						LODCount = (*Count);
+						continue;
 					}
-					else
-					{
-						bool bSearchEmptyLOD = true;
-						int MaterialIndexWithinLOD = 1;
-
-						while (bSearchEmptyLOD)
-						{
-							if (int32* DupeMaterialCount = MaterialLODs.Find(BaseMaterialName + FString::Printf(TEXT("__%d"), MaterialIndexWithinLOD)))
-							{
-								if ((*DupeMaterialCount) < MeshRes->LODRenderData.Num() - 1)
-								{
-									(*DupeMaterialCount)++;
-									LODCount = (*DupeMaterialCount);
-									BaseMaterialName += FString::Printf(TEXT("__%d"), MaterialIndexWithinLOD);
-									bSearchEmptyLOD = false;
-								}
-								else
-								{
-									MaterialIndexWithinLOD++;
-								}
-							}
-							else
-							{
-								MaterialLODs.Add(BaseMaterialName + FString::Printf(TEXT("__%d"), MaterialIndexWithinLOD), 0);
-
-								BaseMaterialName += FString::Printf(TEXT("__%d"), MaterialIndexWithinLOD);
-								bSearchEmptyLOD = false;
-							}
-						}
-					}
-				}
-				else
-				{
-					MaterialLODs.Add(BaseMaterialName, LODCount);
-				}
 
-				BaseMaterialName += FString::Printf(TEXT(" LOD_%d_Component_%d"), LODCount, ComponentIndex);
+					const UMaterial* BaseMaterial = Materials[Section.MaterialIndex]->GetBaseMaterial();
+					FString BaseMaterialName = BaseMaterial->GetName();
+
+					BaseMaterialName += FString::Printf(TEXT(" LOD_%d_Component_%d"), LODIndex, ComponentIndex);
 
-				ArrayUVMaterialOptionString.Add(MakeShareable(new FString(BaseMaterialName)));
+					ArrayUVMaterialOptionString.Add(MakeShareable(new FString(BaseMaterialName)));
+				}
 			}
 		}
 
diff --git a/Engine/Plugins/Experimental/Mutable/Source/CustomizableObjectEditor/Private/MuCOE/SMutableGraphViewer.cpp b/Engine/Plugins/Experimental/Mutable/Source/CustomizableObjectEditor/Private/MuCOE/SMutableGraphViewer.cpp
--- a/Engine/Plugins/Experimental/Mutable/Source/CustomizableObjectEditor/Private/MuCOE/SMutableGraphViewer.cpp
+++ b/Engine/Plugins/Experimental/Mutable/Source/CustomizableObjectEditor/Private/MuCOE/SMutableGraphViewer.cpp
@@ -254,7 +254,7 @@ void SMutableGraphViewer::Construct(const FArguments& InArgs, const mu::NodePtr&
 void SMutableGraphViewer::CompileMutableCodePressed()
 {
 	// Do the compilation to Mutable Code synchronously.
-	TSharedPtr<FCustomizableObjectCompileRunnable> CompileTask = MakeShareable(new FCustomizableObjectCompileRunnable(RootNode, false));
+	TSharedPtr<FCustomizableObjectCompileRunnable> CompileTask = MakeShareable(new FCustomizableObjectCompileRunnable(RootNode));
 	CompileTask->Options = CompileOptions;
 	CompileTask->Init();
 	CompileTask->Run();
diff --git a/Engine/Plugins/Experimental/Mutable/Source/CustomizableObjectEditor/Private/MuCOE/SMutableObjectViewer.cpp b/Engine/Plugins/Experimental/Mutable/Source/CustomizableObjectEditor/Private/MuCOE/SMutableObjectViewer.cpp
--- a/Engine/Plugins/Experimental/Mutable/Source/CustomizableObjectEditor/Private/MuCOE/SMutableObjectViewer.cpp
+++ b/Engine/Plugins/Experimental/Mutable/Source/CustomizableObjectEditor/Private/MuCOE/SMutableObjectViewer.cpp
@@ -215,9 +215,7 @@ void SMutableObjectViewer::Construct(const FArguments& InArgs, UCustomizableObje
 void SMutableObjectViewer::GenerateMutableGraphPressed()
 {
 	// Convert from Unreal graph to Mutable graph.
-	bool bDisableLayouts = false;
-
-	mu::NodePtr RootNode = Compiler.Export(CustomizableObject, CompileOptions, bDisableLayouts);
+	mu::Ptr<mu::Node> RootNode = Compiler.Export(CustomizableObject, CompileOptions);
 	if (!RootNode)
 	{
 		// TODO: Show errors
@@ -241,8 +239,6 @@ void SMutableObjectViewer::GenerateMutableGraphPressed()
 
 void SMutableObjectViewer::CompileMutableCodePressed()
 {
-	bool bDisableLayouts = false;
-
 	if (CompileOptions.bForceLargeLODBias)
 	{
 		// Debug compile with many different biasses
@@ -251,7 +247,7 @@ void SMutableObjectViewer::CompileMutableCodePressed()
 		{
 			CompileOptions.DebugBias = Bias;
 
-			mu::NodePtr RootNode = Compiler.Export(CustomizableObject, CompileOptions, bDisableLayouts);
+			mu::NodePtr RootNode = Compiler.Export(CustomizableObject, CompileOptions);
 			if (!RootNode)
 			{
 				// TODO: Show errors
@@ -260,7 +256,7 @@ void SMutableObjectViewer::CompileMutableCodePressed()
 			}
 
 			// Do the compilation to Mutable Code synchronously.
-			TSharedPtr<FCustomizableObjectCompileRunnable> CompileTask = MakeShareable(new FCustomizableObjectCompileRunnable(RootNode, bDisableLayouts));
+			TSharedPtr<FCustomizableObjectCompileRunnable> CompileTask = MakeShareable(new FCustomizableObjectCompileRunnable(RootNode));
 			CompileTask->Options = CompileOptions;
 			CompileTask->Init();
 			CompileTask->Run();
@@ -269,7 +265,7 @@ void SMutableObjectViewer::CompileMutableCodePressed()
 	}
 
 	// Convert from Unreal graph to Mutable graph.
-	mu::NodePtr RootNode = Compiler.Export(CustomizableObject, CompileOptions, bDisableLayouts);
+	mu::NodePtr RootNode = Compiler.Export(CustomizableObject, CompileOptions);
 	if (!RootNode)
 	{
 		// TODO: Show errors
@@ -278,7 +274,7 @@ void SMutableObjectViewer::CompileMutableCodePressed()
 	}
 
 	// Do the compilation to Mutable Code synchronously.
-	TSharedPtr<FCustomizableObjectCompileRunnable> CompileTask = MakeShareable(new FCustomizableObjectCompileRunnable(RootNode, bDisableLayouts));
+	TSharedPtr<FCustomizableObjectCompileRunnable> CompileTask = MakeShareable(new FCustomizableObjectCompileRunnable(RootNode));
 	CompileTask->Options = CompileOptions;
 	CompileTask->Init();
 	CompileTask->Run();
diff --git a/Engine/Plugins/Experimental/Mutable/Source/MutableRuntime/Private/MuR/CodeRunner.cpp b/Engine/Plugins/Experimental/Mutable/Source/MutableRuntime/Private/MuR/CodeRunner.cpp
--- a/Engine/Plugins/Experimental/Mutable/Source/MutableRuntime/Private/MuR/CodeRunner.cpp
+++ b/Engine/Plugins/Experimental/Mutable/Source/MutableRuntime/Private/MuR/CodeRunner.cpp
@@ -55,7 +55,7 @@
 #include "MuR/OpMeshMerge.h"
 #include "MuR/OpMeshMorph.h"
 #include "MuR/OpMeshRemapIndices.h"
-#include "MuR/OpMeshRemoveChart.h"
+#include "MuR/OpMeshRemove.h"
 #include "MuR/OpMeshReshape.h"
 #include "MuR/OpMeshTransform.h"
 #include "MuR/OpMeshOptimizeSkinning.h"
diff --git a/Engine/Plugins/Experimental/Mutable/Source/MutableRuntime/Private/MuR/OpMeshApplyLayout.h b/Engine/Plugins/Experimental/Mutable/Source/MutableRuntime/Private/MuR/OpMeshApplyLayout.h
--- a/Engine/Plugins/Experimental/Mutable/Source/MutableRuntime/Private/MuR/OpMeshApplyLayout.h
+++ b/Engine/Plugins/Experimental/Mutable/Source/MutableRuntime/Private/MuR/OpMeshApplyLayout.h
@@ -5,22 +5,23 @@
 #include "MuR/MeshPrivate.h"
 #include "MuR/MutableTrace.h"
 #include "MuR/Platform.h"
+#include "MuR/OpMeshRemove.h"
 
 namespace mu
 {
 
 
 	//---------------------------------------------------------------------------------------------
-    inline void MeshApplyLayout( Mesh* pApplied, const Layout* pLayout, int texCoordsSet )
+    inline void MeshApplyLayout( Mesh* pApplied, const Layout* pLayout, int32 texCoordsSet )
 	{
 		MUTABLE_CPUPROFILER_SCOPE(MeshApplyLayout);
 
-        int buffer = -1;
-		int channel = -1;
+        int32 buffer = -1;
+		int32 channel = -1;
 		pApplied->GetVertexBuffers().FindChannel( MBS_TEXCOORDS, texCoordsSet, &buffer, &channel );
 
-        int layoutBuffer = -1;
-        int layoutChannel = -1;
+        int32 layoutBuffer = -1;
+        int32 layoutChannel = -1;
         pApplied->GetVertexBuffers().FindChannel( MBS_LAYOUTBLOCK, texCoordsSet, &layoutBuffer, &layoutChannel );
 
 		if (buffer < 0 || layoutBuffer < 0)
@@ -28,27 +29,18 @@ namespace mu
 			return;
 		}
 
-		//{
-		//	FString debugLog;
-		//	pApplied->GetPrivate()->Log(debugLog, 4);
-		//	UE_LOG(LogMutableCore, Warning, TEXT("before: %s"), *debugLog);
-		//}
-
 		// Get the information about the texture coordinates channel
 		MESH_BUFFER_SEMANTIC semantic;
-		int semanticIndex;
+		int32 semanticIndex;
 		MESH_BUFFER_FORMAT format;
-		int components;
-		int offset;
-		pApplied->GetVertexBuffers().GetChannel
-				( buffer, channel, &semantic, &semanticIndex, &format, &components, &offset );
+		int32 components;
+		int32 offset;
+		pApplied->GetVertexBuffers().GetChannel( buffer, channel, &semantic, &semanticIndex, &format, &components, &offset );
 		check( semantic == MBS_TEXCOORDS );
-		//check( format == MBF_FLOAT32 );
-		//check( components == 2 );
 
-        uint8_t* pData = pApplied->GetVertexBuffers().GetBufferData( buffer );
-		int elemSize = pApplied->GetVertexBuffers().GetElementSize( buffer );
-		int channelOffset = pApplied->GetVertexBuffers().GetChannelOffset( buffer, channel );
+        uint8* pData = pApplied->GetVertexBuffers().GetBufferData( buffer );
+		int32 elemSize = pApplied->GetVertexBuffers().GetElementSize( buffer );
+		int32 channelOffset = pApplied->GetVertexBuffers().GetChannelOffset( buffer, channel );
 		pData += channelOffset;
 
 		struct Box
@@ -77,39 +69,42 @@ namespace mu
         check( pApplied->GetVertexBuffers().m_buffers[layoutBuffer].m_channels[layoutChannel].m_format==MBF_UINT16 );
         check( pApplied->GetVertexBuffers().m_buffers[layoutBuffer].m_channels[layoutChannel].m_componentCount==1 );
 
-        const uint16* pLayoutData = reinterpret_cast<const uint16*>
-                ( pApplied->GetVertexBuffers().GetBufferData( layoutBuffer ) );
+        const uint16* pLayoutData = reinterpret_cast<const uint16*>( pApplied->GetVertexBuffers().GetBufferData( layoutBuffer ) );
+
+		// In some corner case involcing automatic LODs and remove meshes behaving differently among them we may need to remove vertices that don't 
+		// have any block in the current layout. Track them here.
+		TArray<int32> VerticesToRemove;
 
         uint8* pVertices = pData;
-		for ( int v=0; v<pApplied->GetVertexBuffers().GetElementCount(); ++v )
+		for ( int32 v=0; v<pApplied->GetVertexBuffers().GetElementCount(); ++v )
 		{
-			int blockId = pLayoutData[v];
+			int32 blockId = pLayoutData[v];
 
-			// TODO: Optimise
-			int relBlock = pLayout->FindBlock( blockId );
+			// TODO: This could be optimised
+			int32 relBlock = pLayout->FindBlock( blockId );
 
-			// once we do remove the vertices, this assert can be activated.
-			//check( relBlock>=0 && relBlock<(int)transforms.size() );
-
-			// once we do remove the vertices, this condition is not necessary
+			// This may still happen with lower LOD and "remove meshes" in a corner case:
+			// Auto LODs with "Remove Meshes" that behave differently across LODs, and leave geometry in a block that has been removed in the higher LOD.
 			if (relBlock < 0)
 			{
+				VerticesToRemove.Add(v); 
+				pVertices += elemSize;
 				continue;
 			}
 
-			if ( format == MBF_FLOAT32 )
+			switch( format )
+			{
+			case MBF_FLOAT32:
 			{
-				// Not valid in arm because of the potential non-alignement of floats
-				//vec2<float>* pUV = reinterpret_cast<vec2<float>*>( pVertices );
-				//*pUV = (*pUV) * transforms[ relBlock ].size + transforms[ relBlock ].min;
-
 				FVector2f uv;
 				FMemory::Memcpy( &uv, pVertices, sizeof(float)*2 );
 
 				uv = uv * transforms[ relBlock ].size + transforms[ relBlock ].min;
 				FMemory::Memcpy( pVertices, &uv, sizeof(float)*2 );
+				break;
 			}
-			else if ( format == MBF_FLOAT16 )
+
+			case MBF_FLOAT16:
 			{
 				// TODO: Optimise
 				float16* pUV = reinterpret_cast<float16*>( pVertices );
@@ -121,79 +116,91 @@ namespace mu
 
 				pUV[0] = floatToHalf( uv[0] );
 				pUV[1] = floatToHalf( uv[1] );
+				break;
 			}
 
-			else if ( ( format == MBF_NUINT32 )
-						|| ( format == MBF_NINT32 )
-						|| ( format == MBF_UINT32 )
-						|| ( format == MBF_INT32 )
-						)
+			case MBF_NUINT32:
+			case MBF_NINT32:
+			case MBF_UINT32:
+			case MBF_INT32:
 			{
 				// TODO: Optimise
-                uint32_t* pUV = reinterpret_cast<uint32_t*>( pVertices );
+                uint32* pUV = reinterpret_cast<uint32*>( pVertices );
 
 				for ( int c=0; c<2; ++c )
 				{
-                    uint64_t u_32 = pUV[c];
-                    uint64_t u_48 = u_32 * ((uint64_t)(((float)0xffff) * transforms[ relBlock ].size[c]));
-                    u_48 += (uint64_t)(((float)0xffffffffffULL) * transforms[ relBlock ].min[c]);
+                    uint64 u_32 = pUV[c];
+                    uint64 u_48 = u_32 * ((uint64)(((float)0xffff) * transforms[ relBlock ].size[c]));
+                    u_48 += (uint64)(((float)0xffffffffffULL) * transforms[ relBlock ].min[c]);
                     pUV[c] = (uint16)( u_48 >> 16 );
 				}
+				break;
 			}
 
-			else if ( ( format == MBF_NUINT16 )
-						|| ( format == MBF_NINT16 )
-						|| ( format == MBF_UINT16 )
-						|| ( format == MBF_INT16 )
-						)
+			case MBF_NUINT16:
+			case MBF_NINT16:
+			case MBF_UINT16:
+			case MBF_INT16:
 			{
 				// TODO: Optimise
                 uint16* pUV = reinterpret_cast<uint16*>( pVertices );
 
 				for ( int c=0; c<2; ++c )
 				{
-                    uint32_t u_16 = pUV[c];
-                    uint32_t u_32 = u_16 * ((uint8_t)(((float)0xffff) * transforms[ relBlock ].size[c]));
-                    u_32 += (uint32_t)(((float)0xffffffff) * transforms[ relBlock ].min[c]);
+                    uint32 u_16 = pUV[c];
+                    uint32 u_32 = u_16 * ((uint8)(((float)0xffff) * transforms[ relBlock ].size[c]));
+                    u_32 += (uint32)(((float)0xffffffff) * transforms[ relBlock ].min[c]);
                     pUV[c] = (uint16)( u_32 >> 16 );
 				}
+
+				break;
 			}
 
-			else if ( ( format == MBF_NUINT8 )
-						|| ( format == MBF_NINT8 )
-						|| ( format == MBF_UINT8 )
-						|| ( format == MBF_INT8 )
-						)
+			case MBF_NUINT8:
+			case MBF_NINT8:
+			case MBF_UINT8:
+			case MBF_INT8:
 			{
 				// TODO: Optimise
-                uint8_t* pUV = reinterpret_cast<uint8_t*>( pVertices );
+                uint8* pUV = reinterpret_cast<uint8*>( pVertices );
 
 				for ( int c=0; c<2; ++c )
 				{
-                    uint32_t u_8 = pUV[c];
-                    uint32_t u_24 = u_8 * ((uint32_t)(((float)0xffff) * transforms[ relBlock ].size[c]));
-                    u_24 += (uint32_t)(((float)0xffffff) * transforms[ relBlock ].min[c]);
-                    pUV[c] = (uint8_t)( u_24 >> 16 );
+                    uint32 u_8 = pUV[c];
+                    uint32 u_24 = u_8 * ((uint32)(((float)0xffff) * transforms[ relBlock ].size[c]));
+                    u_24 += (uint32)(((float)0xffffff) * transforms[ relBlock ].min[c]);
+                    pUV[c] = (uint8)( u_24 >> 16 );
 				}
+				break;
 			}
 
-			else
-			{
+			default:
 				UE_LOG(LogMutableCore, Warning, TEXT("This case is not implemented.."));
 				check( false );
+				break;
+
 			}
 
 			pVertices += elemSize;
 		}
 
+		if (VerticesToRemove.Num())
+		{
+			// Unpack vertices into a mask
+			TArray<uint8> VertexMask;
+			VertexMask.SetNumZeroed(pApplied->GetVertexCount());
+			for ( int32 VertexIndex: VerticesToRemove )
+			{
+				VertexMask[VertexIndex] = 1;
+			}
+	
+			// Remove
+			MeshRemoveVerticesWithMap( pApplied, VertexMask.GetData(), VertexMask.Num());
+		}
+
 		//
 		pApplied->SetLayout( texCoordsSet, pLayout->Clone() );
 
-		//{
-		//	FString debugLog;
-		//	pApplied->GetPrivate()->Log(debugLog, 4);
-		//	UE_LOG(LogMutableCore, Warning, TEXT("after: %s"), *debugLog);
-		//}
 	}
 
 }
diff --git a/Engine/Plugins/Experimental/Mutable/Source/MutableRuntime/Private/MuR/OpMeshRemove.cpp b/Engine/Plugins/Experimental/Mutable/Source/MutableRuntime/Private/MuR/OpMeshRemove.cpp
new file mode 100644
--- /dev/null
+++ b/Engine/Plugins/Experimental/Mutable/Source/MutableRuntime/Private/MuR/OpMeshRemove.cpp
@@ -0,0 +1,289 @@
+// Copyright Epic Games, Inc. All Rights Reserved.
+
+#include "MuR/OpMeshRemove.h"
+
+#include "MuR/MeshPrivate.h"
+#include "MuR/MutableTrace.h"
+#include "MuR/Platform.h"
+
+
+namespace mu
+{
+
+    //---------------------------------------------------------------------------------------------
+    struct ID_INTERVAL
+    {
+        int32 idStart;
+        int32 idPosition;
+        int32 size;
+    };
+
+
+    void ExtractVertexIndexIntervals( TArray<ID_INTERVAL>& intervals, const Mesh* pSource )
+    {
+        MeshBufferIteratorConst<MBF_UINT32,uint32,1> itVI( pSource->GetVertexBuffers(), MBS_VERTEXINDEX );
+        ID_INTERVAL current;
+		current.idStart = -1;
+		current.idPosition = 0;
+		current.size = 0;
+		for ( int sv=0; sv<pSource->GetVertexBuffers().GetElementCount(); ++sv )
+        {
+            int32 id = (*itVI)[0];
+            ++itVI;
+
+            if (current.idStart<0)
+            {
+                current.idStart = id;
+                current.idPosition = sv;
+                current.size = 1;
+            }
+            else
+            {
+                if (id==current.idStart+current.size)
+                {
+                    ++current.size;
+                }
+                else
+                {
+                    intervals.Add(current);
+                    current.idStart = id;
+                    current.idPosition = sv;
+                    current.size = 1;
+                }
+            }
+        }
+
+        if (current.idStart>=0)
+        {
+            intervals.Add(current);
+        }
+    }
+
+
+    int FindPositionInIntervals( const TArray<ID_INTERVAL>& intervals, int32 id )
+    {
+        for( const auto& interval: intervals )
+        {
+            int deltaId = id - interval.idStart;
+            if (deltaId>=0 && deltaId<interval.size)
+            {
+                return interval.idPosition+deltaId;
+            }
+        }
+        return -1;
+    }
+
+
+	void MeshRemoveVerticesWithMap( Mesh* Result, uint8* removedVertices, int32 resultVertexCount)
+	{
+		int32 firstFreeVertex = 0;
+
+        // Rebuild index buffers
+
+        // Map from source vertex index, to new vertex index for used vertices.
+        // These are indices as in the index buffer, not the absoulte vertex index as in the
+        // vertexbuffer MBS_VERTEXINDEX buffers.
+		TArray<int32> usedVertices;
+		usedVertices.Init( -1, Result->GetVertexCount() );
+        {
+            SIZE_T removedIndices = 0;
+
+            if ( Result->GetIndexBuffers().GetElementSize(0)==4 )
+            {
+                MeshBufferIteratorConst<MBF_UINT32,uint32,1> itSource( Result->GetIndexBuffers(), MBS_VERTEXINDEX );
+                MeshBufferIterator<MBF_UINT32,uint32,1> itDest( Result->GetIndexBuffers(), MBS_VERTEXINDEX );
+
+                int32 indexCount = Result->GetIndexCount();
+                for ( int32 f=0; f<indexCount/3; ++f )
+                {
+                    uint32 sourceIndices[3];
+                    sourceIndices[0] = itSource[0][0];
+                    sourceIndices[1] = itSource[1][0];
+                    sourceIndices[2] = itSource[2][0];
+
+                    bool faceRemoved =  ( removedVertices[ sourceIndices[0] ]
+                            + removedVertices[ sourceIndices[1] ]
+                            + removedVertices[ sourceIndices[2] ]
+                            ) > 0;
+
+                    if ( !faceRemoved )
+                    {
+                        for (int32 i=0;i<3;++i)
+                        {
+                            uint32 sourceIndex = sourceIndices[i];
+
+                            if ( usedVertices[ sourceIndex ] < 0 )
+                            {
+                                usedVertices[ sourceIndex ] = firstFreeVertex;
+                                firstFreeVertex++;
+                            }
+
+                            uint32 destIndex = usedVertices[ sourceIndex ];
+                            *(uint32*)itDest.ptr() = destIndex;
+
+                            itDest++;
+                        }
+                    }
+
+                    itSource+=3;
+                }
+
+                removedIndices = itSource - itDest;
+            }
+
+            else if ( Result->GetIndexBuffers().GetElementSize(0)==2 )
+            {
+                MeshBufferIteratorConst<MBF_UINT16,uint16,1> itSource( Result->GetIndexBuffers(), MBS_VERTEXINDEX );
+                MeshBufferIterator<MBF_UINT16,uint16,1> itDest( Result->GetIndexBuffers(), MBS_VERTEXINDEX );
+
+                int32 indexCount = Result->GetIndexCount();
+                for ( int32 f=0; f<indexCount/3; ++f )
+                {
+                    uint16 sourceIndices[3];
+                    sourceIndices[0] = itSource[0][0];
+                    sourceIndices[1] = itSource[1][0];
+                    sourceIndices[2] = itSource[2][0];
+
+                    bool faceRemoved =  ( removedVertices[ sourceIndices[0] ]
+                            + removedVertices[ sourceIndices[1] ]
+                            + removedVertices[ sourceIndices[2] ]
+                            ) > 0;
+
+                    if ( !faceRemoved )
+                    {
+                        for (int i=0;i<3;++i)
+                        {
+                            uint16 sourceIndex = sourceIndices[i];
+
+                            if ( usedVertices[ sourceIndex ] < 0 )
+                            {
+                                usedVertices[ sourceIndex ] = firstFreeVertex;
+                                firstFreeVertex++;
+                            }
+
+                            uint16 destIndex = (uint16)usedVertices[ sourceIndex ];
+                            *(uint16*)itDest.ptr() = destIndex;
+
+                            itDest++;
+                        }
+                    }
+
+                    itSource+=3;
+                }
+
+                removedIndices = itSource - itDest;
+            }
+
+            else
+            {
+                // Case not implemented
+                check( false );
+            }
+
+            check( removedIndices%3==0 );
+
+            int32 faceCount = Result->GetFaceCount();
+            Result->GetFaceBuffers().SetElementCount( faceCount-(int)removedIndices/3 );
+            Result->GetIndexBuffers().SetElementCount( faceCount*3-(int)removedIndices );
+        }
+
+        // Rebuild the vertex buffers
+        for ( int32 b=0; b<Result->GetVertexBuffers().GetBufferCount(); ++b )
+        {
+            int32 elemSize = Result->GetVertexBuffers().GetElementSize( b );
+            const uint8* pSourceData = Result->GetVertexBuffers().GetBufferData( b );
+            uint8* pData = Result->GetVertexBuffers().GetBufferData( b );
+            for ( int32 v=0; v<Result->GetVertexCount(); ++v )
+            {
+                int32 span = 0;
+                for ( int32 s=0; v+s<Result->GetVertexCount(); ++s )
+                {
+                    if ( usedVertices[v+s]>=0 )
+                    {
+                        if (span==0)
+                        {
+                            ++span;
+                        }
+                        else
+                        {
+                            if ( usedVertices[v+s] == usedVertices[v+s-1]+1 )
+                            {
+                                ++span;
+                            }
+                            else
+                            {
+                                break;
+                            }
+                        }
+                    }
+                    else
+                    {
+                        break;
+                    }
+                }
+
+
+                if (span>0)
+                {
+                    FMemory::Memmove( pData+elemSize*usedVertices[v], pSourceData+elemSize*v, elemSize*span );
+
+                    v += span-1;
+                }
+            }
+        }
+        Result->GetVertexBuffers().SetElementCount( firstFreeVertex );
+
+        // Rebuild surface data.
+        // \todo: For now make single surfaced
+        Result->m_surfaces.Empty();
+        Result->EnsureSurfaceData();
+    }
+
+
+	void MeshRemoveMask(Mesh* Result, const Mesh* pSource, const Mesh* pMask, bool& bOutSuccess)
+    {
+        MUTABLE_CPUPROFILER_SCOPE(MeshRemoveMask);
+		bOutSuccess = true;
+
+        if (!pMask->GetVertexCount() || !pSource->GetVertexCount() || !pSource->GetIndexCount())
+        {
+			bOutSuccess = false;
+            return;
+        }
+
+		Result->CopyFrom(*pSource);
+
+        // TODO
+        Result->SetFaceGroupCount( 0 );
+
+        MeshBufferIteratorConst<MBF_UINT32,uint32,1> itMaskVI( pMask->GetVertexBuffers(), MBS_VERTEXINDEX );
+
+        int32 firstFreeVertex = 0;
+
+        // For each source vertex, true if it is removed.
+        int32 resultVertexCount = Result->GetVertexCount();
+		uint8* removedVertices = reinterpret_cast<uint8*>(FMemory::Malloc(resultVertexCount, 16));
+		FMemory::Memzero(removedVertices,resultVertexCount);
+        {
+			TArray<ID_INTERVAL> intervals;
+            ExtractVertexIndexIntervals(intervals, Result);
+
+            for ( int32 mv=0; mv<pMask->GetVertexBuffers().GetElementCount(); ++mv )
+            {
+                uint32 maskVertexId = (*itMaskVI)[0];
+                ++itMaskVI;
+
+                int32 indexInSource = FindPositionInIntervals(intervals, maskVertexId);
+                if (indexInSource>=0)
+                {
+                    removedVertices[indexInSource] = 1;
+                }
+            }
+        }
+
+		MeshRemoveVerticesWithMap(Result, removedVertices, resultVertexCount);
+
+		FMemory::Free(removedVertices);
+	}
+
+}
diff --git a/Engine/Plugins/Experimental/Mutable/Source/MutableRuntime/Private/MuR/OpMeshRemove.h b/Engine/Plugins/Experimental/Mutable/Source/MutableRuntime/Private/MuR/OpMeshRemove.h
new file mode 100644
--- /dev/null
+++ b/Engine/Plugins/Experimental/Mutable/Source/MutableRuntime/Private/MuR/OpMeshRemove.h
@@ -0,0 +1,19 @@
+// Copyright Epic Games, Inc. All Rights Reserved.
+
+#pragma once
+
+#include "MuR/Ptr.h"
+#include "HAL/Platform.h"
+
+
+namespace mu
+{
+	class Mesh;
+
+	/** Remove a list of vertices and related faces from a mesh. The list of vertices is stored in a specially formattes Mask mesh. */
+	extern void MeshRemoveMask(Mesh* Result, const Mesh* Source, const Mesh* Mask, bool& bOutSuccess);
+
+	/** Remove a list of vertices and related faces from a mesh. The list is stored as a bool map for every vertex in the mesh. */
+	extern void MeshRemoveVerticesWithMap(Mesh* Result, uint8* RemovedVertices, int32 RemovedVerticesCount);
+
+}
diff --git a/Engine/Plugins/Experimental/Mutable/Source/MutableRuntime/Private/MuR/OpMeshRemoveChart.h b/Engine/Plugins/Experimental/Mutable/Source/MutableRuntime/Private/MuR/OpMeshRemoveChart.h
deleted file mode 100644
diff --git a/Engine/Plugins/Experimental/Mutable/Source/MutableRuntime/Private/MuR/Platform.h b/Engine/Plugins/Experimental/Mutable/Source/MutableRuntime/Private/MuR/Platform.h
--- a/Engine/Plugins/Experimental/Mutable/Source/MutableRuntime/Private/MuR/Platform.h
+++ b/Engine/Plugins/Experimental/Mutable/Source/MutableRuntime/Private/MuR/Platform.h
@@ -3,16 +3,9 @@
 #pragma once
 
 // This file contains some hacks to solve differences between platforms
-// This file also contains wrappers for all the libc methods used in the runtime.
-// No other place in the runtime should include any libc header.
 
 #include "HAL/PlatformTime.h"
 #include "Logging/LogMacros.h"
-#include "Stats/Stats.h"
-
-//! There is a define MUTABLE_PROFILE that can be passed on the compiler options to enable
-//! the internal profiling methods. See Config.h for more information.
-
 #include "MuR/Types.h"
 #include "HAL/UnrealMemory.h"
 
@@ -25,44 +18,3 @@ MUTABLERUNTIME_API DECLARE_LOG_CATEGORY_EXTERN(LogMutableCore, Log, All);
         #define MUTABLE_DEBUG
     #endif
 #endif
-
-
-namespace mu
-{
-
-	//---------------------------------------------------------------------------------------------
-	inline int mutable_vsnprintf
-		( char *buffer, size_t sizeOfBuffer, const char *format, va_list argptr )
-	{
-		#ifdef _MSC_VER
-			return vsnprintf_s( buffer, sizeOfBuffer, _TRUNCATE, format, argptr );
-		#else
-			return vsnprintf( buffer, sizeOfBuffer, format, argptr );
-		#endif
-	}
-
-
-	//---------------------------------------------------------------------------------------------
-	inline int mutable_snprintf
-		( char *buffer, size_t sizeOfBuffer, const char *format, ... )
-	{
-		va_list ap;
-		va_start( ap, format );
-		int res = mutable_vsnprintf( buffer, sizeOfBuffer, format, ap );
-		va_end(ap);
-		return res;
-	}
-
-
-	//---------------------------------------------------------------------------------------------
-	// Disgracefully halt the program.
-	inline void Halt()
-	{
-		#ifdef _MSC_VER
-            __debugbreak();
-		#else
-			__builtin_trap();
-		#endif
-	}
-
-}
diff --git a/Engine/Plugins/Experimental/Mutable/Source/MutableRuntime/Public/MuR/Model.h b/Engine/Plugins/Experimental/Mutable/Source/MutableRuntime/Public/MuR/Model.h
--- a/Engine/Plugins/Experimental/Mutable/Source/MutableRuntime/Public/MuR/Model.h
+++ b/Engine/Plugins/Experimental/Mutable/Source/MutableRuntime/Public/MuR/Model.h
@@ -13,7 +13,7 @@
 //! data structures. Compiled models are not necessarily compatible when the runtime is updated,
 //! so this version number can be used externally to verify this. It is not used internally, and
 //! serializing models from different versions than this runtime will probably result in a crash.
-#define MUTABLE_COMPILED_MODEL_CODE_VERSION		uint32( 68 )
+#define MUTABLE_COMPILED_MODEL_CODE_VERSION		uint32( 69 )
 #define MUTABLE_PARAMETERS_VERSION              uint32( 2 )
 
 
diff --git a/Engine/Plugins/Experimental/Mutable/Source/MutableTools/Private/MuT/AST.h b/Engine/Plugins/Experimental/Mutable/Source/MutableTools/Private/MuT/AST.h
--- a/Engine/Plugins/Experimental/Mutable/Source/MutableTools/Private/MuT/AST.h
+++ b/Engine/Plugins/Experimental/Mutable/Source/MutableTools/Private/MuT/AST.h
@@ -180,19 +180,19 @@ namespace mu
     //! Detailed optimization flags
     struct FModelOptimizationOptions
     {
-        bool m_enabled = true;
-        bool m_optimiseOverlappedMasks = false;
-        bool m_constReduction = true;
+        bool bEnabled = true;
+        bool bOptimiseOverlappedMasks = false;
+        bool bConstReduction = true;
 
         //! Preprocess all mesh fragments so that they use the same skeleton, even if not all bones
         //! are relevant for all fragments.
-        bool m_uniformizeSkeleton = true;
+        bool bUniformizeSkeleton = true;
 
         //! Maximum number of iterations when optimising models. If 0 as many as necessary will be performed.
-        int m_maxOptimisationLoopCount = 8;
+        int32 MaxOptimisationLoopCount = 8;
 
         //! Store resource data in disk instead of memory
-        bool m_useDiskCache = false;
+        bool bUseDiskCache = false;
 
 		/** Compile optimizing for the generation of smaller mipmaps of every image. */
 		bool bEnableProgressiveImages = false;
@@ -201,10 +201,10 @@ namespace mu
         //---------------------------------------------------------------------
 
         //! Ratio used to decide if it is worth to generate a crop operation
-        float m_acceptableCropRatio = 0.5f;
+        float AcceptableCropRatio = 0.5f;
 
         //! Ratio used to decide if it is worth to generate a crop operation
-        float m_minRLECompressionGain = 1.2f;
+        float MinRLECompressionGain = 1.2f;
     };
 
 
diff --git a/Engine/Plugins/Experimental/Mutable/Source/MutableTools/Private/MuT/ASTOpImageLayer.cpp b/Engine/Plugins/Experimental/Mutable/Source/MutableTools/Private/MuT/ASTOpImageLayer.cpp
--- a/Engine/Plugins/Experimental/Mutable/Source/MutableTools/Private/MuT/ASTOpImageLayer.cpp
+++ b/Engine/Plugins/Experimental/Mutable/Source/MutableTools/Private/MuT/ASTOpImageLayer.cpp
@@ -474,7 +474,7 @@ namespace mu
 				// Is it worth?
 				float ratio = float(maskUsage.size[0] * maskUsage.size[1])
 					/ float(maskDesc.m_size[0] * maskDesc.m_size[1]);
-				float acceptableCropRatio = options.m_acceptableCropRatio;
+				float acceptableCropRatio = options.AcceptableCropRatio;
 				if (ratio < acceptableCropRatio)
 				{
 					check(maskUsage.size[0] > 0);
diff --git a/Engine/Plugins/Experimental/Mutable/Source/MutableTools/Private/MuT/CodeGenerator.cpp b/Engine/Plugins/Experimental/Mutable/Source/MutableTools/Private/MuT/CodeGenerator.cpp
--- a/Engine/Plugins/Experimental/Mutable/Source/MutableTools/Private/MuT/CodeGenerator.cpp
+++ b/Engine/Plugins/Experimental/Mutable/Source/MutableTools/Private/MuT/CodeGenerator.cpp
@@ -111,24 +111,18 @@ namespace mu
 	{
 		MUTABLE_CPUPROFILER_SCOPE(Generate);
 
-		//        UE_LOG(LogMutableCore,Warning,"astop %9u#  %9u+  %9u- generate root start", size_t(ASTOp::s_count), size_t(ASTOp::s_created), size_t(ASTOp::s_destroyed) );
+		// First pass
+		m_firstPass.Generate(m_pErrorLog, pNode->GetBasePrivate(), m_compilerOptions->bIgnoreStates);
 
-				// First pass
-		m_firstPass.Generate(m_pErrorLog, pNode->GetBasePrivate(), m_compilerOptions->m_ignoreStates);
-
-		//        UE_LOG(LogMutableCore,Warning,"astop %9u#  %9u+  %9u- first pass completed", size_t(ASTOp::s_count), size_t(ASTOp::s_created), size_t(ASTOp::s_destroyed) );
-
-				// Second pass
-		SecondPassGenerator secondPass(&m_firstPass, m_compilerOptions);
-		bool success = secondPass.Generate(m_pErrorLog, pNode->GetBasePrivate());
-		if (!success)
+		// Second pass
+		SecondPassGenerator SecondPass(&m_firstPass, m_compilerOptions);
+		bool bSuccess = SecondPass.Generate(m_pErrorLog, pNode->GetBasePrivate());
+		if (!bSuccess)
 		{
 			return;
 		}
 
-		//        UE_LOG(LogMutableCore,Warning,"astop %9u#  %9u+  %9u- second pass completed ", size_t(ASTOp::s_count), size_t(ASTOp::s_created), size_t(ASTOp::s_destroyed) );
-
-				// Main pass for each state
+		// Main pass for each state
 		{
 			MUTABLE_CPUPROFILER_SCOPE(MainPass);
 
@@ -144,9 +138,7 @@ namespace mu
 			}
 		}
 
-		//        UE_LOG(LogMutableCore,Warning,"astop %9u#  %9u+  %9u- before cache clear", size_t(ASTOp::s_count), size_t(ASTOp::s_created), size_t(ASTOp::s_destroyed) );
-
-				// Free caches
+		// Free caches
 		m_compiled.Reset();
 		m_constantMeshes.Empty();
 		m_generatedLayouts.Empty();
@@ -161,8 +153,6 @@ namespace mu
 		m_currentParents.Empty();
 		m_currentObject.Empty();
 		m_additionalComponents.clear();
-
-		//        UE_LOG(LogMutableCore,Warning,"astop %9u#  %9u+  %9u- exiting generate root", size_t(ASTOp::s_count), size_t(ASTOp::s_created), size_t(ASTOp::s_destroyed) );
 	}
 
 
@@ -194,23 +184,14 @@ namespace mu
 
 		if (const NodeImage* ImageNode = dynamic_cast<const NodeImage*>(pNode.get()))
 		{
-			FImageGenerationResult Result;
-			GenerateImage(Result, ImageNode);
-			return Result.op;
+			// This should never be called for images. Use GenerateImage.
+			check(false);
 		}
 
 		if (const NodeMesh* MeshNode = dynamic_cast<const NodeMesh*>(pNode.get()))
 		{
-			// This should only happen with mesh root nodes
-			FMeshGenerationOptions MeshOptions;
-			MeshOptions.bUniqueVertexIDs = true;
-			MeshOptions.bLayouts = true;
-			MeshOptions.bNormalizeUVs = true;
-			MeshOptions.State = m_currentStateIndex;
-
-			FMeshGenerationResult MeshResult;
-			GenerateMesh(MeshOptions, MeshResult, MeshNode);
-			return MeshResult.meshOp;
+			// This should never be called for meshes. Use GenerateMesh.
+			check(false);
 		}
 
 		if (const NodeProjector* projNode = dynamic_cast<const NodeProjector*>(pNode.get()))
@@ -414,7 +395,8 @@ namespace mu
 	//---------------------------------------------------------------------------------------------
 	Ptr<ASTOp> CodeGenerator::GenerateImageBlockPatch(Ptr<ASTOp> blockAd,
 		const NodePatchImage* pPatch,
-		Ptr<ASTOp> conditionAd)
+		Ptr<ASTOp> conditionAd,
+		const FImageGenerationOptions& ImageOptions )
 	{
 		// Blend operation
 		Ptr<ASTOp> blendAd;
@@ -436,19 +418,18 @@ namespace mu
 			Ptr<ASTOp> blend;
 			if (pImage)
 			{
-				blend = Generate(pImage);
+				FImageGenerationResult BlendResult;
+				GenerateImage(ImageOptions, BlendResult, pImage);
+				blend = BlendResult.op;
 			}
 			else
 			{
-				blend = GenerateMissingImageCode("Blend top image", EImageFormat::IF_RGB_UBYTE,
-					pPatch->GetPrivate()->m_errorContext);
+				blend = GenerateMissingImageCode(TEXT("Blend top image"), EImageFormat::IF_RGB_UBYTE, pPatch->GetPrivate()->m_errorContext, ImageOptions);
 			}
 			blend = GenerateImageFormat(blend, blockAd->GetImageDesc().m_format);
-			blend = GenerateImageSize
-			(
+			blend = GenerateImageSize(
 				blend,
-				FImageSize((uint16)m_imageState.Last().m_imageRect.size[0],
-					(uint16)m_imageState.Last().m_imageRect.size[1])
+				FImageSize((uint16)m_imageState.Last().m_imageRect.size[0], (uint16)m_imageState.Last().m_imageRect.size[1])
 			);
 			op->blend = blend;
 
@@ -456,20 +437,20 @@ namespace mu
 			Ptr<ASTOp> mask;
 			if (pMask)
 			{
-				mask = Generate(pMask);
+				FImageGenerationResult MaskResult;
+				GenerateImage(ImageOptions, MaskResult, pMask);
+				mask = MaskResult.op;
 			}
 			else
 			{
 				// Set the argument default value: affect all pixels.
 				// TODO: Special operation code without mask
-				mask = GeneratePlainImageCode(vec3<float>(1, 1, 1));
+				mask = GeneratePlainImageCode(vec3<float>(1, 1, 1), ImageOptions);
 			}
 			mask = GenerateImageFormat(mask, EImageFormat::IF_L_UBYTE);
-			mask = GenerateImageSize
-			(
+			mask = GenerateImageSize(
 				mask,
-				FImageSize((uint16)m_imageState.Last().m_imageRect.size[0],
-					(uint16)m_imageState.Last().m_imageRect.size[1])
+				FImageSize((uint16)m_imageState.Last().m_imageRect.size[0], (uint16)m_imageState.Last().m_imageRect.size[1])
 			);
 			op->mask = mask;
 
@@ -608,15 +589,11 @@ namespace mu
 
 		if (!pNodeLayout)
 		{
-			char buf[256];
-			mutable_snprintf
-			(
-				buf, 256,
-				"In object [%s] NodePatchImage couldn't find the layout in parent.",
+			FString Msg = FString::Printf(TEXT("In object [%s] NodePatchImage couldn't find the layout in parent."),
 				m_currentParents.Last().m_pObject->m_name.c_str()
 			);
 
-			m_pErrorLog->GetPrivate()->Add(buf, ELMT_ERROR, node.m_errorContext);
+			m_pErrorLog->GetPrivate()->Add(Msg, ELMT_ERROR, node.m_errorContext);
 		}
 
 		return 0;
@@ -701,12 +678,12 @@ namespace mu
 	Ptr<ASTOp> CodeGenerator::ApplyTiling(Ptr<ASTOp> Source, FImageSize Size, EImageFormat Format)
 	{
 		// For now always apply tiling
-		if (m_compilerOptions->m_imageTiling==0)
+		if (m_compilerOptions->ImageTiling==0)
 		{
 			return Source;
 		}
 
-		int32 TileSize = m_compilerOptions->m_imageTiling;
+		int32 TileSize = m_compilerOptions->ImageTiling;
 
 		int32 TilesX = FMath::DivideAndRoundUp<int32>(Size[0], TileSize);
 		int32 TilesY = FMath::DivideAndRoundUp<int32>(Size[1], TileSize);
@@ -829,6 +806,12 @@ namespace mu
             return;
         }
 
+		// This assumes that the lods are processed in order. It checks it this way because some platforms may have empty lods at the top.
+		const bool bIsBaseForSharedSurface = node.SharedSurfaceId != INDEX_NONE && !SharedMeshOptionsMap.Contains(node.SharedSurfaceId);
+
+		// If this is true, we will reuse the surface properties from a higher LOD, se we can skip the generation of material properties and images.
+		const bool bShareSurface = node.SharedSurfaceId != INDEX_NONE && !bIsBaseForSharedSurface;
+
         if ( pMesh )
         {
             MUTABLE_CPUPROFILER_SCOPE(SurfaceMesh);
@@ -845,11 +828,21 @@ namespace mu
 			MeshOptions.State = m_currentStateIndex;
 			MeshOptions.ActiveTags = m_activeTags.Last();
 
+			const FMeshGenerationResult* SharedMeshResults = nullptr;
+			if (bShareSurface)
+			{
+				// Do we have the surface we need to share it with?
+				SharedMeshResults = SharedMeshOptionsMap.Find(node.SharedSurfaceId);
+				check(SharedMeshResults);
+
+				// Override the layouts with the ones from the surface we share
+				MeshOptions.OverrideLayouts = SharedMeshResults->GeneratedLayouts;
+			}
+
 			// Ensure UV islands remain within their main layout block on lower LODs to avoid unexpected reordering 
 			// of the layout blocks when reusing a surface between LODs. Used to fix small displacements on vertices
 			// that may cause them to fall on a different block.
-			const bool bClampUVIslands = node.SharedSurfaceId != INDEX_NONE && m_currentParents.Last().m_lod > 0;
-			MeshOptions.bClampUVIslands = bClampUVIslands;
+			MeshOptions.bClampUVIslands = bShareSurface;
 
 			// Normalize UVs if we're going to work with images and layouts.
 			const bool bNormalizeUVs = node.m_images.Num() && edits.Num();
@@ -881,10 +874,17 @@ namespace mu
 						FMeshGenerationOptions MergedMeshOptions;
 						MergedMeshOptions.bUniqueVertexIDs = true;
 						MergedMeshOptions.bLayouts = true;
-						MergedMeshOptions.bClampUVIslands = bClampUVIslands;
+						MergedMeshOptions.bClampUVIslands = bShareSurface;
 						MergedMeshOptions.bNormalizeUVs = bNormalizeUVs;
 						MergedMeshOptions.State = m_currentStateIndex;
 						MergedMeshOptions.ActiveTags = m_activeTags.Last();
+
+						if (SharedMeshResults)
+						{
+							check(SharedMeshResults->extraMeshLayouts.Num()>editIndex);
+							MergedMeshOptions.OverrideLayouts = SharedMeshResults->extraMeshLayouts[editIndex].GeneratedLayouts;
+						}
+
 						FMeshGenerationResult addResults;
                         GenerateMesh(MergedMeshOptions, addResults, pAdd);
 
@@ -932,6 +932,7 @@ namespace mu
             // Apply mesh removes from child objects "edit surface" nodes.
             // "Removes" need to come after "Adds" because some removes may refer to added meshes,
             // and not the base.
+			// \TODO: Apply base removes first, and then "added meshes" removes here. It may have lower memory footprint during generation.
             Ptr<ASTOpMeshRemoveMask> rop;
             for ( const FirstPassGenerator::SURFACE::EDIT& e: edits )
             {
@@ -1069,89 +1070,92 @@ namespace mu
 				bModifiersForBeforeOperations, node.m_errorContext);
 
             // Layouts
-            for ( int32 l=0; l<meshResults.GeneratedLayouts.Num(); ++l )
+            for ( int32 LayoutIndex=0; LayoutIndex <meshResults.GeneratedLayouts.Num(); ++LayoutIndex)
             {
-                Ptr<ASTOp> layoutOp = nullptr;
+                Ptr<ASTOp> layoutOp;
 
-				Ptr<const Layout> pLayout = meshResults.GeneratedLayouts[l].get();
+				Ptr<const Layout> pLayout = meshResults.GeneratedLayouts[LayoutIndex].get();
                 if ( pLayout )
                 {
-                    // Create the layout expression
-
-                    // Constant layout to start with
-                    {
-                        Ptr<ASTOpConstantResource> op = new ASTOpConstantResource();
-                        op->type = OP_TYPE::LA_CONSTANT;                       
-
-                        op->SetValue(pLayout, m_compilerOptions->m_optimisationOptions.m_useDiskCache );
-                        layoutOp = op;
-                    }
-
-                    // Add children merged meshes layouts
-                    for ( const FMeshGenerationResult::FExtraLayouts& data : meshResults.extraMeshLayouts )
-                    {
-                        if ( !data.meshFragment )
-                        {
-                            // No mesh to add, we assume there are no layouts to add either.
-                            check(data.GeneratedLayouts.IsEmpty());
-                            continue;
-                        }
-
-                        if ( data.GeneratedLayouts.Num() != meshResults.GeneratedLayouts.Num() )
-                        {
-                            m_pErrorLog->GetPrivate()->Add
-                                ( "Merged layout has been ignored because the number of layouts"
-                                    " is different.",
-                                  ELMT_ERROR, node.m_errorContext );
-                        }
-                        else
-                        {
-                            // Constant layout to start with
-                            Ptr<ASTOp> layoutFragmentAd;
-                            {
-                                Ptr<ASTOpConstantResource> op = new ASTOpConstantResource();
-                                op->type = OP_TYPE::LA_CONSTANT;
-
-								Ptr<const Layout> pCloned = data.GeneratedLayouts[l];
-                                op->SetValue( pCloned, m_compilerOptions->m_optimisationOptions.m_useDiskCache );
-
-                                layoutFragmentAd = op;
-                            }
+					if (SharedMeshResults)
+					{
+						check(SharedMeshResults->layoutOps.Num()>LayoutIndex);
+						layoutOp = SharedMeshResults->layoutOps[LayoutIndex];
+					}
+					else
+					{
+						// Create a new layout expression
 
-                            // Merge operation
-                            Ptr<ASTOpLayoutMerge> mergeAd = new ASTOpLayoutMerge();
-                            mergeAd->Base = layoutOp;
-	                        mergeAd->Added = layoutFragmentAd;
+						// Constant layout to start with
+						{
+							Ptr<ASTOpConstantResource> op = new ASTOpConstantResource();
+							op->type = OP_TYPE::LA_CONSTANT;
 
-                            // Condition to apply
-                            if (data.condition)
-                            {
-                                Ptr<ASTOpConditional> conditionalAd = new ASTOpConditional();
-                                conditionalAd->type = OP_TYPE::LA_CONDITIONAL;
-                                conditionalAd->no = layoutOp;
-                                conditionalAd->yes = mergeAd;
-                                conditionalAd->condition = data.condition;
-                                layoutOp = conditionalAd;
-                            }
-                            else
-                            {
-                                layoutOp = mergeAd;
-                            }
-                        }
-                    }
+							op->SetValue(pLayout, m_compilerOptions->OptimisationOptions.bUseDiskCache);
+							layoutOp = op;
+						}
 
+						// Add children merged meshes layouts
+						for (const FMeshGenerationResult::FExtraLayouts& data : meshResults.extraMeshLayouts)
+						{
+							if (!data.meshFragment)
+							{
+								// No mesh to add, we assume there are no layouts to add either.
+								check(data.GeneratedLayouts.IsEmpty());
+								continue;
+							}
+
+							if (data.GeneratedLayouts.Num() != meshResults.GeneratedLayouts.Num())
+							{
+								m_pErrorLog->GetPrivate()->Add(TEXT("Merged layout has been ignored because the number of layouts is different."), ELMT_ERROR, node.m_errorContext);
+							}
+							else
+							{
+								// Constant layout to start with
+								Ptr<ASTOp> layoutFragmentAd;
+								{
+									Ptr<ASTOpConstantResource> op = new ASTOpConstantResource();
+									op->type = OP_TYPE::LA_CONSTANT;
+
+									Ptr<const Layout> pCloned = data.GeneratedLayouts[LayoutIndex];
+									op->SetValue(pCloned, m_compilerOptions->OptimisationOptions.bUseDiskCache);
+
+									layoutFragmentAd = op;
+								}
+
+								// Merge operation
+								Ptr<ASTOpLayoutMerge> mergeAd = new ASTOpLayoutMerge();
+								mergeAd->Base = layoutOp;
+								mergeAd->Added = layoutFragmentAd;
+
+								// Condition to apply
+								if (data.condition)
+								{
+									Ptr<ASTOpConditional> conditionalAd = new ASTOpConditional();
+									conditionalAd->type = OP_TYPE::LA_CONDITIONAL;
+									conditionalAd->no = layoutOp;
+									conditionalAd->yes = mergeAd;
+									conditionalAd->condition = data.condition;
+									layoutOp = conditionalAd;
+								}
+								else
+								{
+									layoutOp = mergeAd;
+								}
+							}
+						}
+					}
 
-                    if ( m_compilerOptions->m_textureLayoutStrategy
-                         ==
-                         CompilerOptions::TextureLayoutStrategy::Pack )
+                    if ( layoutOp )
                     {
                         // Add layout packing instructions
-                        {
+						if (!SharedMeshResults)
+						{
                             // Make sure we removed unnecessary blocks
                             Ptr<ASTOpLayoutFromMesh> ExtractOp = new ASTOpLayoutFromMesh();
 							ExtractOp->Mesh = lastMeshOp;
-							check(l<256);
-							ExtractOp->LayoutIndex = uint8( l );
+							check(LayoutIndex<256);
+							ExtractOp->LayoutIndex = uint8(LayoutIndex);
 
 							Ptr<ASTOpLayoutRemoveBlocks> RemoveOp = new ASTOpLayoutRemoveBlocks();
 							RemoveOp->Source = layoutOp;
@@ -1170,7 +1174,7 @@ namespace mu
                             op->op.type = OP_TYPE::ME_APPLYLAYOUT;
                             op->SetChild(op->op.args.MeshApplyLayout.mesh, lastMeshOp );
                             op->SetChild(op->op.args.MeshApplyLayout.layout, layoutOp );
-                            op->op.args.MeshApplyLayout.channel = (uint16)l;
+                            op->op.args.MeshApplyLayout.channel = (uint16)LayoutIndex;
                             lastMeshOp = op;
                         }
                     }
@@ -1185,561 +1189,565 @@ namespace mu
         }
 
 
-        // Create the expression for each texture
+        // Create the expression for each texture, if we are not reusing the surface from another LOD.
         //------------------------------------------------------------------------
-        for ( int32 t=0; t<node.m_images.Num(); ++t )
-        {
-            //MUTABLE_CPUPROFILER_SCOPE(SurfaceTexture);
-
-            if ( NodeImagePtr pImageNode = node.m_images[t].m_pImage )
-            {
-                // Any image-specific format or mipmapping needs to be applied at the end
-                NodeImageMipmapPtr mipmapNode;
-                NodeImageFormatPtr formatNode;
-                bool found = false;
-                while ( !found)
-                {
-                    if (NodeImageMipmap* tm = dynamic_cast<NodeImageMipmap*>(pImageNode.get()))
-                    {
-                        if (!mipmapNode) mipmapNode = tm;
-                        pImageNode = tm->GetSource();
-                    }
-                    else if (NodeImageFormat* tf = dynamic_cast<NodeImageFormat*>(pImageNode.get()))
-                    {
-                        if (!formatNode) formatNode = tf;
-                        pImageNode = tf->GetSource();
-                    }
-                    else
-                    {
-                        found = true;
-                    }
-                }
-
-                // Find out the size of the image
-                // TODO: What if the image is empty and everything is added?
-                //       Look in the extending images.
-                FImageDesc desc;
-                if (pImageNode)
-                {
-                    MUTABLE_CPUPROFILER_SCOPE(CalculateImageDesc);
-                    desc = CalculateImageDesc( *pImageNode->GetBasePrivate() );
-                }
-
-				// If the image format doesn't come bottom-up, it may come top-down
-				if (desc.m_format == EImageFormat::IF_NONE && formatNode)
-				{
-					desc.m_format = formatNode->GetFormat();
-				}
-
-				const int layout = node.m_images[t].m_layoutIndex;
+		if (!bShareSurface)
+		{
+			for (int32 t = 0; t < node.m_images.Num(); ++t)
+			{
+				//MUTABLE_CPUPROFILER_SCOPE(SurfaceTexture);
 
-				if (desc.m_size[0] == 0 || desc.m_size[1] == 0)
+				if (NodeImagePtr pImageNode = node.m_images[t].m_pImage)
 				{
-					int currentLOD = m_currentParents.Last().m_lod;
-					char buf[256];
-					mutable_snprintf(buf, 256, "An image block for [%s] [%s] [%s] at lod [%d] has zero size and will not be generated. ", 
-						node.m_images[t].m_name.c_str(),
-						node.m_images[t].m_materialName.c_str(),
-						node.m_images[t].m_materialParameterName.c_str(),
-						currentLOD
-					);
-					m_pErrorLog->GetPrivate()->Add(buf, ELMT_INFO, node.m_errorContext);
-				}
-
-                else if ( desc.m_format == EImageFormat::IF_NONE )
-                {
-					char buf[256];
-					mutable_snprintf(buf, 256, "An image [%s] has an unidentified pixel format. ", node.m_images[t].m_name.c_str() );
-                    m_pErrorLog->GetPrivate()->Add(buf, ELMT_ERROR, node.m_errorContext);
-                }
-
-                else if ( 
-					// If the layout index has been set to negative, it means we should ignore the layout node.
-					(layout<0)
-					||
-					(m_compilerOptions->m_textureLayoutStrategy == CompilerOptions::TextureLayoutStrategy::None) 
-					)
-                {
-					check( desc.m_format!= EImageFormat::IF_NONE );
-
-                    // Generate the image
-                    IMAGE_STATE newState;
-					newState.m_imageSize[0] = desc.m_size[0];
-					newState.m_imageSize[1] = desc.m_size[1];
-					newState.m_imageRect = box<UE::Math::TIntVector2<int>>::FromMinSize(UE::Math::TIntVector2<int>(0,0), UE::Math::TIntVector2<int>(desc.m_size) );
-                    newState.m_layoutBlockId = -1;
-                    newState.m_pLayout = nullptr;
-                    m_imageState.Add( newState );
-
-                    Ptr<ASTOp> imageAd = Generate( pImageNode );
-
-                    // Look for patches to this block
-                    for( int32 editIndex=0; editIndex<edits.Num(); ++editIndex )
-                    {
-                        const FirstPassGenerator::SURFACE::EDIT& e = edits[editIndex];
-                        if ( t<e.node->m_textures.Num() )
-                        {
-                            if ( const NodePatchImage* pPatch = e.node->m_textures[t].m_pPatch.get() )
-                            {
-                                imageAd = GenerateImageBlockPatch(imageAd,pPatch,e.condition);
-                            }
-                        }
-
-                    }
-
-                    check(imageAd);
-
-                    m_imageState.Pop();
-
-                    if (mipmapNode)
-                    {
-                        Ptr<ASTOpImageMipmap> op = new ASTOpImageMipmap();
-                        op->Levels = 0;
-                        op->Source = imageAd;
-                        op->BlockLevels = 0;
-
-						op->AddressMode = mipmapNode->GetPrivate()->m_settings.m_addressMode;
-						op->FilterType = mipmapNode->GetPrivate()->m_settings.m_filterType;
-						op->SharpenFactor = mipmapNode->GetPrivate()->m_settings.m_sharpenFactor;
-						op->DitherMipmapAlpha = mipmapNode->GetPrivate()->m_settings.m_ditherMipmapAlpha;
-                        imageAd = op;
-                    }
-
-                    if (formatNode)
-                    {
-                        Ptr<ASTOpImagePixelFormat> fop = new ASTOpImagePixelFormat();
-                        fop->Format = formatNode->GetPrivate()->m_format;
-                        fop->FormatIfAlpha = formatNode->GetPrivate()->m_formatIfAlpha;
-                        fop->Source = imageAd;
-                        check(fop->Format != EImageFormat::IF_NONE);
-                        imageAd = fop;
-                    }
-
-                    Ptr<ASTOpInstanceAdd> op = new ASTOpInstanceAdd();
-                    op->type = OP_TYPE::IN_ADDIMAGE;
-                    op->instance = lastSurfOp;
-                    op->value = imageAd;
-                    op->name = node.m_images[t].m_name;
-
-                    lastSurfOp = op;
-                }
-
-                else if ( m_compilerOptions->m_textureLayoutStrategy
-                          ==
-                          CompilerOptions::TextureLayoutStrategy::Pack )
-                {
-                    if( layout >= meshResults.GeneratedLayouts.Num() ||
-                        layout >= meshResults.layoutOps.Num() )
-                    {
-                        m_pErrorLog->GetPrivate()->Add("Missing layout in object, or its parent.",
-                                                       ELMT_ERROR, node.m_errorContext);
-                    }
-                    else
-                    {
-                        const Layout* pLayout = meshResults.GeneratedLayouts[ layout ].get();
-                        check(pLayout);
-
-                        Ptr<ASTOpInstanceAdd> op = new ASTOpInstanceAdd();
-                        op->type = OP_TYPE::IN_ADDIMAGE;
-                        op->instance = lastSurfOp;
-
-                        // Image
-                        //-------------------------------------
-
-                        // Size of a layout block in pixels
-                        FIntPoint grid = pLayout->GetGridSize();
-
-                        check( desc.m_format!= EImageFormat::IF_NONE );
-
-                        // If the image is too small or not a multiple of the the layout size, 
-						// resize it, but raise a warning.
-						if ( (desc.m_size[0] % grid[0] != 0) || (desc.m_size[1] % grid[1] != 0) )
+					// Any image-specific format or mipmapping needs to be applied at the end
+					NodeImageMipmapPtr mipmapNode;
+					NodeImageFormatPtr formatNode;
+					bool found = false;
+					while (!found)
+					{
+						if (NodeImageMipmap* tm = dynamic_cast<NodeImageMipmap*>(pImageNode.get()))
 						{
-							FImageSize oldSize = desc.m_size;
-							desc.m_size[0] = grid[0] * FMath::Max(1, desc.m_size[0] / grid[0]);
-							desc.m_size[1] = grid[1] * FMath::Max(1, desc.m_size[1] / grid[1]);
-
-							// It is already visible in the error log
-							//UE_LOG( LogMutableCore, Warning, TEXT("A texture [%s] has been resized from [%d x %d] to [%d x %d] because it didn't fit the layout [%d x %d]. "),
-							//	StringCast<TCHAR>(node.m_images[t].m_name.c_str()).Get(), oldSize[0], oldSize[1], desc.m_size[0], desc.m_size[1], grid[0], grid[1] );
-
-							int currentLOD = m_currentParents.Last().m_lod;
-							char buf[256];
-							mutable_snprintf(buf, 256, "A texture [%s] for material [%s] parameter [%s] in LOD [%d] has been resized from [%d x %d] to [%d x %d] because it didn't fit the layout [%d x %d]. ",
-								node.m_images[t].m_name.c_str(), 
-								node.m_images[t].m_materialName.c_str(),
-								node.m_images[t].m_materialParameterName.c_str(),
-								currentLOD,
-								oldSize[0], oldSize[1], desc.m_size[0], desc.m_size[1], grid[0], grid[1] );
-							m_pErrorLog->GetPrivate()->Add(buf, ELMT_INFO, node.m_errorContext);
+							if (!mipmapNode) mipmapNode = tm;
+							pImageNode = tm->GetSource();
 						}
+						else if (NodeImageFormat* tf = dynamic_cast<NodeImageFormat*>(pImageNode.get()))
+						{
+							if (!formatNode) formatNode = tf;
+							pImageNode = tf->GetSource();
+						}
+						else
+						{
+							found = true;
+						}
+					}
 
-                        int blockSizeX = FMath::Max( 1, desc.m_size[0] / grid[0]);
-                        int blockSizeY = FMath::Max( 1, desc.m_size[1] / grid[1]);
-
-						bool bBlocksHaveMips = desc.m_lods > 1;
-
-                        // Start with a blank image
-                        Ptr<ASTOp> imageAd;
-                        {
-                            Ptr<ASTOpFixed> bop = new ASTOpFixed();
-                            bop->op.type = OP_TYPE::IM_BLANKLAYOUT;
-                            bop->SetChild(bop->op.args.ImageBlankLayout.layout, meshResults.layoutOps[ layout ] );
-                            bop->op.args.ImageBlankLayout.blockSize[0] = uint16(blockSizeX);
-                            bop->op.args.ImageBlankLayout.blockSize[1] = uint16(blockSizeY);
-                            // We support block compression directly here, but not non-block compression
-                            if ( GetImageFormatData(desc.m_format).m_pixelsPerBlockX == 0 )
-                            {
-                                // It's something like RLE
-                                bop->op.args.ImageBlankLayout.format = GetUncompressedFormat(desc.m_format);
-                            }
-                            else
-                            {
-                                // Directly supported
-                                bop->op.args.ImageBlankLayout.format = desc.m_format;
-                            }
-                            bop->op.args.ImageBlankLayout.generateMipmaps = bBlocksHaveMips;
-                            bop->op.args.ImageBlankLayout.mipmapCount = 0;
-                            imageAd = bop;
-                        }
-
-                        for( int b=0; b<pLayout->GetBlockCount(); ++b )
-                        {
-                            // Block in layout grid units
-                            box< UE::Math::TIntVector2<uint16> > rectInCells;
-                            pLayout->GetBlock
-                                (
-                                    b,
-                                    &rectInCells.min[0], &rectInCells.min[1],
-                                    &rectInCells.size[0], &rectInCells.size[1]
-                                );
-
-                            // Transform to pixels
-                            box< UE::Math::TIntVector2<uint16> > rect = rectInCells;
-                            rect.min[0] *= blockSizeX;
-                            rect.min[1] *= blockSizeY;
-                            rect.size[0] *= blockSizeX;
-                            rect.size[1] *= blockSizeY;
-
-                            // Generate the image
-                            IMAGE_STATE newState;
-                            newState.m_imageSize = UE::Math::TIntVector2<int32>(desc.m_size);
-							newState.m_imageRect.min = UE::Math::TIntVector2<int32>(rect.min);
-							newState.m_imageRect.size = UE::Math::TIntVector2<int32>(rect.size);
-							newState.m_layoutBlockId = pLayout->m_blocks[b].m_id;
-                            newState.m_pLayout = pLayout;
-                            m_imageState.Add( newState );
-
-                            Ptr<ASTOp> blockAd = Generate( pImageNode );
-
-                            // Look for patches to this block
-                            for( int32 editIndex=0; editIndex<edits.Num(); ++editIndex )
-                            {
-                                const FirstPassGenerator::SURFACE::EDIT& e = edits[editIndex];
-                                if ( t<e.node->m_textures.Num() )
-                                {
-                                    if ( const NodePatchImage* pPatch = e.node->m_textures[t].m_pPatch.get() )
-                                    {
-                                        // Is the current block to be patched?
-                                        if (pPatch->GetPrivate()->m_blocks.Contains(b))
-                                        {
-                                            blockAd = GenerateImageBlockPatch(blockAd,pPatch,e.condition);
-                                        }
-                                    }
-                                }
-                            }
-
-							EImageFormat baseFormat = imageAd->GetImageDesc().m_format;
-							Ptr<ASTOp> FormattedBlock = GenerateImageFormat(blockAd, baseFormat);
-
-							// Apply tiling to avoid generating chunks of image that are too big.
-							FormattedBlock = ApplyTiling(FormattedBlock, FImageSize(rect.size), desc.m_format);
-
-                            // Compose layout operation
-                            Ptr<ASTOpImageCompose> composeOp = new ASTOpImageCompose();
-                            composeOp->Layout = meshResults.layoutOps[ layout ];
-                            composeOp->Base = imageAd;
-							composeOp->BlockImage = FormattedBlock;
+					// Find out the size of the image
+					// TODO: What if the image is empty and everything is added?
+					//       Look in the extending images.
+					FImageDesc desc;
+					if (pImageNode)
+					{
+						MUTABLE_CPUPROFILER_SCOPE(CalculateImageDesc);
+						desc = CalculateImageDesc(*pImageNode->GetBasePrivate());
+					}
 
-                            // Set the absolute block index.
-                            check( pLayout->m_blocks[b].m_id>=0 );
-                            composeOp->BlockIndex = pLayout->m_blocks[b].m_id;
+					// If the image format doesn't come bottom-up, it may come top-down
+					if (desc.m_format == EImageFormat::IF_NONE && formatNode)
+					{
+						desc.m_format = formatNode->GetFormat();
+					}
 
-                            imageAd = composeOp;
+					const int LayoutIndex = node.m_images[t].m_layoutIndex;
 
-                            m_imageState.Pop();
-                        }
-                        check(imageAd);
+					// If the layout index has been set to negative, it means we should ignore the layout for this image.
+					CompilerOptions::TextureLayoutStrategy ImageLayoutStrategy = (LayoutIndex < 0)
+						? CompilerOptions::TextureLayoutStrategy::None
+						: CompilerOptions::TextureLayoutStrategy::Pack
+						;
 
-                        // Apply composition of blocks coming from child objects
-                        for( int32 editIndex=0; editIndex<edits.Num(); ++editIndex )
-                        {
-                            const FirstPassGenerator::SURFACE::EDIT& e = edits[editIndex];
-                            if ( t<e.node->m_textures.Num() )
-                            {
-                                Ptr<NodeImage> pExtend = e.node->m_textures[t].m_pExtend;
-                                if (pExtend)
-                                {
-                                    if ( layout>=meshResults.extraMeshLayouts[editIndex].GeneratedLayouts.Num()
-                                         ||
-                                         !meshResults.extraMeshLayouts[editIndex].GeneratedLayouts[layout] )
-                                    {
-                                        char buf[256];
-                                        mutable_snprintf
-                                            (
-                                                buf, 256,
-                                                "Trying to extend a layout that doesn't exist in object [%s].",
-                                                m_currentParents.Last().m_pObject->m_name.c_str()
-                                            );
-
-                                        m_pErrorLog->GetPrivate()->Add( buf, ELMT_ERROR, node.m_errorContext );
-                                    }
-                                    else
-                                    {
-                                        Ptr<const Layout> pExtendLayout = meshResults.extraMeshLayouts[editIndex].GeneratedLayouts[layout];
-
-                                        // Find out the size of the image
-                                        // TODO: What if the image is empty and everything is added?
-                                        //       Look in the extending images.
-                                        FImageDesc extendDesc = CalculateImageDesc( *pExtend->GetBasePrivate() );
-
-                                        // Size of a layout block in pixels
-                                        FIntPoint extlayout = pExtendLayout->GetGridSize();
-
-                                        Ptr<ASTOp> lastBase = imageAd;
-
-                                        for (int b=0; b<pExtendLayout->GetBlockCount(); ++b)
-                                        {
-                                            // Block in layout grid units
-                                            box< UE::Math::TIntVector2<uint16> > blockRect;
-                                            pExtendLayout->GetBlock
-                                                (
-                                                    b,
-                                                    &blockRect.min[0], &blockRect.min[1],
-                                                    &blockRect.size[0], &blockRect.size[1]
-                                                );
-
-                                            // Transform to pixels
-                                            box< UE::Math::TIntVector2<int32> > rect;
-                                            rect.min[0] = (blockRect.min[0] * extendDesc.m_size[0]) / extlayout[0];
-                                            rect.min[1] = (blockRect.min[1] * extendDesc.m_size[1]) / extlayout[1];
-                                            rect.size[0] = (blockRect.size[0] * extendDesc.m_size[0]) / extlayout[0];
-                                            rect.size[1] = (blockRect.size[1] * extendDesc.m_size[1]) / extlayout[1];
-
-                                            // Generate the image block
-                                            IMAGE_STATE newState;
-                                            newState.m_imageSize = UE::Math::TIntVector2<int32>(extendDesc.m_size);
-                                            newState.m_imageRect = rect;
-                                            newState.m_layoutBlockId = pExtendLayout->m_blocks[b].m_id;
-                                            newState.m_pLayout = pExtendLayout;
-                                            m_imageState.Add( newState );
-
-                                            Ptr<ASTOp> fragmentAd = Generate( pExtend );
-
-                                            m_imageState.Pop();
-
-                                            // Adjust the format and size of the block to be added
-                                            Ptr<ASTOp> formatted = GenerateImageFormat( fragmentAd, GetUncompressedFormat(desc.m_format) );
-                                            FImageSize expectedSize;
-                                            expectedSize[0] = (uint16)(blockSizeX * blockRect.size[0]);
-                                            expectedSize[1] = (uint16)(blockSizeY * blockRect.size[1]);
-                                            formatted = GenerateImageSize( formatted, expectedSize );
-                                            // Direct compression support
-                                            // formatted = GenerateImageUncompressed( formatted );
-                                            //formatted = formatted;
-                                            fragmentAd = formatted;
-
-											// Apply tiling to avoid generating chunks of image that are too big.
-											fragmentAd = ApplyTiling(fragmentAd, expectedSize, desc.m_format);
-
-                                            // Compose operation
-                                            Ptr<ASTOpImageCompose> composeOp = new ASTOpImageCompose();
-                                            composeOp->Layout = meshResults.layoutOps[layout];
-                                            composeOp->Base = lastBase;
-                                            // Direct compression support
-                                            // composeOp.args.ImageCompose.blockImage = GenerateImageUncompressed( fragmentAd );
-                                            composeOp->BlockImage = fragmentAd;
-
-                                            // Set the absolute block index.
-                                            check( pExtendLayout->m_blocks[b].m_id>=0 );
-                                            composeOp->BlockIndex = pExtendLayout->m_blocks[b].m_id;
-
-                                            lastBase = composeOp;
-                                        }
-
-                                        // Condition to enable this image extension
-                                        if (e.condition)
-                                        {
-                                            Ptr<ASTOp> conditionalAd;
-                                            Ptr<ASTOpConditional> cop = new ASTOpConditional();
-                                            cop->type = OP_TYPE::IM_CONDITIONAL;
-                                            cop->no = imageAd;
-                                            cop->yes = lastBase;
-                                            cop->condition = e.condition;
-                                            conditionalAd = cop;
-                                            imageAd = conditionalAd;
-                                        }
-                                        else
-                                        {
-                                            imageAd = lastBase;
-                                        }
-
-                                    }
-                                }
-                            }
-                        }
-
-
-                        // Apply mipmap and format if necessary
-                        if (mipmapNode)
-                        {
-                            Ptr<ASTOpImageMipmap> mop = new ASTOpImageMipmap();
+					if (desc.m_size[0] == 0 || desc.m_size[1] == 0)
+					{
+						int currentLOD = m_currentParents.Last().m_lod;
+						FString Msg = FString::Printf( TEXT("An image block for [%s] [%s] [%s] at lod [%d] has zero size and will not be generated. "),
+							node.m_images[t].m_name.c_str(),
+							node.m_images[t].m_materialName.c_str(),
+							node.m_images[t].m_materialParameterName.c_str(),
+							currentLOD
+						);
+						m_pErrorLog->GetPrivate()->Add(Msg, ELMT_INFO, node.m_errorContext);
+					}
 
-                            // At the end of the day, we want all the mipmaps. Maybe the code
-                            // optimiser will split the process later.
-                            mop->Levels = 0;
-							mop->bOnlyTail = false;
-							mop->Source = imageAd;
+					else if (desc.m_format == EImageFormat::IF_NONE)
+					{
+						FString Msg = FString::Printf(TEXT("An image [%s] has an unidentified pixel format. "), node.m_images[t].m_name.c_str());
+						m_pErrorLog->GetPrivate()->Add(Msg, ELMT_ERROR, node.m_errorContext);
+					}
 
-                            // We have to avoid mips smaller than the image format block size, so
-                            // we will devide the layout block by the format block
-                            const FImageFormatData& finfo = GetImageFormatData( desc.m_format );
+					else if (ImageLayoutStrategy == CompilerOptions::TextureLayoutStrategy::None)
+					{
+						check(desc.m_format != EImageFormat::IF_NONE);
+
+						// Generate the image
+						IMAGE_STATE newState;
+						newState.m_imageSize[0] = desc.m_size[0];
+						newState.m_imageSize[1] = desc.m_size[1];
+						newState.m_imageRect = box<UE::Math::TIntVector2<int>>::FromMinSize(UE::Math::TIntVector2<int>(0, 0), UE::Math::TIntVector2<int>(desc.m_size));
+						newState.m_layoutBlockId = -1;
+						newState.m_pLayout = nullptr;
+						m_imageState.Add(newState);
+
+						FImageGenerationOptions ImageOptions;
+						ImageOptions.ImageLayoutStrategy = ImageLayoutStrategy;
+						FImageGenerationResult Result;
+						GenerateImage(ImageOptions, Result, pImageNode);
+						Ptr<ASTOp> imageAd = Result.op;
+
+						// Look for patches to this block
+						for (int32 editIndex = 0; editIndex < edits.Num(); ++editIndex)
+						{
+							const FirstPassGenerator::SURFACE::EDIT& e = edits[editIndex];
+							if (t < e.node->m_textures.Num())
+							{
+								if (const NodePatchImage* pPatch = e.node->m_textures[t].m_pPatch.get())
+								{
+									imageAd = GenerateImageBlockPatch(imageAd, pPatch, e.condition, ImageOptions);
+								}
+							}
 
-                            int mipsX = (int)ceilf( logf( (float)blockSizeX/finfo.m_pixelsPerBlockX )/logf(2.0f) );
-                            int mipsY = (int)ceilf( logf( (float)blockSizeY/finfo.m_pixelsPerBlockY )/logf(2.0f) );
-                            mop->BlockLevels = (uint8_t)FMath::Max( mipsX, mipsY );
+						}
 
-							mop->AddressMode	   = mipmapNode->GetPrivate()->m_settings.m_addressMode;
-							mop->FilterType		   = mipmapNode->GetPrivate()->m_settings.m_filterType;
-							mop->SharpenFactor	   = mipmapNode->GetPrivate()->m_settings.m_sharpenFactor;
-							mop->DitherMipmapAlpha = mipmapNode->GetPrivate()->m_settings.m_ditherMipmapAlpha;
+						check(imageAd);
 
-                            imageAd = mop;
-                        }
+						m_imageState.Pop();
 
-						else if (bBlocksHaveMips)
+						if (mipmapNode)
 						{
-							// If the blocks had mipmaps, we still need to generate them after the compose, to have the full chain
-							Ptr<ASTOpImageMipmap> mop = new ASTOpImageMipmap();
-
-							// At the end of the day, we want all the mipmaps. Maybe the code
-							// optimiser will split the process later.
-							mop->Levels = 0;
-							mop->bOnlyTail = false;
-							mop->Source = imageAd;
-
-							// We have to avoid mips smaller than the image format block size, so
-							// we will devide the layout block by the format block
-							const FImageFormatData& finfo = GetImageFormatData(desc.m_format);
-
-							int mipsX = (int)ceilf(logf((float)blockSizeX / finfo.m_pixelsPerBlockX) / logf(2.0f));
-							int mipsY = (int)ceilf(logf((float)blockSizeY / finfo.m_pixelsPerBlockY) / logf(2.0f));
-							mop->BlockLevels = (uint8_t)FMath::Max(mipsX, mipsY);
-
-							// Not important for the end of the mip tail?
-							mop->AddressMode = EAddressMode::AM_CLAMP;
-							mop->FilterType = EMipmapFilterType::MFT_SimpleAverage;
-							mop->SharpenFactor = 0;
-							mop->DitherMipmapAlpha = false;
-
-							imageAd = mop;
+							Ptr<ASTOpImageMipmap> op = new ASTOpImageMipmap();
+							op->Levels = 0;
+							op->Source = imageAd;
+							op->BlockLevels = 0;
+
+							op->AddressMode = mipmapNode->GetPrivate()->m_settings.m_addressMode;
+							op->FilterType = mipmapNode->GetPrivate()->m_settings.m_filterType;
+							op->SharpenFactor = mipmapNode->GetPrivate()->m_settings.m_sharpenFactor;
+							op->DitherMipmapAlpha = mipmapNode->GetPrivate()->m_settings.m_ditherMipmapAlpha;
+							imageAd = op;
 						}
 
-                        if (formatNode)
-                        {
-                            Ptr<ASTOpImagePixelFormat> fop = new ASTOpImagePixelFormat();
-                            fop->Format = formatNode->GetPrivate()->m_format;
-                            fop->FormatIfAlpha = formatNode->GetPrivate()->m_formatIfAlpha;
-                            fop->Source = imageAd;
-                            check(fop->Format != EImageFormat::IF_NONE);
-                            imageAd = fop;
-                        }
+						if (formatNode)
+						{
+							Ptr<ASTOpImagePixelFormat> fop = new ASTOpImagePixelFormat();
+							fop->Format = formatNode->GetPrivate()->m_format;
+							fop->FormatIfAlpha = formatNode->GetPrivate()->m_formatIfAlpha;
+							fop->Source = imageAd;
+							check(fop->Format != EImageFormat::IF_NONE);
+							imageAd = fop;
+						}
 
-                        op->value = imageAd;
+						Ptr<ASTOpInstanceAdd> op = new ASTOpInstanceAdd();
+						op->type = OP_TYPE::IN_ADDIMAGE;
+						op->instance = lastSurfOp;
+						op->value = imageAd;
+						op->name = node.m_images[t].m_name;
 
-                        // Name
-                        op->name = node.m_images[t].m_name.c_str();
+						lastSurfOp = op;
+					}
 
-                        lastSurfOp = op;
-                    }
-                }
+					else if (ImageLayoutStrategy == CompilerOptions::TextureLayoutStrategy::Pack)
+					{
+						if (LayoutIndex >= meshResults.GeneratedLayouts.Num() ||
+							LayoutIndex >= meshResults.layoutOps.Num())
+						{
+							m_pErrorLog->GetPrivate()->Add("Missing layout in object, or its parent.", ELMT_ERROR, node.m_errorContext);
+						}
+						else
+						{
+							const Layout* pLayout = meshResults.GeneratedLayouts[LayoutIndex].get();
+							check(pLayout);
+
+							Ptr<ASTOpInstanceAdd> op = new ASTOpInstanceAdd();
+							op->type = OP_TYPE::IN_ADDIMAGE;
+							op->instance = lastSurfOp;
+
+							// Image
+							//-------------------------------------
+
+							// Size of a layout block in pixels
+							FIntPoint grid = pLayout->GetGridSize();
+
+							check(desc.m_format != EImageFormat::IF_NONE);
+
+							// If the image is too small or not a multiple of the the layout size, 
+							// resize it, but raise a warning.
+							if ((desc.m_size[0] % grid[0] != 0) || (desc.m_size[1] % grid[1] != 0))
+							{
+								FImageSize oldSize = desc.m_size;
+								desc.m_size[0] = grid[0] * FMath::Max(1, desc.m_size[0] / grid[0]);
+								desc.m_size[1] = grid[1] * FMath::Max(1, desc.m_size[1] / grid[1]);
+
+								int currentLOD = m_currentParents.Last().m_lod;
+								FString Msg = FString::Printf( TEXT("A texture [%s] for material [%s] parameter [%s] in LOD [%d] has been resized from [%d x %d] to [%d x %d] because it didn't fit the layout [%d x %d]. "),
+									node.m_images[t].m_name.c_str(),
+									node.m_images[t].m_materialName.c_str(),
+									node.m_images[t].m_materialParameterName.c_str(),
+									currentLOD,
+									oldSize[0], oldSize[1], desc.m_size[0], desc.m_size[1], grid[0], grid[1]);
+								m_pErrorLog->GetPrivate()->Add(Msg, ELMT_INFO, node.m_errorContext);
+							}
+
+							int blockSizeX = FMath::Max(1, desc.m_size[0] / grid[0]);
+							int blockSizeY = FMath::Max(1, desc.m_size[1] / grid[1]);
+
+							bool bBlocksHaveMips = desc.m_lods > 1;
+
+							// Start with a blank image
+							Ptr<ASTOp> imageAd;
+							{
+								Ptr<ASTOpFixed> bop = new ASTOpFixed();
+								bop->op.type = OP_TYPE::IM_BLANKLAYOUT;
+								bop->SetChild(bop->op.args.ImageBlankLayout.layout, meshResults.layoutOps[LayoutIndex]);
+								bop->op.args.ImageBlankLayout.blockSize[0] = uint16(blockSizeX);
+								bop->op.args.ImageBlankLayout.blockSize[1] = uint16(blockSizeY);
+								// We support block compression directly here, but not non-block compression
+								if (GetImageFormatData(desc.m_format).m_pixelsPerBlockX == 0)
+								{
+									// It's something like RLE
+									bop->op.args.ImageBlankLayout.format = GetUncompressedFormat(desc.m_format);
+								}
+								else
+								{
+									// Directly supported
+									bop->op.args.ImageBlankLayout.format = desc.m_format;
+								}
+								bop->op.args.ImageBlankLayout.generateMipmaps = bBlocksHaveMips;
+								bop->op.args.ImageBlankLayout.mipmapCount = 0;
+								imageAd = bop;
+							}
+
+							for (int b = 0; b < pLayout->GetBlockCount(); ++b)
+							{
+								// Block in layout grid units
+								box< UE::Math::TIntVector2<uint16> > rectInCells;
+								pLayout->GetBlock
+								(
+									b,
+									&rectInCells.min[0], &rectInCells.min[1],
+									&rectInCells.size[0], &rectInCells.size[1]
+								);
+
+								// Transform to pixels
+								box< UE::Math::TIntVector2<uint16> > rect = rectInCells;
+								rect.min[0] *= blockSizeX;
+								rect.min[1] *= blockSizeY;
+								rect.size[0] *= blockSizeX;
+								rect.size[1] *= blockSizeY;
+
+								// Generate the image
+								IMAGE_STATE newState;
+								newState.m_imageSize = UE::Math::TIntVector2<int32>(desc.m_size);
+								newState.m_imageRect.min = UE::Math::TIntVector2<int32>(rect.min);
+								newState.m_imageRect.size = UE::Math::TIntVector2<int32>(rect.size);
+								newState.m_layoutBlockId = pLayout->m_blocks[b].m_id;
+								newState.m_pLayout = pLayout;
+								m_imageState.Add(newState);
+
+								FImageGenerationOptions ImageOptions;
+								ImageOptions.ImageLayoutStrategy = ImageLayoutStrategy;
+								FImageGenerationResult Result;
+								GenerateImage(ImageOptions, Result, pImageNode);
+								Ptr<ASTOp> blockAd = Result.op;
+
+								// Look for patches to this block
+								for (int32 editIndex = 0; editIndex < edits.Num(); ++editIndex)
+								{
+									const FirstPassGenerator::SURFACE::EDIT& e = edits[editIndex];
+									if (t < e.node->m_textures.Num())
+									{
+										if (const NodePatchImage* pPatch = e.node->m_textures[t].m_pPatch.get())
+										{
+											// Is the current block to be patched?
+											if (pPatch->GetPrivate()->m_blocks.Contains(b))
+											{
+												blockAd = GenerateImageBlockPatch(blockAd, pPatch, e.condition, ImageOptions);
+											}
+										}
+									}
+								}
+
+								EImageFormat baseFormat = imageAd->GetImageDesc().m_format;
+								Ptr<ASTOp> FormattedBlock = GenerateImageFormat(blockAd, baseFormat);
+
+								// Apply tiling to avoid generating chunks of image that are too big.
+								FormattedBlock = ApplyTiling(FormattedBlock, FImageSize(rect.size), desc.m_format);
+
+								// Compose layout operation
+								Ptr<ASTOpImageCompose> composeOp = new ASTOpImageCompose();
+								composeOp->Layout = meshResults.layoutOps[LayoutIndex];
+								composeOp->Base = imageAd;
+								composeOp->BlockImage = FormattedBlock;
+
+								// Set the absolute block index.
+								check(pLayout->m_blocks[b].m_id >= 0);
+								composeOp->BlockIndex = pLayout->m_blocks[b].m_id;
+
+								imageAd = composeOp;
+
+								m_imageState.Pop();
+							}
+							check(imageAd);
+
+							// Apply composition of blocks coming from child objects
+							for (int32 editIndex = 0; editIndex < edits.Num(); ++editIndex)
+							{
+								const FirstPassGenerator::SURFACE::EDIT& e = edits[editIndex];
+								if (t < e.node->m_textures.Num())
+								{
+									Ptr<NodeImage> pExtend = e.node->m_textures[t].m_pExtend;
+									if (pExtend)
+									{
+										if (LayoutIndex >= meshResults.extraMeshLayouts[editIndex].GeneratedLayouts.Num()
+											||
+											!meshResults.extraMeshLayouts[editIndex].GeneratedLayouts[LayoutIndex])
+										{
+											FString Msg = FString::Printf(TEXT("Trying to extend a layout that doesn't exist in object [%s]."),
+												m_currentParents.Last().m_pObject->m_name.c_str()
+											);
+
+											m_pErrorLog->GetPrivate()->Add(Msg, ELMT_ERROR, node.m_errorContext);
+										}
+										else
+										{
+											Ptr<const Layout> pExtendLayout = meshResults.extraMeshLayouts[editIndex].GeneratedLayouts[LayoutIndex];
+
+											// Find out the size of the image
+											// TODO: What if the image is empty and everything is added?
+											//       Look in the extending images.
+											FImageDesc extendDesc = CalculateImageDesc(*pExtend->GetBasePrivate());
+
+											// Size of a layout block in pixels
+											FIntPoint extlayout = pExtendLayout->GetGridSize();
+
+											Ptr<ASTOp> lastBase = imageAd;
+
+											for (int b = 0; b < pExtendLayout->GetBlockCount(); ++b)
+											{
+												// Block in layout grid units
+												box< UE::Math::TIntVector2<uint16> > blockRect;
+												pExtendLayout->GetBlock
+												(
+													b,
+													&blockRect.min[0], &blockRect.min[1],
+													&blockRect.size[0], &blockRect.size[1]
+												);
+
+												// Transform to pixels
+												box< UE::Math::TIntVector2<int32> > rect;
+												rect.min[0] = (blockRect.min[0] * extendDesc.m_size[0]) / extlayout[0];
+												rect.min[1] = (blockRect.min[1] * extendDesc.m_size[1]) / extlayout[1];
+												rect.size[0] = (blockRect.size[0] * extendDesc.m_size[0]) / extlayout[0];
+												rect.size[1] = (blockRect.size[1] * extendDesc.m_size[1]) / extlayout[1];
+
+												// Generate the image block
+												IMAGE_STATE newState;
+												newState.m_imageSize = UE::Math::TIntVector2<int32>(extendDesc.m_size);
+												newState.m_imageRect = rect;
+												newState.m_layoutBlockId = pExtendLayout->m_blocks[b].m_id;
+												newState.m_pLayout = pExtendLayout;
+												m_imageState.Add(newState);
+
+												FImageGenerationOptions ImageOptions;
+												ImageOptions.ImageLayoutStrategy = ImageLayoutStrategy;
+												FImageGenerationResult ExtendResult;
+												GenerateImage(ImageOptions, ExtendResult, pExtend);
+												Ptr<ASTOp> fragmentAd = ExtendResult.op;
+
+												m_imageState.Pop();
+
+												// Adjust the format and size of the block to be added
+												Ptr<ASTOp> formatted = GenerateImageFormat(fragmentAd, GetUncompressedFormat(desc.m_format));
+												FImageSize expectedSize;
+												expectedSize[0] = (uint16)(blockSizeX * blockRect.size[0]);
+												expectedSize[1] = (uint16)(blockSizeY * blockRect.size[1]);
+												formatted = GenerateImageSize(formatted, expectedSize);
+												fragmentAd = formatted;
+
+												// Apply tiling to avoid generating chunks of image that are too big.
+												fragmentAd = ApplyTiling(fragmentAd, expectedSize, desc.m_format);
+
+												// Compose operation
+												Ptr<ASTOpImageCompose> composeOp = new ASTOpImageCompose();
+												composeOp->Layout = meshResults.layoutOps[LayoutIndex];
+												composeOp->Base = lastBase;
+												composeOp->BlockImage = fragmentAd;
+
+												// Set the absolute block index.
+												check(pExtendLayout->m_blocks[b].m_id >= 0);
+												composeOp->BlockIndex = pExtendLayout->m_blocks[b].m_id;
+
+												lastBase = composeOp;
+											}
+
+											// Condition to enable this image extension
+											if (e.condition)
+											{
+												Ptr<ASTOp> conditionalAd;
+												Ptr<ASTOpConditional> cop = new ASTOpConditional();
+												cop->type = OP_TYPE::IM_CONDITIONAL;
+												cop->no = imageAd;
+												cop->yes = lastBase;
+												cop->condition = e.condition;
+												conditionalAd = cop;
+												imageAd = conditionalAd;
+											}
+											else
+											{
+												imageAd = lastBase;
+											}
+
+										}
+									}
+								}
+							}
+
+
+							// Apply mipmap and format if necessary
+							if (mipmapNode)
+							{
+								Ptr<ASTOpImageMipmap> mop = new ASTOpImageMipmap();
+
+								// At the end of the day, we want all the mipmaps. Maybe the code
+								// optimiser will split the process later.
+								mop->Levels = 0;
+								mop->bOnlyTail = false;
+								mop->Source = imageAd;
+
+								// We have to avoid mips smaller than the image format block size, so
+								// we will devide the layout block by the format block
+								const FImageFormatData& finfo = GetImageFormatData(desc.m_format);
+
+								int mipsX = (int)ceilf(logf((float)blockSizeX / finfo.m_pixelsPerBlockX) / logf(2.0f));
+								int mipsY = (int)ceilf(logf((float)blockSizeY / finfo.m_pixelsPerBlockY) / logf(2.0f));
+								mop->BlockLevels = (uint8_t)FMath::Max(mipsX, mipsY);
+
+								mop->AddressMode = mipmapNode->GetPrivate()->m_settings.m_addressMode;
+								mop->FilterType = mipmapNode->GetPrivate()->m_settings.m_filterType;
+								mop->SharpenFactor = mipmapNode->GetPrivate()->m_settings.m_sharpenFactor;
+								mop->DitherMipmapAlpha = mipmapNode->GetPrivate()->m_settings.m_ditherMipmapAlpha;
+
+								imageAd = mop;
+							}
+
+							else if (bBlocksHaveMips)
+							{
+								// If the blocks had mipmaps, we still need to generate them after the compose, to have the full chain
+								Ptr<ASTOpImageMipmap> mop = new ASTOpImageMipmap();
+
+								// At the end of the day, we want all the mipmaps. Maybe the code
+								// optimiser will split the process later.
+								mop->Levels = 0;
+								mop->bOnlyTail = false;
+								mop->Source = imageAd;
+
+								// We have to avoid mips smaller than the image format block size, so
+								// we will devide the layout block by the format block
+								const FImageFormatData& finfo = GetImageFormatData(desc.m_format);
+
+								int mipsX = (int)ceilf(logf((float)blockSizeX / finfo.m_pixelsPerBlockX) / logf(2.0f));
+								int mipsY = (int)ceilf(logf((float)blockSizeY / finfo.m_pixelsPerBlockY) / logf(2.0f));
+								mop->BlockLevels = (uint8_t)FMath::Max(mipsX, mipsY);
+
+								// Not important for the end of the mip tail?
+								mop->AddressMode = EAddressMode::AM_CLAMP;
+								mop->FilterType = EMipmapFilterType::MFT_SimpleAverage;
+								mop->SharpenFactor = 0;
+								mop->DitherMipmapAlpha = false;
+
+								imageAd = mop;
+							}
+
+							if (formatNode)
+							{
+								Ptr<ASTOpImagePixelFormat> fop = new ASTOpImagePixelFormat();
+								fop->Format = formatNode->GetPrivate()->m_format;
+								fop->FormatIfAlpha = formatNode->GetPrivate()->m_formatIfAlpha;
+								fop->Source = imageAd;
+								check(fop->Format != EImageFormat::IF_NONE);
+								imageAd = fop;
+							}
+
+							op->value = imageAd;
+
+							// Name
+							op->name = node.m_images[t].m_name.c_str();
+
+							lastSurfOp = op;
+						}
+					}
 
-                else
-                {
-                    // Unimplemented texture layout strategy
-                    check(false);
-                }
-            }
-        }
+					else
+					{
+						// Unimplemented texture layout strategy
+						check(false);
+					}
+				}
+			}
 
-        // Create the expression for each vector
-        //------------------------------------------------------------------------
-        for (int32 t = 0; t<node.m_vectors.Num(); ++t)
-        {
-            //MUTABLE_CPUPROFILER_SCOPE(SurfaceVector);
+			// Create the expression for each vector
+			//------------------------------------------------------------------------
+			for (int32 t = 0; t < node.m_vectors.Num(); ++t)
+			{
+				//MUTABLE_CPUPROFILER_SCOPE(SurfaceVector);
 
-            if (NodeColourPtr pVectorNode = node.m_vectors[t].m_pVector)
-            {
-                Ptr<ASTOpInstanceAdd> op = new ASTOpInstanceAdd();
-                op->type = OP_TYPE::IN_ADDVECTOR;
-                op->instance = lastSurfOp;
+				if (NodeColourPtr pVectorNode = node.m_vectors[t].m_pVector)
+				{
+					Ptr<ASTOpInstanceAdd> op = new ASTOpInstanceAdd();
+					op->type = OP_TYPE::IN_ADDVECTOR;
+					op->instance = lastSurfOp;
 
-                // Vector
-                Ptr<ASTOp> vectorAd = Generate(pVectorNode);
-                op->value = vectorAd;
+					// Vector
+					Ptr<ASTOp> vectorAd = Generate(pVectorNode);
+					op->value = vectorAd;
 
-                // Name
-                op->name = node.m_vectors[t].m_name;
+					// Name
+					op->name = node.m_vectors[t].m_name;
 
-                lastSurfOp = op;
-            }
-        }
+					lastSurfOp = op;
+				}
+			}
 
-        // Create the expression for each scalar
-        //------------------------------------------------------------------------
-        for ( int32 t = 0; t < node.m_scalars.Num(); ++t )
-        {
-            // MUTABLE_CPUPROFILER_SCOPE(SurfaceScalar);
+			// Create the expression for each scalar
+			//------------------------------------------------------------------------
+			for (int32 t = 0; t < node.m_scalars.Num(); ++t)
+			{
+				// MUTABLE_CPUPROFILER_SCOPE(SurfaceScalar);
 
-            if ( NodeScalarPtr pScalarNode = node.m_scalars[t].m_pScalar )
-            {
-                Ptr<ASTOpInstanceAdd> op = new ASTOpInstanceAdd();
-                op->type = OP_TYPE::IN_ADDSCALAR;
-                op->instance = lastSurfOp;
+				if (NodeScalarPtr pScalarNode = node.m_scalars[t].m_pScalar)
+				{
+					Ptr<ASTOpInstanceAdd> op = new ASTOpInstanceAdd();
+					op->type = OP_TYPE::IN_ADDSCALAR;
+					op->instance = lastSurfOp;
 
-                // Scalar
-                Ptr<ASTOp> scalarAd = Generate( pScalarNode );
-                op->value = scalarAd;
+					// Scalar
+					Ptr<ASTOp> scalarAd = Generate(pScalarNode);
+					op->value = scalarAd;
 
-                // Name
-                op->name = node.m_scalars[t].m_name;
+					// Name
+					op->name = node.m_scalars[t].m_name;
 
-                lastSurfOp = op;
-            }
-        }
+					lastSurfOp = op;
+				}
+			}
 
-        // Create the expression for each string
-        //------------------------------------------------------------------------
-        for ( int32 t = 0; t < node.m_strings.Num(); ++t )
-        {
-            if ( NodeStringPtr pStringNode = node.m_strings[t].m_pString )
-            {
-                Ptr<ASTOpInstanceAdd> op = new ASTOpInstanceAdd();
-                op->type = OP_TYPE::IN_ADDSTRING;
-                op->instance = lastSurfOp;
+			// Create the expression for each string
+			//------------------------------------------------------------------------
+			for (int32 t = 0; t < node.m_strings.Num(); ++t)
+			{
+				if (NodeStringPtr pStringNode = node.m_strings[t].m_pString)
+				{
+					Ptr<ASTOpInstanceAdd> op = new ASTOpInstanceAdd();
+					op->type = OP_TYPE::IN_ADDSTRING;
+					op->instance = lastSurfOp;
 
-                Ptr<ASTOp> stringAd = Generate( pStringNode );
-                op->value = stringAd;
+					Ptr<ASTOp> stringAd = Generate(pStringNode);
+					op->value = stringAd;
 
-                // Name
-                op->name = node.m_strings[t].m_name;
+					// Name
+					op->name = node.m_strings[t].m_name;
 
-                lastSurfOp = op;
-            }
-        }
+					lastSurfOp = op;
+				}
+			}
+		}
 
         result.surfaceOp = lastSurfOp;
         targetSurface->resultSurfaceOp = lastSurfOp;
+
+		// If we are going to share this surface properties, remember it.
+		if (bIsBaseForSharedSurface)
+		{
+			check(!SharedMeshOptionsMap.Contains(node.SharedSurfaceId));
+			SharedMeshOptionsMap.Add( node.SharedSurfaceId, meshResults );
+		}
     }
 
 
@@ -2060,14 +2068,10 @@ namespace mu
 					!=
 					usedNames.end() )
 				{
-					char buf[256];
-					mutable_snprintf
-					(
-						buf, 256,
-						"Object group has more than one children with the same name [%s].",
+					FString Msg = FString::Printf(TEXT("Object group has more than one children with the same name [%s]."),
 						strChildName
 					);
-					m_pErrorLog->GetPrivate()->Add(buf, ELMT_WARNING, node.m_errorContext);
+					m_pErrorLog->GetPrivate()->Add(Msg, ELMT_WARNING, node.m_errorContext);
 				}
 				else
 				{
@@ -2080,19 +2084,13 @@ namespace mu
     }
 
     //---------------------------------------------------------------------------------------------
-    Ptr<ASTOp> CodeGenerator::GenerateMissingBoolCode( const char* strWhere,
+    Ptr<ASTOp> CodeGenerator::GenerateMissingBoolCode(const TCHAR* strWhere,
                                                       bool value,
                                                       const void* errorContext )
     {
         // Log a warning
-        char buf[256];
-        mutable_snprintf
-            (
-                buf, 256,
-                "Required connection not found: %s",
-                strWhere
-            );
-        m_pErrorLog->GetPrivate()->Add( buf, ELMT_ERROR, errorContext );
+		FString Msg = FString::Printf(TEXT("Required connection not found: %s"), strWhere );
+        m_pErrorLog->GetPrivate()->Add( Msg, ELMT_ERROR, errorContext );
 
         // Create a constant node
         NodeBoolConstantPtr pNode = new NodeBoolConstant();
diff --git a/Engine/Plugins/Experimental/Mutable/Source/MutableTools/Private/MuT/CodeGenerator.h b/Engine/Plugins/Experimental/Mutable/Source/MutableTools/Private/MuT/CodeGenerator.h
--- a/Engine/Plugins/Experimental/Mutable/Source/MutableTools/Private/MuT/CodeGenerator.h
+++ b/Engine/Plugins/Experimental/Mutable/Source/MutableTools/Private/MuT/CodeGenerator.h
@@ -333,18 +333,12 @@ namespace mu
         Ptr<ASTOp> GenerateTableVariable( TablePtr pTable, const string& strName );
 
         //!
-        Ptr<ASTOp> GenerateMissingBoolCode( const char* strWhere, bool value,
-                                         const void* errorContext );
+        Ptr<ASTOp> GenerateMissingBoolCode(const TCHAR* strWhere, bool value, const void* errorContext );
 
         //!
 		template<class NODE_TABLE_PRIVATE, TABLE_COLUMN_TYPE TYPE, OP_TYPE OPTYPE, typename F>
 		Ptr<ASTOp> GenerateTableSwitch( const NODE_TABLE_PRIVATE& node, F&& GenerateOption );
 
-        //!
-        Ptr<ASTOp> GenerateImageBlockPatch( Ptr<ASTOp> blockAd,
-                                             const NodePatchImage* pPatch,
-                                             Ptr<ASTOp> conditionAd );
-
 
     private:
 
@@ -369,50 +363,91 @@ namespace mu
 
 		//-----------------------------------------------------------------------------------------
 		// Images
+				
+		/** Options that affect the generation of images. It is like list of what required data we want while parsing down the image node graph. */
+		struct FImageGenerationOptions
+		{
+			/** */
+			CompilerOptions::TextureLayoutStrategy ImageLayoutStrategy = CompilerOptions::TextureLayoutStrategy::None;
+
+			friend FORCEINLINE uint32 GetTypeHash(const FImageGenerationOptions& InKey)
+			{
+				uint32 KeyHash = 0;
+				KeyHash = HashCombine(KeyHash, ::GetTypeHash(InKey.ImageLayoutStrategy));
+				return KeyHash;
+			}
+
+			FORCEINLINE bool operator==(const FImageGenerationOptions& Other) const
+			{
+				return ImageLayoutStrategy == Other.ImageLayoutStrategy;
+			}
+
+		};
+
 		struct FImageGenerationResult
 		{
 			Ptr<ASTOp> op;
 		};
 
+
+		struct FGeneratedImageCacheKey
+		{
+			NodePtrConst Node;
+			FImageGenerationOptions Options;
+
+			friend FORCEINLINE uint32 GetTypeHash(const FGeneratedImageCacheKey& InKey)
+			{
+				uint32 KeyHash = 0;
+				KeyHash = HashCombine(KeyHash, ::GetTypeHash(InKey.Node.get()));
+				KeyHash = HashCombine(KeyHash, GetTypeHash(InKey.Options));
+				return KeyHash;
+			}
+
+			FORCEINLINE bool operator==(const FGeneratedImageCacheKey& Other) const
+			{
+				return Node == Other.Node && Options == Other.Options;
+			}
+		};
+
 		typedef TMap<FVisitedKeyMap, FImageGenerationResult> GeneratedImagesMap;
 		GeneratedImagesMap m_generatedImages;
 
-		void GenerateImage(FImageGenerationResult& result, const NodeImagePtrConst& node);
-		void GenerateImage_Constant(FImageGenerationResult&, const NodeImageConstant*);
-		void GenerateImage_Reference(FImageGenerationResult&, const NodeImageReference*);
-		void GenerateImage_Interpolate(FImageGenerationResult&, const NodeImageInterpolate*);
-		void GenerateImage_Saturate(FImageGenerationResult&, const NodeImageSaturate*);
-		void GenerateImage_Table(FImageGenerationResult&, const NodeImageTable*);
-		void GenerateImage_Swizzle(FImageGenerationResult&, const NodeImageSwizzle*);
-		void GenerateImage_ColourMap(FImageGenerationResult&, const NodeImageColourMap*);
-		void GenerateImage_Gradient(FImageGenerationResult&, const NodeImageGradient*);
-		void GenerateImage_Binarise(FImageGenerationResult&, const NodeImageBinarise*);
-		void GenerateImage_Luminance(FImageGenerationResult&, const NodeImageLuminance*);
-		void GenerateImage_Layer(FImageGenerationResult&, const NodeImageLayer*);
-		void GenerateImage_LayerColour(FImageGenerationResult&, const NodeImageLayerColour*);
-		void GenerateImage_Resize(FImageGenerationResult&, const NodeImageResize*);
-		void GenerateImage_PlainColour(FImageGenerationResult&, const NodeImagePlainColour*);
-		void GenerateImage_Interpolate3(FImageGenerationResult&, const NodeImageInterpolate3*);
-		void GenerateImage_Project(FImageGenerationResult&, const NodeImageProject*);
-		void GenerateImage_Mipmap(FImageGenerationResult&, const NodeImageMipmap*);
-		void GenerateImage_Switch(FImageGenerationResult&, const NodeImageSwitch*);
-		void GenerateImage_Conditional(FImageGenerationResult&, const NodeImageConditional*);
-		void GenerateImage_Format(FImageGenerationResult&, const NodeImageFormat*);
-		void GenerateImage_Parameter(FImageGenerationResult&, const NodeImageParameter*);
-		void GenerateImage_MultiLayer(FImageGenerationResult&, const NodeImageMultiLayer*);
-		void GenerateImage_Invert(FImageGenerationResult&, const NodeImageInvert*);
-		void GenerateImage_Variation(FImageGenerationResult&, const NodeImageVariation*);
-		void GenerateImage_NormalComposite(FImageGenerationResult&, const NodeImageNormalComposite*);
-		void GenerateImage_Transform(FImageGenerationResult&, const NodeImageTransform*);
+		void GenerateImage(const FImageGenerationOptions&, FImageGenerationResult& result, const NodeImagePtrConst& node);
+		void GenerateImage_Constant(const FImageGenerationOptions&, FImageGenerationResult&, const NodeImageConstant*);
+		void GenerateImage_Reference(const FImageGenerationOptions&, FImageGenerationResult&, const NodeImageReference*);
+		void GenerateImage_Interpolate(const FImageGenerationOptions&, FImageGenerationResult&, const NodeImageInterpolate*);
+		void GenerateImage_Saturate(const FImageGenerationOptions&, FImageGenerationResult&, const NodeImageSaturate*);
+		void GenerateImage_Table(const FImageGenerationOptions&, FImageGenerationResult&, const NodeImageTable*);
+		void GenerateImage_Swizzle(const FImageGenerationOptions&, FImageGenerationResult&, const NodeImageSwizzle*);
+		void GenerateImage_ColourMap(const FImageGenerationOptions&, FImageGenerationResult&, const NodeImageColourMap*);
+		void GenerateImage_Gradient(const FImageGenerationOptions&, FImageGenerationResult&, const NodeImageGradient*);
+		void GenerateImage_Binarise(const FImageGenerationOptions&, FImageGenerationResult&, const NodeImageBinarise*);
+		void GenerateImage_Luminance(const FImageGenerationOptions&, FImageGenerationResult&, const NodeImageLuminance*);
+		void GenerateImage_Layer(const FImageGenerationOptions&, FImageGenerationResult&, const NodeImageLayer*);
+		void GenerateImage_LayerColour(const FImageGenerationOptions&, FImageGenerationResult&, const NodeImageLayerColour*);
+		void GenerateImage_Resize(const FImageGenerationOptions&, FImageGenerationResult&, const NodeImageResize*);
+		void GenerateImage_PlainColour(const FImageGenerationOptions&, FImageGenerationResult&, const NodeImagePlainColour*);
+		void GenerateImage_Interpolate3(const FImageGenerationOptions&, FImageGenerationResult&, const NodeImageInterpolate3*);
+		void GenerateImage_Project(const FImageGenerationOptions&, FImageGenerationResult&, const NodeImageProject*);
+		void GenerateImage_Mipmap(const FImageGenerationOptions&, FImageGenerationResult&, const NodeImageMipmap*);
+		void GenerateImage_Switch(const FImageGenerationOptions&, FImageGenerationResult&, const NodeImageSwitch*);
+		void GenerateImage_Conditional(const FImageGenerationOptions&, FImageGenerationResult&, const NodeImageConditional*);
+		void GenerateImage_Format(const FImageGenerationOptions&, FImageGenerationResult&, const NodeImageFormat*);
+		void GenerateImage_Parameter(const FImageGenerationOptions&, FImageGenerationResult&, const NodeImageParameter*);
+		void GenerateImage_MultiLayer(const FImageGenerationOptions&, FImageGenerationResult&, const NodeImageMultiLayer*);
+		void GenerateImage_Invert(const FImageGenerationOptions&, FImageGenerationResult&, const NodeImageInvert*);
+		void GenerateImage_Variation(const FImageGenerationOptions&, FImageGenerationResult&, const NodeImageVariation*);
+		void GenerateImage_NormalComposite(const FImageGenerationOptions&, FImageGenerationResult&, const NodeImageNormalComposite*);
+		void GenerateImage_Transform(const FImageGenerationOptions&, FImageGenerationResult&, const NodeImageTransform*);
 
 		//!
-		ImagePtr GenerateMissingImage(EImageFormat);
+		Ptr<Image> GenerateMissingImage(EImageFormat);
 
 		//!
-		Ptr<ASTOp> GenerateMissingImageCode(const char* strWhere, EImageFormat, const void* errorContext);
+		Ptr<ASTOp> GenerateMissingImageCode(const TCHAR* strWhere, EImageFormat, const void* errorContext, const FImageGenerationOptions& Options);
 
 		//!
-		Ptr<ASTOp> GeneratePlainImageCode(const vec3<float>& colour);
+		Ptr<ASTOp> GeneratePlainImageCode(const vec3<float>& colour, const FImageGenerationOptions& Options);
 
 		//!
 		Ptr<ASTOp> GenerateImageFormat(Ptr<ASTOp>, EImageFormat);
@@ -429,6 +464,9 @@ namespace mu
 		/** Evaluate if the image to generate is big enough to be split in separate operations and tiled afterwards. */
 		Ptr<ASTOp> ApplyTiling(Ptr<ASTOp> Source, FImageSize Size, EImageFormat Format);
 
+		//!
+		Ptr<ASTOp> GenerateImageBlockPatch(Ptr<ASTOp> blockAd, const NodePatchImage* pPatch, Ptr<ASTOp> conditionAd, const FImageGenerationOptions& ImageOptions);
+
         //-----------------------------------------------------------------------------------------
         // Meshes
 
@@ -458,7 +496,7 @@ namespace mu
 			bool bNormalizeUVs = false;
 
 			/** If this has something the layouts in constant meshes will be ignored, because
-			* they are supposed to match some other set of layouts. If the vector is empty, layouts
+			* they are supposed to match some other set of layouts. If the array is empty, layouts
 			* are generated normally.
 			*/
 			TArray<Ptr<const Layout>> OverrideLayouts;			
@@ -530,12 +568,17 @@ namespace mu
         typedef TMap<FGeneratedMeshCacheKey,FMeshGenerationResult> GeneratedMeshMap;
         GeneratedMeshMap m_generatedMeshes;
 
+		/** Store the mesh generation data for surfaces that we intend to share across LODs. 
+		* The key is the SharedSurfaceId.
+		*/
+		TMap<int32, FMeshGenerationResult> SharedMeshOptionsMap;
+
 		//! Map of layouts found in the code already generated. The map is from the source layout
 		//! pointer of the layouts in the node meshes to the cloned and modified layout. 
 		//! The cloned layout will have absolute block ids assigned.
 		TMap<Ptr<const Layout>, Ptr<const Layout>> m_generatedLayouts;
 
-        void GenerateMesh(const FMeshGenerationOptions&, FMeshGenerationResult& result, const NodeMeshPtrConst& node);
+        void GenerateMesh(const FMeshGenerationOptions&, FMeshGenerationResult& result, const NodeMeshPtrConst&);
         void GenerateMesh_Constant(const FMeshGenerationOptions&, FMeshGenerationResult&, const NodeMeshConstant* );
         void GenerateMesh_Format(const FMeshGenerationOptions&, FMeshGenerationResult&, const NodeMeshFormat* );
         void GenerateMesh_Morph(const FMeshGenerationOptions&, FMeshGenerationResult&, const NodeMeshMorph* );
@@ -627,7 +670,7 @@ namespace mu
 		void GenerateScalar_Arithmetic(FScalarGenerationResult&, const Ptr<const NodeScalarArithmeticOperation>&);
 		void GenerateScalar_Variation(FScalarGenerationResult&, const Ptr<const NodeScalarVariation>&);
 		void GenerateScalar_Table(FScalarGenerationResult&, const Ptr<const NodeScalarTable>&);
-		Ptr<ASTOp> GenerateMissingScalarCode(const char* strWhere, float value, const void* errorContext);
+		Ptr<ASTOp> GenerateMissingScalarCode(const TCHAR* strWhere, float value, const void* errorContext);
 
 		//-----------------------------------------------------------------------------------------
 		// Colors
@@ -648,7 +691,7 @@ namespace mu
 		void GenerateColor_Arithmetic(FColorGenerationResult&, const Ptr<const NodeColourArithmeticOperation>&);
 		void GenerateColor_Variation(FColorGenerationResult&, const Ptr<const NodeColourVariation>&);
 		void GenerateColor_Table(FColorGenerationResult&, const Ptr<const NodeColourTable>&);
-		Ptr<ASTOp> GenerateMissingColourCode(const char* strWhere, const void* errorContext);
+		Ptr<ASTOp> GenerateMissingColourCode(const TCHAR* strWhere, const void* errorContext);
 
 		//-----------------------------------------------------------------------------------------
 		// Strings
diff --git a/Engine/Plugins/Experimental/Mutable/Source/MutableTools/Private/MuT/CodeGenerator_Color.cpp b/Engine/Plugins/Experimental/Mutable/Source/MutableTools/Private/MuT/CodeGenerator_Color.cpp
--- a/Engine/Plugins/Experimental/Mutable/Source/MutableTools/Private/MuT/CodeGenerator_Color.cpp
+++ b/Engine/Plugins/Experimental/Mutable/Source/MutableTools/Private/MuT/CodeGenerator_Color.cpp
@@ -189,8 +189,7 @@ namespace mu
 		if (node.m_options.Num() == 0)
 		{
 			// No options in the switch!
-			Ptr<ASTOp> missingOp = GenerateMissingColourCode("Switch option",
-				node.m_errorContext);
+			Ptr<ASTOp> missingOp = GenerateMissingColourCode(TEXT("Switch option"), node.m_errorContext);
 			result.op = missingOp;
 			return;
 		}
@@ -206,7 +205,7 @@ namespace mu
 		else
 		{
 			// This argument is required
-			op->variable = GenerateMissingScalarCode("Switch variable", 0.0f, node.m_errorContext);
+			op->variable = GenerateMissingScalarCode(TEXT("Switch variable"), 0.0f, node.m_errorContext);
 		}
 
 		// Options
@@ -220,8 +219,7 @@ namespace mu
 			else
 			{
 				// This argument is required
-				branch = GenerateMissingColourCode("Switch option",
-					node.m_errorContext);
+				branch = GenerateMissingColourCode(TEXT("Switch option"), node.m_errorContext);
 			}
 			op->cases.Emplace((int16)t, op, branch);
 		}
@@ -246,7 +244,7 @@ namespace mu
 		else
 		{
 			// This argument is required
-			currentOp = GenerateMissingColourCode("Variation default", node.m_errorContext);
+			currentOp = GenerateMissingColourCode(TEXT("Variation default"), node.m_errorContext);
 		}
 
 		// Process variations in reverse order, since conditionals are built bottom-up.
@@ -264,11 +262,8 @@ namespace mu
 
 			if (tagIndex < 0)
 			{
-				char buf[256];
-				mutable_snprintf(buf, 256, "Unknown tag found in image variation [%s].",
-					tag.c_str());
-
-				m_pErrorLog->GetPrivate()->Add(buf, ELMT_WARNING, node.m_errorContext);
+				FString Msg = FString::Printf(TEXT("Unknown tag found in image variation [%s]."), *FString(tag.c_str()));
+				m_pErrorLog->GetPrivate()->Add(Msg, ELMT_WARNING, node.m_errorContext);
 				continue;
 			}
 
@@ -280,7 +275,7 @@ namespace mu
 			else
 			{
 				// This argument is required
-				variationOp = GenerateMissingColourCode("Variation option", node.m_errorContext);
+				variationOp = GenerateMissingColourCode(TEXT("Variation option"), node.m_errorContext);
 			}
 
 
@@ -306,13 +301,14 @@ namespace mu
 		Ptr<ASTOpFixed> op = new ASTOpFixed();
 		op->op.type = OP_TYPE::CO_SAMPLEIMAGE;
 
-
 		// Source image
+		FImageGenerationOptions ImageOptions;
+		ImageOptions.ImageLayoutStrategy = CompilerOptions::TextureLayoutStrategy::None;
 		Ptr<ASTOp> base;
-		if (Node* pSource = node.m_pImage.get())
+		if (node.m_pImage)
 		{
 			// We take whatever size will be produced
-			FImageDesc desc = CalculateImageDesc(*pSource->GetBasePrivate());
+			FImageDesc desc = CalculateImageDesc(*node.m_pImage->GetBasePrivate());
 			IMAGE_STATE newState;
 			newState.m_imageSize[0] = desc.m_size[0];
 			newState.m_imageSize[1] = desc.m_size[1];
@@ -323,7 +319,9 @@ namespace mu
 			m_imageState.Add(newState);
 
 			// Generate
-			base = Generate(pSource);
+			FImageGenerationResult MapResult;
+			GenerateImage(ImageOptions, MapResult, node.m_pImage);
+			base = MapResult.op;
 
 			// Restore rect
 			m_imageState.Pop();
@@ -331,7 +329,7 @@ namespace mu
 		else
 		{
 			// This argument is required
-			base = GenerateMissingImageCode("Sample image", EImageFormat::IF_RGB_UBYTE, node.m_errorContext);
+			base = GenerateMissingImageCode(TEXT("Sample image"), EImageFormat::IF_RGB_UBYTE, node.m_errorContext, ImageOptions);
 		}
 		base = GenerateImageFormat(base, EImageFormat::IF_RGB_UBYTE);
 		op->SetChild(op->op.args.ColourSampleImage.image, base);
@@ -459,11 +457,7 @@ namespace mu
 		else
 		{
 			op->SetChild(op->op.args.ColourArithmetic.a,
-				CodeGenerator::GenerateMissingColourCode
-				(
-					"ColourArithmetic A",
-					node.m_errorContext
-				));
+				CodeGenerator::GenerateMissingColourCode(TEXT("ColourArithmetic A"), node.m_errorContext));
 		}
 
 		// B
@@ -474,11 +468,7 @@ namespace mu
 		else
 		{
 			op->SetChild(op->op.args.ColourArithmetic.b,
-				CodeGenerator::GenerateMissingColourCode
-				(
-					"ColourArithmetic B",
-					node.m_errorContext
-				));
+				CodeGenerator::GenerateMissingColourCode(TEXT("ColourArithmetic B"),node.m_errorContext));
 		}
 
 		result.op = op;
@@ -502,12 +492,11 @@ namespace mu
 
 
 	//---------------------------------------------------------------------------------------------
-	mu::Ptr<ASTOp> CodeGenerator::GenerateMissingColourCode(const char* strWhere, const void* errorContext)
+	mu::Ptr<ASTOp> CodeGenerator::GenerateMissingColourCode(const TCHAR* strWhere, const void* errorContext)
 	{
 		// Log a warning
-		char buf[256];
-		mutable_snprintf(buf, 256, "Required connection not found: %s", strWhere);
-		m_pErrorLog->GetPrivate()->Add(buf, ELMT_ERROR, errorContext);
+		FString Msg = FString::Printf(TEXT("Required connection not found: %s"), strWhere);
+		m_pErrorLog->GetPrivate()->Add(Msg, ELMT_ERROR, errorContext);
 
 		// Create a constant colour node
 		NodeColourConstantPtr pNode = new NodeColourConstant();
diff --git a/Engine/Plugins/Experimental/Mutable/Source/MutableTools/Private/MuT/CodeGenerator_Image.cpp b/Engine/Plugins/Experimental/Mutable/Source/MutableTools/Private/MuT/CodeGenerator_Image.cpp
--- a/Engine/Plugins/Experimental/Mutable/Source/MutableTools/Private/MuT/CodeGenerator_Image.cpp
+++ b/Engine/Plugins/Experimental/Mutable/Source/MutableTools/Private/MuT/CodeGenerator_Image.cpp
@@ -114,7 +114,7 @@ namespace mu
 
 
 	//-------------------------------------------------------------------------------------------------
-	void CodeGenerator::GenerateImage(FImageGenerationResult& result, const NodeImagePtrConst& Untyped)
+	void CodeGenerator::GenerateImage(const FImageGenerationOptions& Options, FImageGenerationResult& result, const NodeImagePtrConst& Untyped)
 	{
 		if (!Untyped)
 		{
@@ -140,6 +140,7 @@ namespace mu
 		}
 
 		// See if it was already generated
+		// \TODO: complete and use FGeneratedImageCacheKey instead of generic keys
 		FVisitedKeyMap key = GetCurrentCacheKey(Untyped);
 		GeneratedImagesMap::ValueType* it = m_generatedImages.Find(key);
 		if (it)
@@ -153,33 +154,33 @@ namespace mu
 			// Generate for each different type of node
 			switch (Untyped->GetImageNodeType())
 			{
-			case NodeImage::EType::Constant: GenerateImage_Constant(result, static_cast<const NodeImageConstant*>(Node)); break;
+			case NodeImage::EType::Constant: GenerateImage_Constant(Options, result, static_cast<const NodeImageConstant*>(Node)); break;
 			case NodeImage::EType::Difference_Deprecated: check(false); break;
-			case NodeImage::EType::Interpolate: GenerateImage_Interpolate(result, static_cast<const NodeImageInterpolate*>(Node)); break;
-			case NodeImage::EType::Saturate: GenerateImage_Saturate(result, static_cast<const NodeImageSaturate*>(Node)); break;
-			case NodeImage::EType::Table: GenerateImage_Table(result, static_cast<const NodeImageTable*>(Node)); break;
-			case NodeImage::EType::Swizzle: GenerateImage_Swizzle(result, static_cast<const NodeImageSwizzle*>(Node)); break;
-			case NodeImage::EType::ColourMap: GenerateImage_ColourMap(result, static_cast<const NodeImageColourMap*>(Node)); break;
-			case NodeImage::EType::Gradient: GenerateImage_Gradient(result, static_cast<const NodeImageGradient*>(Node)); break;
-			case NodeImage::EType::Binarise: GenerateImage_Binarise(result, static_cast<const NodeImageBinarise*>(Node)); break;
-			case NodeImage::EType::Luminance: GenerateImage_Luminance(result, static_cast<const NodeImageLuminance*>(Node)); break;
-			case NodeImage::EType::Layer: GenerateImage_Layer(result, static_cast<const NodeImageLayer*>(Node)); break;
-			case NodeImage::EType::LayerColour: GenerateImage_LayerColour(result, static_cast<const NodeImageLayerColour*>(Node)); break;
-			case NodeImage::EType::Resize: GenerateImage_Resize(result, static_cast<const NodeImageResize*>(Node)); break;
-			case NodeImage::EType::PlainColour: GenerateImage_PlainColour(result, static_cast<const NodeImagePlainColour*>(Node)); break;
-			case NodeImage::EType::Interpolate3: GenerateImage_Interpolate3(result, static_cast<const NodeImageInterpolate3*>(Node)); break;
-			case NodeImage::EType::Project: GenerateImage_Project(result, static_cast<const NodeImageProject*>(Node)); break;
-			case NodeImage::EType::Mipmap: GenerateImage_Mipmap(result, static_cast<const NodeImageMipmap*>(Node)); break;
-			case NodeImage::EType::Switch: GenerateImage_Switch(result, static_cast<const NodeImageSwitch*>(Node)); break;
-			case NodeImage::EType::Conditional: GenerateImage_Conditional(result, static_cast<const NodeImageConditional*>(Node)); break;
-			case NodeImage::EType::Format: GenerateImage_Format(result, static_cast<const NodeImageFormat*>(Node)); break;
-			case NodeImage::EType::Parameter: GenerateImage_Parameter(result, static_cast<const NodeImageParameter*>(Node)); break;
-			case NodeImage::EType::MultiLayer: GenerateImage_MultiLayer(result, static_cast<const NodeImageMultiLayer*>(Node)); break;
-			case NodeImage::EType::Invert: GenerateImage_Invert(result, static_cast<const NodeImageInvert*>(Node)); break;
-			case NodeImage::EType::Variation: GenerateImage_Variation(result, static_cast<const NodeImageVariation*>(Node)); break;
-			case NodeImage::EType::NormalComposite: GenerateImage_NormalComposite(result, static_cast<const NodeImageNormalComposite*>(Node)); break;
-			case NodeImage::EType::Transform: GenerateImage_Transform(result, static_cast<const NodeImageTransform*>(Node)); break;
-			case NodeImage::EType::Reference: GenerateImage_Reference(result, static_cast<const NodeImageReference*>(Node)); break;
+			case NodeImage::EType::Interpolate: GenerateImage_Interpolate(Options, result, static_cast<const NodeImageInterpolate*>(Node)); break;
+			case NodeImage::EType::Saturate: GenerateImage_Saturate(Options, result, static_cast<const NodeImageSaturate*>(Node)); break;
+			case NodeImage::EType::Table: GenerateImage_Table(Options, result, static_cast<const NodeImageTable*>(Node)); break;
+			case NodeImage::EType::Swizzle: GenerateImage_Swizzle(Options, result, static_cast<const NodeImageSwizzle*>(Node)); break;
+			case NodeImage::EType::ColourMap: GenerateImage_ColourMap(Options, result, static_cast<const NodeImageColourMap*>(Node)); break;
+			case NodeImage::EType::Gradient: GenerateImage_Gradient(Options, result, static_cast<const NodeImageGradient*>(Node)); break;
+			case NodeImage::EType::Binarise: GenerateImage_Binarise(Options, result, static_cast<const NodeImageBinarise*>(Node)); break;
+			case NodeImage::EType::Luminance: GenerateImage_Luminance(Options, result, static_cast<const NodeImageLuminance*>(Node)); break;
+			case NodeImage::EType::Layer: GenerateImage_Layer(Options, result, static_cast<const NodeImageLayer*>(Node)); break;
+			case NodeImage::EType::LayerColour: GenerateImage_LayerColour(Options, result, static_cast<const NodeImageLayerColour*>(Node)); break;
+			case NodeImage::EType::Resize: GenerateImage_Resize(Options, result, static_cast<const NodeImageResize*>(Node)); break;
+			case NodeImage::EType::PlainColour: GenerateImage_PlainColour(Options, result, static_cast<const NodeImagePlainColour*>(Node)); break;
+			case NodeImage::EType::Interpolate3: GenerateImage_Interpolate3(Options, result, static_cast<const NodeImageInterpolate3*>(Node)); break;
+			case NodeImage::EType::Project: GenerateImage_Project(Options, result, static_cast<const NodeImageProject*>(Node)); break;
+			case NodeImage::EType::Mipmap: GenerateImage_Mipmap(Options, result, static_cast<const NodeImageMipmap*>(Node)); break;
+			case NodeImage::EType::Switch: GenerateImage_Switch(Options, result, static_cast<const NodeImageSwitch*>(Node)); break;
+			case NodeImage::EType::Conditional: GenerateImage_Conditional(Options, result, static_cast<const NodeImageConditional*>(Node)); break;
+			case NodeImage::EType::Format: GenerateImage_Format(Options, result, static_cast<const NodeImageFormat*>(Node)); break;
+			case NodeImage::EType::Parameter: GenerateImage_Parameter(Options, result, static_cast<const NodeImageParameter*>(Node)); break;
+			case NodeImage::EType::MultiLayer: GenerateImage_MultiLayer(Options, result, static_cast<const NodeImageMultiLayer*>(Node)); break;
+			case NodeImage::EType::Invert: GenerateImage_Invert(Options, result, static_cast<const NodeImageInvert*>(Node)); break;
+			case NodeImage::EType::Variation: GenerateImage_Variation(Options, result, static_cast<const NodeImageVariation*>(Node)); break;
+			case NodeImage::EType::NormalComposite: GenerateImage_NormalComposite(Options, result, static_cast<const NodeImageNormalComposite*>(Node)); break;
+			case NodeImage::EType::Transform: GenerateImage_Transform(Options, result, static_cast<const NodeImageTransform*>(Node)); break;
+			case NodeImage::EType::Reference: GenerateImage_Reference(Options, result, static_cast<const NodeImageReference*>(Node)); break;
 			case NodeImage::EType::None: check(false);
 			}
 
@@ -197,7 +198,7 @@ namespace mu
 
 
     //---------------------------------------------------------------------------------------------
-    void CodeGenerator::GenerateImage_Constant(FImageGenerationResult& result, const NodeImageConstant* InNode)
+    void CodeGenerator::GenerateImage_Constant(const FImageGenerationOptions& Options, FImageGenerationResult& result, const NodeImageConstant* InNode)
     {
 		const NodeImageConstant::Private& node = *InNode->GetPrivate();
 		
@@ -243,13 +244,13 @@ namespace mu
         if ( pImage->GetSizeX()!=cropRect.size[0] || pImage->GetSizeY()!=cropRect.size[1] )
         {
             Ptr<Image> pCropped = new Image(cropRect.size[0], cropRect.size[1], 1, pImage->GetFormat());
-			ImageCrop(pCropped.get(), m_compilerOptions->m_imageCompressionQuality, pImage.get(), cropRect );
-            op->SetValue( pCropped, m_compilerOptions->m_optimisationOptions.m_useDiskCache );
+			ImageCrop(pCropped.get(), m_compilerOptions->ImageCompressionQuality, pImage.get(), cropRect );
+            op->SetValue( pCropped, m_compilerOptions->OptimisationOptions.bUseDiskCache );
         }
         else
         {
             check( cropRect.min[0]==0 && cropRect.min[1]==0 );
-            op->SetValue( pImage, m_compilerOptions->m_optimisationOptions.m_useDiskCache );
+            op->SetValue( pImage, m_compilerOptions->OptimisationOptions.bUseDiskCache );
         }
 
 		result.op = op;
@@ -257,7 +258,7 @@ namespace mu
 
 
 	//---------------------------------------------------------------------------------------------
-	void CodeGenerator::GenerateImage_Reference(FImageGenerationResult& result, const NodeImageReference* InNode)
+	void CodeGenerator::GenerateImage_Reference(const FImageGenerationOptions& Options, FImageGenerationResult& result, const NodeImageReference* InNode)
 	{
 		const NodeImageReference::Private& node = *InNode->GetPrivate();
 
@@ -271,7 +272,7 @@ namespace mu
 
 
     //---------------------------------------------------------------------------------------------
-	void CodeGenerator::GenerateImage_Parameter(FImageGenerationResult& result, const NodeImageParameter* InNode )
+	void CodeGenerator::GenerateImage_Parameter(const FImageGenerationOptions& Options, FImageGenerationResult& result, const NodeImageParameter* InNode )
     {
 		const NodeImageParameter::Private& node = *InNode->GetPrivate();
 
@@ -308,7 +309,7 @@ namespace mu
 
 
     //---------------------------------------------------------------------------------------------
-	void CodeGenerator::GenerateImage_Layer(FImageGenerationResult& result, const NodeImageLayer* InNode)
+	void CodeGenerator::GenerateImage_Layer(const FImageGenerationOptions& Options, FImageGenerationResult& result, const NodeImageLayer* InNode)
 	{
 		const NodeImageLayer::Private& node = *InNode->GetPrivate();
 		
@@ -320,15 +321,16 @@ namespace mu
 
         // Base image
         Ptr<ASTOp> base;
-        if ( Node* pBase = node.m_pBase.get() )
+        if ( node.m_pBase )
         {
-            base = Generate( pBase );
+			FImageGenerationResult BaseResult;
+			GenerateImage(Options, BaseResult, node.m_pBase);
+			base = BaseResult.op;
         }
         else
         {
             // This argument is required
-            base = GenerateMissingImageCode( "Image Layer base", EImageFormat::IF_RGB_UBYTE,
-                                             node.m_errorContext );
+            base = GenerateMissingImageCode(TEXT("Image Layer base"), EImageFormat::IF_RGB_UBYTE, node.m_errorContext, Options);
         }
 
 		FImageSize TargetSize((uint16)m_imageState.Last().m_imageRect.size[0],(uint16)m_imageState.Last().m_imageRect.size[1]);
@@ -339,10 +341,13 @@ namespace mu
         op->base = base;
 
         // Mask of the effect
-        Ptr<ASTOp> mask = 0;
-        if ( Node* pMask = node.m_pMask.get() )
+        Ptr<ASTOp> mask;
+        if ( node.m_pMask )
         {
-            mask = Generate( pMask );
+			FImageGenerationResult MaskResult;
+			GenerateImage(Options, MaskResult, node.m_pMask);
+			mask = MaskResult.op;
+
             mask = GenerateImageFormat( mask, EImageFormat::IF_L_UBYTE );
             mask = GenerateImageSize( mask, TargetSize);
         }
@@ -350,14 +355,16 @@ namespace mu
 
         // Image to apply
         Ptr<ASTOp> blended = 0;
-        if ( Node* pBlended = node.m_pBlended.get() )
+        if ( node.m_pBlended )
         {
-            blended = Generate( pBlended );
+			FImageGenerationResult BlendedResult;
+			GenerateImage(Options, BlendedResult, node.m_pBlended);
+			blended = BlendedResult.op;
         }
         else
         {
             // This argument is required
-            blended = GeneratePlainImageCode( vec3<float>( 1,1,0 ) );
+            blended = GeneratePlainImageCode( vec3<float>( 1,1,0 ), Options );
         }
         //blended = GenerateImageFormat( blended, EImageFormat::IF_RGB_UBYTE );
         blended = GenerateImageFormat( blended, baseFormat );
@@ -369,7 +376,7 @@ namespace mu
 
 
     //---------------------------------------------------------------------------------------------
-	void CodeGenerator::GenerateImage_LayerColour(FImageGenerationResult& result, const NodeImageLayerColour* InNode)
+	void CodeGenerator::GenerateImage_LayerColour(const FImageGenerationOptions& Options, FImageGenerationResult& result, const NodeImageLayerColour* InNode)
 	{
 		const NodeImageLayerColour::Private& node = *InNode->GetPrivate();
 		
@@ -380,39 +387,46 @@ namespace mu
 
         // Base image
         Ptr<ASTOp> base;
-        if ( Node* pBase = node.m_pBase.get() )
+        if ( node.m_pBase )
         {
-            base = Generate( pBase );
+			FImageGenerationResult BaseResult;
+			GenerateImage(Options, BaseResult, node.m_pBase);
+			base = BaseResult.op;
         }
         else
         {
             // This argument is required
-            base = GenerateMissingImageCode( "Layer base image", EImageFormat::IF_RGB_UBYTE, node.m_errorContext );
+            base = GenerateMissingImageCode(TEXT("Layer base image"), EImageFormat::IF_RGB_UBYTE, node.m_errorContext, Options );
         }
         base = GenerateImageFormat( base, EImageFormat::IF_RGB_UBYTE );
         base = GenerateImageSize( base, FImageSize((uint16)m_imageState.Last().m_imageRect.size[0],(uint16)m_imageState.Last().m_imageRect.size[1]) );
         op->base = base;
 
         // Mask of the effect
-        Ptr<ASTOp> mask = 0;
-        if ( Node* pMask = node.m_pMask.get() )
+        Ptr<ASTOp> mask;
+        if ( node.m_pMask )
         {
-            mask = Generate( pMask );
-            mask = GenerateImageFormat( mask, EImageFormat::IF_L_UBYTE );
+			FImageGenerationResult MaskResult;
+			GenerateImage(Options, MaskResult, node.m_pMask);
+			mask = MaskResult.op;
+			
+			mask = GenerateImageFormat( mask, EImageFormat::IF_L_UBYTE );
             mask = GenerateImageSize( mask, FImageSize((uint16)m_imageState.Last().m_imageRect.size[0],(uint16)m_imageState.Last().m_imageRect.size[1]) );
         }
         op->mask = mask;
 
         // Colour to apply
         Ptr<ASTOp> colour = 0;
-        if ( Node* pColour = node.m_pColour.get() )
+        if ( node.m_pColour )
         {
-            colour = Generate( pColour );
+			FColorGenerationResult ColorResult;
+            GenerateColor(ColorResult, node.m_pColour);
+			colour = ColorResult.op;
         }
         else
         {
             // This argument is required
-            colour = GenerateMissingColourCode( "Layer colour", node.m_errorContext );
+            colour = GenerateMissingColourCode(TEXT("Layer colour"), node.m_errorContext );
         }
         op->color = colour;
 
@@ -421,7 +435,7 @@ namespace mu
 
 
     //---------------------------------------------------------------------------------------------
-	void CodeGenerator::GenerateImage_MultiLayer(FImageGenerationResult& result, const NodeImageMultiLayer* InNode)
+	void CodeGenerator::GenerateImage_MultiLayer(const FImageGenerationOptions& Options, FImageGenerationResult& result, const NodeImageMultiLayer* InNode)
 	{
 		const NodeImageMultiLayer::Private& node = *InNode->GetPrivate();
 
@@ -433,29 +447,32 @@ namespace mu
 
         // Base image
         Ptr<ASTOp> base;
-        if ( Node* pBase = node.m_pBase.get() )
+        if ( node.m_pBase )
         {
-            base = Generate( pBase );
-        }
+			FImageGenerationResult BaseResult;
+			GenerateImage(Options, BaseResult, node.m_pBase);
+			base = BaseResult.op;
+		}
         else
         {
             // This argument is required
-            base = GenerateMissingImageCode( "Image MultiLayer base", EImageFormat::IF_RGB_UBYTE,
-                                             node.m_errorContext );
+            base = GenerateMissingImageCode(TEXT("Image MultiLayer base"), EImageFormat::IF_RGB_UBYTE, node.m_errorContext, Options);
         }
 
 		EImageFormat baseFormat = base->GetImageDesc().m_format;
-        base = GenerateImageSize
-                ( base, FImageSize((uint16)m_imageState.Last().m_imageRect.size[0],
-                                   (uint16)m_imageState.Last().m_imageRect.size[1]) );
+        base = GenerateImageSize( base, FImageSize((uint16)m_imageState.Last().m_imageRect.size[0], 
+			                                       (uint16)m_imageState.Last().m_imageRect.size[1]) );
         op->base = base;
 
         // Mask of the effect
         Ptr<ASTOp> mask;
-        if ( Node* pMask = node.m_pMask.get() )
+        if ( node.m_pMask )
         {
-            mask = Generate( pMask );
-            mask = GenerateImageFormat( mask, EImageFormat::IF_L_UBYTE );
+			FImageGenerationResult MaskResult;
+			GenerateImage(Options, MaskResult, node.m_pMask);
+			mask = MaskResult.op;
+
+			mask = GenerateImageFormat( mask, EImageFormat::IF_L_UBYTE );
             mask = GenerateImageSize
                     ( mask, FImageSize((uint16)m_imageState.Last().m_imageRect.size[0],
                                        (uint16)m_imageState.Last().m_imageRect.size[1]) );
@@ -464,16 +481,17 @@ namespace mu
 
         // Image to apply
         Ptr<ASTOp> blended;
-        if ( Node* pBlended = node.m_pBlended.get() )
+        if (node.m_pBlended)
         {
-            blended = Generate( pBlended );
+			FImageGenerationResult MaskResult;
+			GenerateImage(Options, MaskResult, node.m_pBlended);
+			blended = MaskResult.op;
         }
         else
         {
             // This argument is required
-            blended = GeneratePlainImageCode( vec3<float>( 1,1,0 ) );
+            blended = GeneratePlainImageCode( vec3<float>( 1,1,0 ), Options);
         }
-        //blended = GenerateImageFormat( blended, EImageFormat::IF_RGB_UBYTE );
         blended = GenerateImageFormat( blended, baseFormat );
         blended = GenerateImageSize
                 ( blended, FImageSize((uint16)m_imageState.Last().m_imageRect.size[0],
@@ -483,7 +501,6 @@ namespace mu
         // Range of iteration
         if ( node.m_pRange )
         {
-            //op->range = Generate( pRange );
             FRangeGenerationResult rangeResult;
             GenerateRange( rangeResult, node.m_pRange );
 
@@ -497,7 +514,7 @@ namespace mu
 
 
 	//---------------------------------------------------------------------------------------------
-	void CodeGenerator::GenerateImage_NormalComposite(FImageGenerationResult& result, const NodeImageNormalComposite* InNode)
+	void CodeGenerator::GenerateImage_NormalComposite(const FImageGenerationOptions& Options, FImageGenerationResult& result, const NodeImageNormalComposite* InNode)
 	{
 		const NodeImageNormalComposite::Private& node = *InNode->GetPrivate();
 
@@ -510,15 +527,16 @@ namespace mu
 
         // Base image
         Ptr<ASTOp> base;
-        if ( Node* pBase = node.m_pBase.get() )
+        if ( node.m_pBase )
         {
-            base = Generate( pBase );
-        }
+			FImageGenerationResult BaseResult;
+			GenerateImage(Options, BaseResult, node.m_pBase);
+			base = BaseResult.op;
+		}
         else
         {
             // This argument is required
-            base = GenerateMissingImageCode( "Image Composite Base", EImageFormat::IF_RGB_UBYTE,
-                                             node.m_errorContext );
+            base = GenerateMissingImageCode(TEXT("Image Composite Base"), EImageFormat::IF_RGB_UBYTE, node.m_errorContext, Options);
         }
 
 		EImageFormat baseFormat = base->GetImageDesc().m_format;
@@ -528,9 +546,12 @@ namespace mu
         op->Base = base;
 
         Ptr<ASTOp> normal;
-        if ( Node* pNormal = node.m_pNormal.get() )
+        if ( node.m_pNormal )
         {
-            normal = Generate( pNormal );
+			FImageGenerationResult BaseResult;
+			GenerateImage(Options, BaseResult, node.m_pNormal);
+			normal = BaseResult.op;
+
             normal = GenerateImageFormat( normal, EImageFormat::IF_RGB_UBYTE );
             normal = GenerateImageSize
                     ( normal, FImageSize((uint16)m_imageState.Last().m_imageRect.size[0],
@@ -539,8 +560,7 @@ namespace mu
 		else
 		{
             // This argument is required
-            normal = GenerateMissingImageCode( "Image Composite Normal", EImageFormat::IF_RGB_UBYTE,
-                                               node.m_errorContext );
+            normal = GenerateMissingImageCode(TEXT("Image Composite Normal"), EImageFormat::IF_RGB_UBYTE, node.m_errorContext, Options);
 		}
 
         op->Normal = normal;
@@ -548,7 +568,7 @@ namespace mu
 		result.op = op;
     }
 
-	void CodeGenerator::GenerateImage_Transform(FImageGenerationResult& result, const NodeImageTransform* InNode)
+	void CodeGenerator::GenerateImage_Transform(const FImageGenerationOptions& Options, FImageGenerationResult& result, const NodeImageTransform* InNode)
     {
 		const NodeImageTransform::Private& node = *InNode->GetPrivate();
 
@@ -595,15 +615,16 @@ namespace mu
 
         // Base image
         Ptr<ASTOp> base;
-        if ( Node* pBase = node.m_pBase.get() )
+        if ( node.m_pBase )
         {
-            base = Generate( pBase );
-        }
+			FImageGenerationResult BaseResult;
+			GenerateImage(Options, BaseResult, node.m_pBase);
+			base = BaseResult.op;
+		}
         else
         {
             // This argument is required
-            base = GenerateMissingImageCode( "Image Transform Base", EImageFormat::IF_RGB_UBYTE,
-                                             node.m_errorContext );
+            base = GenerateMissingImageCode(TEXT("Image Transform Base"), EImageFormat::IF_RGB_UBYTE, node.m_errorContext, Options);
         }
 
 		EImageFormat baseFormat = base->GetImageDesc().m_format;
@@ -616,7 +637,7 @@ namespace mu
     }
 
     //---------------------------------------------------------------------------------------------
-	void CodeGenerator::GenerateImage_Interpolate(FImageGenerationResult& result, const NodeImageInterpolate* InNode)
+	void CodeGenerator::GenerateImage_Interpolate(const FImageGenerationOptions& Options, FImageGenerationResult& result, const NodeImageInterpolate* InNode)
 	{
 		const NodeImageInterpolate::Private& node = *InNode->GetPrivate();
 
@@ -634,7 +655,7 @@ namespace mu
         {
             // This argument is required
             op->SetChild( op->op.args.ImageInterpolate.factor,
-                          GenerateMissingScalarCode( "Interpolation factor", 0.5f, node.m_errorContext ));
+                          GenerateMissingScalarCode(TEXT("Interpolation factor"), 0.5f, node.m_errorContext ));
         }
 
         // Target images
@@ -644,9 +665,11 @@ namespace mu
             ; t< node.m_targets.Num() && numTargets<MUTABLE_OP_MAX_INTERPOLATE_COUNT
             ; ++t )
         {
-            if ( Node* pA = node.m_targets[t].get() )
+            if ( node.m_targets[t] )
             {
-                Ptr<ASTOp> target = Generate( pA );
+				FImageGenerationResult BaseResult;
+				GenerateImage(Options, BaseResult, node.m_targets[t]);
+				Ptr<ASTOp> target = BaseResult.op;
 
                 // TODO: Support other formats
                 target = GenerateImageFormat( target, EImageFormat::IF_RGB_UBYTE );
@@ -661,10 +684,8 @@ namespace mu
         // At least one target is required
         if (!op->op.args.ImageInterpolate.targets[0])
         {
-            Ptr<ASTOp> target = GenerateMissingImageCode( "First interpolation image", EImageFormat::IF_RGB_UBYTE,
-                                                       node.m_errorContext );
-            target = GenerateImageSize
-                    ( target, FImageSize((uint16)m_imageState.Last().m_imageRect.size[0],(uint16)m_imageState.Last().m_imageRect.size[1]) );
+            Ptr<ASTOp> target = GenerateMissingImageCode(TEXT("First interpolation image"), EImageFormat::IF_RGB_UBYTE, node.m_errorContext, Options);
+            target = GenerateImageSize( target, FImageSize((uint16)m_imageState.Last().m_imageRect.size[0],(uint16)m_imageState.Last().m_imageRect.size[1]) );
             op->SetChild( op->op.args.ImageInterpolate.targets[0], target);
         }
 
@@ -673,7 +694,7 @@ namespace mu
 
 
     //---------------------------------------------------------------------------------------------
-	void CodeGenerator::GenerateImage_Interpolate3(FImageGenerationResult& result, const NodeImageInterpolate3* InNode)
+	void CodeGenerator::GenerateImage_Interpolate3(const FImageGenerationOptions& Options, FImageGenerationResult& result, const NodeImageInterpolate3* InNode)
 	{
 		const NodeImageInterpolate3::Private& node = *InNode->GetPrivate();
 
@@ -691,7 +712,7 @@ namespace mu
         {
             // This argument is required
             op->SetChild( op->op.args.ImageInterpolate3.factor1,
-                          GenerateMissingScalarCode( "Interpolation factor 1", 0.3f, node.m_errorContext ));
+                          GenerateMissingScalarCode(TEXT("Interpolation factor 1"), 0.3f, node.m_errorContext ));
         }
 
         if ( node.m_pFactor2.get() )
@@ -702,20 +723,21 @@ namespace mu
         {
             // This argument is required
             op->SetChild( op->op.args.ImageInterpolate3.factor2,
-                          GenerateMissingScalarCode( "Interpolation factor 2", 0.3f, node.m_errorContext ));
+                          GenerateMissingScalarCode(TEXT("Interpolation factor 2"), 0.3f, node.m_errorContext ));
         }
 
         // Target 0
         Ptr<ASTOp> target;
-        if ( node.m_pTarget0.get() )
+        if ( node.m_pTarget0 )
         {
-            target = Generate( node.m_pTarget0 );
+			FImageGenerationResult BaseResult;
+			GenerateImage(Options, BaseResult, node.m_pTarget0);
+			target = BaseResult.op;
         }
         else
         {
             // This argument is required
-            target = GenerateMissingImageCode( "Interpolation target 0", EImageFormat::IF_RGB_UBYTE,
-                                               node.m_errorContext );
+            target = GenerateMissingImageCode(TEXT("Interpolation target 0"), EImageFormat::IF_RGB_UBYTE, node.m_errorContext, Options);
         }
         target = GenerateImageFormat( target, EImageFormat::IF_RGB_UBYTE );
         target = GenerateImageSize
@@ -723,15 +745,16 @@ namespace mu
         op->SetChild( op->op.args.ImageInterpolate3.target0, target);
 
         // Target 1
-        if ( node.m_pTarget1.get() )
+        if ( node.m_pTarget1 )
         {
-            target = Generate( node.m_pTarget1 );
-        }
+			FImageGenerationResult BaseResult;
+			GenerateImage(Options, BaseResult, node.m_pTarget1);
+			target = BaseResult.op;
+		}
         else
         {
             // This argument is required
-            target = GenerateMissingImageCode( "Interpolation target 1", EImageFormat::IF_RGB_UBYTE,
-                                               node.m_errorContext );
+            target = GenerateMissingImageCode(TEXT("Interpolation target 1"), EImageFormat::IF_RGB_UBYTE, node.m_errorContext, Options);
         }
         target = GenerateImageFormat( target, EImageFormat::IF_RGB_UBYTE );
         target = GenerateImageSize
@@ -740,15 +763,16 @@ namespace mu
         op->SetChild( op->op.args.ImageInterpolate3.target1, target);
 
         // Target 2
-        if ( node.m_pTarget2.get() )
+        if ( node.m_pTarget2 )
         {
-            target = Generate( node.m_pTarget2 );
-        }
+			FImageGenerationResult BaseResult;
+			GenerateImage(Options, BaseResult, node.m_pTarget2);
+			target = BaseResult.op;
+		}
         else
         {
             // This argument is required
-            target = GenerateMissingImageCode( "Interpolation target 2", EImageFormat::IF_RGB_UBYTE,
-                                               node.m_errorContext );
+            target = GenerateMissingImageCode(TEXT("Interpolation target 2"), EImageFormat::IF_RGB_UBYTE, node.m_errorContext, Options);
         }
         target = GenerateImageFormat( target, EImageFormat::IF_RGB_UBYTE );
         target = GenerateImageSize
@@ -762,7 +786,7 @@ namespace mu
 
 
 	//---------------------------------------------------------------------------------------------
-	void CodeGenerator::GenerateImage_Swizzle(FImageGenerationResult& result, const NodeImageSwizzle* InNode)
+	void CodeGenerator::GenerateImage_Swizzle(const FImageGenerationOptions& Options, FImageGenerationResult& result, const NodeImageSwizzle* InNode)
 	{
 		const NodeImageSwizzle::Private& node = *InNode->GetPrivate();
 
@@ -819,15 +843,18 @@ namespace mu
         Ptr<ASTOp> first;
 		for (int32 t = 0; t<node.m_sources.Num(); ++t)
 		{
-			if (Node* pChannel = node.m_sources[t].get())
+			if (node.m_sources[t])
 			{
-                Ptr<ASTOp> source = Generate(pChannel);
+				FImageGenerationResult BaseResult;
+				GenerateImage(Options, BaseResult, node.m_sources[t]);
+                Ptr<ASTOp> source = BaseResult.op;
+
 				source = GenerateImageUncompressed(source);
 
 				if (!source)
 				{
 					// TODO: Warn?
-					source = GenerateMissingImageCode("Swizzle channel", EImageFormat::IF_L_UBYTE, node.m_errorContext);
+					source = GenerateMissingImageCode(TEXT("Swizzle channel"), EImageFormat::IF_L_UBYTE, node.m_errorContext, Options);
 				}
 
                 Ptr<ASTOp> sizedSource;
@@ -853,7 +880,7 @@ namespace mu
 		// At least one source is required
         if (!op->Sources[0])
 		{
-            Ptr<ASTOp> source = GenerateMissingImageCode("First swizzle image", EImageFormat::IF_RGBA_UBYTE, node.m_errorContext);
+            Ptr<ASTOp> source = GenerateMissingImageCode(TEXT("First swizzle image"), EImageFormat::IF_RGBA_UBYTE, node.m_errorContext, Options);
             op->Sources[0] = source;
 		}
 
@@ -872,7 +899,7 @@ namespace mu
 
 
 	//---------------------------------------------------------------------------------------------
-	void CodeGenerator::GenerateImage_Format(FImageGenerationResult& result, const NodeImageFormat* InNode)
+	void CodeGenerator::GenerateImage_Format(const FImageGenerationOptions& Options, FImageGenerationResult& result, const NodeImageFormat* InNode)
 	{
 		const NodeImageFormat::Private& node = *InNode->GetPrivate();
 
@@ -887,11 +914,13 @@ namespace mu
 		// Source is required
 		if (!node.m_source)
 		{
-            fop->Source = GenerateMissingImageCode("Source image for format.", EImageFormat::IF_RGBA_UBYTE, node.m_errorContext);
+            fop->Source = GenerateMissingImageCode(TEXT("Source image for format."), EImageFormat::IF_RGBA_UBYTE, node.m_errorContext, Options);
 		}
 		else
 		{
-            fop->Source = Generate(node.m_source.get());
+			FImageGenerationResult BaseResult;
+			GenerateImage(Options, BaseResult, node.m_source);
+			fop->Source = BaseResult.op;
 		}
 
         result.op = fop;
@@ -899,7 +928,7 @@ namespace mu
 
 
     //---------------------------------------------------------------------------------------------
-	void CodeGenerator::GenerateImage_Saturate(FImageGenerationResult& result, const NodeImageSaturate* InNode)
+	void CodeGenerator::GenerateImage_Saturate(const FImageGenerationOptions& Options, FImageGenerationResult& result, const NodeImageSaturate* InNode)
 	{
 		const NodeImageSaturate::Private& node = *InNode->GetPrivate();
 
@@ -909,14 +938,16 @@ namespace mu
 
         // Source image
         Ptr<ASTOp> base = 0;
-        if ( Node* pSource = node.m_pSource.get() )
+        if ( node.m_pSource )
         {
-            base = Generate( pSource );
-        }
+			FImageGenerationResult BaseResult;
+			GenerateImage(Options, BaseResult, node.m_pSource);
+			base = BaseResult.op;
+		}
         else
         {
             // This argument is required
-            base = GenerateMissingImageCode( "Saturate image", EImageFormat::IF_RGB_UBYTE, node.m_errorContext );
+            base = GenerateMissingImageCode(TEXT("Saturate image"), EImageFormat::IF_RGB_UBYTE, node.m_errorContext, Options);
         }
 		
         base = GenerateImageFormat(base, GetRGBOrRGBAFormat(base->GetImageDesc().m_format));
@@ -933,8 +964,7 @@ namespace mu
         else
         {
             // This argument is required
-            op->SetChild( op->op.args.ImageSaturate.factor, GenerateMissingScalarCode( "Saturation factor", 0.5f,
-                                                                      node.m_errorContext ) );
+            op->SetChild( op->op.args.ImageSaturate.factor, GenerateMissingScalarCode(TEXT("Saturation factor"), 0.5f, node.m_errorContext ) );
         }
 
         result.op = op;
@@ -942,7 +972,7 @@ namespace mu
 
 
     //---------------------------------------------------------------------------------------------
-	void CodeGenerator::GenerateImage_Luminance(FImageGenerationResult& result, const NodeImageLuminance* InNode)
+	void CodeGenerator::GenerateImage_Luminance(const FImageGenerationOptions& Options, FImageGenerationResult& result, const NodeImageLuminance* InNode)
 	{
 		const NodeImageLuminance::Private& node = *InNode->GetPrivate();
 
@@ -950,16 +980,17 @@ namespace mu
         op->op.type = OP_TYPE::IM_LUMINANCE;
 
         // Source image
-        Ptr<ASTOp> base = 0;
-        if ( Node* pSource = node.m_pSource.get() )
+        Ptr<ASTOp> base;
+        if ( node.m_pSource )
         {
-            base = Generate( pSource );
+			FImageGenerationResult BaseResult;
+			GenerateImage(Options, BaseResult, node.m_pSource);
+			base = BaseResult.op;
         }
         else
         {
             // This argument is required
-            base = GenerateMissingImageCode( "Image luminance", EImageFormat::IF_RGB_UBYTE,
-                                             node.m_errorContext );
+            base = GenerateMissingImageCode(TEXT("Image luminance"), EImageFormat::IF_RGB_UBYTE, node.m_errorContext, Options);
         }
         base = GenerateImageFormat( base, EImageFormat::IF_RGB_UBYTE );
         op->SetChild( op->op.args.ImageLuminance.base, base);
@@ -969,7 +1000,7 @@ namespace mu
 
 
     //---------------------------------------------------------------------------------------------
-	void CodeGenerator::GenerateImage_ColourMap(FImageGenerationResult& result, const NodeImageColourMap* InNode)
+	void CodeGenerator::GenerateImage_ColourMap(const FImageGenerationOptions& Options, FImageGenerationResult& result, const NodeImageColourMap* InNode)
 	{
 		const NodeImageColourMap::Private& node = *InNode->GetPrivate();
 
@@ -977,16 +1008,17 @@ namespace mu
         op->op.type = OP_TYPE::IM_COLOURMAP;
 
         // Base image
-        Ptr<ASTOp> base = 0;
-        if ( Node* pBase = node.m_pBase.get() )
+        Ptr<ASTOp> base ;
+        if ( node.m_pBase )
         {
-            base = Generate( pBase );
-        }
+			FImageGenerationResult BaseResult;
+			GenerateImage(Options, BaseResult, node.m_pBase);
+			base = BaseResult.op;
+		}
         else
         {
             // This argument is required
-            base = GenerateMissingImageCode( "Colourmap base image", EImageFormat::IF_RGB_UBYTE,
-                                             node.m_errorContext );
+            base = GenerateMissingImageCode(TEXT("Colourmap base image"), EImageFormat::IF_RGB_UBYTE, node.m_errorContext, Options);
         }
         base = GenerateImageSize
                 ( base, FImageSize( (uint16)m_imageState.Last().m_imageRect.size[0],
@@ -994,16 +1026,18 @@ namespace mu
         op->SetChild( op->op.args.ImageColourMap.base, base);
 
         // Mask of the effect
-        Ptr<ASTOp> mask = 0;
-        if ( Node* pMask = node.m_pMask.get() )
+        Ptr<ASTOp> mask;
+        if ( node.m_pMask )
         {
-            mask = Generate( pMask );
+			FImageGenerationResult BaseResult;
+			GenerateImage(Options, BaseResult, node.m_pMask);
+			mask = BaseResult.op;
         }
         else
         {
             // Set the argument default value: affect all pixels.
             // TODO: Special operation code without mask
-            mask = GeneratePlainImageCode( vec3<float>( 1,1,1 ) );
+            mask = GeneratePlainImageCode( vec3<float>( 1,1,1 ), Options);
         }
         mask = GenerateImageFormat( mask, EImageFormat::IF_L_UBYTE );
         mask = GenerateImageSize
@@ -1013,15 +1047,17 @@ namespace mu
 
         // Map image
 		// \TODO: We are forcing an map image size as it was a normal texture, and this is incorrect.
-        Ptr<ASTOp> mapImage = 0;
-        if ( Node* pMap = node.m_pMap.get() )
+        Ptr<ASTOp> mapImage;
+        if ( node.m_pMap )
         {
-            mapImage = Generate( pMap );
+			FImageGenerationResult BaseResult;
+			GenerateImage(Options, BaseResult, node.m_pMap);
+			mapImage = BaseResult.op;
         }
         else
         {
             // This argument is required
-            mapImage = GenerateMissingImageCode( "Map image", EImageFormat::IF_RGB_UBYTE, node.m_errorContext );
+            mapImage = GenerateMissingImageCode(TEXT("Map image"), EImageFormat::IF_RGB_UBYTE, node.m_errorContext, Options);
         }
         mapImage = GenerateImageFormat( mapImage, EImageFormat::IF_RGB_UBYTE );
         mapImage = GenerateImageSize
@@ -1034,7 +1070,7 @@ namespace mu
 
 
     //---------------------------------------------------------------------------------------------
-	void CodeGenerator::GenerateImage_Gradient(FImageGenerationResult& result, const NodeImageGradient* InNode)
+	void CodeGenerator::GenerateImage_Gradient(const FImageGenerationOptions& Options, FImageGenerationResult& result, const NodeImageGradient* InNode)
 	{
 		const NodeImageGradient::Private& node = *InNode->GetPrivate();
 
@@ -1050,7 +1086,7 @@ namespace mu
         else
         {
             // This argument is required
-            colour0 = GenerateMissingColourCode( "Gradient colour 0", node.m_errorContext );
+            colour0 = GenerateMissingColourCode(TEXT("Gradient colour 0"), node.m_errorContext );
         }
         op->SetChild( op->op.args.ImageGradient.colour0, colour0);
 
@@ -1063,7 +1099,7 @@ namespace mu
         else
         {
             // This argument is required
-            colour1 = GenerateMissingColourCode( "Gradient colour 1", node.m_errorContext );
+            colour1 = GenerateMissingColourCode(TEXT("Gradient colour 1"), node.m_errorContext );
         }
         op->SetChild( op->op.args.ImageGradient.colour1, colour1);
 
@@ -1075,7 +1111,7 @@ namespace mu
 
 
     //---------------------------------------------------------------------------------------------
-	void CodeGenerator::GenerateImage_Binarise(FImageGenerationResult& result, const NodeImageBinarise* InNode)
+	void CodeGenerator::GenerateImage_Binarise(const FImageGenerationOptions& Options, FImageGenerationResult& result, const NodeImageBinarise* InNode)
 	{
 		const NodeImageBinarise::Private& node = *InNode->GetPrivate();
 
@@ -1084,15 +1120,16 @@ namespace mu
 
         // A image
         Ptr<ASTOp> a;
-        if ( Node* pBase = node.m_pBase.get() )
+        if ( node.m_pBase )
         {
-            a = Generate( pBase );
-        }
+			FImageGenerationResult BaseResult;
+			GenerateImage(Options, BaseResult, node.m_pBase);
+			a = BaseResult.op;
+		}
         else
         {
             // This argument is required
-            a = GenerateMissingImageCode( "Image Binarise Base", EImageFormat::IF_RGB_UBYTE,
-                                          node.m_errorContext );
+            a = GenerateMissingImageCode(TEXT("Image Binarise Base"), EImageFormat::IF_RGB_UBYTE, node.m_errorContext, Options);
         }
         a = GenerateImageFormat( a, EImageFormat::IF_RGB_UBYTE );
         a = GenerateImageSize
@@ -1109,7 +1146,7 @@ namespace mu
         else
         {
             // This argument is required
-            b = GenerateMissingScalarCode( "Image Binarise Threshold", 0.5f, node.m_errorContext );
+            b = GenerateMissingScalarCode(TEXT("Image Binarise Threshold"), 0.5f, node.m_errorContext );
         }
         op->SetChild( op->op.args.ImageBinarise.threshold, b );
 
@@ -1118,7 +1155,7 @@ namespace mu
 
 
     //---------------------------------------------------------------------------------------------
-	void CodeGenerator::GenerateImage_Resize(FImageGenerationResult& result, const NodeImageResize* InNode)
+	void CodeGenerator::GenerateImage_Resize(const FImageGenerationOptions& Options, FImageGenerationResult& result, const NodeImageResize* InNode)
 	{
 		const NodeImageResize::Private& node = *InNode->GetPrivate();
 
@@ -1127,8 +1164,8 @@ namespace mu
         Ptr<ASTOp> at = 0;
 
         // Source image
-        Ptr<ASTOp> base = 0;
-        if ( Node* pSource = node.m_pBase.get() )
+        Ptr<ASTOp> base;
+        if ( node.m_pBase )
         {
             IMAGE_STATE newState = m_imageState.Last();
             if ( node.m_relative )
@@ -1143,14 +1180,16 @@ namespace mu
 
             m_imageState.Add(newState);
 
-            base = Generate( pSource );
+			FImageGenerationResult BaseResult;
+			GenerateImage(Options, BaseResult, node.m_pBase);
+			base = BaseResult.op;
 
             m_imageState.Pop();
         }
         else
         {
             // This argument is required
-            base = GenerateMissingImageCode( "Image resize base", EImageFormat::IF_RGB_UBYTE, node.m_errorContext );
+            base = GenerateMissingImageCode(TEXT("Image resize base"), EImageFormat::IF_RGB_UBYTE, node.m_errorContext, Options);
         }
 
         // Size
@@ -1178,7 +1217,7 @@ namespace mu
 
 
     //---------------------------------------------------------------------------------------------
-	void CodeGenerator::GenerateImage_PlainColour(FImageGenerationResult& result, const NodeImagePlainColour* InNode)
+	void CodeGenerator::GenerateImage_PlainColour(const FImageGenerationOptions& Options, FImageGenerationResult& result, const NodeImagePlainColour* InNode)
 	{
 		const NodeImagePlainColour::Private& node = *InNode->GetPrivate();
 
@@ -1191,7 +1230,7 @@ namespace mu
         else
         {
             // This argument is required
-            base = GenerateMissingColourCode( "Image plain colour base", node.m_errorContext );
+            base = GenerateMissingColourCode(TEXT("Image plain colour base"), node.m_errorContext );
         }
 
         Ptr<ASTOpFixed> op = new ASTOpFixed();
@@ -1221,7 +1260,7 @@ namespace mu
 
 
     //---------------------------------------------------------------------------------------------
-	void CodeGenerator::GenerateImage_Switch(FImageGenerationResult& result, const NodeImageSwitch* InNode)
+	void CodeGenerator::GenerateImage_Switch(const FImageGenerationOptions& Options, FImageGenerationResult& result, const NodeImageSwitch* InNode)
 	{
 		const NodeImageSwitch::Private& node = *InNode->GetPrivate();
 
@@ -1230,9 +1269,7 @@ namespace mu
         if (node.m_options.Num() == 0)
 		{
 			// No options in the switch!
-            Ptr<ASTOp> missingOp = GenerateMissingImageCode("Switch option",
-				EImageFormat::IF_RGBA_UBYTE,
-				node.m_errorContext);
+            Ptr<ASTOp> missingOp = GenerateMissingImageCode(TEXT("Switch option"), EImageFormat::IF_RGBA_UBYTE, node.m_errorContext, Options);
 			result.op = missingOp;
 			return;
 		}
@@ -1248,7 +1285,7 @@ namespace mu
 		else
 		{
 			// This argument is required
-            op->variable = GenerateMissingScalarCode( "Switch variable", 0.0f, node.m_errorContext );
+            op->variable = GenerateMissingScalarCode(TEXT("Switch variable"), 0.0f, node.m_errorContext );
 		}
 
 		// Options
@@ -1258,14 +1295,14 @@ namespace mu
 
             if (node.m_options[t])
             {
-                branch = Generate( node.m_options[t].get() );
-            }
+				FImageGenerationResult BaseResult;
+				GenerateImage(Options, BaseResult, node.m_options[t]);
+				branch = BaseResult.op;
+			}
             else
             {
                 // This argument is required
-                branch = GenerateMissingImageCode("Switch option",
-                                                    EImageFormat::IF_RGBA_UBYTE,
-                                                    node.m_errorContext);
+                branch = GenerateMissingImageCode(TEXT("Switch option"), EImageFormat::IF_RGBA_UBYTE, node.m_errorContext, Options);
             }
 
             op->cases.Emplace((int16_t)t,op,branch);
@@ -1307,7 +1344,7 @@ namespace mu
 
 
     //---------------------------------------------------------------------------------------------
-	void CodeGenerator::GenerateImage_Conditional(FImageGenerationResult& result, const NodeImageConditional* InNode)
+	void CodeGenerator::GenerateImage_Conditional(const FImageGenerationOptions& Options, FImageGenerationResult& result, const NodeImageConditional* InNode)
 	{
 		const NodeImageConditional::Private& node = *InNode->GetPrivate();
 
@@ -1322,19 +1359,24 @@ namespace mu
         else
         {
             // This argument is required
-            op->condition = GenerateMissingBoolCode( "Conditional condition", true, node.m_errorContext );
+            op->condition = GenerateMissingBoolCode(TEXT("Conditional condition"), true, node.m_errorContext );
         }
 
         // Options
-        op->yes = Generate( node.m_true.get() );
-        op->no = Generate( node.m_false.get() );
+		FImageGenerationResult YesResult;
+		GenerateImage(Options, YesResult, node.m_true);
+		op->yes = YesResult.op;
 
+		FImageGenerationResult NoResult;
+		GenerateImage(Options, NoResult, node.m_false);
+		op->no = NoResult.op;
+		
         result.op = op;
     }
 
 
     //---------------------------------------------------------------------------------------------
-	void CodeGenerator::GenerateImage_Project(FImageGenerationResult& result, const NodeImageProject* InNode)
+	void CodeGenerator::GenerateImage_Project(const FImageGenerationOptions& Options, FImageGenerationResult& result, const NodeImageProject* InNode)
 	{
 		const NodeImageProject::Private& node = *InNode->GetPrivate();
 
@@ -1412,11 +1454,10 @@ namespace mu
                 // will remove the faces that are not in the layout block we are generating.
                 Ptr<ASTOpConstantResource> cop = new ASTOpConstantResource();
                 cop->type = OP_TYPE::ME_CONSTANT;
-
 				Ptr<Mesh> FormatMeshResult = new Mesh();
 				CreateMeshOptimisedForWrappingProjection(FormatMeshResult.get(), node.m_layout);
 
-                cop->SetValue(FormatMeshResult, m_compilerOptions->m_optimisationOptions.m_useDiskCache);
+                cop->SetValue(FormatMeshResult, m_compilerOptions->OptimisationOptions.bUseDiskCache);
 
                 Ptr<ASTOpMeshFormat> fop = new ASTOpMeshFormat();
                 fop->Buffers = OP::MeshFormatArgs::BT_VERTEX
@@ -1448,7 +1489,7 @@ namespace mu
 				Ptr<Mesh> FormatMeshResult = new Mesh();
                 CreateMeshOptimisedForProjection(FormatMeshResult.get(), node.m_layout);
 
-                cop->SetValue(FormatMeshResult, m_compilerOptions->m_optimisationOptions.m_useDiskCache);
+                cop->SetValue(FormatMeshResult, m_compilerOptions->OptimisationOptions.bUseDiskCache);
 
                 Ptr<ASTOpMeshFormat> fop = new ASTOpMeshFormat();
                 fop->Buffers = OP::MeshFormatArgs::BT_VERTEX
@@ -1466,7 +1507,7 @@ namespace mu
             Ptr<const Mesh> pMesh = new Mesh();
             Ptr<ASTOpConstantResource> cop = new ASTOpConstantResource();
             cop->type = OP_TYPE::ME_CONSTANT;
-            cop->SetValue( pMesh, m_compilerOptions->m_optimisationOptions.m_useDiskCache );
+            cop->SetValue( pMesh, m_compilerOptions->OptimisationOptions.bUseDiskCache );
             pop->SetChild( pop->op.args.MeshProject.mesh, cop );
             m_pErrorLog->GetPrivate()->Add( "Projector mesh not set.", ELMT_ERROR, node.m_errorContext );
         }
@@ -1493,10 +1534,15 @@ namespace mu
             m_imageState.Add( newState );
 
             // Generate
-            op->image = Generate( node.m_pImage.get() );
+			// \TODO: Build new Options with the above data
+			FImageGenerationResult ImageResult;
+			GenerateImage(Options, ImageResult, node.m_pImage);
+			op->image = ImageResult.op;
+
 			op->SourceSizeX = desc.m_size[0];
 			op->SourceSizeY = desc.m_size[1];
 
+
             // Restore rect
             m_imageState.Pop();
 
@@ -1504,7 +1550,7 @@ namespace mu
         else
         {
             // This argument is required
-            op->image = GenerateMissingImageCode( "Projector image", EImageFormat::IF_RGB_UBYTE, node.m_errorContext );
+            op->image = GenerateMissingImageCode(TEXT("Projector image"), EImageFormat::IF_RGB_UBYTE, node.m_errorContext, Options);
         }
 
         // Image size, from the current block being generated
@@ -1538,7 +1584,10 @@ namespace mu
         // Target mask
         if ( node.m_pMask )
         {
-            Ptr<ASTOp> mask = Generate( node.m_pMask.get() );
+			FImageGenerationResult MaskResult;
+			GenerateImage(Options, MaskResult, node.m_pMask);
+			Ptr<ASTOp> mask = MaskResult.op;
+
             mask = GenerateImageFormat( mask, EImageFormat::IF_L_UBYTE );
             op->mask = GenerateImageSize( mask, FImageSize( m_imageState.Last().m_imageRect.size ) );
         }
@@ -1564,7 +1613,7 @@ namespace mu
 		MakeGrowMapOp->Border = MUTABLE_GROW_BORDER_VALUE;
 		
 		// If we want to be able to generate progressive mips efficiently, we need mipmaps for the "displacement map".
-		if (m_compilerOptions->m_optimisationOptions.bEnableProgressiveImages)
+		if (m_compilerOptions->OptimisationOptions.bEnableProgressiveImages)
 		{
 			Ptr<ASTOpImageMipmap> MipMask = new ASTOpImageMipmap;
 			MipMask->Source = MakeGrowMapOp->Mask.child();
@@ -1582,7 +1631,7 @@ namespace mu
 
 
     //---------------------------------------------------------------------------------------------
-	void CodeGenerator::GenerateImage_Mipmap(FImageGenerationResult& result, const NodeImageMipmap* InNode)
+	void CodeGenerator::GenerateImage_Mipmap(const FImageGenerationOptions& Options, FImageGenerationResult& result, const NodeImageMipmap* InNode)
 	{
 		const NodeImageMipmap::Private& node = *InNode->GetPrivate();
 
@@ -1592,21 +1641,22 @@ namespace mu
 
         Ptr<ASTOpImageMipmap> op = new ASTOpImageMipmap();
 
-        // At the end of the day, we want all the mipmaps. Maybe the code optimiser will split the
-        // process later.
+        // At the end of the day, we want all the mipmaps. Maybe the code optimiser will split the process later.
         op->Levels = 0;
 
         // Source image
         Ptr<ASTOp> base;
-        if ( Node* pSource = node.m_pSource.get() )
+        if ( node.m_pSource )
         {
             MUTABLE_CPUPROFILER_SCOPE(Base);
-            base = Generate( pSource );
-        }
+			FImageGenerationResult BaseResult;
+			GenerateImage(Options, BaseResult, node.m_pSource);
+			base = BaseResult.op;
+		}
         else
         {
             // This argument is required
-            base = GenerateMissingImageCode( "Mipmap image", EImageFormat::IF_RGB_UBYTE, node.m_errorContext );
+            base = GenerateMissingImageCode(TEXT("Mipmap image"), EImageFormat::IF_RGB_UBYTE, node.m_errorContext, Options);
         }
 
         op->Source = base;
@@ -1615,7 +1665,7 @@ namespace mu
         // optimisation operations. Scan the source image code looking for this info
         int32 blockX = 0;
         int32 blockY = 0;
-        if ( m_compilerOptions->m_textureLayoutStrategy
+        if ( Options.ImageLayoutStrategy
              !=
              CompilerOptions::TextureLayoutStrategy::None )
         {
@@ -1647,7 +1697,7 @@ namespace mu
 
 
 	//---------------------------------------------------------------------------------------------
-	void CodeGenerator::GenerateImage_Invert(FImageGenerationResult& result, const NodeImageInvert* InNode)
+	void CodeGenerator::GenerateImage_Invert(const FImageGenerationOptions& Options, FImageGenerationResult& result, const NodeImageInvert* InNode)
 	{
 		const NodeImageInvert::Private& node = *InNode->GetPrivate();
 
@@ -1656,15 +1706,16 @@ namespace mu
 
 		// A image
 		Ptr<ASTOp> a;
-		if (Node* pBase = node.m_pBase.get())
+		if (node.m_pBase)
 		{
-			a = Generate(pBase);
+			FImageGenerationResult BaseResult;
+			GenerateImage(Options, BaseResult, node.m_pBase);
+			a = BaseResult.op;
 		}
 		else
 		{
 			// This argument is required
-			a = GenerateMissingImageCode("Image Invert Color", EImageFormat::IF_RGB_UBYTE,
-				node.m_errorContext);
+			a = GenerateMissingImageCode(TEXT("Image Invert Color"), EImageFormat::IF_RGB_UBYTE, node.m_errorContext, Options);
 		}
 		a = GenerateImageFormat(a, EImageFormat::IF_RGB_UBYTE);
 		a = GenerateImageSize
@@ -1677,7 +1728,7 @@ namespace mu
 
 
     //---------------------------------------------------------------------------------------------
-	void CodeGenerator::GenerateImage_Variation(FImageGenerationResult& result, const NodeImageVariation* InNode)
+	void CodeGenerator::GenerateImage_Variation(const FImageGenerationOptions& Options, FImageGenerationResult& result, const NodeImageVariation* InNode)
 	{
 		const NodeImageVariation::Private& node = *InNode->GetPrivate();
 
@@ -1686,14 +1737,14 @@ namespace mu
         // Default case
         if ( node.m_defaultImage )
         {
-            FImageGenerationResult branchResults;
-			GenerateImage(branchResults, node.m_defaultImage);
-			currentOp = branchResults.op;
+            FImageGenerationResult BranchResults;
+			GenerateImage(Options, BranchResults, node.m_defaultImage);
+			currentOp = BranchResults.op;
         }
         else
         {
             // This argument is required
-            currentOp = GenerateMissingImageCode( "Variation default", EImageFormat::IF_RGBA_UBYTE, node.m_errorContext );
+            currentOp = GenerateMissingImageCode(TEXT("Variation default"), EImageFormat::IF_RGBA_UBYTE, node.m_errorContext, Options);
         }
 
         // Process variations in reverse order, since conditionals are built bottom-up.
@@ -1711,23 +1762,23 @@ namespace mu
 
             if ( tagIndex < 0 )
             {
-                char buf[256];
-                mutable_snprintf( buf, 256, "Unknown tag found in image variation [%s].", tag.c_str() );
+				FString Msg = FString::Printf(TEXT("Unknown tag found in image variation [%s]."), tag.c_str() );
 
-                m_pErrorLog->GetPrivate()->Add( buf, ELMT_WARNING, node.m_errorContext );
+                m_pErrorLog->GetPrivate()->Add( Msg, ELMT_WARNING, node.m_errorContext );
                 continue;
             }
 
             Ptr<ASTOp> variationOp;
             if ( node.m_variations[t].m_image )
             {
-                variationOp = Generate( node.m_variations[t].m_image );
+				FImageGenerationResult Result;
+				GenerateImage(Options, Result, node.m_variations[t].m_image);
+				variationOp = Result.op;
             }
             else
             {
                 // This argument is required
-                variationOp = GenerateMissingImageCode( "Variation option", EImageFormat::IF_RGBA_UBYTE,
-                                                        node.m_errorContext );
+                variationOp = GenerateMissingImageCode(TEXT("Variation option"), EImageFormat::IF_RGBA_UBYTE, node.m_errorContext, Options);
             }
 
 
@@ -1771,12 +1822,12 @@ namespace mu
 
 
 	//---------------------------------------------------------------------------------------------
-	void CodeGenerator::GenerateImage_Table(FImageGenerationResult& result, const NodeImageTable* InNode)
+	void CodeGenerator::GenerateImage_Table(const FImageGenerationOptions& Options, FImageGenerationResult& result, const NodeImageTable* InNode)
 	{
 		const NodeImageTable::Private& node = *InNode->GetPrivate();
 
 		result.op = GenerateTableSwitch<NodeImageTable::Private, TCT_IMAGE, OP_TYPE::IM_SWITCH>(node,
-			[this](const NodeImageTable::Private& node, int colIndex, int row, ErrorLog* pErrorLog)
+			[this,Options](const NodeImageTable::Private& node, int colIndex, int row, ErrorLog* pErrorLog)
 			{
 				TABLE_VALUE CellData = node.m_pTable->GetPrivate()->m_rows[row].m_values[colIndex];
 				ImagePtrConst pImage = nullptr;
@@ -1789,10 +1840,8 @@ namespace mu
 
 				if (!pImage)
 				{
-					char temp[256];
-					mutable_snprintf(temp, 256,
-						"Table has a missing image in column %d, row %d.", colIndex, row);
-					pErrorLog->GetPrivate()->Add(temp, ELMT_ERROR, node.m_errorContext);
+					FString Msg = FString::Printf(TEXT("Table has a missing image in column %d, row %d."), colIndex, row);
+					pErrorLog->GetPrivate()->Add(Msg, ELMT_ERROR, node.m_errorContext);
 				}
 				else
 				{
@@ -1812,13 +1861,15 @@ namespace mu
 					}
 				}
 
-				return Generate(CellImage);
+				FImageGenerationResult Result;
+				GenerateImage(Options, Result, CellImage);
+				return Result.op;
 			});
 	}
 
 
     //---------------------------------------------------------------------------------------------
-    ImagePtr CodeGenerator::GenerateMissingImage( EImageFormat format )
+    Ptr<Image> CodeGenerator::GenerateMissingImage( EImageFormat format )
     {
         // Create the image node if it hasn't been created yet.
         if (!m_missingImage[size_t(format)])
@@ -1912,19 +1963,11 @@ namespace mu
 
 
     //---------------------------------------------------------------------------------------------
-    Ptr<ASTOp> CodeGenerator::GenerateMissingImageCode( const char* strWhere,
-                                                     EImageFormat format,
-                                                     const void* errorContext )
+    Ptr<ASTOp> CodeGenerator::GenerateMissingImageCode(const TCHAR* strWhere, EImageFormat format, const void* errorContext, const FImageGenerationOptions& Options )
     {
         // Log an error message
-        char buf[256];
-        mutable_snprintf
-            (
-                buf, 256,
-                "Required connection not found: %s",
-                strWhere
-            );
-        m_pErrorLog->GetPrivate()->Add( buf, ELMT_ERROR, errorContext );
+		FString Msg = FString::Printf(TEXT("Required connection not found: %s"), strWhere );
+        m_pErrorLog->GetPrivate()->Add( Msg, ELMT_ERROR, errorContext );
 
         // Make a checkered debug image
         ImagePtr pImage = GenerateMissingImage( format );
@@ -1932,14 +1975,15 @@ namespace mu
         NodeImageConstantPtr pNode = new NodeImageConstant();
         pNode->SetValue( pImage.get() );
 
-        Ptr<ASTOp> result = Generate( pNode );
+		FImageGenerationResult Result;
+		GenerateImage(Options, Result, pNode);
 
-        return result;
+        return Result.op;
     }
 
 
     //---------------------------------------------------------------------------------------------
-    Ptr<ASTOp> CodeGenerator::GeneratePlainImageCode( const vec3<float>& colour )
+    Ptr<ASTOp> CodeGenerator::GeneratePlainImageCode( const vec3<float>& colour, const FImageGenerationOptions& Options )
     {
         const int size = 4;
         ImagePtr pImage = new Image( size, size, 1, EImageFormat::IF_RGB_UBYTE );
@@ -1953,10 +1997,12 @@ namespace mu
             pData += 3;
         }
 
-        NodeImageConstantPtr pNode = new NodeImageConstant();
-        pNode->SetValue( pImage.get() );
+        NodeImageConstantPtr ConstantNode = new NodeImageConstant();
+		ConstantNode->SetValue( pImage.get() );
 
-        Ptr<ASTOp> result = Generate( pNode );
+		FImageGenerationResult TempResult;
+		GenerateImage(Options, TempResult, ConstantNode);
+        Ptr<ASTOp> result = TempResult.op;
 
         return result;
     }
diff --git a/Engine/Plugins/Experimental/Mutable/Source/MutableTools/Private/MuT/CodeGenerator_Mesh.cpp b/Engine/Plugins/Experimental/Mutable/Source/MutableTools/Private/MuT/CodeGenerator_Mesh.cpp
--- a/Engine/Plugins/Experimental/Mutable/Source/MutableTools/Private/MuT/CodeGenerator_Mesh.cpp
+++ b/Engine/Plugins/Experimental/Mutable/Source/MutableTools/Private/MuT/CodeGenerator_Mesh.cpp
@@ -280,7 +280,7 @@ class Node;
 			return;
 		}
 
-		// The layout we are adding must have block ids.
+		// The layout must have block ids.
 		check(GeneratedLayout->m_blocks.IsEmpty() || GeneratedLayout->m_blocks[0].m_id != -1);
 
 
@@ -289,7 +289,7 @@ class Node;
 		currentLayoutMesh->AddLayout(Layout);
 
 		// Create the layout block vertex buffer
-		uint16* LayoutData = 0;
+		uint16* LayoutData = nullptr;
 		{
 			const int32 LayoutBufferIndex = currentLayoutMesh->GetVertexBuffers().GetBufferCount();
 			currentLayoutMesh->GetVertexBuffers().SetBufferCount(LayoutBufferIndex + 1);
@@ -375,14 +375,10 @@ class Node;
 		// Notify Overlapping layout blocks
 		if (!OverlappingBlocks.IsEmpty())
 		{
-			char buf[256];
-			mutable_snprintf
-			(
-				buf, 256,
-				"Source mesh has %d layout block overlapping in LOD %d",
+			FString Msg = FString::Printf(TEXT("Source mesh has %d layout block overlapping in LOD %d"),
 				OverlappingBlocks.Num() + 1, m_currentParents.Last().m_lod
 			);
-			m_pErrorLog->GetPrivate()->Add(buf, ELMT_WARNING, errorContext);
+			m_pErrorLog->GetPrivate()->Add(Msg, ELMT_WARNING, errorContext);
 		}
 
 		// Get the information about the texture coordinates channel
@@ -446,14 +442,8 @@ class Node;
 			// Mutable does not support non-normalized UVs
 			if (bNonNormalizedUVs)
 			{
-				char buf[256];
-				mutable_snprintf
-				(
-					buf, 256,
-					"Source mesh has non-normalized UVs in LOD %d",
-					m_currentParents.Last().m_lod
-				);
-				m_pErrorLog->GetPrivate()->Add(buf, ELMT_WARNING, errorContext);
+				FString Msg = FString::Printf(TEXT("Source mesh has non-normalized UVs in LOD %d"), m_currentParents.Last().m_lod );
+				m_pErrorLog->GetPrivate()->Add(Msg, ELMT_WARNING, errorContext);
 			}
 		}
 
@@ -642,19 +632,13 @@ class Node;
 
 			if (!UnassignedUVs.IsEmpty())
 			{
-				char buf[256];
-				mutable_snprintf
-				(
-					buf, 256,
-					"Source mesh has %d vertices not assigned to any layout block in LOD %d",
-					UnassignedUVs.Num(), m_currentParents.Last().m_lod
-				);
+				FString Msg = FString::Printf(TEXT("Source mesh has %d vertices not assigned to any layout block in LOD %d"), UnassignedUVs.Num(), m_currentParents.Last().m_lod);
 
 				ErrorLogMessageAttachedDataView attachedDataView;
 				attachedDataView.m_unassignedUVs = UnassignedUVs.GetData();
 				attachedDataView.m_unassignedUVsSize = (size_t)UnassignedUVs.Num();
 
-				m_pErrorLog->GetPrivate()->Add(buf, attachedDataView, ELMT_WARNING, errorContext);
+				m_pErrorLog->GetPrivate()->Add(Msg, attachedDataView, ELMT_WARNING, errorContext);
 			}
 		}
 
@@ -772,7 +756,7 @@ class Node;
         else
         {
             // This argument is required
-            OpMorph->Factor = GenerateMissingScalarCode( "Morph factor", 0.5f, node.m_errorContext );
+            OpMorph->Factor = GenerateMissingScalarCode(TEXT("Morph factor"), 0.5f, node.m_errorContext );
         }
 
         // Base
@@ -1003,7 +987,7 @@ class Node;
         {
             // This argument is required
             op->SetChild( op->op.args.MeshInterpolate.factor,
-                          GenerateMissingScalarCode( "Interpolation factor", 0.5f, node.m_errorContext ) );
+                          GenerateMissingScalarCode(TEXT("Interpolation factor"), 0.5f, node.m_errorContext ) );
         }
 
         //
@@ -1094,7 +1078,7 @@ class Node;
         else
         {
             // This argument is required
-            op->variable = GenerateMissingScalarCode( "Switch variable", 0.0f, node.m_errorContext );
+            op->variable = GenerateMissingScalarCode(TEXT("Switch variable"), 0.0f, node.m_errorContext );
         }
 
         // Options
@@ -1280,7 +1264,7 @@ class Node;
 		{
 			// This data is required
 			MeshPtr pTempMesh = new Mesh();
-			op->SetValue(pTempMesh, m_compilerOptions->m_optimisationOptions.m_useDiskCache);
+			op->SetValue(pTempMesh, m_compilerOptions->OptimisationOptions.bUseDiskCache);
 			m_constantMeshes.Add(pTempMesh);
 
 			// Log an error message
@@ -1365,7 +1349,7 @@ class Node;
 				}
 			}
 
-			op->SetValue(DuplicateOf, m_compilerOptions->m_optimisationOptions.m_useDiskCache);
+			op->SetValue(DuplicateOf, m_compilerOptions->OptimisationOptions.bUseDiskCache);
 		}
 		else
 		{
@@ -1444,7 +1428,7 @@ class Node;
 
 			// Add the constant data
 			m_constantMeshes.Add(pCloned);
-			op->SetValue(pCloned.get(), m_compilerOptions->m_optimisationOptions.m_useDiskCache);
+			op->SetValue(pCloned.get(), m_compilerOptions->OptimisationOptions.bUseDiskCache);
 		}
 
  
@@ -1498,7 +1482,7 @@ class Node;
 
             Ptr<ASTOpConstantResource> cop = new ASTOpConstantResource();
             cop->type = OP_TYPE::ME_CONSTANT;
-            cop->SetValue( pFormatMesh, m_compilerOptions->m_optimisationOptions.m_useDiskCache );
+            cop->SetValue( pFormatMesh, m_compilerOptions->OptimisationOptions.bUseDiskCache );
             op->Format = cop;
 
             m_constantMeshes.Add(pFormatMesh);
diff --git a/Engine/Plugins/Experimental/Mutable/Source/MutableTools/Private/MuT/CodeGenerator_Scalar.cpp b/Engine/Plugins/Experimental/Mutable/Source/MutableTools/Private/MuT/CodeGenerator_Scalar.cpp
--- a/Engine/Plugins/Experimental/Mutable/Source/MutableTools/Private/MuT/CodeGenerator_Scalar.cpp
+++ b/Engine/Plugins/Experimental/Mutable/Source/MutableTools/Private/MuT/CodeGenerator_Scalar.cpp
@@ -108,7 +108,7 @@ namespace mu
 		else
 		{
 			check(false);
-			mu::Halt();
+			return;
 		}
 
 		// Cache the result
@@ -225,7 +225,7 @@ namespace mu
 		if (node.m_options.Num() == 0)
 		{
 			// No options in the switch!
-			Ptr<ASTOp> missingOp = GenerateMissingScalarCode("Switch option",
+			Ptr<ASTOp> missingOp = GenerateMissingScalarCode(TEXT("Switch option"),
 				1.0f,
 				node.m_errorContext);
 			result.op = missingOp;
@@ -245,7 +245,7 @@ namespace mu
 		else
 		{
 			// This argument is required
-			op->variable = GenerateMissingScalarCode("Switch variable", 0.0f, node.m_errorContext);
+			op->variable = GenerateMissingScalarCode(TEXT("Switch variable"), 0.0f, node.m_errorContext);
 		}
 
 		// Options
@@ -261,7 +261,7 @@ namespace mu
 			else
 			{
 				// This argument is required
-				branch = GenerateMissingScalarCode("Switch option", 1.0f, node.m_errorContext);
+				branch = GenerateMissingScalarCode(TEXT("Switch option"), 1.0f, node.m_errorContext);
 			}
 			op->cases.Emplace((int16_t)t, op, branch);
 		}
@@ -289,7 +289,7 @@ namespace mu
 		else
 		{
 			// This argument is required
-			op = GenerateMissingScalarCode("Variation default", 0.0f, node.m_errorContext);
+			op = GenerateMissingScalarCode(TEXT("Variation default"), 0.0f, node.m_errorContext);
 		}
 
 		// Process variations in reverse order, since conditionals are built bottom-up.
@@ -307,11 +307,9 @@ namespace mu
 
 			if (tagIndex < 0)
 			{
-				char buf[256];
-				mutable_snprintf(buf, 256, "Unknown tag found in image variation [%s].",
-					tag.c_str());
+				FString Msg = FString::Printf(TEXT("Unknown tag found in image variation [%s]."), *FString(tag.c_str()));
 
-				m_pErrorLog->GetPrivate()->Add(buf, ELMT_WARNING, node.m_errorContext);
+				m_pErrorLog->GetPrivate()->Add(Msg, ELMT_WARNING, node.m_errorContext);
 				continue;
 			}
 
@@ -325,7 +323,7 @@ namespace mu
 			else
 			{
 				// This argument is required
-				variationOp = GenerateMissingScalarCode("Variation option", 0.0f,
+				variationOp = GenerateMissingScalarCode(TEXT("Variation option"), 0.0f,
 					node.m_errorContext);
 			}
 
@@ -357,7 +355,7 @@ namespace mu
 		}
 		else
 		{
-			op->time = CodeGenerator::GenerateMissingScalarCode("Curve T", 0.5f, node.m_errorContext);
+			op->time = CodeGenerator::GenerateMissingScalarCode(TEXT("Curve T"), 0.5f, node.m_errorContext);
 		}
 
 		op->curve = node.m_curve;
@@ -394,12 +392,7 @@ namespace mu
 		else
 		{
 			op->SetChild(op->op.args.ScalarArithmetic.a,
-				CodeGenerator::GenerateMissingScalarCode
-				(
-					"ScalarArithmetic A",
-					1.0f,
-					node.m_errorContext
-				)
+				CodeGenerator::GenerateMissingScalarCode( TEXT("ScalarArithmetic A"), 1.0f, node.m_errorContext )
 			);
 		}
 
@@ -411,12 +404,7 @@ namespace mu
 		else
 		{
 			op->SetChild(op->op.args.ScalarArithmetic.b,
-				CodeGenerator::GenerateMissingScalarCode
-				(
-					"ScalarArithmetic B",
-					1.0f,
-					node.m_errorContext
-				)
+				CodeGenerator::GenerateMissingScalarCode( TEXT("ScalarArithmetic B"), 1.0f, node.m_errorContext )
 			);
 		}
 
@@ -443,17 +431,11 @@ namespace mu
 
 
 	//---------------------------------------------------------------------------------------------
-	mu::Ptr<ASTOp> CodeGenerator::GenerateMissingScalarCode(const char* strWhere, float value, const void* errorContext)
+	mu::Ptr<ASTOp> CodeGenerator::GenerateMissingScalarCode(const TCHAR* strWhere, float value, const void* errorContext)
 	{
 		// Log a warning
-		char buf[256];
-		mutable_snprintf
-		(
-			buf, 256,
-			"Required connection not found: %s",
-			strWhere
-		);
-		m_pErrorLog->GetPrivate()->Add(buf, ELMT_ERROR, errorContext);
+		FString Msg = FString::Printf(TEXT("Required connection not found: %s"), strWhere );
+		m_pErrorLog->GetPrivate()->Add(Msg, ELMT_ERROR, errorContext);
 
 		// Create a constant node
 		NodeScalarConstantPtr pNode = new NodeScalarConstant();
diff --git a/Engine/Plugins/Experimental/Mutable/Source/MutableTools/Private/MuT/CodeGenerator_SecondPass.cpp b/Engine/Plugins/Experimental/Mutable/Source/MutableTools/Private/MuT/CodeGenerator_SecondPass.cpp
--- a/Engine/Plugins/Experimental/Mutable/Source/MutableTools/Private/MuT/CodeGenerator_SecondPass.cpp
+++ b/Engine/Plugins/Experimental/Mutable/Source/MutableTools/Private/MuT/CodeGenerator_SecondPass.cpp
@@ -139,7 +139,7 @@ namespace mu
 			}
 
 			// Optimise the condition now.
-			//PartialOptimise( c, m_pCompilerOptions->m_optimisationOptions );
+			//PartialOptimise( c, m_pCompilerOptions->OptimisationOptions );
 		}
 
 
@@ -513,7 +513,7 @@ namespace mu
 			}
 		}
 
-		//PartialOptimise( c, m_pCompilerOptions->m_optimisationOptions );
+		//PartialOptimise( c, m_pCompilerOptions->OptimisationOptions );
 
 		return c;
 	}
diff --git a/Engine/Plugins/Experimental/Mutable/Source/MutableTools/Private/MuT/CodeOptimiser.cpp b/Engine/Plugins/Experimental/Mutable/Source/MutableTools/Private/MuT/CodeOptimiser.cpp
--- a/Engine/Plugins/Experimental/Mutable/Source/MutableTools/Private/MuT/CodeOptimiser.cpp
+++ b/Engine/Plugins/Experimental/Mutable/Source/MutableTools/Private/MuT/CodeOptimiser.cpp
@@ -432,8 +432,8 @@ namespace mu
 	public:
 		// input
 		mu::Ptr<ASTOp> m_source;
-		bool m_useDiskCache = false;
-		int m_imageCompressionQuality = 0;
+		bool bUseDiskCache = false;
+		int ImageCompressionQuality = 0;
 
 	private:
 
@@ -455,7 +455,7 @@ namespace mu
 
 			Ptr<Settings> pSettings = new Settings;
 			pSettings->SetProfile( false );
-			pSettings->SetImageCompressionQuality( m_imageCompressionQuality );
+			pSettings->SetImageCompressionQuality( ImageCompressionQuality );
 			SystemPtr pSystem = new System( pSettings );
 
 			// Don't generate mips suring linking here.
@@ -486,7 +486,7 @@ namespace mu
 				{
 					mu::Ptr<ASTOpConstantResource> constantOp = new ASTOpConstantResource();
 					constantOp->type = OP_TYPE::ME_CONSTANT;
-					constantOp->SetValue( pMesh, m_useDiskCache );
+					constantOp->SetValue( pMesh, bUseDiskCache );
 					m_result = constantOp;
 				  }
 				break;
@@ -502,7 +502,7 @@ namespace mu
 				{
 					mu::Ptr<ASTOpConstantResource> constantOp = new ASTOpConstantResource();
 					constantOp->type = OP_TYPE::IM_CONSTANT;
-					constantOp->SetValue( pImage, m_useDiskCache );
+					constantOp->SetValue( pImage, bUseDiskCache );
 					m_result = constantOp;
 				}
 
@@ -520,7 +520,7 @@ namespace mu
 				{
 					mu::Ptr<ASTOpConstantResource> constantOp = new ASTOpConstantResource();
 					constantOp->type = OP_TYPE::LA_CONSTANT;
-					constantOp->SetValue( pLayout, m_useDiskCache );
+					constantOp->SetValue( pLayout, bUseDiskCache );
 					m_result = constantOp;
 				}
 				break;
@@ -590,7 +590,7 @@ namespace mu
 		bool modified = false;
 
 		// don't do this if constant optimization has been disabled, usually for debugging.
-		if (!options->m_optimisationOptions.m_constReduction)
+		if (!options->OptimisationOptions.bConstReduction)
 		{
 			return false;
 		}
@@ -749,9 +749,9 @@ namespace mu
 					modified = true;
 
 					ConstantTask* constantTask = new ConstantTask;
-					constantTask->m_useDiskCache = options->m_optimisationOptions.m_useDiskCache;
+					constantTask->bUseDiskCache = options->OptimisationOptions.bUseDiskCache;
 					constantTask->m_source = n;
-					constantTask->m_imageCompressionQuality = options->m_imageCompressionQuality;
+					constantTask->ImageCompressionQuality = options->ImageCompressionQuality;
 
 					constantTask->Run();
 					constantTask->Complete();
@@ -811,13 +811,13 @@ namespace mu
 			// All kind of optimisations that depend on the meaning of each operation
 			// \TODO: We are doing it for all states.
 			UE_LOG(LogMutableCore, Verbose, TEXT(" - semantic optimiser"));
-			modified |= SemanticOptimiserAST( roots, m_options->GetPrivate()->m_optimisationOptions, Pass );
+			modified |= SemanticOptimiserAST( roots, m_options->GetPrivate()->OptimisationOptions, Pass );
 			//AXE_INT_VALUE("Mutable", Verbose, "ast size", (int64_t)ASTOp::CountNodes(roots));
 			UE_LOG(LogMutableCore, Verbose, TEXT("(int) %s : %ld"), TEXT("ast size"), int64(ASTOp::CountNodes(roots)));
 			ASTOp::LogHistogram(roots);
 
 			UE_LOG(LogMutableCore, Verbose, TEXT(" - sink optimiser"));
-			modified |= SinkOptimiserAST( roots, m_options->GetPrivate()->m_optimisationOptions );
+			modified |= SinkOptimiserAST( roots, m_options->GetPrivate()->OptimisationOptions );
 			//AXE_INT_VALUE("Mutable", Verbose, "ast size", (int64_t)ASTOp::CountNodes(roots));
 			ASTOp::LogHistogram(roots);
 
@@ -1120,7 +1120,7 @@ namespace mu
 							NewMesh->CheckIntegrity();
 							mu::Ptr<ASTOpConstantResource> newOp = new ASTOpConstantResource();
 							newOp->type = OP_TYPE::ME_CONSTANT;
-							newOp->SetValue(NewMesh, options.m_useDiskCache);
+							newOp->SetValue(NewMesh, options.bUseDiskCache);
 
 							ASTOp::Replace(at, newOp);
 						}
@@ -1139,7 +1139,7 @@ namespace mu
 	//        // Add new skeleton constant instruction
 	//        Ptr<ASTOpConstantResource> cop = new ASTOpConstantResource();
 	//        cop->type = OP_TYPE::ME_CONSTANT;
-	//        cop->SetValue( rep.m_pFinalSkeleton, options.m_useDiskCache );
+	//        cop->SetValue( rep.m_pFinalSkeleton, options.bUseDiskCache );
 
 	//        // Add new "apply skeleton" instruction
 	//        Ptr<ASTOpFixed> skop = new ASTOpFixed();
@@ -1171,11 +1171,11 @@ namespace mu
 
 		UE_LOG(LogMutableCore, Verbose, TEXT("(int) %s : %ld"), TEXT("ast size"), int64(ASTOp::CountNodes(roots)));
 
-		if ( m_options->GetPrivate()->m_optimisationOptions.m_enabled )
+		if ( m_options->GetPrivate()->OptimisationOptions.bEnabled )
 		{
 			// We use 4 times the count because at the time we moved to sharing this count it
 			// was being used 4 times, and we want to keep the tests consistent.
-			m_optimizeIterationsLeft = m_options->GetPrivate()->m_optimisationOptions.m_maxOptimisationLoopCount * 4;
+			m_optimizeIterationsLeft = m_options->GetPrivate()->OptimisationOptions.MaxOptimisationLoopCount * 4;
 			m_optimizeIterationsMax = m_optimizeIterationsLeft;
 
 			// The first duplicated data remover has the special mission of removing
@@ -1193,12 +1193,12 @@ namespace mu
 			UE_LOG(LogMutableCore, Verbose, TEXT("(int) %s : %ld"), TEXT("ast size"), int64(ASTOp::CountNodes(roots)));
 
 			// Special optimization stages
-			if ( m_options->GetPrivate()->m_optimisationOptions.m_uniformizeSkeleton )
+			if ( m_options->GetPrivate()->OptimisationOptions.bUniformizeSkeleton )
 			{
 				UE_LOG(LogMutableCore, Verbose, TEXT(" - skeleton cleaner"));
 				ASTOp::LogHistogram(roots);
 
-				SkeletonCleanerAST( roots, m_options->GetPrivate()->m_optimisationOptions );
+				SkeletonCleanerAST( roots, m_options->GetPrivate()->OptimisationOptions );
 				UE_LOG(LogMutableCore, Verbose, TEXT("(int) %s : %ld"), TEXT("ast size"), int64(ASTOp::CountNodes(roots)));
 				ASTOp::LogHistogram(roots);
 			}
@@ -1247,7 +1247,7 @@ namespace mu
 					auto pMesh = static_cast<const Mesh*>(typed->GetValue().get());
 					pMesh->ResetStaticFormatFlags();
 					typed->SetValue( pMesh,
-									 m_options->GetPrivate()->m_optimisationOptions.m_useDiskCache );
+									 m_options->GetPrivate()->OptimisationOptions.bUseDiskCache );
 				}
 			});
 
@@ -1273,7 +1273,7 @@ namespace mu
 		}
 
 	//        // Minimal optimisation of constant subtrees
-		else if ( m_options->GetPrivate()->m_optimisationOptions.m_constReduction )
+		else if ( m_options->GetPrivate()->OptimisationOptions.bConstReduction )
 		{
 			// The first duplicated data remover has the special mission of removing
 			// duplicated data (meshes) that may have been specified in the source
diff --git a/Engine/Plugins/Experimental/Mutable/Source/MutableTools/Private/MuT/CodeOptimiser.h b/Engine/Plugins/Experimental/Mutable/Source/MutableTools/Private/MuT/CodeOptimiser.h
--- a/Engine/Plugins/Experimental/Mutable/Source/MutableTools/Private/MuT/CodeOptimiser.h
+++ b/Engine/Plugins/Experimental/Mutable/Source/MutableTools/Private/MuT/CodeOptimiser.h
@@ -253,7 +253,7 @@ namespace mu
 
         bool m_modified;
 
-        FModelOptimizationOptions m_optimisationOptions;
+        FModelOptimizationOptions OptimisationOptions;
 
         RuntimeParameterVisitorAST m_hasRuntimeParamVisitor;
 
diff --git a/Engine/Plugins/Experimental/Mutable/Source/MutableTools/Private/MuT/CodeOptimiser_Parameters.cpp b/Engine/Plugins/Experimental/Mutable/Source/MutableTools/Private/MuT/CodeOptimiser_Parameters.cpp
--- a/Engine/Plugins/Experimental/Mutable/Source/MutableTools/Private/MuT/CodeOptimiser_Parameters.cpp
+++ b/Engine/Plugins/Experimental/Mutable/Source/MutableTools/Private/MuT/CodeOptimiser_Parameters.cpp
@@ -274,7 +274,7 @@ namespace mu
             )
             : m_stateProps(s)
             , m_modified(false)
-            , m_optimisationOptions(optimisationOptions)
+            , OptimisationOptions(optimisationOptions)
             , m_hasRuntimeParamVisitor(&s)
     {
     }
@@ -1358,7 +1358,7 @@ namespace mu
                     case OP_TYPE::IM_DODGECOLOUR:
                     case OP_TYPE::IM_MULTIPLYCOLOUR:
                     {
-                        if ( m_optimisationOptions.m_optimiseOverlappedMasks )
+                        if ( OptimisationOptions.m_optimiseOverlappedMasks )
                         {
                             OP::ADDRESS maskAt = op.args.ImageLayer.mask;
                             OP::ADDRESS baseMaskAt =
@@ -2051,7 +2051,7 @@ namespace mu
 
                     UE_LOG(LogMutableCore, Verbose, TEXT(" - before parameter optimiser"));
 
-                    ParameterOptimiserAST param( m_states[s], m_options->GetPrivate()->m_optimisationOptions );
+                    ParameterOptimiserAST param( m_states[s], m_options->GetPrivate()->OptimisationOptions );
                     modified = param.Apply();
 
                     TArray<Ptr<ASTOp>> roots;
@@ -2062,12 +2062,12 @@ namespace mu
 
                     // All kind of optimisations that depend on the meaning of each operation
                     UE_LOG(LogMutableCore, Verbose, TEXT(" - semantic optimiser"));
-                    modified |= SemanticOptimiserAST( roots, m_options->GetPrivate()->m_optimisationOptions, 1 );
+                    modified |= SemanticOptimiserAST( roots, m_options->GetPrivate()->OptimisationOptions, 1 );
 					UE_LOG(LogMutableCore, Verbose, TEXT("(int) %s : %ld"), TEXT("ast size"), int64(ASTOp::CountNodes(roots)));
 					//ASTOp::LogHistogram(roots);
 
                     UE_LOG(LogMutableCore, Verbose, TEXT(" - sink optimiser"));
-                    modified |= SinkOptimiserAST( roots, m_options->GetPrivate()->m_optimisationOptions );
+                    modified |= SinkOptimiserAST( roots, m_options->GetPrivate()->OptimisationOptions );
 					UE_LOG(LogMutableCore, Verbose, TEXT("(int) %s : %ld"), TEXT("ast size"), int64(ASTOp::CountNodes(roots)));
 					//ASTOp::LogHistogram(roots);
 
@@ -2123,26 +2123,6 @@ namespace mu
             }
         }
 
-        // If we are targetting GPU build, gpuize the code
-//        m_pModelReport->GetPrivate()->m_states.resize( pModel->GetStateCount() );
-
-//        for ( int s=0; s<pModel->GetStateCount(); ++s )
-//        {
-//            if ( m_gpuPlatformProps.type!=GPU_NONE
-//                 &&
-//                 ( m_states[s].optimisationFlags.m_gpu.m_internal || m_states[s].optimisationFlags.m_gpu.m_external ) )
-//            {
-//                MUTABLE_CPUPROFILER_SCOPE(GPUTranslator);
-//                GPUTranslator( program,
-//                               s,
-//                               m_gpuPlatformProps,
-//                               m_states[s].optimisationFlags.m_gpu,
-//                               m_pModelReport );
-//            }
-//        }
-
-        //program.Check();
-
         // Reoptimise because of state cache reformats and gpuization
         {
             MUTABLE_CPUPROFILER_SCOPE(Reoptimise);
@@ -2165,7 +2145,7 @@ namespace mu
 
                 UE_LOG(LogMutableCore, Verbose, TEXT(" - semantic optimiser"));
                 modified |=
-                    SemanticOptimiserAST( roots, m_options->GetPrivate()->m_optimisationOptions, 1 );
+                    SemanticOptimiserAST( roots, m_options->GetPrivate()->OptimisationOptions, 1 );
 				UE_LOG(LogMutableCore, Verbose, TEXT("(int) %s : %ld"), TEXT("ast size"), int64(ASTOp::CountNodes(roots)));
 
                 // Image size operations are treated separately
@@ -2207,8 +2187,8 @@ namespace mu
         {
             MUTABLE_CPUPROFILER_SCOPE(DataFormats);
 
-            DataOptimiseAST( m_options->GetPrivate()->m_imageCompressionQuality, roots,
-                             m_options->GetPrivate()->m_optimisationOptions );
+            DataOptimiseAST( m_options->GetPrivate()->ImageCompressionQuality, roots,
+                             m_options->GetPrivate()->OptimisationOptions );
 
             // After optimising the data formats, we may remove more constants
             DuplicatedDataRemoverAST( roots );
diff --git a/Engine/Plugins/Experimental/Mutable/Source/MutableTools/Private/MuT/Compiler.cpp b/Engine/Plugins/Experimental/Mutable/Source/MutableTools/Private/MuT/Compiler.cpp
--- a/Engine/Plugins/Experimental/Mutable/Source/MutableTools/Private/MuT/Compiler.cpp
+++ b/Engine/Plugins/Experimental/Mutable/Source/MutableTools/Private/MuT/Compiler.cpp
@@ -59,7 +59,6 @@ namespace mu
     CompilerOptions::CompilerOptions()
     {
         m_pD = new Private();
-        m_pD->m_log = false;
     }
 
 
@@ -80,69 +79,62 @@ namespace mu
 
 
     //---------------------------------------------------------------------------------------------
-    void CompilerOptions::SetLogEnabled( bool enabled )
+    void CompilerOptions::SetLogEnabled( bool bEnabled)
     {
-        m_pD->m_log = enabled;
+        m_pD->bLog = bEnabled;
     }
 
 
     //---------------------------------------------------------------------------------------------
-    void CompilerOptions::SetOptimisationEnabled( bool enabled )
+    void CompilerOptions::SetOptimisationEnabled( bool bEnabled)
     {
-        m_pD->m_optimisationOptions.m_enabled = enabled;
-        if (enabled)
+        m_pD->OptimisationOptions.bEnabled = bEnabled;
+        if (bEnabled)
         {
-            m_pD->m_optimisationOptions.m_constReduction = true;
+            m_pD->OptimisationOptions.bConstReduction = true;
         }
     }
 
 
     //---------------------------------------------------------------------------------------------
-    void CompilerOptions::SetConstReductionEnabled( bool constReduction )
+    void CompilerOptions::SetConstReductionEnabled( bool bConstReductionEnabled )
     {
-        m_pD->m_optimisationOptions.m_constReduction = constReduction;
+        m_pD->OptimisationOptions.bConstReduction = bConstReductionEnabled;
     }
 
 
     //---------------------------------------------------------------------------------------------
     void CompilerOptions::SetUseDiskCache( bool enabled )
     {
-        m_pD->m_optimisationOptions.m_useDiskCache = enabled;
+        m_pD->OptimisationOptions.bUseDiskCache = enabled;
     }
 
 
     //---------------------------------------------------------------------------------------------
     void CompilerOptions::SetOptimisationMaxIteration( int maxIterations )
     {
-        m_pD->m_optimisationOptions.m_maxOptimisationLoopCount = maxIterations;
+        m_pD->OptimisationOptions.MaxOptimisationLoopCount = maxIterations;
     }
 
 
     //---------------------------------------------------------------------------------------------
-    void CompilerOptions::SetIgnoreStates( bool ignore )
+    void CompilerOptions::SetIgnoreStates( bool bIgnore )
     {
-        m_pD->m_ignoreStates = ignore;
-    }
-
-
-    //---------------------------------------------------------------------------------------------
-    void CompilerOptions::SetTextureLayoutStrategy( TextureLayoutStrategy strategy )
-    {
-        m_pD->m_textureLayoutStrategy = strategy;
+        m_pD->bIgnoreStates = bIgnore;
     }
 
 
 	//---------------------------------------------------------------------------------------------
-	void CompilerOptions::SetImageCompressionQuality(int quality)
+	void CompilerOptions::SetImageCompressionQuality(int32 Quality)
 	{
-		m_pD->m_imageCompressionQuality = quality;
+		m_pD->ImageCompressionQuality = Quality;
 	}
 
 
 	//---------------------------------------------------------------------------------------------
 	void CompilerOptions::SetImageTiling(int32 Tiling)
 	{
-		m_pD->m_imageTiling = Tiling;
+		m_pD->ImageTiling = Tiling;
 	}
 
 
@@ -157,7 +149,7 @@ namespace mu
 	//---------------------------------------------------------------------------------------------
 	void CompilerOptions::SetEnableProgressiveImages(bool bEnabled)
 	{
-		m_pD->m_optimisationOptions.bEnableProgressiveImages = bEnabled;
+		m_pD->OptimisationOptions.bEnableProgressiveImages = bEnabled;
 	}
 
 
diff --git a/Engine/Plugins/Experimental/Mutable/Source/MutableTools/Private/MuT/CompilerPrivate.h b/Engine/Plugins/Experimental/Mutable/Source/MutableTools/Private/MuT/CompilerPrivate.h
--- a/Engine/Plugins/Experimental/Mutable/Source/MutableTools/Private/MuT/CompilerPrivate.h
+++ b/Engine/Plugins/Experimental/Mutable/Source/MutableTools/Private/MuT/CompilerPrivate.h
@@ -87,19 +87,16 @@ namespace mu
     public:
 
         //! Detailed optimization options
-        FModelOptimizationOptions m_optimisationOptions;
-
-        bool m_ignoreStates = false;
-        CompilerOptions::TextureLayoutStrategy m_textureLayoutStrategy = CompilerOptions::TextureLayoutStrategy::Pack;
+        FModelOptimizationOptions OptimisationOptions;
 
+        bool bIgnoreStates = false;
 		int MinRomSize = 3;
 		int MinTextureResidentMipCount = 3;
 
-        int m_imageCompressionQuality = 0;
-		int32 m_imageTiling=0 ;
-
-        bool m_log = false;
+        int ImageCompressionQuality = 0;
+		int32 ImageTiling=0 ;
 
+        bool bLog = false;
     };
 
 
diff --git a/Engine/Plugins/Experimental/Mutable/Source/MutableTools/Private/MuT/DataPacker.cpp b/Engine/Plugins/Experimental/Mutable/Source/MutableTools/Private/MuT/DataPacker.cpp
--- a/Engine/Plugins/Experimental/Mutable/Source/MutableTools/Private/MuT/DataPacker.cpp
+++ b/Engine/Plugins/Experimental/Mutable/Source/MutableTools/Private/MuT/DataPacker.cpp
@@ -771,9 +771,9 @@ namespace mu
                     // Only replace if the compression was worth!
                     size_t oldSize = pOld->GetDataSize();
                     size_t newSize = pNew->GetDataSize();
-                    if (float(oldSize) > float(newSize) * options.m_minRLECompressionGain)
+                    if (float(oldSize) > float(newSize) * options.MinRLECompressionGain)
                     {
-                        typed->SetValue(pNew, options.m_useDiskCache);
+                        typed->SetValue(pNew, options.bUseDiskCache);
                     }
                 }
                 else if ( accFormat.m_supportedFormats[typed][(size_t)EImageFormat::IF_L_UBYTE_RLE] )
@@ -783,9 +783,9 @@ namespace mu
                     // Only replace if the compression was worth!
                     size_t oldSize = pOld->GetDataSize();
                     size_t newSize = pNew->GetDataSize();
-                    if (float(oldSize) > float(newSize) * options.m_minRLECompressionGain)
+                    if (float(oldSize) > float(newSize) * options.MinRLECompressionGain)
                     {
-                        typed->SetValue(pNew, options.m_useDiskCache);
+                        typed->SetValue(pNew, options.bUseDiskCache);
                     }
                 }
             }
@@ -804,7 +804,7 @@ namespace mu
 				ASTOpConstantResource* typed = dynamic_cast<ASTOpConstantResource*>(n.get());
                 Ptr<Mesh> pMesh = static_cast<const Mesh*>(typed->GetValue().get())->Clone();
                 MeshRemoveUnusedBufferSemantics( pMesh.get(), meshSemanticsVisitor.m_requiredSemantics[typed]);
-                typed->SetValue(pMesh, options.m_useDiskCache);
+                typed->SetValue(pMesh, options.bUseDiskCache);
             }
         });
 
diff --git a/Engine/Plugins/Experimental/Mutable/Source/MutableTools/Private/MuT/ErrorLog.cpp b/Engine/Plugins/Experimental/Mutable/Source/MutableTools/Private/MuT/ErrorLog.cpp
--- a/Engine/Plugins/Experimental/Mutable/Source/MutableTools/Private/MuT/ErrorLog.cpp
+++ b/Engine/Plugins/Experimental/Mutable/Source/MutableTools/Private/MuT/ErrorLog.cpp
@@ -111,7 +111,7 @@ namespace mu
 
 		if ( index >=0 && index<(int)m_pD->m_messages.Num() )
 		{
-			const Private::MSG& message = m_pD->m_messages[index];
+			const Private::FMessage& message = m_pD->m_messages[index];
             
             if ( message.m_data ) 
             {
@@ -123,35 +123,25 @@ namespace mu
 		return result;
 	}
 
-	//---------------------------------------------------------------------------------------------
-	void ErrorLog::Private::Add(const char* strMessage, ErrorLogMessageType type,
-		const void* context)
-	{
-		m_messages.Add(MSG());
-		m_messages.Last().m_type = type;
-		m_messages.Last().m_text = strMessage;
-		m_messages.Last().m_context = context;
-	}
-
 	//---------------------------------------------------------------------------------------------
 	void ErrorLog::Private::Add(const FString& InMessage, ErrorLogMessageType InType, const void* InContext)
 	{
-		m_messages.Add(MSG());
+		m_messages.Add(FMessage());
 		m_messages.Last().m_type = InType;
 		m_messages.Last().m_text = InMessage;
 		m_messages.Last().m_context = InContext;
 	}
 
 	//---------------------------------------------------------------------------------------------
-	void ErrorLog::Private::Add( const char* strMessage, 
-                                 const ErrorLogMessageAttachedDataView& dataView,
-                                 ErrorLogMessageType type, const void* context )
+	void ErrorLog::Private::Add(const FString& InMessage,
+                                const ErrorLogMessageAttachedDataView& dataView,
+                                ErrorLogMessageType type, const void* context )
 	{
-		m_messages.Add( MSG() );
+		m_messages.Add(FMessage() );
 		m_messages.Last().m_type = type;
-		m_messages.Last().m_text = strMessage;
+		m_messages.Last().m_text = InMessage;
 		m_messages.Last().m_context = context;
-		m_messages.Last().m_data = std::make_shared< DATA >();
+		m_messages.Last().m_data = MakeShared<FErrorData>();
 
         if ( dataView.m_unassignedUVs && dataView.m_unassignedUVsSize > 0 )
         {
@@ -165,7 +155,7 @@ namespace mu
 	{
 		UE_LOG(LogMutableCore, Log, TEXT(" Error Log :\n"));
 
-		for ( const Private::MSG& msg : m_pD->m_messages )
+		for ( const Private::FMessage& msg : m_pD->m_messages )
 		{
 			switch ( msg.m_type )
 			{
diff --git a/Engine/Plugins/Experimental/Mutable/Source/MutableTools/Private/MuT/ErrorLogPrivate.h b/Engine/Plugins/Experimental/Mutable/Source/MutableTools/Private/MuT/ErrorLogPrivate.h
--- a/Engine/Plugins/Experimental/Mutable/Source/MutableTools/Private/MuT/ErrorLogPrivate.h
+++ b/Engine/Plugins/Experimental/Mutable/Source/MutableTools/Private/MuT/ErrorLogPrivate.h
@@ -22,26 +22,20 @@ namespace mu
 		{
 		}
 
-        struct DATA
+        struct FErrorData
         {
             TArray< float > m_unassignedUVs;
         };
 
-		struct MSG
+		struct FMessage
 		{
-			MSG()
-			{
-                m_type = ELMT_NONE;
-				m_context = 0;
-			}
-
-			ErrorLogMessageType m_type;
+			ErrorLogMessageType m_type = ELMT_NONE;
 			FString m_text;
-            std::shared_ptr< DATA > m_data;
-			const void* m_context;
+            TSharedPtr<FErrorData> m_data;
+			const void* m_context = nullptr;
 		};
 
-		TArray< MSG > m_messages;
+		TArray<FMessage> m_messages;
 
 
 		//-----------------------------------------------------------------------------------------
@@ -49,20 +43,11 @@ namespace mu
 		//!
 		void Add(const FString& Message, ErrorLogMessageType Type, const void* Context);
 
-		//!
-		void Add(const char* strMessage, ErrorLogMessageType type, const void* context);
-
         //!
-        void Add( const char* strMessage, const ErrorLogMessageAttachedDataView& data, 
-                  ErrorLogMessageType type, const void* context );
+        void Add(const FString& InMessage, const ErrorLogMessageAttachedDataView& data, ErrorLogMessageType type, const void* context );
 	};
 
 
-	//---------------------------------------------------------------------------------------------
-	class Model;
-	typedef Ptr<Model> ModelPtr;
-	typedef Ptr<const Model> ModelPtrConst;
-
     extern const TCHAR* GetOpName( OP_TYPE type );
 
 }
diff --git a/Engine/Plugins/Experimental/Mutable/Source/MutableTools/Public/MuT/Compiler.h b/Engine/Plugins/Experimental/Mutable/Source/MutableTools/Public/MuT/Compiler.h
--- a/Engine/Plugins/Experimental/Mutable/Source/MutableTools/Public/MuT/Compiler.h
+++ b/Engine/Plugins/Experimental/Mutable/Source/MutableTools/Public/MuT/Compiler.h
@@ -30,10 +30,6 @@ namespace mu
     using ErrorLogPtr=Ptr<ErrorLog>;
     using ErrorLogPtrConst=Ptr<const ErrorLog>;
 
-    class ModelReport;
-    using ModelReportPtr=Ptr<ModelReport>;
-    using ModelReportPtrConst=Ptr<const ModelReport>;
-
     class Node;
     class NodeTransformedObject;
 
@@ -55,18 +51,18 @@ namespace mu
         void SetConstReductionEnabled( bool bEnabled );
 
         //!
-        void SetOptimisationMaxIteration( int maxIterations );
+        void SetOptimisationMaxIteration( int32 MaxIterations );
 
         //!
-        void SetIgnoreStates( bool ignore );
+        void SetIgnoreStates( bool bIgnore );
 
         //! If enabled, the disk will be used as temporary memory. This will make the compilation
         //! process very slow, but will be able to compile very large models.
-        void SetUseDiskCache( bool enabled );
+        void SetUseDiskCache( bool bEnabled );
 
 		//! Set the quality for the image compression algorithms. The level value is used internally
 		//! with System::SetImagecompressionQuality
-		void SetImageCompressionQuality(int quality);
+		void SetImageCompressionQuality(int32 quality);
 
 		/** Set the image tiling strategy :
 		 * If 0 (default) there is no tiling. Otherwise, images will be generated in tiles of the given size or less, and assembled afterwards as a final step.
@@ -80,7 +76,7 @@ namespace mu
 		void SetEnableProgressiveImages(bool bEnabled);
 
         //! Different data packing strategies
-        enum class TextureLayoutStrategy : std::uint8_t
+        enum class TextureLayoutStrategy : uint8
         {
             //! Pack texture layouts without changing any scale
             Pack,
@@ -95,12 +91,6 @@ namespace mu
         //! Return a readable string for a GPU.
         static const char* GetTextureLayoutStrategyName( TextureLayoutStrategy s );
 
-        //! Enable or disable the process that splits data into related sets to serialise in
-        //! separate files. By default it is disabled. In order to get multiple files, this has to
-        //! be enabled AND the appropiate serialisation function needs to be called (see
-        //! rutime/Model.h).
-        void SetTextureLayoutStrategy( TextureLayoutStrategy strategy );
-
         //-----------------------------------------------------------------------------------------
         // Interface pattern
         //-----------------------------------------------------------------------------------------
