[Mutable] Remove incorrect check of MeshDifference success result.

#rnx

diff --git a/Engine/Plugins/Experimental/Mutable/Source/MutableRuntime/Private/MuR/CodeRunner.cpp b/Engine/Plugins/Experimental/Mutable/Source/MutableRuntime/Private/MuR/CodeRunner.cpp
--- a/Engine/Plugins/Experimental/Mutable/Source/MutableRuntime/Private/MuR/CodeRunner.cpp
+++ b/Engine/Plugins/Experimental/Mutable/Source/MutableRuntime/Private/MuR/CodeRunner.cpp
@@ -1101,8 +1101,6 @@ namespace mu
                                NumChannels, Semantics.GetData(), SemanticIndices.GetData(),
                                bIgnoreTextureCoords != 0, bOutSuccess);
 			
-				check(bOutSuccess);
-
 				Release(pBase);
 				Release(pTarget);
 
