[mutable] Remove unused compilation option.
Ensure image tiling options are honored when packaging.
#rnx

diff --git a/Engine/Plugins/Experimental/Mutable/Source/CustomizableObject/Classes/MuCO/CustomizableObject.h b/Engine/Plugins/Experimental/Mutable/Source/CustomizableObject/Classes/MuCO/CustomizableObject.h
--- a/Engine/Plugins/Experimental/Mutable/Source/CustomizableObject/Classes/MuCO/CustomizableObject.h
+++ b/Engine/Plugins/Experimental/Mutable/Source/CustomizableObject/Classes/MuCO/CustomizableObject.h
@@ -165,12 +165,6 @@ struct FCompilationOptions
 	UPROPERTY()
 	int32 OptimizationLevel = 1;
 
-	// If enable, multi-threaded compilation will be used when compiling this models.
-	// For some memory-heavy models, this may need to be disabled, since it may multiply the
-	// use of memory during compilation for each thread.
-	UPROPERTY()
-	bool bUseParallelCompilation = true;
-
 	// Use the disk to store intermediate compilation data. This slows down the object compilation
 	// but it may be necessary for huge objects.
 	UPROPERTY()
diff --git a/Engine/Plugins/Experimental/Mutable/Source/CustomizableObject/Private/MuCO/CustomizableObject.cpp b/Engine/Plugins/Experimental/Mutable/Source/CustomizableObject/Private/MuCO/CustomizableObject.cpp
--- a/Engine/Plugins/Experimental/Mutable/Source/CustomizableObject/Private/MuCO/CustomizableObject.cpp
+++ b/Engine/Plugins/Experimental/Mutable/Source/CustomizableObject/Private/MuCO/CustomizableObject.cpp
@@ -714,9 +714,6 @@ void UCustomizableObject::CompileForTargetPlatform(const ITargetPlatform* Target
 		Options.TargetPlatform = TargetPlatform;
 		Options.CustomizableObjectNumBoneInfluences = ICustomizableObjectModule::Get().GetNumBoneInfluences();
 
-		// If this is enabled, there are determinism problems. Disable it for packaging.
-		Options.bUseParallelCompilation = false;
-
 		Compiler->Compile(*this, Options, false);
 	}
 
diff --git a/Engine/Plugins/Experimental/Mutable/Source/CustomizableObjectEditor/Private/MuCOE/CustomizableObjectCompiler.cpp b/Engine/Plugins/Experimental/Mutable/Source/CustomizableObjectEditor/Private/MuCOE/CustomizableObjectCompiler.cpp
--- a/Engine/Plugins/Experimental/Mutable/Source/CustomizableObjectEditor/Private/MuCOE/CustomizableObjectCompiler.cpp
+++ b/Engine/Plugins/Experimental/Mutable/Source/CustomizableObjectEditor/Private/MuCOE/CustomizableObjectCompiler.cpp
@@ -924,6 +924,7 @@ void FCustomizableObjectCompiler::CompileInternal(UCustomizableObject* Object, c
 	Options.bSkinWeightProfilesEnabled = Object->bEnableAltSkinWeightProfiles;
 	Options.bPhysicsAssetMergeEnabled = Object->bEnablePhysicsAssetMerge;
 	Options.bAnimBpPhysicsManipulationEnabled = Object->bEnableAnimBpPhysicsAssetsManipualtion;
+	Options.ImageTiling = Object->CompileOptions.ImageTiling;
 
 	if (Object->IsLocked() || !UCustomizableObjectSystem::GetInstance()->LockObject(Object))
 	{
diff --git a/Engine/Plugins/Experimental/Mutable/Source/CustomizableObjectEditor/Private/MuCOE/CustomizableObjectEditor.cpp b/Engine/Plugins/Experimental/Mutable/Source/CustomizableObjectEditor/Private/MuCOE/CustomizableObjectEditor.cpp
--- a/Engine/Plugins/Experimental/Mutable/Source/CustomizableObjectEditor/Private/MuCOE/CustomizableObjectEditor.cpp
+++ b/Engine/Plugins/Experimental/Mutable/Source/CustomizableObjectEditor/Private/MuCOE/CustomizableObjectEditor.cpp
@@ -814,12 +814,6 @@ void FCustomizableObjectEditor::BindCommands()
 		FCanExecuteAction(),
 		FIsActionChecked::CreateSP(this, &FCustomizableObjectEditor::CompileOptions_TextureCompression_IsChecked));
 
-	ToolkitCommands->MapAction(
-		Commands.CompileOptions_UseParallelCompilation,
-		FExecuteAction::CreateSP(this, &FCustomizableObjectEditor::CompileOptions_UseParallelCompilation_Toggled),
-		FCanExecuteAction(),
-		FIsActionChecked::CreateSP(this, &FCustomizableObjectEditor::CompileOptions_UseParallelCompilation_IsChecked));
-
 	ToolkitCommands->MapAction(
 		Commands.CompileOptions_UseDiskCompilation,
 		FExecuteAction::CreateSP(this, &FCustomizableObjectEditor::CompileOptions_UseDiskCompilation_Toggled),
@@ -1069,7 +1063,6 @@ TSharedRef<SWidget> FCustomizableObjectEditor::GenerateCompileOptionsMenuContent
 				});
 		MenuBuilder.AddWidget(CompileTilingCombo.ToSharedRef(), LOCTEXT("MutableCompileImageTiling", "Image Tiling"));
 
-		MenuBuilder.AddMenuEntry(FCustomizableObjectEditorCommands::Get().CompileOptions_UseParallelCompilation);
 		MenuBuilder.AddMenuEntry(FCustomizableObjectEditorCommands::Get().CompileOptions_UseDiskCompilation);
 		MenuBuilder.AddMenuEntry(FCustomizableObjectEditorCommands::Get().CompileOptions_EnableTextureCompression);
 	}
@@ -1350,20 +1343,6 @@ void FCustomizableObjectEditor::OnChangeCompileOptimizationLevel(TSharedPtr<FStr
 }
 
 
-void FCustomizableObjectEditor::CompileOptions_UseParallelCompilation_Toggled()
-{
-	const FScopedTransaction Transaction(LOCTEXT("ChangedEnableCompilingInMultipleThreadsTransaction", "Changed Enable compiling in multiple threads"));
-	CustomizableObject->Modify();
-	CustomizableObject->CompileOptions.bUseParallelCompilation = !CustomizableObject->CompileOptions.bUseParallelCompilation;
-}
-
-
-bool FCustomizableObjectEditor::CompileOptions_UseParallelCompilation_IsChecked()
-{
-	return CustomizableObject ? CustomizableObject->CompileOptions.bUseParallelCompilation : false;
-}
-
-
 void FCustomizableObjectEditor::CompileOptions_UseDiskCompilation_Toggled()
 {
 	const FScopedTransaction Transaction(LOCTEXT("ChangedEnableCompilingUsingTheDiskAsMemory", "Changed Enable compiling using the disk as memory"));
diff --git a/Engine/Plugins/Experimental/Mutable/Source/CustomizableObjectEditor/Private/MuCOE/CustomizableObjectEditor.h b/Engine/Plugins/Experimental/Mutable/Source/CustomizableObjectEditor/Private/MuCOE/CustomizableObjectEditor.h
--- a/Engine/Plugins/Experimental/Mutable/Source/CustomizableObjectEditor/Private/MuCOE/CustomizableObjectEditor.h
+++ b/Engine/Plugins/Experimental/Mutable/Source/CustomizableObjectEditor/Private/MuCOE/CustomizableObjectEditor.h
@@ -287,8 +287,6 @@ private:
 	TArray< TSharedPtr<FString> > CompileOptimizationStrings;
 	TSharedPtr<SNumericDropDown<float>> CompileTilingCombo;
 
-	void CompileOptions_UseParallelCompilation_Toggled();
-	bool CompileOptions_UseParallelCompilation_IsChecked();
 	void CompileOptions_UseDiskCompilation_Toggled();
 	bool CompileOptions_UseDiskCompilation_IsChecked();
 	void CompileOptions_TextureCompression_Toggled();
diff --git a/Engine/Plugins/Experimental/Mutable/Source/CustomizableObjectEditor/Private/MuCOE/CustomizableObjectEditorActions.cpp b/Engine/Plugins/Experimental/Mutable/Source/CustomizableObjectEditor/Private/MuCOE/CustomizableObjectEditorActions.cpp
--- a/Engine/Plugins/Experimental/Mutable/Source/CustomizableObjectEditor/Private/MuCOE/CustomizableObjectEditorActions.cpp
+++ b/Engine/Plugins/Experimental/Mutable/Source/CustomizableObjectEditor/Private/MuCOE/CustomizableObjectEditorActions.cpp
@@ -25,7 +25,6 @@ void FCustomizableObjectEditorCommands::RegisterCommands()
 	UI_COMMAND(CompileOnlySelected, "Compile Only Selected", "Compile the source graph of the customizable object and update the previews, only for the selected options in the preview. The rest of options are discarded. If they are selected, press again this button to see their effect in the preview.", EUserInterfaceActionType::Button, FInputChord());
 	UI_COMMAND(ResetCompileOptions, "Reset Compilation Options", "Set reasonable defaults for the compilation options.", EUserInterfaceActionType::Button, FInputChord());
 	UI_COMMAND(CompileOptions_EnableTextureCompression, "Enable texture compression.", "Only for debug. Do not use.", EUserInterfaceActionType::ToggleButton, FInputChord());
-	UI_COMMAND(CompileOptions_UseParallelCompilation, "Enable compiling in multiple threads.", "This is faster but use more memory.", EUserInterfaceActionType::ToggleButton, FInputChord());
 	UI_COMMAND(CompileOptions_UseDiskCompilation, "Enable compiling using the disk as memory.", "This is very slow but supports compiling huge objects. It requires a lot of free space in the OS disk.", EUserInterfaceActionType::ToggleButton, FInputChord());
 
 	UI_COMMAND(Debug, "Debug", "Debug the object.", EUserInterfaceActionType::Button, FInputChord());
diff --git a/Engine/Plugins/Experimental/Mutable/Source/CustomizableObjectEditor/Private/MuCOE/SMutableGraphViewer.cpp b/Engine/Plugins/Experimental/Mutable/Source/CustomizableObjectEditor/Private/MuCOE/SMutableGraphViewer.cpp
--- a/Engine/Plugins/Experimental/Mutable/Source/CustomizableObjectEditor/Private/MuCOE/SMutableGraphViewer.cpp
+++ b/Engine/Plugins/Experimental/Mutable/Source/CustomizableObjectEditor/Private/MuCOE/SMutableGraphViewer.cpp
@@ -318,19 +318,6 @@ TSharedRef<SWidget> SMutableGraphViewer::GenerateCompileOptionsMenuContent()
 			;
 		MenuBuilder.AddWidget(CompileTilingCombo.ToSharedRef(), LOCTEXT("MutableCompileImageTiling", "Image Tiling"));
 
-		// Parallel compilation
-		MenuBuilder.AddMenuEntry(
-			LOCTEXT("MutableEnableParallelCompilation", "Enable compiling in multiple threads."),
-			LOCTEXT("MutableEnableParallelCompilationTooltip", "This is faster but use more memory."),
-			FSlateIcon(),
-			FUIAction(
-				FExecuteAction::CreateLambda([this]() { CompileOptions.bUseParallelCompilation = !CompileOptions.bUseParallelCompilation; }),
-				FCanExecuteAction(),
-				FIsActionChecked::CreateLambda([this]() { return CompileOptions.bUseParallelCompilation; })),
-			NAME_None,
-			EUserInterfaceActionType::ToggleButton
-		);
-
 		// Disk as cache
 		MenuBuilder.AddMenuEntry(
 			LOCTEXT("MutableDiskAsMemory", "Enable compiling using the disk as memory."),
diff --git a/Engine/Plugins/Experimental/Mutable/Source/CustomizableObjectEditor/Private/MuCOE/SMutableObjectViewer.cpp b/Engine/Plugins/Experimental/Mutable/Source/CustomizableObjectEditor/Private/MuCOE/SMutableObjectViewer.cpp
--- a/Engine/Plugins/Experimental/Mutable/Source/CustomizableObjectEditor/Private/MuCOE/SMutableObjectViewer.cpp
+++ b/Engine/Plugins/Experimental/Mutable/Source/CustomizableObjectEditor/Private/MuCOE/SMutableObjectViewer.cpp
@@ -389,19 +389,6 @@ TSharedRef<SWidget> SMutableObjectViewer::GenerateCompileOptionsMenuContent()
 			;
 		MenuBuilder.AddWidget(CompileTilingCombo.ToSharedRef(), LOCTEXT("MutableCompileImageTiling", "Image Tiling"));
 
-		// Parallel compilation
-		MenuBuilder.AddMenuEntry(
-			LOCTEXT("Generate_MutableEnableParallelCompilation", "Enable compiling in multiple threads."),
-			LOCTEXT("Generate_MutableEnableParallelCompilationTooltip", "This is faster but use more memory."),
-			FSlateIcon(),
-			FUIAction(
-				FExecuteAction::CreateLambda([this]() { CompileOptions.bUseParallelCompilation = !CompileOptions.bUseParallelCompilation; }),
-				FCanExecuteAction(),
-				FIsActionChecked::CreateLambda([this]() { return CompileOptions.bUseParallelCompilation; })),
-			NAME_None,
-			EUserInterfaceActionType::ToggleButton
-		);
-
 		// Disk as cache
 		MenuBuilder.AddMenuEntry(
 			LOCTEXT("Generate_MutableUseDisk", "Enable compiling using the disk as memory."),
diff --git a/Engine/Plugins/Experimental/Mutable/Source/CustomizableObjectEditor/Public/MuCOE/CustomizableObjectEditorActions.h b/Engine/Plugins/Experimental/Mutable/Source/CustomizableObjectEditor/Public/MuCOE/CustomizableObjectEditorActions.h
--- a/Engine/Plugins/Experimental/Mutable/Source/CustomizableObjectEditor/Public/MuCOE/CustomizableObjectEditorActions.h
+++ b/Engine/Plugins/Experimental/Mutable/Source/CustomizableObjectEditor/Public/MuCOE/CustomizableObjectEditorActions.h
@@ -19,7 +19,6 @@ public:
 	TSharedPtr< FUICommandInfo > CompileOnlySelected;
 	TSharedPtr< FUICommandInfo > ResetCompileOptions;
 	TSharedPtr< FUICommandInfo > CompileOptions_EnableTextureCompression;
-	TSharedPtr< FUICommandInfo > CompileOptions_UseParallelCompilation;
 	TSharedPtr< FUICommandInfo > CompileOptions_UseDiskCompilation;
 
 	TSharedPtr< FUICommandInfo > Debug;
