[Mutable] Fix posible null dereferece in ClipMorphPlane, GeometryOperation, BindShape, ApplyShape and MorphReshape ops when computing the memory budget to reserve.

#rnx

diff --git a/Engine/Plugins/Experimental/Mutable/Source/MutableRuntime/Private/MuR/CodeRunner.cpp b/Engine/Plugins/Experimental/Mutable/Source/MutableRuntime/Private/MuR/CodeRunner.cpp
--- a/Engine/Plugins/Experimental/Mutable/Source/MutableRuntime/Private/MuR/CodeRunner.cpp
+++ b/Engine/Plugins/Experimental/Mutable/Source/MutableRuntime/Private/MuR/CodeRunner.cpp
@@ -1987,7 +1987,7 @@ namespace mu
                     // Should be None or an axis aligned box
                     const FShape& selectionShape = pModel->GetPrivate()->m_program.m_constantShapes[args.vertexSelectionShapeOrBone];
 
-					Ptr<Mesh> Result = CreateMesh(Source->GetDataSize());
+					Ptr<Mesh> Result = CreateMesh(Source ? Source->GetDataSize() : 0);
 
 					bool bOutSuccess = false;
 					MeshClipMorphPlane(Result.get(), Source.get(), origin, normal, args.dist, args.factor, morphShape.size[0], morphShape.size[1], morphShape.size[2], selectionShape, bOutSuccess, mu::string(), -1);
@@ -2034,7 +2034,7 @@ namespace mu
                     FShape selectionShape;
                     selectionShape.type = (uint8)FShape::Type::None;
 
-					Ptr<Mesh> Result = CreateMesh(Source->GetDataSize());
+					Ptr<Mesh> Result = CreateMesh(Source ? Source->GetDataSize() : 0);
 
 					bool bOutSuccess = false;
 					MeshClipMorphPlane(Result.get(), Source.get(), origin, normal, args.dist, args.factor, morphShape.size[0], morphShape.size[1], morphShape.size[2], selectionShape, bOutSuccess, mu::string(), -1.0f);
@@ -2213,7 +2213,7 @@ namespace mu
                 Ptr<const Mesh> pReference = LoadMesh(FCacheAddress(args.reference,item));
 
                 // Only if both are valid.
-                MeshPtr Result = CreateMesh(Source->GetDataSize());
+                MeshPtr Result = CreateMesh(Source ? Source->GetDataSize() : 0);
                 
 				bool bOutSuccess = false;
 				if (Source && pReference)
@@ -2337,7 +2337,7 @@ namespace mu
 				float ScalarA = LoadScalar(FCacheAddress(args.scalarA, item));
 				float ScalarB = LoadScalar(FCacheAddress(args.scalarB, item));
 
-				Ptr<Mesh> Result = CreateMesh(MeshA->GetDataSize());
+				Ptr<Mesh> Result = CreateMesh(MeshA ? MeshA->GetDataSize() : 0);
 
 				bool bOutSuccess = false;
 				MeshGeometryOperation(Result.get(), MeshA.get(), MeshB.get(), ScalarA, ScalarB, bOutSuccess);
@@ -2463,7 +2463,7 @@ namespace mu
 				}	
 				else
 				{
-					Ptr<Mesh> Result = CreateMesh(BaseMesh->GetDataSize());
+					Ptr<Mesh> Result = CreateMesh(BaseMesh ? BaseMesh->GetDataSize() : 0);
 
 					bool bOutSuccess = false;
 					MeshBindShapeClipDeform(Result.get(), BaseMesh.get(), Shape.get(), BindingMethod, bOutSuccess);
@@ -2521,7 +2521,7 @@ namespace mu
 				const EMeshBindShapeFlags ReshapeFlags = static_cast<EMeshBindShapeFlags>(args.flags);
 				const bool bReshapeVertices = EnumHasAnyFlags(ReshapeFlags, EMeshBindShapeFlags::ReshapeVertices);
 
-				Ptr<Mesh> ReshapedMeshResult = CreateMesh(BaseMesh->GetDataSize());
+				Ptr<Mesh> ReshapedMeshResult = CreateMesh(BaseMesh ? BaseMesh->GetDataSize() : 0);
 
 				bool bOutSuccess = false;
 				MeshApplyShape(ReshapedMeshResult.get(), BaseMesh.get(), Shape.get(), ReshapeFlags, bOutSuccess);
@@ -2597,7 +2597,7 @@ namespace mu
 				Ptr<const Mesh> MorphedMesh = LoadMesh(FCacheAddress(Args.Morph, item));
 				Ptr<const Mesh> ReshapeMesh = LoadMesh(FCacheAddress(Args.Reshape, item));
 
-				if (ReshapeMesh)
+				if (ReshapeMesh && MorphedMesh)
 				{
 					// Copy without Skeleton, Physics or Poses 
 					EMeshCopyFlags CopyFlags = ~(
