[Mutable] Added Reroute Pins (also known as Knot Nodes).

- Fixed issue where FollowInput/Output functions could hang due to a cycle.
- Changed the MaterialAsset type color from White to Cian (White is now used by Wildcard in line with other UE Graph Editors).
- Fixed issue with Material Texture Pin icons.
- Fixed where nodes with parent Material Node where not finding the parent behind an Import/Export node.

#rnx

diff --git a/Engine/Plugins/Experimental/Mutable/Source/CustomizableObjectEditor/Classes/MuCOE/EdGraphSchema_CustomizableObject.h b/Engine/Plugins/Experimental/Mutable/Source/CustomizableObjectEditor/Classes/MuCOE/EdGraphSchema_CustomizableObject.h
--- a/Engine/Plugins/Experimental/Mutable/Source/CustomizableObjectEditor/Classes/MuCOE/EdGraphSchema_CustomizableObject.h
+++ b/Engine/Plugins/Experimental/Mutable/Source/CustomizableObjectEditor/Classes/MuCOE/EdGraphSchema_CustomizableObject.h
@@ -104,6 +104,7 @@ public:
 	static const FName PC_Enum;
 	static const FName PC_Stack;
 	static const FName PC_MaterialAsset;
+	static const FName PC_Wildcard;
 
 private:
 	/** List of all node types. */
diff --git a/Engine/Plugins/Experimental/Mutable/Source/CustomizableObjectEditor/Classes/MuCOE/Nodes/CustomizableObjectNodeReroute.h b/Engine/Plugins/Experimental/Mutable/Source/CustomizableObjectEditor/Classes/MuCOE/Nodes/CustomizableObjectNodeReroute.h
new file mode 100644
--- /dev/null
+++ b/Engine/Plugins/Experimental/Mutable/Source/CustomizableObjectEditor/Classes/MuCOE/Nodes/CustomizableObjectNodeReroute.h
@@ -0,0 +1,47 @@
+// Copyright Epic Games, Inc. All Rights Reserved.
+
+#pragma once
+
+#include "MuCOE/Nodes/CustomizableObjectNode.h"
+
+#include "CustomizableObjectNodeReroute.generated.h"
+
+
+/** Code based on UK2Node_Knot. */
+UCLASS()
+class UCustomizableObjectNodeReroute : public UCustomizableObjectNode
+{
+	GENERATED_BODY()
+
+public:
+	// UEdGraphNode interface
+	virtual FText GetTooltipText() const override;
+	virtual FText GetNodeTitle(ENodeTitleType::Type TitleType) const override;
+	virtual bool ShouldOverridePinNames() const override;
+	virtual FText GetPinNameOverride(const UEdGraphPin& Pin) const override;
+	virtual void OnRenameNode(const FString& NewName) override;
+	virtual bool ShouldDrawNodeAsControlPointOnly(int32& OutInputPinIndex, int32& OutOutputPinIndex) const override;
+	virtual TSharedPtr<SGraphNode> CreateVisualWidget() override;
+	void NotifyPinConnectionListChanged(UEdGraphPin* Pin);
+	virtual void PinConnectionListChanged(UEdGraphPin* Pin) override;
+
+	// UCustomizableObjectNode interface
+	virtual void AllocateDefaultPins(UCustomizableObjectNodeRemapPins* RemapPins) override;
+	virtual void ReconstructNode(UCustomizableObjectNodeRemapPins* RemapPinsMode) override;
+	virtual bool IsSingleOutputNode() const override;
+
+	// Own interface
+	UEdGraphPin* GetInputPin() const;
+
+	UEdGraphPin* GetOutputPin() const;
+	
+private:
+	/** Code based on UK2Node_Knot::PropagatePinType. */
+	void PropagatePinType();
+
+	/** Code based on UK2Node_Knot::PropagatePinTypeFromDirection. */
+	void PropagatePinTypeFromDirection(bool bFromInput);
+	
+	/** Code based on UK2Node_Knot. Recursion guard boolean to prevent PropagatePinType from hanging if there is a cycle of reroute nodes. */
+	bool bRecursionGuard = false;
+};
\ No newline at end of file
diff --git a/Engine/Plugins/Experimental/Mutable/Source/CustomizableObjectEditor/Private/MuCOE/EdGraphSchema_CustomizableObject.cpp b/Engine/Plugins/Experimental/Mutable/Source/CustomizableObjectEditor/Private/MuCOE/EdGraphSchema_CustomizableObject.cpp
--- a/Engine/Plugins/Experimental/Mutable/Source/CustomizableObjectEditor/Private/MuCOE/EdGraphSchema_CustomizableObject.cpp
+++ b/Engine/Plugins/Experimental/Mutable/Source/CustomizableObjectEditor/Private/MuCOE/EdGraphSchema_CustomizableObject.cpp
@@ -74,6 +74,7 @@
 #include "Settings/EditorStyleSettings.h"
 #include "ToolMenu.h"
 #include "ToolMenuSection.h"
+#include "MuCOE/Nodes/CustomizableObjectNodeReroute.h"
 #include "Toolkits/ToolkitManager.h"
 
 class IToolkit;
@@ -224,6 +225,7 @@ const FName UEdGraphSchema_CustomizableObject::PC_Bool("bool");
 const FName UEdGraphSchema_CustomizableObject::PC_Enum("enum");
 const FName UEdGraphSchema_CustomizableObject::PC_Stack("stack");
 const FName UEdGraphSchema_CustomizableObject::PC_MaterialAsset("materialAsset");
+const FName UEdGraphSchema_CustomizableObject::PC_Wildcard("wildcard");
 
 
 UEdGraphSchema_CustomizableObject::UEdGraphSchema_CustomizableObject()
@@ -327,7 +329,7 @@ namespace
 		TemplateNode->BeginConstruct();
 		TemplateNode->ReconstructNode();
 		
-		for (const UEdGraphPin* Pin : TemplateNode->GetAllPins())
+		for (const UEdGraphPin* Pin : TemplateNode->GetAllNonOrphanPins())
 		{
 			const UEdGraphPin* InputPin = nullptr;
 			const UEdGraphPin* OutputPin = nullptr;
@@ -604,6 +606,12 @@ void UEdGraphSchema_CustomizableObject::GetGraphContextActions(FGraphContextMenu
 		Action->NodeTemplate = Node;
 	}
 
+	{
+		UCustomizableObjectNodeReroute* Node = NewObject<UCustomizableObjectNodeReroute>();
+		TSharedPtr<FCustomizableObjectSchemaAction_NewNode> Action = AddNewNodeAction(ContextMenuBuilder, FString(), Node->GetNodeTitle(ENodeTitleType::ListView), FText(), 1);
+		Action->NodeTemplate = Node;
+	}
+
 	// Add Paste here if appropriate
 	if (!ContextMenuBuilder.FromPin)
 	{
@@ -690,59 +698,63 @@ FLinearColor UEdGraphSchema_CustomizableObject::GetPinTypeColor(const FName& Typ
 {
 	if (TypeString == PC_Enum)
 	{
-		return FLinearColor(0.357667f, 0.500000f, 0.060000f, 1.000000f);
+		return FLinearColor(0.357667f, 0.500000f, 0.060000f, 1.000000f); // Light green
 	}
 	else if (TypeString == PC_Float)
 	{
-		return FLinearColor(0.357667f, 1.000000f, 0.060000f, 1.000000f);
+		return FLinearColor(0.357667f, 1.000000f, 0.060000f, 1.000000f); // Green
 	}
 	else if (TypeString == PC_Color)
 	{
-		return FLinearColor(1.000000f, 0.591255f, 0.016512f, 1.000000f);
+		return FLinearColor(1.000000f, 0.591255f, 0.016512f, 1.000000f); // Yellow
 	}
 	else if (TypeString == PC_Bool)
 	{
-		return FLinearColor(0.470000f, 0.0f, 0.000000f, 1.000000f);
+		return FLinearColor(0.470000f, 0.0f, 0.000000f, 1.000000f); // Red
 	}
 	else if (TypeString == PC_Projector)
 	{
-		return FLinearColor(1.000000f, 0.500000f, 1.000000f, 0.600000f);
+		return FLinearColor(1.000000f, 0.500000f, 1.000000f, 0.600000f); // Light pink
 	}
 	else if (TypeString == PC_GroupProjector)
 	{
-		return FLinearColor(1.000000f, 0.172585f, 0.000000f, 1.000000f);
+		return FLinearColor(1.000000f, 0.172585f, 0.000000f, 1.000000f); // Orange
 	}
 	else if (TypeString == PC_Mesh)
 	{
-		return FLinearColor(0.100000f, 0.000000f, 0.500000f, 1.000000f);
+		return FLinearColor(0.100000f, 0.000000f, 0.500000f, 1.000000f); // Purple
 	}
 	else if (TypeString == PC_Layout)
 	{
-		return FLinearColor(0.500000f, 0.500000f, 0.100000f, 1.000000f);
+		return FLinearColor(0.500000f, 0.500000f, 0.100000f, 1.000000f); // Light yellow
 	}
 	else if (TypeString == PC_Image)
 	{
-		return FLinearColor(0.353393f, 0.454175f, 1.000000f, 1.000000f);
+		return FLinearColor(0.353393f, 0.454175f, 1.000000f, 1.000000f); // Grey
 	}
 	else if (TypeString == PC_PassThroughImage)
 	{
-		return FLinearColor(0.353393f, 0.454175f, 0.353393f, 1.000000f);
+		return FLinearColor(0.353393f, 0.454175f, 0.353393f, 1.000000f); // Dark green
 	}
 	else if (TypeString == PC_Material)
 	{
-		return FLinearColor(0.000000f, 0.100000f, 0.600000f, 1.000000f);
+		return FLinearColor(0.000000f, 0.100000f, 0.600000f, 1.000000f); // Blue
 	}
 	else if (TypeString == PC_Object)
 	{
-		return FLinearColor(0.000000f, 0.400000f, 0.910000f, 1.000000f);
+		return FLinearColor(0.000000f, 0.400000f, 0.910000f, 1.000000f); // Light blue
 	}
 	else if (TypeString == PC_Stack)
 	{
-		return FLinearColor(1.000000f, 0.000000f, 0.800000f, 1.000000f);
+		return FLinearColor(1.000000f, 0.000000f, 0.800000f, 1.000000f); // Pink
 	}
 	else if (TypeString == PC_MaterialAsset)
 	{
-		return FLinearColor(1.000000f, 1.000000f, 1.000000f, 1.000000f);
+		return FLinearColor(0.000000f, 1.000000f, 0.1000f, 1.000000f); // Cian
+	}
+	else if (TypeString == PC_Wildcard)
+	{
+		return FLinearColor(1.000000f, 1.000000f, 1.000000f, 1.000000f); // White
 	}
 
 	for (const FRegisteredCustomizableObjectPinType& PinType : ICustomizableObjectModule::Get().GetExtendedPinTypes())
@@ -1268,6 +1280,10 @@ FText UEdGraphSchema_CustomizableObject::GetPinCategoryName(const FName& PinCate
 	{
 		return LOCTEXT("Material_Asset_Pin_Category", "materialAsset");
 	}
+	else if (PinCategory == UEdGraphSchema_CustomizableObject::PC_MaterialAsset)
+	{
+		return LOCTEXT("Wildcard_Pin_Category", "Wildcard");
+	}
 	else
 	{
 		for (const FRegisteredCustomizableObjectPinType& PinType : ICustomizableObjectModule::Get().GetExtendedPinTypes())
diff --git a/Engine/Plugins/Experimental/Mutable/Source/CustomizableObjectEditor/Private/MuCOE/GraphTraversal.cpp b/Engine/Plugins/Experimental/Mutable/Source/CustomizableObjectEditor/Private/MuCOE/GraphTraversal.cpp
--- a/Engine/Plugins/Experimental/Mutable/Source/CustomizableObjectEditor/Private/MuCOE/GraphTraversal.cpp
+++ b/Engine/Plugins/Experimental/Mutable/Source/CustomizableObjectEditor/Private/MuCOE/GraphTraversal.cpp
@@ -24,86 +24,24 @@
 #include "MuCOE/Nodes/CustomizableObjectNodeTable.h"
 #include "MuCOE/Nodes/CustomizableObjectNodeAnimationPose.h"
 #include "MuCOE/Nodes/CustomizableObjectNodeObjectGroup.h"
+#include "MuCOE/Nodes/CustomizableObjectNodeReroute.h"
 #include "UObject/UObjectIterator.h"
 
 
-/** Function to follow an External Pin node.
- *
- * @param Node Possible node to follow.
- * @param PinsToVisit Discovered pins (output parameter).
- * @return True if the node was a External Pin node (node followed).
- */
-bool CastExternalPinNode(const UEdGraphNode& Node, TArray<const UEdGraphPin*>& PinsToVisit)
-{
-	const UCustomizableObjectNodeExternalPin* ExternalPinNode = Cast<UCustomizableObjectNodeExternalPin>(&Node);
-	if (ExternalPinNode)
-	{
-		if (const UCustomizableObjectNodeExposePin* ExposePinNode = ExternalPinNode->GetNodeExposePin())
-		{
-			if (const UEdGraphPin* ExposePin = ExposePinNode->InputPin())
-			{
-				PinsToVisit.Add(ExposePin);
-			}
-		}
-	}
-
-	return ExternalPinNode != nullptr;
-}
-
-
-/** Function to follow an Expose Pin node.
- *
- * @param Node Possible node to follow.
- * @param PinsToVisit Discovered pins (output parameter).
- * @return True if the node was a Expose Pin node (node followed).
- */
-bool CastExposePinNode(const UEdGraphNode& Node, TArray<const UEdGraphPin*>& PinsToVisit)
-{
-	const UCustomizableObjectNodeExposePin* ExposePinNode = Cast<UCustomizableObjectNodeExposePin>(&Node);
-	if (ExposePinNode)
-	{
-		const UCustomizableObjectNodeExternalPin* NodeExternalPin = nullptr;
-		for (TObjectIterator<UCustomizableObjectNodeExternalPin> It; It; ++It)
-		{
-			if ((*It)->GetNodeExposePin() == ExposePinNode)
-			{
-				NodeExternalPin = *It;
-				break;
-			}
-		}
-		
-		if (NodeExternalPin)
-		{
-			if (const UEdGraphPin* ExternalPin = NodeExternalPin->GetExternalPin())
-			{
-				PinsToVisit.Add(ExternalPin);
-			}
-		}
-	}
-
-	return ExposePinNode != nullptr;
-}
-
-
 /** Follow the given pin returning its connected pin.
  *
  * - Skips all orphan pins.
- * - Follows External Pin nodes.
+ * - Follows External Pin and Reroute nodes.
  *
  * @param Pin Pin to follow.
- * @param CycleDetected If provided, it will set to true if a cycle has been found.
- * @param CastNode Function to follow External Pin nodes (following an output pin) or Expose Pin nodes (following an input pin).
- * @param Direction Direction of the pin to explore.
- */
-TArray<UEdGraphPin*> FollowPinArray(const UEdGraphPin& Pin, bool* CycleDetected, const EEdGraphPinDirection Direction, const TFunction<bool(const UEdGraphNode&, TArray<const UEdGraphPin*>&)> CastNode)
+ * @param bOutCycleDetected If provided, it will set to true if a cycle has been found.
+ * @param Direction Direction of the pin to explore.  */
+TArray<UEdGraphPin*> FollowPinArray(const UEdGraphPin& Pin, bool* bOutCycleDetected, const EEdGraphPinDirection Direction)
 {
 	check(Pin.Direction == Direction); // Can only follow input pins. To follow output pins see FollowOutputPin, but be aware of its limitations!
 
-	if (CycleDetected)
-	{
-		*CycleDetected = false;
-	}
-	
+	bool bCycleDetected = false;
+		
 	TArray<UEdGraphPin*> Result;
 	
 	TSet<const UEdGraphPin*> Visited;
@@ -119,32 +57,76 @@ TArray<UEdGraphPin*> FollowPinArray(const UEdGraphPin& Pin, bool* CycleDetected,
 			continue;
 		}
 
-		Visited.FindOrAdd(&CurrentPin, CycleDetected);
-		if (CycleDetected && *CycleDetected)
+		Visited.FindOrAdd(&CurrentPin, &bCycleDetected);
+		if (bCycleDetected)
 		{
 			continue;
 		}
 
 		for (UEdGraphPin* LinkedPin : CurrentPin.LinkedTo)
 		{
-			if (!IsPinOrphan(*LinkedPin))
+			if (IsPinOrphan(*LinkedPin))
 			{
-				if (!CastNode(*LinkedPin->GetOwningNode(), PinsToVisit))
+				continue;
+			}
+			
+			if (const UCustomizableObjectNodeExposePin* ExposePinNode = Cast<UCustomizableObjectNodeExposePin>(LinkedPin->GetOwningNode()))
+			{
+				check(Direction == EGPD_Output);
+
+				const UCustomizableObjectNodeExternalPin* LinkedNode = nullptr;
+				for (TObjectIterator<UCustomizableObjectNodeExternalPin> It; It; ++It)
 				{
-					Result.Add(LinkedPin);
+					if ((*It)->GetNodeExposePin() == ExposePinNode)
+					{
+						LinkedNode = *It;
+						break;
+					}
+				}
+	
+				if (LinkedNode)
+				{
+					if (const UEdGraphPin* ExternalPin = LinkedNode->GetExternalPin())
+					{
+						PinsToVisit.Add(ExternalPin);
+					}
 				}
 			}
+			else if (const UCustomizableObjectNodeExternalPin* ExternalPinNode = Cast<UCustomizableObjectNodeExternalPin>(LinkedPin->GetOwningNode()))
+			{
+				check(Direction == EGPD_Input);
+				
+				if (const UCustomizableObjectNodeExposePin* LinkedNode = ExternalPinNode->GetNodeExposePin())
+				{
+					if (const UEdGraphPin* ExposePin = LinkedNode->InputPin())
+					{
+						PinsToVisit.Add(ExposePin);
+					}
+				}
+			}
+			else if (const UCustomizableObjectNodeReroute* NodeReroute = Cast<UCustomizableObjectNodeReroute>(LinkedPin->GetOwningNode()))
+			{
+				PinsToVisit.Add(Direction == EGPD_Input ? NodeReroute->GetInputPin() : NodeReroute->GetOutputPin());
+			}
+			else
+			{
+				Result.Add(LinkedPin);
+			}
 		}
-		
+	}
+
+	if (bOutCycleDetected)
+	{
+		*bOutCycleDetected = bCycleDetected;	
 	}
 
 	return Result;
 }
 
 
-TArray<UEdGraphPin*> FollowInputPinArray(const UEdGraphPin& Pin, bool* CycleDetected)
+TArray<UEdGraphPin*> FollowInputPinArray(const UEdGraphPin& Pin, bool* bOutCycleDetected)
 {
-	return FollowPinArray(Pin, CycleDetected, EGPD_Input, &CastExternalPinNode);
+	return FollowPinArray(Pin, bOutCycleDetected, EGPD_Input);
 }
 
 
@@ -164,9 +146,9 @@ UEdGraphPin* FollowInputPin(const UEdGraphPin& Pin, bool* CycleDetected)
 }
 
 
-TArray<UEdGraphPin*> FollowOutputPinArray(const UEdGraphPin& Pin, bool* CycleDetected)
+TArray<UEdGraphPin*> FollowOutputPinArray(const UEdGraphPin& Pin, bool* bOutCycleDetected)
 {
-	return FollowPinArray(Pin, CycleDetected, EGPD_Output, &CastExposePinNode);
+	return FollowPinArray(Pin, bOutCycleDetected, EGPD_Output);
 }
 
 
diff --git a/Engine/Plugins/Experimental/Mutable/Source/CustomizableObjectEditor/Private/MuCOE/Nodes/CustomizableObjectNode.cpp b/Engine/Plugins/Experimental/Mutable/Source/CustomizableObjectEditor/Private/MuCOE/Nodes/CustomizableObjectNode.cpp
--- a/Engine/Plugins/Experimental/Mutable/Source/CustomizableObjectEditor/Private/MuCOE/Nodes/CustomizableObjectNode.cpp
+++ b/Engine/Plugins/Experimental/Mutable/Source/CustomizableObjectEditor/Private/MuCOE/Nodes/CustomizableObjectNode.cpp
@@ -314,20 +314,17 @@ void UCustomizableObjectNode::AutowireNewNode(UEdGraphPin* FromPin)
 
 	const UEdGraphSchema_CustomizableObject* Schema = GetDefault<UEdGraphSchema_CustomizableObject>();
 
-	for ( int PinIndex=0; PinIndex<Pins.Num(); ++PinIndex )
-	{
-		if ( FromPin->PinType.PinCategory == Pins[PinIndex]->PinType.PinCategory 
-			&&
-			FromPin->Direction != Pins[PinIndex]->Direction )
+	for (UEdGraphPin* Pin : GetAllNonOrphanPins())
+	{	
+		UEdGraphNode* OwningNode = FromPin->GetOwningNode(); // TryCreateConnection can reconstruct the node invalidating the FromPin. Get the OwningNode before.
+
+		if (Schema->TryCreateConnection(FromPin, Pin))
 		{
-			UEdGraphNode* OwningNode = FromPin->GetOwningNode(); // TryCreateConnection can reconstruct the node invalidating the FromPin. Save the OwningNode here.
-			if (Schema->TryCreateConnection(FromPin, Pins[PinIndex]))
-			{
-				OwningNode->NodeConnectionListChanged();
-				NodeConnectionListChanged();
-			}
+			OwningNode->NodeConnectionListChanged();
+			NodeConnectionListChanged();
+
 			break;
-		}
+		}		
 	}
 }
 
@@ -404,7 +401,10 @@ bool UCustomizableObjectNode::CanConnect( const UEdGraphPin* InOwnedInputPin, co
 {
 	bOutIsOtherNodeBlocklisted = false;
 
-	bOutArePinsCompatible = InOwnedInputPin->PinType.PinCategory == InOutputPin->PinType.PinCategory;
+	bOutArePinsCompatible = InOwnedInputPin->PinType.PinCategory == InOutputPin->PinType.PinCategory ||
+		InOwnedInputPin->PinType.PinCategory == UEdGraphSchema_CustomizableObject::PC_Wildcard ||
+		InOutputPin->PinType.PinCategory == UEdGraphSchema_CustomizableObject::PC_Wildcard;
+	
 	return bOutArePinsCompatible;
 }
 
@@ -551,11 +551,11 @@ int32 UCustomizableObjectNode::GetLOD() const
 	const UCustomizableObjectNode* CurrentElement;
 	while (PotentialCustomizableNodeObjects.Dequeue(CurrentElement))
 	{
-		for (UEdGraphPin* Pin : CurrentElement->GetAllNonOrphanPins())
+		for (const UEdGraphPin* Pin : CurrentElement->GetAllNonOrphanPins())
 		{
-			if (Pin->Direction == EEdGraphPinDirection::EGPD_Output)
+			if (Pin->Direction == EGPD_Output)
 			{
-				for (UEdGraphPin* LinkedPin : Pin->LinkedTo)
+				for (UEdGraphPin* LinkedPin : FollowOutputPinArray(*Pin))
 				{
 					if (UCustomizableObjectNodeObject* CurrentCustomizableObjectNode = Cast<UCustomizableObjectNodeObject>(LinkedPin->GetOwningNode()))
 					{
diff --git a/Engine/Plugins/Experimental/Mutable/Source/CustomizableObjectEditor/Private/MuCOE/Nodes/CustomizableObjectNodeReroute.cpp b/Engine/Plugins/Experimental/Mutable/Source/CustomizableObjectEditor/Private/MuCOE/Nodes/CustomizableObjectNodeReroute.cpp
new file mode 100644
--- /dev/null
+++ b/Engine/Plugins/Experimental/Mutable/Source/CustomizableObjectEditor/Private/MuCOE/Nodes/CustomizableObjectNodeReroute.cpp
@@ -0,0 +1,217 @@
+// Copyright Epic Games, Inc. All Rights Reserved.
+
+#include "MuCOE/Nodes/CustomizableObjectNodeReroute.h"
+
+#include "SGraphNodeKnot.h"
+#include "MuCOE/EdGraphSchema_CustomizableObject.h"
+
+#define LOCTEXT_NAMESPACE "CustomizableObjectEditor"
+
+
+void UCustomizableObjectNodeReroute::AllocateDefaultPins(UCustomizableObjectNodeRemapPins* RemapPins)
+{
+	const FName InputPinName(TEXT("InputPin"));
+	const FName OutputPinName(TEXT("OutputPin"));
+
+	UEdGraphPin* MyInputPin = CreatePin(EGPD_Input, UEdGraphSchema_CustomizableObject::PC_Wildcard, InputPinName);
+	MyInputPin->bDefaultValueIsIgnored = true;
+
+	CreatePin(EGPD_Output, UEdGraphSchema_CustomizableObject::PC_Wildcard, OutputPinName);
+}
+
+
+void UCustomizableObjectNodeReroute::ReconstructNode(UCustomizableObjectNodeRemapPins* RemapPinsMode)
+{
+	Super::ReconstructNode(RemapPinsMode);
+	PropagatePinType();
+}
+
+
+bool UCustomizableObjectNodeReroute::IsSingleOutputNode() const
+{
+	return true;
+}
+
+
+FText UCustomizableObjectNodeReroute::GetTooltipText() const
+{
+	return LOCTEXT("KnotTooltip", "Reroute Node (reroutes wires)");
+}
+
+FText UCustomizableObjectNodeReroute::GetNodeTitle(ENodeTitleType::Type TitleType) const
+{
+	if (TitleType == ENodeTitleType::EditableTitle)
+	{
+		return FText::FromString(NodeComment);
+	}
+	else if (TitleType == ENodeTitleType::MenuTitle)
+	{
+		return LOCTEXT("KnotListTitle", "Add Reroute Node...");
+	}
+	else
+	{
+		return LOCTEXT("KnotTitle", "Reroute Node");
+	}
+}
+
+
+bool UCustomizableObjectNodeReroute::ShouldOverridePinNames() const
+{
+	return true;
+}
+
+
+FText UCustomizableObjectNodeReroute::GetPinNameOverride(const UEdGraphPin& Pin) const
+{
+	return FText::GetEmpty(); // Keep the pin size tiny
+}
+
+
+void UCustomizableObjectNodeReroute::OnRenameNode(const FString& NewName)
+{
+	NodeComment = NewName;
+}
+
+
+bool UCustomizableObjectNodeReroute::ShouldDrawNodeAsControlPointOnly(int32& OutInputPinIndex, int32& OutOutputPinIndex) const
+{
+	OutInputPinIndex = 0;
+	OutOutputPinIndex = 1;
+	return true;
+}
+
+
+TSharedPtr<SGraphNode> UCustomizableObjectNodeReroute::CreateVisualWidget()
+{
+	return SNew(SGraphNodeKnot, this);
+}
+
+
+void UCustomizableObjectNodeReroute::NotifyPinConnectionListChanged(UEdGraphPin* Pin)
+{
+	PropagatePinType();
+}
+
+
+void UCustomizableObjectNodeReroute::PinConnectionListChanged(UEdGraphPin* Pin)
+{
+	Super::PinConnectionListChanged(Pin);
+	PropagatePinType();
+}
+
+
+UEdGraphPin* UCustomizableObjectNodeReroute::GetInputPin() const
+{
+	return Pins[0];
+}
+
+
+UEdGraphPin* UCustomizableObjectNodeReroute::GetOutputPin() const
+{
+	return Pins[1];
+}
+
+
+void UCustomizableObjectNodeReroute::PropagatePinType()
+{
+	UEdGraphPin* MyInputPin  = GetInputPin();
+	UEdGraphPin* MyOutputPin = GetOutputPin();
+
+	for (const UEdGraphPin* Inputs : MyInputPin->LinkedTo)
+	{
+		if (Inputs->PinType.PinCategory != UEdGraphSchema_CustomizableObject::PC_Wildcard)
+		{
+			PropagatePinTypeFromDirection(true);
+			return;
+		}
+	}
+
+	for (const UEdGraphPin* Outputs : MyOutputPin->LinkedTo)
+	{
+		if (Outputs->PinType.PinCategory != UEdGraphSchema_CustomizableObject::PC_Wildcard)
+		{
+			PropagatePinTypeFromDirection(false);
+			return;
+		}
+	}
+
+	// if all inputs/outputs are wildcards, still favor the inputs first (propagate array/reference/etc. state)
+	if (MyInputPin->LinkedTo.Num() > 0)
+	{
+		// If we can't mirror from output type, we should at least get the type information from the input connection chain
+		PropagatePinTypeFromDirection(true);
+	}
+	else if (MyOutputPin->LinkedTo.Num() > 0)
+	{
+		// Try to mirror from output first to make sure we get appropriate member references
+		PropagatePinTypeFromDirection(false);
+	}
+	else
+	{
+		// Revert to wildcard
+		MyInputPin->BreakAllPinLinks();
+		MyInputPin->PinType.ResetToDefaults();
+		MyInputPin->PinType.PinCategory = UEdGraphSchema_CustomizableObject::PC_Wildcard;
+
+		MyOutputPin->BreakAllPinLinks();
+		MyOutputPin->PinType.ResetToDefaults();
+		MyOutputPin->PinType.PinCategory = UEdGraphSchema_CustomizableObject::PC_Wildcard;
+	}
+}
+
+
+void UCustomizableObjectNodeReroute::PropagatePinTypeFromDirection(bool bFromInput)
+{
+	if (bRecursionGuard)
+	{
+		return;
+	}
+	
+	// Set the type of the pin based on the source connection, and then percolate
+	// that type information up until we no longer reach another Reroute node
+	UEdGraphPin* MySourcePin = bFromInput ? GetInputPin() : GetOutputPin();
+	UEdGraphPin* MyDestinationPin = bFromInput ? GetOutputPin() : GetInputPin();
+
+	TGuardValue<bool> RecursionGuard(bRecursionGuard, true);
+
+	// Make sure any source knot pins compute their type, this will try to call back
+	// into this function but the recursion guard will stop it
+	for (const UEdGraphPin* InPin : MySourcePin->LinkedTo)
+	{
+		if (InPin)
+		{
+			if (UCustomizableObjectNodeReroute * KnotNode = Cast<UCustomizableObjectNodeReroute>(InPin->GetOwningNode()))
+			{
+				KnotNode->PropagatePinTypeFromDirection(bFromInput);
+			}
+		}
+	}
+
+	if (const UEdGraphPin* TypeSource = MySourcePin->LinkedTo.Num() ? MySourcePin->LinkedTo[0] : nullptr)
+	{
+		MySourcePin->PinType = TypeSource->PinType;
+		MySourcePin->PinType.ContainerType = EPinContainerType::None;
+		
+		MyDestinationPin->PinType = TypeSource->PinType;
+		MyDestinationPin->PinType.ContainerType = EPinContainerType::None;
+
+		for (const UEdGraphPin* LinkPin : MyDestinationPin->LinkedTo)
+		{
+			// Order of reconstruction can be such that nulls haven't been cleared out of the destination node's list yet so
+			// must protect against null here
+			if (LinkPin)
+			{
+				UEdGraphNode* OwningNode = LinkPin->GetOwningNode();
+
+				// Notify any pins in the destination direction
+				if (UCustomizableObjectNodeReroute* KnotNode = Cast<UCustomizableObjectNodeReroute>(OwningNode))
+				{
+					KnotNode->PropagatePinTypeFromDirection(bFromInput);
+				}
+			}
+		}
+	}
+}
+
+
+#undef LOCTEXT_NAMESPACE
diff --git a/Engine/Plugins/Experimental/Mutable/Source/CustomizableObjectEditor/Private/MuCOE/Nodes/SCustomizableObjectNodeMaterialPinImage.cpp b/Engine/Plugins/Experimental/Mutable/Source/CustomizableObjectEditor/Private/MuCOE/Nodes/SCustomizableObjectNodeMaterialPinImage.cpp
--- a/Engine/Plugins/Experimental/Mutable/Source/CustomizableObjectEditor/Private/MuCOE/Nodes/SCustomizableObjectNodeMaterialPinImage.cpp
+++ b/Engine/Plugins/Experimental/Mutable/Source/CustomizableObjectEditor/Private/MuCOE/Nodes/SCustomizableObjectNodeMaterialPinImage.cpp
@@ -50,7 +50,7 @@ const FSlateBrush* SCustomizableObjectNodeMaterialPinImage::GetPinIcon() const
 {
 	if (CastChecked<UCustomizableObjectNodeMaterial>(GraphPinObj->GetOwningNode())->IsImageMutableMode(*GraphPinObj))
 	{
-		if (FollowInputPin(*GraphPinObj))
+		if (GraphPinObj->LinkedTo.Num())
 		{
 			return CachedPinMutableConnected;
 		}
diff --git a/Engine/Plugins/Experimental/Mutable/Source/CustomizableObjectEditor/Public/MuCOE/GraphTraversal.h b/Engine/Plugins/Experimental/Mutable/Source/CustomizableObjectEditor/Public/MuCOE/GraphTraversal.h
--- a/Engine/Plugins/Experimental/Mutable/Source/CustomizableObjectEditor/Public/MuCOE/GraphTraversal.h
+++ b/Engine/Plugins/Experimental/Mutable/Source/CustomizableObjectEditor/Public/MuCOE/GraphTraversal.h
@@ -14,13 +14,12 @@ class UEdGraphPin;
  * - Follows External Pin nodes.
  *
  * @param Pin Pin to follow.
- * @param CycleDetected If provided, it will set to true if a cycle has been found.
- * @return Connected pins.
- */
-TArray<UEdGraphPin*> FollowInputPinArray(const UEdGraphPin& Pin, bool* CycleDetected = nullptr);
+ * @param bOutCycleDetected If provided, it will set to true if a cycle has been found.
+ * @return Connected pins. */
+TArray<UEdGraphPin*> FollowInputPinArray(const UEdGraphPin& Pin, bool* bOutCycleDetected = nullptr);
 
 /** Non-array version of FollowInputPinArray. The pin can only have one connected pin. */
-UEdGraphPin* FollowInputPin(const UEdGraphPin& Pin, bool* CycleDetected = nullptr);
+UEdGraphPin* FollowInputPin(const UEdGraphPin& Pin, bool* bOutCycleDetected = nullptr);
 
 /** Follow the given input output returning the input connected pin.
  * 
@@ -29,10 +28,9 @@ UEdGraphPin* FollowInputPin(const UEdGraphPin& Pin, bool* CycleDetected = nullpt
  * - It will only follow External Pin nodes of loaded CO (i.e., Expose Pin nodes of CO which are NOT loaded will not be found)!
  *
  * @param Pin Pin to follow.
- * @param CycleDetected If provided, it will set to true if a cycle has been found.
- * @return Connected pins.
- */
-TArray<UEdGraphPin*> FollowOutputPinArray(const UEdGraphPin& Pin, bool* CycleDetected = nullptr);
+ * @param bOutCycleDetected If provided, it will set to true if a cycle has been found.
+ * @return Connected pins. */
+TArray<UEdGraphPin*> FollowOutputPinArray(const UEdGraphPin& Pin, bool* bOutCycleDetected = nullptr);
 
 /** Non-array version of FollowOutputPinArray. The pin can only have one connected pin. */
 UEdGraphPin* FollowOutputPin(const UEdGraphPin& Pin, bool* CycleDetected = nullptr);
