[Mutable] Fix crash when opening a CO Instance during asset discovery.
Slight clean-up to the CO Instance editor class.

#rnx

diff --git a/Engine/Plugins/Experimental/Mutable/Source/CustomizableObjectEditor/Private/MuCOE/CustomizableObjectInstanceEditor.cpp b/Engine/Plugins/Experimental/Mutable/Source/CustomizableObjectEditor/Private/MuCOE/CustomizableObjectInstanceEditor.cpp
--- a/Engine/Plugins/Experimental/Mutable/Source/CustomizableObjectEditor/Private/MuCOE/CustomizableObjectInstanceEditor.cpp
+++ b/Engine/Plugins/Experimental/Mutable/Source/CustomizableObjectEditor/Private/MuCOE/CustomizableObjectInstanceEditor.cpp
@@ -80,29 +80,11 @@ void FCustomizableObjectInstanceEditor::UnregisterTabSpawners(const TSharedRef<c
 }
 
 
-FCustomizableObjectInstanceEditor::FCustomizableObjectInstanceEditor():
-	CreatePreviewInstanceAfterCOCompile(false)
+FCustomizableObjectInstanceEditor::FCustomizableObjectInstanceEditor()
 {
 	CustomizableObjectInstance = nullptr;
-	//PreviewCustomizableSkeletalComponent = nullptr;
 	PreviewStaticMeshComponent = nullptr;
-	//PreviewSkeletalMeshComponent = nullptr;
-	AssetRegistryLoaded = false;
-	UpdateSkeletalMeshAfterAssetLoaded = false;
-	ObjectToEditCopy = nullptr;
-
-	FAssetRegistryModule& AssetRegistryModule = FModuleManager::LoadModuleChecked<FAssetRegistryModule>(TEXT("AssetRegistry"));
-	if (AssetRegistryModule.Get().IsLoadingAssets())
-	{
-		AssetRegistryModule.Get().OnFilesLoaded().AddRaw(this, &FCustomizableObjectInstanceEditor::OnAssetRegistryLoadComplete);
-	}
-	else
-	{
-		AssetRegistryLoaded = true;
-	}
-
 	HelperCallback = nullptr;
-
 	PoseAsset = nullptr;
 }
 
@@ -127,14 +109,17 @@ FCustomizableObjectInstanceEditor::~FCustomizableObjectInstanceEditor()
 	}
 
 	PreviewCustomizableSkeletalComponents.Reset();
-	PreviewStaticMeshComponent = nullptr;
 	PreviewSkeletalMeshComponents.Reset();
+	PreviewStaticMeshComponent = nullptr;
 
 	CustomizableInstanceDetailsView.Reset();
 	Viewport.Reset();
 
 	FCoreUObjectDelegates::OnObjectModified.Remove(OnObjectModifiedHandle);
-	Compiler->ForceFinishCompilation();
+	if (Compiler)
+	{
+		Compiler->ForceFinishCompilation();
+	}
 
 	FAssetRegistryModule& AssetRegistryModule = FModuleManager::LoadModuleChecked<FAssetRegistryModule>(TEXT("AssetRegistry"));
 	AssetRegistryModule.Get().OnFilesLoaded().RemoveAll(this);
@@ -143,7 +128,7 @@ FCustomizableObjectInstanceEditor::~FCustomizableObjectInstanceEditor()
 }
 
 
-void FCustomizableObjectInstanceEditor::InitCustomizableObjectInstanceEditor( const EToolkitMode::Type Mode, const TSharedPtr< class IToolkitHost >& InitToolkitHost, UCustomizableObjectInstance* ObjectToEdit )
+void FCustomizableObjectInstanceEditor::InitCustomizableObjectInstanceEditor( const EToolkitMode::Type Mode, const TSharedPtr< class IToolkitHost >& InitToolkitHost, UCustomizableObjectInstance* InCustomizableObjectInstance )
 {
 	// Register our commands. This will only register them if not previously registered
 	FCustomizableObjectInstanceEditorCommands::Register();
@@ -163,43 +148,44 @@ void FCustomizableObjectInstanceEditor::InitCustomizableObjectInstanceEditor( co
 	//DetailsViewArgs.bShowActorLabel = false;
 	DetailsViewArgs.bShowObjectLabel = false;
 	CustomizableInstanceDetailsView = PropPlugin.CreateDetailView( DetailsViewArgs );
+	CustomizableInstanceDetailsView->SetObject(InCustomizableObjectInstance, true);
 	//CustomizableInstanceDetailsView->SetIsPropertyVisibleDelegate( FIsPropertyVisible::CreateStatic(&IsPropertyVisible) );
 
 	TextureAnalyzer = SNew(SCustomizableObjecEditorTextureAnalyzer).CustomizableObjectInstanceEditor(this).CustomizableObjectEditor(nullptr);
 
-	Compiler = MakeUnique<FCustomizableObjectCompiler>();
-	check(Compiler);
 	Viewport = SNew(SCustomizableObjectEditorViewportTabBody)
 		.CustomizableObjectEditor(SharedThis(this));
 
-	if (AssetRegistryLoaded)
+	// Set the instance
+	check(InCustomizableObjectInstance);
+	CustomizableObjectInstance = InCustomizableObjectInstance;
+	bOnlyRelevantParameters = InCustomizableObjectInstance->bShowOnlyRelevantParameters;
+	bOnlyRuntimeParameters = InCustomizableObjectInstance->bShowOnlyRuntimeParameters;
+
+	// Check if the asset registry has finished loading assets before compiling the object or updating the instance
+	FAssetRegistryModule& AssetRegistryModule = FModuleManager::LoadModuleChecked<FAssetRegistryModule>(TEXT("AssetRegistry"));
+	if (AssetRegistryModule.Get().IsLoadingAssets())
+	{
+		AssetRegistryModule.Get().OnFilesLoaded().AddRaw(this, &FCustomizableObjectInstanceEditor::OnAssetRegistryLoadComplete);
+		AssetRegistryLoaded = false;
+	}
+	else
 	{
-		if (ObjectToEdit)
+		UCustomizableObject* CustomizableObject = CustomizableObjectInstance->GetCustomizableObject();
+
+		// Compile for the first time if necessary
+		if (CustomizableObject && !CustomizableObject->IsCompiled())
 		{
-			if (UCustomizableObject* CustomizableObject = ObjectToEdit->GetCustomizableObject();
-				CustomizableObject &&
-				!CustomizableObject->IsCompiled())
-			{
-				// Compile for the first time if necessary
-				CompileObject(CustomizableObject);
-				CreatePreviewInstanceAfterCOCompile = true;
-				ObjectToEditCopy = ObjectToEdit;
-			}
+			
+			CompileObject(CustomizableObject);
 		}
-
-		SetEditorMesh(ObjectToEdit);
-		if (!CreatePreviewInstanceAfterCOCompile)
+		else if (CustomizableObject)
 		{
 			CreatePreviewInstance();
 		}
 		
 		Viewport->SetAssetRegistryLoaded(true);
 	}
-	else
-	{
-		ObjectToEditCopy = ObjectToEdit;
-		UpdateSkeletalMeshAfterAssetLoaded = true;
-	}
 
 	FAdvancedPreviewSceneModule& AdvancedPreviewSceneModule = FModuleManager::LoadModuleChecked<FAdvancedPreviewSceneModule>("AdvancedPreviewScene");
 
@@ -232,7 +218,7 @@ void FCustomizableObjectInstanceEditor::InitCustomizableObjectInstanceEditor( co
 
 	const bool bCreateDefaultStandaloneMenu = true;
 	const bool bCreateDefaultToolbar = true;
-	FAssetEditorToolkit::InitAssetEditor( Mode, InitToolkitHost, CustomizableObjectInstanceEditorAppIdentifier, StandaloneDefaultLayout, bCreateDefaultStandaloneMenu, bCreateDefaultToolbar, ObjectToEdit );
+	FAssetEditorToolkit::InitAssetEditor( Mode, InitToolkitHost, CustomizableObjectInstanceEditorAppIdentifier, StandaloneDefaultLayout, bCreateDefaultStandaloneMenu, bCreateDefaultToolbar, CustomizableObjectInstance );
 
 	ExtendToolbar();
 	RegenerateMenusAndToolbars();
@@ -241,15 +227,6 @@ void FCustomizableObjectInstanceEditor::InitCustomizableObjectInstanceEditor( co
 	OnInstancePropertySelectionChanged(NULL);
 
 	OnObjectModifiedHandle = FCoreUObjectDelegates::OnObjectModified.AddRaw(this, &FCustomizableObjectInstanceEditor::OnObjectModified);
-
-	bOnlyRelevantParameters = true;
-	bOnlyRuntimeParameters = true;
-	
-	if (CustomizableObjectInstance)
-	{
-		bOnlyRelevantParameters = CustomizableObjectInstance->bShowOnlyRelevantParameters;
-		bOnlyRuntimeParameters = CustomizableObjectInstance->bShowOnlyRuntimeParameters;
-	}
 }
 
 
@@ -267,37 +244,23 @@ FText FCustomizableObjectInstanceEditor::GetBaseToolkitName() const
 
 void FCustomizableObjectInstanceEditor::CreatePreviewInstance()
 {
+	check(CustomizableObjectInstance);
 	UCustomizableObject* CustomizableObject = CustomizableObjectInstance->GetCustomizableObject();
+	
+	if (!CustomizableObject || (!CustomizableObject->IsCompiled() && !UCustomizableObjectSystem::GetInstance()->IsCompilationDisabled()))
+	{
+		return;
+	}
 
+	// Bind update delegate
 	HelperCallback = NewObject<UUpdateClassWrapperClass>(GetTransientPackage());
 	HelperCallback->Delegate.BindSP(this, &FCustomizableObjectInstanceEditor::OnUpdatePreviewInstance);
 	CustomizableObjectInstance->UpdatedDelegate.AddDynamic(HelperCallback, &UUpdateClassWrapperClass::DelegatedCallback);
 
-	if (!CustomizableObject->IsCompiled() && !UCustomizableObjectSystem::GetInstance()->IsCompilationDisabled()) // In case the compilation failed
-	{
-		return;
-	}
-
 	PreviewStaticMeshComponent = nullptr;
 
-	// Getting the number of mesh components from the root node
-	int32 NumMeshComponents = 0;
-
-	if (CustomizableObject->Source)
-	{
-		TArray<UCustomizableObjectNodeObject*> RootNode;
-		CustomizableObject->Source->GetNodesOfClass< UCustomizableObjectNodeObject>(RootNode);
-
-		for (int32 i = 0; i < RootNode.Num(); ++i)
-		{
-			if (RootNode[i]->bIsBase)
-			{
-				NumMeshComponents = RootNode[i]->NumMeshComponents;
-				break;
-			}
-		}
-	}
-
+	// Create a SkeletalMeshComponent for each component in the CO
+	int32 NumMeshComponents = CustomizableObject->GetComponentCount();
 	PreviewSkeletalMeshComponents.AddZeroed(NumMeshComponents);
 	PreviewCustomizableSkeletalComponents.AddZeroed(NumMeshComponents);
 
@@ -338,39 +301,33 @@ void FCustomizableObjectInstanceEditor::CreatePreviewInstance()
 void FCustomizableObjectInstanceEditor::OnAssetRegistryLoadComplete()
 {
 	AssetRegistryLoaded = true;
-
 	Viewport->SetAssetRegistryLoaded(true);
 
-	if (UpdateSkeletalMeshAfterAssetLoaded)
+	check(CustomizableObjectInstance);
+	UCustomizableObject* CustomizableObject = CustomizableObjectInstance->GetCustomizableObject();
+	if (!CustomizableObject)
 	{
-		if (ObjectToEditCopy && ObjectToEditCopy->GetCustomizableObject() && !ObjectToEditCopy->GetCustomizableObject()->IsCompiled())
-		{
-			// Compile for the first time if necessary
-			CompileObject(ObjectToEditCopy->GetCustomizableObject());
-			CreatePreviewInstanceAfterCOCompile = true;
-		}
-		else
-		{
-			SetEditorMesh(ObjectToEditCopy);
-			CreatePreviewInstance();
-		}
-		UpdateSkeletalMeshAfterAssetLoaded = false;
+		return;
+	}
+
+	// Compile for the first time if necessary
+	if (!CustomizableObject->IsCompiled())
+	{
+		CompileObject(CustomizableObject);
+	}
+	else
+	{
+		CreatePreviewInstance();
 	}
 }
 
 
 void FCustomizableObjectInstanceEditor::UpdatePreviewVisibility()
 {
+	const bool bEnableVisibility = CustomizableObjectInstance->SkeletalMeshStatus == ESkeletalMeshState::Correct;
 	for (UDebugSkelMeshComponent* PreviewSkeletalMeshComponent : PreviewSkeletalMeshComponents)
 	{
-		if (CustomizableObjectInstance->SkeletalMeshStatus == ESkeletalMeshState::Correct)
-		{
-			PreviewSkeletalMeshComponent->SetVisibility(true, true);
-		}
-		else
-		{
-			PreviewSkeletalMeshComponent->SetVisibility(false, true);
-		}
+		PreviewSkeletalMeshComponent->SetVisibility(bEnableVisibility, true);
 	}
 }
 
@@ -379,27 +336,32 @@ void FCustomizableObjectInstanceEditor::SaveAsset_Execute()
 {
 	if (CustomizableObjectInstance)
 	{
-		UPackage* Package = CustomizableObjectInstance->GetOutermost();
+		return;
+	}
 
-		if (Package)
-		{
-			TArray<UPackage*> PackagesToSave;
-			PackagesToSave.Add(Package);
+	if (UPackage* Package = CustomizableObjectInstance->GetOutermost())
+	{
+		TArray<UPackage*> PackagesToSave;
+		PackagesToSave.Add(Package);
 
-			FEditorFileUtils::PromptForCheckoutAndSave(PackagesToSave, false, false);
-		}
+		FEditorFileUtils::PromptForCheckoutAndSave(PackagesToSave, false, false);
 	}
 }
 
-
 bool FCustomizableObjectInstanceEditor::CanOpenOrShowParent()
 {
+	if (!CustomizableObjectInstance)
+	{
+		return false;
+	}
+
 	return CustomizableObjectInstance->GetCustomizableObject() != nullptr;
 }
 
 
 void FCustomizableObjectInstanceEditor::ShowParentInContentBrowser()
 {
+	check(CustomizableObjectInstance);
 	FContentBrowserModule& ContentBrowserModule = FModuleManager::Get().LoadModuleChecked<FContentBrowserModule>(TEXT("ContentBrowser"));
 	ContentBrowserModule.Get().SyncBrowserToAssets(TArray<FAssetData>{FAssetData(CustomizableObjectInstance->GetCustomizableObject())});
 }
@@ -407,6 +369,7 @@ void FCustomizableObjectInstanceEditor::ShowParentInContentBrowser()
 
 void FCustomizableObjectInstanceEditor::OpenParentInEditor()
 {
+	check(CustomizableObjectInstance);
 	GEditor->GetEditorSubsystem<UAssetEditorSubsystem>()->OpenEditorForAsset(CustomizableObjectInstance->GetCustomizableObject());
 }
 
@@ -480,14 +443,6 @@ void FCustomizableObjectInstanceEditor::BindCommands()
 	const FCustomizableObjectInstanceEditorCommands& Commands = FCustomizableObjectInstanceEditorCommands::Get();
 
 	const TSharedRef<FUICommandList>& UICommandList = GetToolkitCommands();
-
-	// Mesh menu
-	//UICommandList->MapAction(
-	//	Commands.Compile,
-	//	FExecuteAction::CreateSP( this, &FCustomizableObjectInstanceEditor::Compile ),
-	//	FCanExecuteAction(),
-	//	FIsActionChecked() );
-
 	
 	UICommandList->MapAction(
 		Commands.ShowParentCO,
@@ -684,38 +639,6 @@ bool FCustomizableObjectInstanceEditor::GetAssetRegistryLoaded()
 }
 
 
-void FCustomizableObjectInstanceEditor::SetEditorMesh(UCustomizableObjectInstance* InCustomizableObjectInstance)
-{
-	CustomizableObjectInstance = InCustomizableObjectInstance;
-
-	// Set the details view.
-	TArray<UObject*> SelectedObjects;
-	SelectedObjects.Add(CustomizableObjectInstance);
-	CustomizableInstanceDetailsView->SetObjects(SelectedObjects);
-}
-
-
-void FCustomizableObjectInstanceEditor::OnChangeMesh()
-{
-	FEditorDelegates::LoadSelectedAssetsIfNeeded.Broadcast();
-	UCustomizableObjectInstance* SelectedMesh = GEditor->GetSelectedObjects()->GetTop<UCustomizableObjectInstance>();
-	if(SelectedMesh && SelectedMesh != CustomizableObjectInstance)
-	{
-		GEditor->GetEditorSubsystem<UAssetEditorSubsystem>()->NotifyEditorClosed(this);
-		GEditor->GetEditorSubsystem<UAssetEditorSubsystem>()->NotifyAssetOpened(SelectedMesh, this);
-
-		SetEditorMesh(SelectedMesh);
-//		LevelOfDetailSettings->MeshChanged();
-
-		// Refresh the tool so it will update it's LOD list.
-		//if(GenerateUniqueUVs.IsValid())
-		//{
-		//	GenerateUniqueUVs->RefreshTool();
-		//}
-	}
-}
-
-
 void FCustomizableObjectInstanceEditor::OnInstancePropertySelectionChanged(FProperty* InProperty)
 {
 	RefreshViewport();
@@ -764,11 +687,9 @@ void FCustomizableObjectInstanceEditor::OnUpdatePreviewInstance()
 		CustomizableObjectInstance->LastSelectedProjectorParameterWithIndex = "";
 	}
 
-	if ( !CustomizableObjectInstance->ProjectorUpdatedInViewport)
+	if (UpdateDetails && !CustomizableObjectInstance->ProjectorUpdatedInViewport)
 	{
-		TArray<UObject*> instances;
-		instances.Add(CustomizableObjectInstance);
-		CustomizableInstanceDetailsView->SetObjects(instances, UpdateDetails);
+		CustomizableInstanceDetailsView->ForceRefresh();
 	}
 
 	if (TextureAnalyzer.IsValid())
@@ -780,30 +701,22 @@ void FCustomizableObjectInstanceEditor::OnUpdatePreviewInstance()
 
 void FCustomizableObjectInstanceEditor::CompileObject(UCustomizableObject* Object)
 {
-	UE_LOG(LogMutable, Verbose, TEXT("PROFILE: -----------------------------------------------------------"));
-	UE_LOG(LogMutable, Verbose, TEXT("PROFILE: [ %16.8f ] FCustomizableObjectInstanceEditor::CompileObject start."), FPlatformTime::Seconds());
 
-	if (!AssetRegistryLoaded)
+	if (!Object || !Object->Source)
 	{
-		FNotificationInfo Info(NSLOCTEXT("CustomizableObject", "CustomizableObjectCompileTryLater", "Please wait until asset registry loads all assets"));
-		Info.bFireAndForget = true;
-		Info.bUseThrobber = true;
-		Info.FadeOutDuration = 1.0f;
-		Info.ExpireDuration = 2.0f;
-		FSlateNotificationManager::Get().AddNotification(Info);
 		return;
 	}
-
+	
 	Viewport->UpdateGizmoDataToOrigin();
 
-	if (Object && Object->Source)
+	if (!Compiler)
 	{
-		FCompilationOptions Options = Object->CompileOptions;
-		Options.bSilentCompilation = false;
-		Compiler->Compile(*Object, Options, true);
+		Compiler = MakeUnique<FCustomizableObjectCompiler>();
 	}
-
-	UE_LOG(LogMutable, Verbose, TEXT("PROFILE: [ %16.8f ] FCustomizableObjectInstanceEditor::CompileObject end."), FPlatformTime::Seconds());
+	
+	FCompilationOptions Options = Object->CompileOptions;
+	Options.bSilentCompilation = false;
+	Compiler->Compile(*Object, Options, true);
 }
 
 
@@ -815,20 +728,11 @@ bool FCustomizableObjectInstanceEditor::IsTickable(void) const
 
 void FCustomizableObjectInstanceEditor::Tick(float InDeltaTime)
 {
-	bool bUpdated = Compiler->Tick();
-
-	if (bUpdated)
+	check(CustomizableObjectInstance);
+	if (Compiler && Compiler->Tick())
 	{
-		if ((CustomizableObjectInstance == nullptr) && (ObjectToEditCopy != nullptr))
-		{
-			CustomizableObjectInstance = ObjectToEditCopy;
-		}
-
-		if (!CustomizableObjectInstance) return;
-
 		if (const UCustomizableObject* CustomizableObject = CustomizableObjectInstance->GetCustomizableObject();
-			CustomizableObject &&
-			CustomizableObject->Source)
+			CustomizableObject && CustomizableObject->Source)
 		{
 			if (PreviewCustomizableSkeletalComponents.Num() > 0)
 			{
@@ -837,19 +741,14 @@ void FCustomizableObjectInstanceEditor::Tick(float InDeltaTime)
 			else
 			{
 				CreatePreviewInstance();
-
-				if (CreatePreviewInstanceAfterCOCompile)
-				{
-					CreatePreviewInstanceAfterCOCompile = false;
-				}
 			}
+
+			CustomizableInstanceDetailsView->ForceRefresh();
 		}
 
-		SetEditorMesh(CustomizableObjectInstance);
+		Compiler.Reset();
 	}
 
-	if (!CustomizableObjectInstance) return;
-
 	// TEMP CODE
 	if (CustomizableObjectInstance->TempUpdateGizmoInViewport)
 	{
diff --git a/Engine/Plugins/Experimental/Mutable/Source/CustomizableObjectEditor/Private/MuCOE/CustomizableObjectInstanceEditor.h b/Engine/Plugins/Experimental/Mutable/Source/CustomizableObjectEditor/Private/MuCOE/CustomizableObjectInstanceEditor.h
--- a/Engine/Plugins/Experimental/Mutable/Source/CustomizableObjectEditor/Private/MuCOE/CustomizableObjectInstanceEditor.h
+++ b/Engine/Plugins/Experimental/Mutable/Source/CustomizableObjectEditor/Private/MuCOE/CustomizableObjectInstanceEditor.h
@@ -111,16 +111,6 @@ private:
 	/** Binds commands associated with the Static Mesh Editor. */
 	void BindCommands();
 		
-	/**
-	 *	Sets the editor's current mesh and refreshes various settings to correspond with the new data.
-	 *
-	 *	@param	InCustomizableObject		The static mesh to use for the editor.
-	 */
-	void SetEditorMesh(UCustomizableObjectInstance* InCustomizableObjectInstance);
-
-	/** Change the mesh the editor is viewing. */
-	void OnChangeMesh();
-
 	/** Callback when selection changes in the Property Tree. */
 	void OnInstancePropertySelectionChanged(FProperty* InProperty);
 
@@ -188,13 +178,7 @@ private:
 	FDelegateHandle OnObjectModifiedHandle;
 
 	/** Flag to know when the asset registry initial loading has completed */
-	bool AssetRegistryLoaded;
-
-	/** Flag to know whether the lasts steps when making a new preview instance need to be done when asset registry completes asset loading */
-	bool UpdateSkeletalMeshAfterAssetLoaded;
-
-	/** Copy of variable needed when asset registry finishes loading assets to complete editor loading */
-	UCustomizableObjectInstance* ObjectToEditCopy;
+	bool AssetRegistryLoaded = true;
 
 	/** UObject class to be able to use the update callback */
 	TObjectPtr<UUpdateClassWrapperClass> HelperCallback;
@@ -205,17 +189,12 @@ private:
 	/** Object compiler */
 	TUniquePtr<class FCustomizableObjectCompiler> Compiler;
 
-	/** Used in InitCustomizableObjectInstanceEditor, in the case the CO being opened is not compiled. Since referencer assets are loaded asynchronously,
-	* this flag is used to complete editor initialization */
-	bool CreatePreviewInstanceAfterCOCompile;
-
 	/** Pose asset when doing drag and drop of an UPoseAsset to the viewport */
 	UPoseAsset* PoseAsset;
 
 	/** Texture Analyzer table widget which shows the information of the transient textures used in the customizable object instance */
 	TSharedPtr<class SCustomizableObjecEditorTextureAnalyzer> TextureAnalyzer;
 
-
 	/** Variables used to force the refresh of the details view widget. These are needed because sometimes the scrollbar of the window doesn't appear
 	    until we force the refresh */
 	bool bOnlyRuntimeParameters;
diff --git a/Engine/Plugins/Experimental/Mutable/Source/CustomizableObjectEditor/Private/MuCOE/SCustomizableInstanceProperties.cpp b/Engine/Plugins/Experimental/Mutable/Source/CustomizableObjectEditor/Private/MuCOE/SCustomizableInstanceProperties.cpp
--- a/Engine/Plugins/Experimental/Mutable/Source/CustomizableObjectEditor/Private/MuCOE/SCustomizableInstanceProperties.cpp
+++ b/Engine/Plugins/Experimental/Mutable/Source/CustomizableObjectEditor/Private/MuCOE/SCustomizableInstanceProperties.cpp
@@ -437,7 +437,7 @@ void SCustomizableInstanceProperties::ResetParamBox()
 			{
 				FCustomizableInstancePropertyArrayElem ParameterSortInfo;
 				ParameterSortInfo.PropertyCOIndex = ParamIndexInObject;
-				if (CustomInstance->IsParameterRelevant(ParameterSortInfo.PropertyCOIndex))
+				if (!CustomInstance->bShowOnlyRelevantParameters || CustomInstance->IsParameterRelevant(ParameterSortInfo.PropertyCOIndex))
 				{
 					ParameterSortInfo.UIOrder = CustomizableObject->GetParameterUIMetadataFromIndex(ParameterSortInfo.PropertyCOIndex).ParamUIMetadata.UIOrder;
 					ParameterSortInfo.ParameterName = CustomizableObject->GetParameterName(ParameterSortInfo.PropertyCOIndex);
@@ -449,20 +449,14 @@ void SCustomizableInstanceProperties::ResetParamBox()
 
 			for (int32 ParamIndexInObject = 0; ParamIndexInObject < ParamIndexesInObject.Num(); ++ParamIndexInObject)
 			{
-				if (CustomInstance->IsParameterRelevant(ParamIndexesInObject[ParamIndexInObject].PropertyCOIndex))
-				{
-					FillChildrenMap(ParamIndexesInObject[ParamIndexInObject].PropertyCOIndex);
-				}
+				FillChildrenMap(ParamIndexesInObject[ParamIndexInObject].PropertyCOIndex);
 			}
 
 			for (int32 ParamIndexInObject = 0; ParamIndexInObject < ParamIndexesInObject.Num(); ++ParamIndexInObject)
 			{
-				if (CustomInstance->IsParameterRelevant(ParamIndexesInObject[ParamIndexInObject].PropertyCOIndex))
+				if (!ParamHasParent.Find(ParamIndexesInObject[ParamIndexInObject].PropertyCOIndex))
 				{
-					if (!ParamHasParent.Find(ParamIndexesInObject[ParamIndexInObject].PropertyCOIndex))
-					{
-						RecursivelyAddParamAndChildren(ParamIndexesInObject[ParamIndexInObject].PropertyCOIndex);
-					}
+					RecursivelyAddParamAndChildren(ParamIndexesInObject[ParamIndexInObject].PropertyCOIndex);
 				}
 			}
 		}
