[Mutable] Fix false positives in cycle detection.

An example of this issue was when compiling Group Projector nodes.

The false positive was due to the __COUNTER__ not global but independent for each compilation unit. Changing it to the function address ensures that the key is unique for each function.

#rnx

diff --git a/Engine/Plugins/Experimental/Mutable/Source/CustomizableObjectEditor/Private/MuCOE/GenerateMutableSource/GenerateMutableSource.cpp b/Engine/Plugins/Experimental/Mutable/Source/CustomizableObjectEditor/Private/MuCOE/GenerateMutableSource/GenerateMutableSource.cpp
--- a/Engine/Plugins/Experimental/Mutable/Source/CustomizableObjectEditor/Private/MuCOE/GenerateMutableSource/GenerateMutableSource.cpp
+++ b/Engine/Plugins/Experimental/Mutable/Source/CustomizableObjectEditor/Private/MuCOE/GenerateMutableSource/GenerateMutableSource.cpp
@@ -330,7 +330,7 @@ void FPinData::Push(const UEdGraphPin* Pin)
 }
 
 
-FGraphCycleKey::FGraphCycleKey(const UEdGraphPin& Pin, uint32 Id) :
+FGraphCycleKey::FGraphCycleKey(const UEdGraphPin& Pin, const FString& Id) :
 	Pin(Pin),
 	Id(Id)
 {
diff --git a/Engine/Plugins/Experimental/Mutable/Source/CustomizableObjectEditor/Private/MuCOE/GenerateMutableSource/GenerateMutableSource.h b/Engine/Plugins/Experimental/Mutable/Source/CustomizableObjectEditor/Private/MuCOE/GenerateMutableSource/GenerateMutableSource.h
--- a/Engine/Plugins/Experimental/Mutable/Source/CustomizableObjectEditor/Private/MuCOE/GenerateMutableSource/GenerateMutableSource.h
+++ b/Engine/Plugins/Experimental/Mutable/Source/CustomizableObjectEditor/Private/MuCOE/GenerateMutableSource/GenerateMutableSource.h
@@ -454,7 +454,7 @@ struct FGraphCycleKey
 {
 	friend uint32 GetTypeHash(const FGraphCycleKey& Key);
 
-	FGraphCycleKey(const UEdGraphPin& Pin, uint32 Id);
+	FGraphCycleKey(const UEdGraphPin& Pin, const FString& Id);
 
 	bool operator==(const FGraphCycleKey& Other) const;
 	
@@ -462,7 +462,7 @@ struct FGraphCycleKey
 	const UEdGraphPin& Pin;
 
 	/** Unique id. */
-	uint32 Id;
+	FString Id;
 };
 
 /** Graph Cycle scope.
@@ -488,7 +488,7 @@ private:
 
 /** Return the default value if there is a cycle. */
 #define RETURN_ON_CYCLE(Pin, GenerationContext) \
-	FGraphCycle GraphCycle(FGraphCycleKey(Pin, __COUNTER__), GenerationContext); \
+	FGraphCycle GraphCycle(FGraphCycleKey(Pin, TEXT(__FILE__ PREPROCESSOR_TO_STRING(__LINE__))), GenerationContext); \
 	if (GraphCycle.FoundCycle()) \
 	{ \
 		return {}; \
