[mutable] Reduce performance impact of memory-budget control code.
#rnx

diff --git a/Engine/Plugins/Experimental/Mutable/Source/MutableRuntime/Private/MuR/System.cpp b/Engine/Plugins/Experimental/Mutable/Source/MutableRuntime/Private/MuR/System.cpp
--- a/Engine/Plugins/Experimental/Mutable/Source/MutableRuntime/Private/MuR/System.cpp
+++ b/Engine/Plugins/Experimental/Mutable/Source/MutableRuntime/Private/MuR/System.cpp
@@ -1065,6 +1065,7 @@ namespace mu
         bool bFinished = TotalBytes <= BudgetBytes;
 
 		// Try to free pooled resources first
+		if (!bFinished)
 		{
 			MUTABLE_CPUPROFILER_SCOPE(EnsureBudgetBelow_FreePooled);
 			while (PooledImages.Num() && !bFinished)
@@ -1079,67 +1080,65 @@ namespace mu
 		}
 		
 		// Try to free a loaded roms
+		if (!bFinished)
 		{
 			MUTABLE_CPUPROFILER_SCOPE(EnsureBudgetBelow_FreeRoms);
 
-			while (!bFinished)
+			struct FRomRef
+			{
+				const Model* Model = nullptr;
+				int32 RomIndex = 0;
+			};
+
+			TArray<TPair<float, FRomRef>> Candidates;
+			Candidates.Reserve(512);
+
+			for (FModelCacheEntry& ModelCache : CachePerModel)
 			{
-				TSharedPtr<const Model> LowestPriorityModel;
-				int32 LowestPriorityRom = -1;
-				float LowestPriority = 0.0f;
-				for (FModelCacheEntry& ModelCache : CachePerModel)
+				TSharedPtr<const Model> CacheModel = ModelCache.Model.Pin();
+				if (CacheModel)
 				{
-					TSharedPtr<const Model> CacheModel = ModelCache.Model.Pin();
-					if (CacheModel)
-					{
-						mu::FProgram& program = CacheModel->GetPrivate()->m_program;
-						check(ModelCache.RomWeights.Num() == program.m_roms.Num());
+					mu::FProgram& program = CacheModel->GetPrivate()->m_program;
+					check(ModelCache.RomWeights.Num() == program.m_roms.Num());
 
-						for (int32 RomIndex = 0; RomIndex < program.m_roms.Num(); ++RomIndex)
+					for (int32 RomIndex = 0; RomIndex < program.m_roms.Num(); ++RomIndex)
+					{
+						const FRomData& Rom = program.m_roms[RomIndex];
+						bool bIsLoaded = program.IsRomLoaded(RomIndex);
+
+						// We cannot unload a rom if some operation is expecting it.
+						bool bIsRomLocked = ModelCache.PendingOpsPerRom.IsValidIndex(RomIndex)
+							&& 
+							ModelCache.PendingOpsPerRom[RomIndex]>0;
+						if (bIsLoaded && !bIsRomLocked)
 						{
-							const FRomData& Rom = program.m_roms[RomIndex];
-							bool bIsLoaded = program.IsRomLoaded(RomIndex);
-
-							// We cannot unload a rom if some operation is expecting it.
-							bool bIsRomLocked = ModelCache.PendingOpsPerRom.IsValidIndex(RomIndex)
-								&& 
-								ModelCache.PendingOpsPerRom[RomIndex]>0;
-							if (bIsLoaded && !bIsRomLocked)
-							{
-								constexpr float FactorWeight = 100.0f;
-								constexpr float FactorTime = -1.0f;
-								float Priority = FactorWeight * float(ModelCache.RomWeights[RomIndex].Key)
-									+
-									FactorTime * float((RomTick - ModelCache.RomWeights[RomIndex].Value));
-
-								if (LowestPriorityRom < 0 || Priority < LowestPriority)
-								{
-									LowestPriorityRom = RomIndex;
-									LowestPriority = Priority;
-									LowestPriorityModel = CacheModel;
-								}
-							}
+							constexpr float FactorWeight = 100.0f;
+							constexpr float FactorTime = -1.0f;
+							float Priority = FactorWeight * float(ModelCache.RomWeights[RomIndex].Key)
+								+
+								FactorTime * float((RomTick - ModelCache.RomWeights[RomIndex].Value));
+
+							FRomRef Ref = { CacheModel.Get(), RomIndex };
+							Candidates.Add(TPair<float, FRomRef>(Priority,Ref));
 						}
 					}
 				}
+			}
 
-				if (LowestPriorityRom < 0)
-				{
-					// If we reached this it means we couldn't find more streaming data to free
-					// Try to continue anyway.
-					break;
-				}
-				else
-				{
-					MUTABLE_CPUPROFILER_SCOPE(EnsureBudgetBelow_UnloadRom);
+			Candidates.Sort([](const TPair<float, FRomRef>& A, const TPair<float, FRomRef>& B) { return A.Key > B.Key; });
 
-					// UE_LOG(LogMutableCore,Log, "Unloading rom because of memory budget: %d.", lowestPriorityRom);
-					int32 UnloadedSize = LowestPriorityModel->GetPrivate()->m_program.UnloadRom(LowestPriorityRom);
-					TotalBytes -= UnloadedSize;
-					TrackedBudgetBytes_Rom -= UnloadedSize;
+			while (!bFinished && Candidates.Num())
+			{
+				MUTABLE_CPUPROFILER_SCOPE(EnsureBudgetBelow_UnloadRom);
 
-					bFinished = TotalBytes <= BudgetBytes;
-				}
+				TPair<float, FRomRef> Candidate = Candidates.Pop(false);
+
+				// UE_LOG(LogMutableCore,Log, "Unloading rom because of memory budget: %d.", lowestPriorityRom);
+				int32 UnloadedSize = Candidate.Value.Model->GetPrivate()->m_program.UnloadRom(Candidate.Value.RomIndex);
+				TotalBytes -= UnloadedSize;
+				TrackedBudgetBytes_Rom -= UnloadedSize;
+
+				bFinished = TotalBytes <= BudgetBytes;
 			}
 		}
 
@@ -1159,52 +1158,14 @@ namespace mu
 
 				// Gather all data in the cache for this instance
 				TArray<const Resource*> CacheUnique;
-				CodeContainer<FProgramCache::FOpExecutionData>::iterator It = Instance.Cache->OpExecutionData.begin();
-				for (; It.IsValid(); ++It)
+				CacheUnique.Reserve(1024);
 				{
-					FProgramCache::FOpExecutionData& Data = *It;
+					MUTABLE_CPUPROFILER_SCOPE(EnsureBudgetBelow_FreeCached_Gather_Other);
 
-					if (!Data.DataTypeIndex)
+					CodeContainer<FProgramCache::FOpExecutionData>::iterator It = Instance.Cache->OpExecutionData.begin();
+					for (; It.IsValid(); ++It)
 					{
-						continue;
-					}
-
-					const Resource* Value = nullptr;
-					switch (Data.DataType)
-					{
-					case DATATYPE::DT_IMAGE:
-						Value = Instance.Cache->ImageResults[Data.DataTypeIndex].get();
-						if (Value)
-						{
-							CacheUnique.AddUnique(Value);
-						}
-						break;
-
-					case DATATYPE::DT_MESH:
-						Value = Instance.Cache->MeshResults[Data.DataTypeIndex].get();
-						if (Value)
-						{
-							CacheUnique.AddUnique(Value);
-						}
-						break;
-
-					default:
-						break;
-					}
-				}
-
-				while (!bFinished && CacheUnique.Num())
-				{
-					// Free one
-					const Resource* Removed = CacheUnique.Pop();
-
-					int32 RemovedDataSize = Removed->GetDataSize();
-
-					// Clear its cache references
-					CodeContainer<FProgramCache::FOpExecutionData>::iterator RemIt = Instance.Cache->OpExecutionData.begin();
-					for (; RemIt.IsValid(); ++RemIt)
-					{
-						FProgramCache::FOpExecutionData& Data = *RemIt;
+						FProgramCache::FOpExecutionData& Data = *It;
 
 						if (!Data.DataTypeIndex)
 						{
@@ -1216,36 +1177,82 @@ namespace mu
 						{
 						case DATATYPE::DT_IMAGE:
 							Value = Instance.Cache->ImageResults[Data.DataTypeIndex].get();
+							if (Value)
+							{
+								CacheUnique.AddUnique(Value);
+							}
 							break;
 
 						case DATATYPE::DT_MESH:
 							Value = Instance.Cache->MeshResults[Data.DataTypeIndex].get();
+							if (Value)
+							{
+								CacheUnique.AddUnique(Value);
+							}
 							break;
 
 						default:
 							break;
 						}
+					}
+				}
+
+				{
+					MUTABLE_CPUPROFILER_SCOPE(EnsureBudgetBelow_FreeCached_Free_Other);
+
+					while (!bFinished && CacheUnique.Num())
+					{
+						// Free one
+						const Resource* Removed = CacheUnique.Pop(false);
 
-						if (Value == Removed)
+						int32 RemovedDataSize = Removed->GetDataSize();
+
+						// Clear its cache references
+						CodeContainer<FProgramCache::FOpExecutionData>::iterator RemIt = Instance.Cache->OpExecutionData.begin();
+						for (; RemIt.IsValid(); ++RemIt)
 						{
-							int32 RemoveCount = CacheResources.Remove(Removed);
-							if (RemoveCount > 0)
+							FProgramCache::FOpExecutionData& Data = *RemIt;
+
+							if (!Data.DataTypeIndex)
 							{
-								TrackedBudgetBytes_Cached -= RemovedDataSize;
+								continue;
 							}
 
-							// \TODO: This is not very efficient, since it will search the data again.
-							Instance.Cache->SetUnused(RemIt.get_address());
+							const Resource* Value = nullptr;
+							switch (Data.DataType)
+							{
+							case DATATYPE::DT_IMAGE:
+								Value = Instance.Cache->ImageResults[Data.DataTypeIndex].get();
+								break;
+
+							case DATATYPE::DT_MESH:
+								Value = Instance.Cache->MeshResults[Data.DataTypeIndex].get();
+								break;
+
+							default:
+								break;
+							}
+
+							if (Value == Removed)
+							{
+								int32 RemoveCount = CacheResources.Remove(Removed);
+								if (RemoveCount > 0)
+								{
+									TrackedBudgetBytes_Cached -= RemovedDataSize;
+								}
+
+								Instance.Cache->SetUnused(*RemIt);
+							}
 						}
-					}
 
-					TotalBytes -= RemovedDataSize;
-					bFinished = TotalBytes <= BudgetBytes;
-				}
+						TotalBytes -= RemovedDataSize;
+						bFinished = TotalBytes <= BudgetBytes;
+					}
 
-				if (bFinished)
-				{
-					break;
+					if (bFinished)
+					{
+						break;
+					}
 				}
 			}
 
@@ -1255,54 +1262,16 @@ namespace mu
 			{
 				// Gather all data in the cache for this instance
 				TArray<const Resource*> CacheUnique;
-				CodeContainer<FProgramCache::FOpExecutionData>::iterator It = CurrentInstanceCache->OpExecutionData.begin();
-				for (; It.IsValid(); ++It)
+				CacheUnique.Reserve(1024);
 				{
-					FProgramCache::FOpExecutionData& Data = *It;
+					MUTABLE_CPUPROFILER_SCOPE(EnsureBudgetBelow_FreeCached_Gather_Current);
 
-					if (!Data.DataTypeIndex || Data.OpHitCount>0 || !Data.IsCacheLocked)
+					CodeContainer<FProgramCache::FOpExecutionData>::iterator It = CurrentInstanceCache->OpExecutionData.begin();
+					for (; It.IsValid(); ++It)
 					{
-						continue;
-					}
+						FProgramCache::FOpExecutionData& Data = *It;
 
-					const Resource* Value = nullptr;
-					switch (Data.DataType)
-					{
-					case DATATYPE::DT_IMAGE:
-						Value = CurrentInstanceCache->ImageResults[Data.DataTypeIndex].get();
-						if (Value)
-						{
-							CacheUnique.AddUnique(Value);
-						}
-						break;
-
-					case DATATYPE::DT_MESH:
-						Value = CurrentInstanceCache->MeshResults[Data.DataTypeIndex].get();
-						if (Value)
-						{
-							CacheUnique.AddUnique(Value);
-						}
-						break;
-
-					default:
-						break;
-					}
-				}
-
-				while (!bFinished && CacheUnique.Num())
-				{
-					// Free one
-					const Resource* Removed = CacheUnique.Pop();
-
-					// Does this data have any other cache references with op-hit-count bigger than 0?
-					bool bStillUsed = false;
-
-					CodeContainer<FProgramCache::FOpExecutionData>::iterator CheckIt = CurrentInstanceCache->OpExecutionData.begin();
-					for (; CheckIt.IsValid(); ++CheckIt)
-					{
-						const FProgramCache::FOpExecutionData& Data = *CheckIt;
-
-						if (!Data.DataTypeIndex)
+						if (!Data.DataTypeIndex || Data.OpHitCount > 0)
 						{
 							continue;
 						}
@@ -1312,73 +1281,120 @@ namespace mu
 						{
 						case DATATYPE::DT_IMAGE:
 							Value = CurrentInstanceCache->ImageResults[Data.DataTypeIndex].get();
+							if (Value)
+							{
+								CacheUnique.AddUnique(Value);
+							}
 							break;
 
 						case DATATYPE::DT_MESH:
 							Value = CurrentInstanceCache->MeshResults[Data.DataTypeIndex].get();
+							if (Value)
+							{
+								CacheUnique.AddUnique(Value);
+							}
 							break;
 
 						default:
 							break;
 						}
-
-						if (Value == Removed && Data.OpHitCount>0)
-						{
-							bStillUsed = true;
-							break;
-						}
 					}
+				}
 
-					if (bStillUsed)
+				{
+					MUTABLE_CPUPROFILER_SCOPE(EnsureBudgetBelow_FreeCached_Free_Current);
+
+					while (!bFinished && CacheUnique.Num())
 					{
-						continue;
-					}
+						// Free one
+						const Resource* Removed = CacheUnique.Pop(false);
 
+						// Does this data have any other cache references with op-hit-count bigger than 0?
+						bool bStillUsed = false;
 
-					int32 RemovedDataSize = Removed->GetDataSize();
+						CodeContainer<FProgramCache::FOpExecutionData>::iterator CheckIt = CurrentInstanceCache->OpExecutionData.begin();
+						for (; CheckIt.IsValid(); ++CheckIt)
+						{
+							const FProgramCache::FOpExecutionData& Data = *CheckIt;
 
-					// Clear its cache references
-					CodeContainer<FProgramCache::FOpExecutionData>::iterator RemIt = CurrentInstanceCache->OpExecutionData.begin();
-					for (; RemIt.IsValid(); ++RemIt)
-					{
-						FProgramCache::FOpExecutionData& Data = *RemIt;
+							if (!Data.DataTypeIndex)
+							{
+								continue;
+							}
 
-						if (!Data.DataTypeIndex)
+							const Resource* Value = nullptr;
+							switch (Data.DataType)
+							{
+							case DATATYPE::DT_IMAGE:
+								Value = CurrentInstanceCache->ImageResults[Data.DataTypeIndex].get();
+								break;
+
+							case DATATYPE::DT_MESH:
+								Value = CurrentInstanceCache->MeshResults[Data.DataTypeIndex].get();
+								break;
+
+							default:
+								break;
+							}
+
+							if (Value == Removed && Data.OpHitCount > 0)
+							{
+								bStillUsed = true;
+								break;
+							}
+						}
+
+						if (bStillUsed)
 						{
 							continue;
 						}
 
-						const Resource* Value = nullptr;
-						switch (Data.DataType)
+
+						int32 RemovedDataSize = Removed->GetDataSize();
+
+						// Clear its cache references
+						CodeContainer<FProgramCache::FOpExecutionData>::iterator RemIt = CurrentInstanceCache->OpExecutionData.begin();
+						for (; RemIt.IsValid(); ++RemIt)
 						{
-						case DATATYPE::DT_IMAGE:
-							Value = CurrentInstanceCache->ImageResults[Data.DataTypeIndex].get();
-							break;
+							FProgramCache::FOpExecutionData& Data = *RemIt;
 
-						case DATATYPE::DT_MESH:
-							Value = CurrentInstanceCache->MeshResults[Data.DataTypeIndex].get();
-							break;
+							if (!Data.DataTypeIndex)
+							{
+								continue;
+							}
 
-						default:
-							break;
-						}
+							const Resource* Value = nullptr;
+							switch (Data.DataType)
+							{
+							case DATATYPE::DT_IMAGE:
+								Value = CurrentInstanceCache->ImageResults[Data.DataTypeIndex].get();
+								break;
 
-						if (Value == Removed)
-						{
-							int32 RemoveCount = CacheResources.Remove(Removed);
-							if (RemoveCount > 0)
+							case DATATYPE::DT_MESH:
+								Value = CurrentInstanceCache->MeshResults[Data.DataTypeIndex].get();
+								break;
+
+							default:
+								break;
+							}
+
+							if (Value == Removed)
 							{
-								TrackedBudgetBytes_Cached -= RemovedDataSize;
+								MUTABLE_CPUPROFILER_SCOPE(EnsureBudgetBelow_FreeCached_Free_Current_ActualFree);
+
+								int32 RemoveCount = CacheResources.Remove(Removed);
+								if (RemoveCount > 0)
+								{
+									TrackedBudgetBytes_Cached -= RemovedDataSize;
+								}
+								CurrentInstanceCache->SetUnused(*RemIt);
 							}
-							// \TODO: This is not very efficient, since it will search the data again.
-							CurrentInstanceCache->SetUnused(RemIt.get_address());
 						}
-					}
 
-					TotalBytes -= RemovedDataSize;
-					bFinished = TotalBytes <= BudgetBytes;
+						TotalBytes -= RemovedDataSize;
+						bFinished = TotalBytes <= BudgetBytes;
+					}
 				}
-
 			}
 
 		}
diff --git a/Engine/Plugins/Experimental/Mutable/Source/MutableRuntime/Private/MuR/SystemPrivate.h b/Engine/Plugins/Experimental/Mutable/Source/MutableRuntime/Private/MuR/SystemPrivate.h
--- a/Engine/Plugins/Experimental/Mutable/Source/MutableRuntime/Private/MuR/SystemPrivate.h
+++ b/Engine/Plugins/Experimental/Mutable/Source/MutableRuntime/Private/MuR/SystemPrivate.h
@@ -537,17 +537,16 @@ namespace mu
 			}
 		}
 
-		void SetUnused(FCacheAddress at)
+		void SetUnused(FOpExecutionData& Data)
 		{
 			// Only clear datatypes that have results that use relevant amounts of memory			
-			FOpExecutionData& Data = OpExecutionData[at];
 			uint32 DataTypeIndex = Data.DataTypeIndex;
 			Data.IsValueValid = false;
 
 			if (DataTypeIndex)
 			{
-				check(OpExecutionData[at].DataType < DATATYPE::DT_COUNT);
-				switch ((DATATYPE)OpExecutionData[at].DataType)
+				check(Data.DataType < DATATYPE::DT_COUNT);
+				switch ((DATATYPE)Data.DataType)
 				{
 				case DATATYPE::DT_IMAGE:					
 					ImageResults[DataTypeIndex] = nullptr;
@@ -715,7 +714,7 @@ namespace mu
 			--Data->OpHitCount;
 			if (Data->OpHitCount == 0 && !Data->IsCacheLocked)
 			{
-				SetUnused(at);
+				SetUnused(*Data);
 			}
 
 			return Result;
@@ -740,7 +739,7 @@ namespace mu
 			--Data->OpHitCount;
 			if (Data->OpHitCount == 0 && !Data->IsCacheLocked)
 			{
-				SetUnused(at);
+				SetUnused(*Data);
 				bIsLastReference = true;
 			}
 
@@ -766,7 +765,7 @@ namespace mu
 			--Data->OpHitCount;
 			if (Data->OpHitCount == 0 && !Data->IsCacheLocked)
 			{
-				SetUnused(at);
+				SetUnused(*Data);
 				bIsLastReference = true;
 			}
 
@@ -1138,7 +1137,7 @@ namespace mu
 		TArray<Ptr<const Image>> TempImages;
 		TArray<Ptr<const Mesh>> TempMeshes;
 
-		/** */
+		/** List of resources that are currently in any cache position, and the number of positions they are in. */
 		TMap<Ptr<const Resource>, int32> CacheResources;
 
 		/** Given a mutable model, find or create its rom cache. */
