[Mutable] Allow Removed/Disabled Sections in Automatic LODs from Mesh

- Now when using Automatic LODs From Mesh, if a Skeletal Mesh does not contains the requested section, we will generate an empty Mutable Core Mesh (mesh without data).
- Support meshes now will only use only the connected LOD regardless of Automatic LODs from Mesh. Support meshes are:
    - NodeRemoveMesh mesh
    - NodeMeshReshape BaseShape and TargetShape meshes
    - NodeClipMeshWithMesh Clip mesh
    - NodeClipDeform ClipShape mesh
- Fixed issue where tables where generating more sections than necessary (not necessary sections where already not used).
- Changed incorrectly named variables from MaterialIndex to SectionIndex.
- Unified LOD variable names. Now we use LODIndex (for the LOD we are generating) and LODIndexConnected (for the LOD where the pin is connected).
- Fixed Grammar Implicits doc issues.

#rnx

diff --git a/Engine/Plugins/Experimental/Mutable/Source/CustomizableObjectEditor/Private/MuCOE/Compiler/GrammarImplicits.txt b/Engine/Plugins/Experimental/Mutable/Source/CustomizableObjectEditor/Private/MuCOE/Compiler/GrammarImplicits.txt
--- a/Engine/Plugins/Experimental/Mutable/Source/CustomizableObjectEditor/Private/MuCOE/Compiler/GrammarImplicits.txt
+++ b/Engine/Plugins/Experimental/Mutable/Source/CustomizableObjectEditor/Private/MuCOE/Compiler/GrammarImplicits.txt
@@ -18,6 +18,8 @@ LAYOUT -> MeshLayoutNode
 
 MORPHSTACK -> NodeMeshMorphStackDefinition MESH
 
+MORPH -> SKELETASLMESHASSET
+
 STATICMESHASSET -> UStaticMesh
 
 SKELETASLMESHASSET-> USkeletalMesh
@@ -28,7 +30,7 @@ MESH -> NodeMeshMorphStackApplication MORPHSTACK MESH
 MESH -> NodeMeshReshape MESH:BaseMesh MESH:BaseShape MESH:TargetShape
 MESH -> NodeMeshSwitch ENUM MESH*
 MESH -> NodeMeshVariation MESH MESH*
-MESH -> NodeSkeletalMesh SKELETALMESH LAYOUT*
+MESH -> NodeSkeletalMesh SKELETASLMESHASSET LAYOUT* MORPH*
 MESH -> NodeStaticMesh STATICMESHASSET
 MESH -> NodeAnimationPose MESH
 MESH -> NodeTable SKELETASLMESHASSET* STATICMESHASSET*
@@ -83,7 +85,7 @@ SURFACE -> NodeMaterialVariation SURFACENODE SURFACENODE*
 SURFACE -> NodeCopyMaterial MESH SURFACENODE
 SURFACE -> NodeRemoveMesh SURFACENODE MESH
 SURFACE -> NodeRemoveMeshBlocks SURFACENODE 
-SURFACE -> NodeMorphMaterial SURFACENODE
+SURFACE -> NodeMorphMaterial SURFACENODE MORPH
 
 MODIFIER -> NodeMeshClipDeform MESH
 MODIFIER -> NodeMeshClipWithMesh MESH SURFACENODE
diff --git a/Engine/Plugins/Experimental/Mutable/Source/CustomizableObjectEditor/Private/MuCOE/CustomizableObjectLayout.cpp b/Engine/Plugins/Experimental/Mutable/Source/CustomizableObjectEditor/Private/MuCOE/CustomizableObjectLayout.cpp
--- a/Engine/Plugins/Experimental/Mutable/Source/CustomizableObjectEditor/Private/MuCOE/CustomizableObjectLayout.cpp
+++ b/Engine/Plugins/Experimental/Mutable/Source/CustomizableObjectEditor/Private/MuCOE/CustomizableObjectLayout.cpp
@@ -93,7 +93,7 @@ void UCustomizableObjectLayout::GenerateBlocksFromUVs()
 			GenerationContext.MeshGenerationFlags.Push(ShapeFlags);
 
 			GenerationContext.ComponentInfos.Add(SkeletalMesh);
-			MutableMesh = ConvertSkeletalMeshToMutable(SkeletalMesh, TSoftClassPtr<UAnimInstance>(), LOD, Material, GenerationContext, Node);
+			MutableMesh = ConvertSkeletalMeshToMutable(SkeletalMesh, TSoftClassPtr<UAnimInstance>(), LOD, Material, LOD, Material, GenerationContext, Node);
 		}
 		else if (UStaticMesh* StaticMesh = Cast<UStaticMesh>(Mesh))
 		{
diff --git a/Engine/Plugins/Experimental/Mutable/Source/CustomizableObjectEditor/Private/MuCOE/GenerateMutableSource/GenerateMutableSource.cpp b/Engine/Plugins/Experimental/Mutable/Source/CustomizableObjectEditor/Private/MuCOE/GenerateMutableSource/GenerateMutableSource.cpp
--- a/Engine/Plugins/Experimental/Mutable/Source/CustomizableObjectEditor/Private/MuCOE/GenerateMutableSource/GenerateMutableSource.cpp
+++ b/Engine/Plugins/Experimental/Mutable/Source/CustomizableObjectEditor/Private/MuCOE/GenerateMutableSource/GenerateMutableSource.cpp
@@ -254,7 +254,7 @@ FMutableComponentInfo& FMutableGraphGenerationContext::GetCurrentComponentInfo()
 }
 
 
-FGeneratedKey::FGeneratedKey(void* InFunctionAddress, const UEdGraphPin& InPin, const UCustomizableObjectNode& Node, FMutableGraphGenerationContext& GenerationContext, const bool UseMesh)
+FGeneratedKey::FGeneratedKey(void* InFunctionAddress, const UEdGraphPin& InPin, const UCustomizableObjectNode& Node, FMutableGraphGenerationContext& GenerationContext, const bool UseMesh, const bool InbOnlyConnectedLOD)
 {
 	FunctionAddress = InFunctionAddress;
 	Pin = &InPin;
@@ -264,6 +264,7 @@ FGeneratedKey::FGeneratedKey(void* InFunctionAddress, const UEdGraphPin& InPin,
 	{
 		Flags = GenerationContext.MeshGenerationFlags.Last();
 		MeshMorphStack = GenerationContext.MeshMorphStack;
+		bOnlyConnectedLOD = InbOnlyConnectedLOD;
 	}
 }
 
@@ -274,7 +275,8 @@ bool FGeneratedKey::operator==(const FGeneratedKey& Other) const
 		Pin == Other.Pin &&
 		LOD == Other.LOD &&
 		Flags == Other.Flags &&
-		MeshMorphStack == Other.MeshMorphStack;
+		MeshMorphStack == Other.MeshMorphStack &&
+		bOnlyConnectedLOD == Other.bOnlyConnectedLOD;
 }
 
 
@@ -285,6 +287,7 @@ uint32 GetTypeHash(const FGeneratedKey& Key)
 	Hash = HashCombine(Hash, GetTypeHash(Key.LOD));
 	Hash = HashCombine(Hash, GetTypeHash(Key.Flags));
 	//Hash = HashCombine(Hash, GetTypeHash(Key.MeshMorphStack)); // Does not support array
+	Hash = HashCombine(Hash, GetTypeHash(Key.bOnlyConnectedLOD));
 	
 	return Hash;
 }
diff --git a/Engine/Plugins/Experimental/Mutable/Source/CustomizableObjectEditor/Private/MuCOE/GenerateMutableSource/GenerateMutableSource.h b/Engine/Plugins/Experimental/Mutable/Source/CustomizableObjectEditor/Private/MuCOE/GenerateMutableSource/GenerateMutableSource.h
--- a/Engine/Plugins/Experimental/Mutable/Source/CustomizableObjectEditor/Private/MuCOE/GenerateMutableSource/GenerateMutableSource.h
+++ b/Engine/Plugins/Experimental/Mutable/Source/CustomizableObjectEditor/Private/MuCOE/GenerateMutableSource/GenerateMutableSource.h
@@ -113,7 +113,7 @@ class FGeneratedKey
 	friend uint32 GetTypeHash(const FGeneratedKey& Key);
 
 public:
-	FGeneratedKey(void* InFunctionAddress, const UEdGraphPin& InPin, const UCustomizableObjectNode& Node, FMutableGraphGenerationContext& GenerationContext, const bool UseMesh = false);
+	FGeneratedKey(void* InFunctionAddress, const UEdGraphPin& InPin, const UCustomizableObjectNode& Node, FMutableGraphGenerationContext& GenerationContext, const bool UseMesh = false, bool bOnlyConnectedLOD = false);
 	
 	bool operator==(const FGeneratedKey& Other) const;
 
@@ -130,6 +130,9 @@ private:
 
 	/** Active morphs at the time of mesh generation. */
 	TArray<FMorphNodeData> MeshMorphStack;
+
+	/** When caching a generated mesh, true if we force to generate the connected LOD when using Automatic LODs From Mesh. */
+	bool bOnlyConnectedLOD = false;
 };
 
 
diff --git a/Engine/Plugins/Experimental/Mutable/Source/CustomizableObjectEditor/Private/MuCOE/GenerateMutableSource/GenerateMutableSourceColor.cpp b/Engine/Plugins/Experimental/Mutable/Source/CustomizableObjectEditor/Private/MuCOE/GenerateMutableSource/GenerateMutableSourceColor.cpp
--- a/Engine/Plugins/Experimental/Mutable/Source/CustomizableObjectEditor/Private/MuCOE/GenerateMutableSource/GenerateMutableSourceColor.cpp
+++ b/Engine/Plugins/Experimental/Mutable/Source/CustomizableObjectEditor/Private/MuCOE/GenerateMutableSource/GenerateMutableSourceColor.cpp
@@ -317,7 +317,9 @@ mu::NodeColourPtr GenerateMutableSourceColor(const UEdGraphPin* Pin, FMutableGra
 					// Generating a new Color column if not exists
 					if (Table->FindColumn(StringCast<ANSICHAR>(*ColumnName).Get()) == INDEX_NONE)
 					{
-						bSuccess = GenerateTableColumn(TypedNodeTable, Pin, Table, ColumnName, GenerationContext.CurrentLOD, GenerationContext);
+						int32 Dummy = -1; // TODO MTBL-1512
+						bool Dummy2 = false;
+						bSuccess = GenerateTableColumn(TypedNodeTable, Pin, Table, ColumnName, Dummy, Dummy, GenerationContext.CurrentLOD, Dummy, Dummy2, GenerationContext);
 
 						if (!bSuccess)
 						{
diff --git a/Engine/Plugins/Experimental/Mutable/Source/CustomizableObjectEditor/Private/MuCOE/GenerateMutableSource/GenerateMutableSourceFloat.cpp b/Engine/Plugins/Experimental/Mutable/Source/CustomizableObjectEditor/Private/MuCOE/GenerateMutableSource/GenerateMutableSourceFloat.cpp
--- a/Engine/Plugins/Experimental/Mutable/Source/CustomizableObjectEditor/Private/MuCOE/GenerateMutableSource/GenerateMutableSourceFloat.cpp
+++ b/Engine/Plugins/Experimental/Mutable/Source/CustomizableObjectEditor/Private/MuCOE/GenerateMutableSource/GenerateMutableSourceFloat.cpp
@@ -343,7 +343,9 @@ mu::NodeScalarPtr GenerateMutableSourceFloat(const UEdGraphPin* Pin, FMutableGra
 					// Generating a new Float column if not exists
 					if (Table->FindColumn(StringCast<ANSICHAR>(*ColumnName).Get()) == INDEX_NONE)
 					{
-						bSuccess = GenerateTableColumn(TypedNodeTable, Pin, Table, ColumnName, GenerationContext.CurrentLOD, GenerationContext);
+						int32 Dummy = -1; // TODO MTBL-1512
+						bool Dummy2 = false;
+						bSuccess = GenerateTableColumn(TypedNodeTable, Pin, Table, ColumnName, Dummy, Dummy, GenerationContext.CurrentLOD, Dummy, Dummy2, GenerationContext);
 
 						if (!bSuccess)
 						{
diff --git a/Engine/Plugins/Experimental/Mutable/Source/CustomizableObjectEditor/Private/MuCOE/GenerateMutableSource/GenerateMutableSourceImage.cpp b/Engine/Plugins/Experimental/Mutable/Source/CustomizableObjectEditor/Private/MuCOE/GenerateMutableSource/GenerateMutableSourceImage.cpp
--- a/Engine/Plugins/Experimental/Mutable/Source/CustomizableObjectEditor/Private/MuCOE/GenerateMutableSource/GenerateMutableSourceImage.cpp
+++ b/Engine/Plugins/Experimental/Mutable/Source/CustomizableObjectEditor/Private/MuCOE/GenerateMutableSource/GenerateMutableSourceImage.cpp
@@ -648,7 +648,7 @@ mu::NodeImagePtr GenerateMutableSourceImage(const UEdGraphPin* Pin, FMutableGrap
 		if (const UEdGraphPin* ConnectedPin = FollowInputPin(*TypedNodeProject->MeshPin()))
 		{
 			FMutableGraphMeshGenerationData DummyMeshData;
-			mu::NodeMeshPtr MeshNode = GenerateMutableSourceMesh(ConnectedPin, GenerationContext, DummyMeshData);
+			mu::NodeMeshPtr MeshNode = GenerateMutableSourceMesh(ConnectedPin, GenerationContext, DummyMeshData, false, false);
 			ImageNode->SetMesh(MeshNode);
 		}
 
@@ -984,7 +984,9 @@ mu::NodeImagePtr GenerateMutableSourceImage(const UEdGraphPin* Pin, FMutableGrap
 						// Generating a new Texture column if not exists
 						if (Table->FindColumn(StringCast<ANSICHAR>(*ColumnName).Get()) == INDEX_NONE)
 						{
-							bSuccess = GenerateTableColumn(TypedNodeTable, Pin, Table, ColumnName, GenerationContext.CurrentLOD, GenerationContext);
+							int32 Dummy = -1; // TODO MTBL-1512
+							bool Dummy2 = false;
+							bSuccess = GenerateTableColumn(TypedNodeTable, Pin, Table, ColumnName, Dummy, Dummy, GenerationContext.CurrentLOD, Dummy, Dummy2, GenerationContext);
 
 							if (!bSuccess)
 							{
diff --git a/Engine/Plugins/Experimental/Mutable/Source/CustomizableObjectEditor/Private/MuCOE/GenerateMutableSource/GenerateMutableSourceMesh.cpp b/Engine/Plugins/Experimental/Mutable/Source/CustomizableObjectEditor/Private/MuCOE/GenerateMutableSource/GenerateMutableSourceMesh.cpp
--- a/Engine/Plugins/Experimental/Mutable/Source/CustomizableObjectEditor/Private/MuCOE/GenerateMutableSource/GenerateMutableSourceMesh.cpp
+++ b/Engine/Plugins/Experimental/Mutable/Source/CustomizableObjectEditor/Private/MuCOE/GenerateMutableSource/GenerateMutableSourceMesh.cpp
@@ -49,25 +49,21 @@
 #define LOCTEXT_NAMESPACE "CustomizableObjectEditor"
 
 
-/** Returns the corrected LOD and Section Index when using Automatic LOD From Mesh strategy.
- *
- * Do not confuse Section Index and Material Index, they are not the same.
- * 
- * @param Context 
- * @param Node 
- * @param SkeletalMesh 
- * @param InOutLODIndex In connected pin LOD Index. Out corrected Skeletal Mesh LOD Index.
- * @param InOutPinSectionIndex In connected pin Section Index. Out corrected Skeletal Mesh Section Index. */
-void GetEffectiveLODAndSection(const FMutableGraphGenerationContext& Context, const UCustomizableObjectNode& Node, const USkeletalMesh& SkeletalMesh, int32& InOutLODIndex, int32& InOutPinSectionIndex)
+void GetLODAndSectionForAutomaticLODs(const FMutableGraphGenerationContext& Context, const UCustomizableObjectNode& Node, const USkeletalMesh& SkeletalMesh,
+	const int32 LODIndexConnected, const int32 SectionIndexConnected, int32& OutLODIndex, int32& OutSectionIndex, const bool bOnlyConnectedLOD)
 {
-	if (Context.CurrentAutoLODStrategy != ECustomizableObjectAutomaticLODStrategy::AutomaticFromMesh)
+	OutLODIndex = LODIndexConnected;
+	OutSectionIndex = SectionIndexConnected;
+	
+	if (Context.CurrentAutoLODStrategy != ECustomizableObjectAutomaticLODStrategy::AutomaticFromMesh ||
+		bOnlyConnectedLOD)
 	{
 		return;
 	}
 	
-	const int32 CurrentLOD = InOutLODIndex + Context.CurrentLOD;
+	const int32 CompilingLODIndex = LODIndexConnected + Context.CurrentLOD;
 	
-	if (CurrentLOD == InOutLODIndex)
+	if (CompilingLODIndex == LODIndexConnected)
 	{
 		return;
 	}
@@ -78,52 +74,49 @@ void GetEffectiveLODAndSection(const FMutableGraphGenerationContext& Context, co
 		return;
 	}
 
-	if (!ImportedModel->LODModels.IsValidIndex(InOutLODIndex) ||
-		!ImportedModel->LODModels[InOutLODIndex].Sections.IsValidIndex(InOutPinSectionIndex))
+	if (!ImportedModel->LODModels.IsValidIndex(LODIndexConnected) ||
+		!ImportedModel->LODModels[LODIndexConnected].Sections.IsValidIndex(SectionIndexConnected))
 	{
 		return;
 	}
 	
-	const int32 SearchLODMaterialIndex = ImportedModel->LODModels[InOutLODIndex].Sections[InOutPinSectionIndex].MaterialIndex; // Material Index of the connected pin
+	const int32 SearchLODMaterialIndex = ImportedModel->LODModels[LODIndexConnected].Sections[SectionIndexConnected].MaterialIndex; // Material Index of the connected pin
 
-	bool bRepeatedSections = false;
-	
-	const int32 StartLOD = FMath::Min(CurrentLOD, ImportedModel->LODModels.Num() - 1); // Not all meshes have the number of LODs specified in the Mutable graph
-	for (int32 LODIndex = StartLOD; LODIndex >= 0; --LODIndex)
+	if (!ImportedModel->LODModels.IsValidIndex(CompilingLODIndex))
 	{
-		const FSkeletalMeshLODModel& LODModel = ImportedModel->LODModels[LODIndex];
-		const TArray<int32>& MaterialMap = SkeletalMesh.GetLODInfoArray()[LODIndex].LODMaterialMap;
+		OutLODIndex = -1;
+		OutSectionIndex = -1;
+		return;
+	}
+
+	const FSkeletalMeshLODModel& LODModel = ImportedModel->LODModels[CompilingLODIndex];
+	const TArray<int32>& MaterialMap = SkeletalMesh.GetLODInfoArray()[CompilingLODIndex].LODMaterialMap;
 
-		bool bFound = false;
-		for (int32 SectionIndex = 0; SectionIndex < LODModel.Sections.Num(); ++SectionIndex)
+	bool bFound = false;
+	for (int32 SectionIndex = 0; SectionIndex < LODModel.Sections.Num(); ++SectionIndex)
+	{
+		const int32 MaterialIndex = MaterialMap.IsValidIndex(SectionIndex) ? MaterialMap[SectionIndex] : LODModel.Sections[SectionIndex].MaterialIndex; // MaterialMap overrides the MaterialIndex in the section
+			
+		if (MaterialIndex == SearchLODMaterialIndex &&
+			!LODModel.Sections[SectionIndex].bDisabled)
 		{
-			const int32 MaterialIndex = MaterialMap.IsValidIndex(SectionIndex) ? MaterialMap[SectionIndex] : LODModel.Sections[SectionIndex].MaterialIndex; // MaterialMap overrides the MaterialIndex in the section
-				
-			if (MaterialIndex == SearchLODMaterialIndex &&
-				!LODModel.Sections[SectionIndex].bDisabled)
+			if (!bFound)
 			{
-				if (!bFound)
-				{
-					InOutLODIndex = LODIndex;
-					InOutPinSectionIndex = SectionIndex;
-					bFound = true;
-				}
-				else
-				{
-					bRepeatedSections = true;
-				}
+				OutLODIndex = CompilingLODIndex;
+				OutSectionIndex = SectionIndex;
+				bFound = true;
+			}
+			else
+			{
+				Context.Compiler->CompilerLog(FText::Format(LOCTEXT("MeshMultipleMaterialIndex", "Mesh {0} contains multiple sections with the same Material Index"), FText::FromString(SkeletalMesh.GetName())), &Node);
 			}
-		}
-
-		if (bFound)
-		{
-			break;
 		}
 	}
 
-	if (bRepeatedSections)
+	if (!bFound)
 	{
-		Context.Compiler->CompilerLog(FText::Format(LOCTEXT("MeshMultipleMaterialIndex", "Mesh {0} contains multiple sections with the same Material Index"), FText::FromString(SkeletalMesh.GetName())), &Node);
+		OutLODIndex = -1;
+		OutSectionIndex = -1;
 	}
 }
 
@@ -627,7 +620,7 @@ mu::Ptr<mu::PhysicsBody> MakePhysicsBodyFromAsset(UPhysicsAsset* Asset, const TA
 	return PhysicsBody;
 }
 
-mu::MeshPtr ConvertSkeletalMeshToMutable(const USkeletalMesh* InSkeletalMesh, const TSoftClassPtr<UAnimInstance>& AnimBp, int LOD, int MaterialIndex, FMutableGraphGenerationContext& GenerationContext, const UCustomizableObjectNode* CurrentNode)
+mu::MeshPtr ConvertSkeletalMeshToMutable(const USkeletalMesh* InSkeletalMesh, const TSoftClassPtr<UAnimInstance>& AnimBp, int32 LODIndexConnected, int32 SectionIndexConnected, int32 LODIndex, int32 SectionIndex, FMutableGraphGenerationContext& GenerationContext, const UCustomizableObjectNode* CurrentNode)
 {
 	if(!InSkeletalMesh)
 	{
@@ -635,38 +628,54 @@ mu::MeshPtr ConvertSkeletalMeshToMutable(const USkeletalMesh* InSkeletalMesh, co
 	}
 
 	const FSkeletalMeshModel* ImportedModel = InSkeletalMesh->GetImportedModel();
+	if (!ImportedModel)
+	{
+		const FString Msg = FString::Printf(TEXT("The SkeletalMesh [%s] doesn't have an imported resource."), *InSkeletalMesh->GetName());
+		GenerationContext.Compiler->CompilerLog(FText::FromString(Msg), CurrentNode);
 
-	// Check in case the data has changed
-	if (!(ImportedModel
-		&& ImportedModel->LODModels.Num() > LOD
-		&& ImportedModel->LODModels[LOD].Sections.Num() > MaterialIndex))
+		return nullptr;
+	}
+
+	if (!ImportedModel->LODModels.IsValidIndex(LODIndex))
 	{
-		FString Msg;
-		if (!ImportedModel)
+		if (GenerationContext.CurrentAutoLODStrategy == ECustomizableObjectAutomaticLODStrategy::AutomaticFromMesh && 
+			LODIndex != LODIndexConnected) // If we are using automatic LODs and not generating the base LOD (the connected one) is not an error.
 		{
-			Msg = FString::Printf(TEXT("The SkeletalMesh [%s] doesn't have an imported resource."), *InSkeletalMesh->GetName());
+			return new mu::Mesh(); // Return empty mesh to preserve the layouts
 		}
-		else if (LOD >= ImportedModel->LODModels.Num())
+		else
 		{
-			Msg = FString::Printf(
+			const FString Msg = FString::Printf(
 				TEXT("The SkeletalMesh [%s] doesn't have the expected number of LODs [need %d, has %d]. Changed after reimporting?"),
 				InSkeletalMesh ? *InSkeletalMesh->GetName() : TEXT("none"),
-				LOD + 1,
+				LODIndex + 1,
 				ImportedModel->LODModels.Num());
+			GenerationContext.Compiler->CompilerLog(FText::FromString(Msg), CurrentNode);
+
+			return nullptr;
+		}
+	}
+
+	if (!ImportedModel->LODModels[LODIndex].Sections.IsValidIndex(SectionIndex))
+	{
+		if (GenerationContext.CurrentAutoLODStrategy == ECustomizableObjectAutomaticLODStrategy::AutomaticFromMesh && 
+			SectionIndex != SectionIndexConnected) // If we are using automatic LODs and not generating the base LOD (the connected one) is not an error.
+		{
+			return new mu::Mesh(); // Return empty mesh to preserve the layouts
 		}
 		else
 		{
-			Msg = FString::Printf(
+			const FString Msg = FString::Printf(
 				TEXT("The SkeletalMesh [%s] doesn't have the expected structure. Maybe the number of LODs [need %d, has %d] or Materials [need %d, has %d] has changed after reimporting?"),
 				*InSkeletalMesh->GetName(),
-				LOD + 1,
+				LODIndex + 1,
 				ImportedModel->LODModels.Num(),
-				MaterialIndex + 1,
-				ImportedModel->LODModels[LOD].Sections.Num()
-			);
+				SectionIndex + 1,
+				ImportedModel->LODModels[LODIndex].Sections.Num());
+			GenerationContext.Compiler->CompilerLog(FText::FromString(Msg), CurrentNode);
+
+			return nullptr;
 		}
-		GenerationContext.Compiler->CompilerLog(FText::FromString(Msg), CurrentNode);
-		return nullptr;
 	}
 
 	// Get the mesh generation flags to use
@@ -681,7 +690,6 @@ mu::MeshPtr ConvertSkeletalMeshToMutable(const USkeletalMesh* InSkeletalMesh, co
 	// should ignore skinning.
 	check(!(GenerationContext.CurrentMeshComponent < 0) || bIgnoreSkeleton);
 
-	// 
 	bool bBoneMapModified = false;
 	TArray<FBoneIndexType> BoneMap;
 	TArray<FBoneIndexType> RemappedBoneMapIndices;
@@ -761,11 +769,11 @@ mu::MeshPtr ConvertSkeletalMeshToMutable(const USkeletalMesh* InSkeletalMesh, co
 			}
 		}
 
-		const TArray<uint16>& SourceRequiredBones = ImportedModel->LODModels[LOD].RequiredBones;
+		const TArray<uint16>& SourceRequiredBones = ImportedModel->LODModels[LODIndex].RequiredBones;
 
 		// Remove bones and build an array to remap indices of the BoneMap
 		TArray<FBoneIndexType> RemappedBones;
-		BuildRemappedBonesArray(MutComponentInfo, InSkeletalMesh, LOD, SourceRequiredBones, RemappedBones);
+		BuildRemappedBonesArray(MutComponentInfo, InSkeletalMesh, LODIndex, SourceRequiredBones, RemappedBones);
 		
 		// Build RequiredBones array
 		TArray<FBoneIndexType> RequiredBones;
@@ -777,7 +785,7 @@ mu::MeshPtr ConvertSkeletalMeshToMutable(const USkeletalMesh* InSkeletalMesh, co
 		}
 
 		// Build BoneMap
-		const TArray<uint16>& SourceBoneMap = ImportedModel->LODModels[LOD].Sections[MaterialIndex].BoneMap;
+		const TArray<uint16>& SourceBoneMap = ImportedModel->LODModels[LODIndex].Sections[SectionIndex].BoneMap;
 		const int32 NumBonesInBoneMap = SourceBoneMap.Num();
 		const int32 NumRemappedBones = RemappedBones.Num();
 
@@ -848,9 +856,9 @@ mu::MeshPtr ConvertSkeletalMeshToMutable(const USkeletalMesh* InSkeletalMesh, co
 
 	// Vertices
 	TArray<FSoftSkinVertex> Vertices;
-	ImportedModel->LODModels[LOD].GetVertices(Vertices);
-	int VertexStart = ImportedModel->LODModels[LOD].Sections[MaterialIndex].GetVertexBufferIndex();
-	int VertexCount = ImportedModel->LODModels[LOD].Sections[MaterialIndex].GetNumVertices();
+	ImportedModel->LODModels[LODIndex].GetVertices(Vertices);
+	int VertexStart = ImportedModel->LODModels[LODIndex].Sections[SectionIndex].GetVertexBufferIndex();
+	int VertexCount = ImportedModel->LODModels[LODIndex].Sections[SectionIndex].GetNumVertices();
 
 	MutableMesh->GetVertexBuffers().SetElementCount(VertexCount);
 
@@ -860,7 +868,7 @@ mu::MeshPtr ConvertSkeletalMeshToMutable(const USkeletalMesh* InSkeletalMesh, co
 
 	MutableMesh->GetVertexBuffers().SetBufferCount(VertexBuffersCount);
 
-	const int32 MaxSectionInfluences = ImportedModel->LODModels[LOD].Sections[MaterialIndex].MaxBoneInfluences;
+	const int32 MaxSectionInfluences = ImportedModel->LODModels[LODIndex].Sections[SectionIndex].MaxBoneInfluences;
 	const bool bUseUnlimitedInfluences = FGPUBaseSkinVertexFactory::UseUnlimitedBoneInfluences(MaxSectionInfluences, GenerationContext.Options.TargetPlatform);
 
 	if (bIgnoreSkeleton)
@@ -1176,17 +1184,17 @@ mu::MeshPtr ConvertSkeletalMeshToMutable(const USkeletalMesh* InSkeletalMesh, co
 
 					const TArray<FMorphTargetLODModel>& MorphLODModels = MorphTarget->GetMorphLODModels();
 
-					if (LOD >= MorphLODModels.Num()
-						|| !MorphLODModels[LOD].SectionIndices.Contains(MaterialIndex))
+					if (LODIndex >= MorphLODModels.Num()
+						|| !MorphLODModels[LODIndex].SectionIndices.Contains(SectionIndex))
 					{
 						continue;
 					}
 
 					// The vertices should be sorted by SourceIdx
-					check(MorphLODModels[LOD].Vertices.Num() < 2 || MorphLODModels[LOD].Vertices[0].SourceIdx < MorphLODModels[LOD].Vertices.Last().SourceIdx);
+					check(MorphLODModels[LODIndex].Vertices.Num() < 2 || MorphLODModels[LODIndex].Vertices[0].SourceIdx < MorphLODModels[LODIndex].Vertices.Last().SourceIdx);
 
 					const int32 VertexFoundIndex = Algo::BinarySearchBy(
-						MorphLODModels[LOD].Vertices,
+						MorphLODModels[LODIndex].Vertices,
 						(uint32)VertexIdx,
 						[](const FMorphTargetDelta& Element) { return Element.SourceIdx; });
 
@@ -1195,7 +1203,7 @@ mu::MeshPtr ConvertSkeletalMeshToMutable(const USkeletalMesh* InSkeletalMesh, co
 						continue;
 					}
 
-					const FMorphTargetDelta& VertexFound = MorphLODModels[LOD].Vertices[VertexFoundIndex];
+					const FMorphTargetDelta& VertexFound = MorphLODModels[LODIndex].Vertices[VertexFoundIndex];
 					const FName MorphTargetName = MorphTarget->GetFName();
 
 					TArray<FMorphTargetInfo>& ContributingMorphTargetsInfo = GenerationContext.ContributingMorphTargetsInfo;
@@ -1252,7 +1260,7 @@ mu::MeshPtr ConvertSkeletalMeshToMutable(const USkeletalMesh* InSkeletalMesh, co
 		// clothing assets are shared among all LODs in a section
 		const int32 ClothingAssetIndex = Invoke([&]() -> int32
 		{
-			const UClothingAssetBase* ClothingAssetBase = InSkeletalMesh->GetSectionClothingAsset(LOD, MaterialIndex);
+			const UClothingAssetBase* ClothingAssetBase = InSkeletalMesh->GetSectionClothingAsset(LODIndex, SectionIndex);
 
 			if (!ClothingAssetBase)
 			{
@@ -1307,7 +1315,7 @@ mu::MeshPtr ConvertSkeletalMeshToMutable(const USkeletalMesh* InSkeletalMesh, co
 				FirstElem.SourceAssetIndex = INDEX_NONE;
 			}
 
-			const FSkelMeshSection& SectionData = ImportedModel->LODModels[LOD].Sections[MaterialIndex];
+			const FSkelMeshSection& SectionData = ImportedModel->LODModels[LODIndex].Sections[SectionIndex];
 			const TArray<FMeshToMeshVertData>& ClothMappingData = SectionData.ClothMappingDataLODs[0];
 
 
@@ -1351,7 +1359,7 @@ mu::MeshPtr ConvertSkeletalMeshToMutable(const USkeletalMesh* InSkeletalMesh, co
 			(GenerationContext.ClothMeshToMeshVertData.GetData() + ClothDataIndexBase, NumClothMappingDataVerts);
 
 			const FCustomizableObjectClothingAssetData& ClothingAssetData = GenerationContext.ContributingClothingAssetsData[ClothingAssetIndex];
-			const int16 ClothingAssetLODIndex = static_cast<int16>(ClothingAssetData.LodMap[LOD]);
+			const int16 ClothingAssetLODIndex = static_cast<int16>(ClothingAssetData.LodMap[LODIndex]);
 
 			for (FCustomizableObjectMeshToMeshVertData& ClothingDataElem : AppendedClothingDataView)
 			{
@@ -1386,12 +1394,12 @@ mu::MeshPtr ConvertSkeletalMeshToMutable(const USkeletalMesh* InSkeletalMesh, co
 		const int32 BoneWeightsSize = MutableBonesPerVertex * BoneWeightTypeSizeBytes;
 		const int32 SkinWeightProfileVertexSize = sizeof(int32) + BoneIndicesSize + BoneWeightsSize;
 
-		const int32 MaxSectionBoneMapIndex = ImportedModel->LODModels[LOD].Sections[MaterialIndex].BoneMap.Num();
+		const int32 MaxSectionBoneMapIndex = ImportedModel->LODModels[LODIndex].Sections[SectionIndex].BoneMap.Num();
 
 		const TArray<FSkinWeightProfileInfo>& SkinWeightProfilesInfo = InSkeletalMesh->GetSkinWeightProfiles();
 		for (const FSkinWeightProfileInfo& Profile : SkinWeightProfilesInfo)
 		{
-			const FImportedSkinWeightProfileData* ImportedProfileData = ImportedModel->LODModels[LOD].SkinWeightProfiles.Find(Profile.Name);
+			const FImportedSkinWeightProfileData* ImportedProfileData = ImportedModel->LODModels[LODIndex].SkinWeightProfiles.Find(Profile.Name);
 			if (!ImportedProfileData)
 			{
 				continue;
@@ -1476,8 +1484,8 @@ mu::MeshPtr ConvertSkeletalMeshToMutable(const USkeletalMesh* InSkeletalMesh, co
 
 	// Indices
 	{
-		int IndexStart = ImportedModel->LODModels[LOD].Sections[MaterialIndex].BaseIndex;
-		int IndexCount = ImportedModel->LODModels[LOD].Sections[MaterialIndex].NumTriangles * 3;
+		int IndexStart = ImportedModel->LODModels[LODIndex].Sections[SectionIndex].BaseIndex;
+		int IndexCount = ImportedModel->LODModels[LODIndex].Sections[SectionIndex].NumTriangles * 3;
 		MutableMesh->GetIndexBuffers().SetBufferCount(1);
 		MutableMesh->GetIndexBuffers().SetElementCount(IndexCount);
 		MutableMesh->GetFaceBuffers().SetElementCount(IndexCount / 3);
@@ -1486,8 +1494,8 @@ mu::MeshPtr ConvertSkeletalMeshToMutable(const USkeletalMesh* InSkeletalMesh, co
 		// For some reason, the indices in 4.25 (and 4.24) are in different order in the Imported and Rendering data structures. The strange thing 
 		// is actually that the vertices in the imported model seem to match the rendering model indices. Maybe there is some mapping that
 		// we are missing, but for now this will do:
-		const int ElementSize = InSkeletalMesh->GetResourceForRendering()->LODRenderData[LOD].MultiSizeIndexContainer.GetDataTypeSize();
-		void* IndexDataPointer = InSkeletalMesh->GetResourceForRendering()->LODRenderData[LOD].MultiSizeIndexContainer.GetIndexBuffer()->GetPointerTo(IndexStart);
+		const int ElementSize = InSkeletalMesh->GetResourceForRendering()->LODRenderData[LODIndex].MultiSizeIndexContainer.GetDataTypeSize();
+		void* IndexDataPointer = InSkeletalMesh->GetResourceForRendering()->LODRenderData[LODIndex].MultiSizeIndexContainer.GetIndexBuffer()->GetPointerTo(IndexStart);
 		const int FinalElementSize = sizeof(uint32_t);
 		const int ChannelCount = 1;
 		const MESH_BUFFER_SEMANTIC Semantics[ChannelCount] = { MBS_VERTEXINDEX };
@@ -1800,13 +1808,13 @@ mu::MeshPtr ConvertSkeletalMeshToMutable(const USkeletalMesh* InSkeletalMesh, co
 }
 
 
-mu::MeshPtr ConvertStaticMeshToMutable(const UStaticMesh* StaticMesh, int LOD, int MaterialIndex, FMutableGraphGenerationContext& GenerationContext, const UCustomizableObjectNode* CurrentNode)
+mu::MeshPtr ConvertStaticMeshToMutable(const UStaticMesh* StaticMesh, int32 LODIndex, int32 SectionIndex, FMutableGraphGenerationContext& GenerationContext, const UCustomizableObjectNode* CurrentNode)
 {
 	if (!StaticMesh->GetRenderData() ||
-		!StaticMesh->GetRenderData()->LODResources.IsValidIndex(LOD) ||
-		!StaticMesh->GetRenderData()->LODResources[LOD].Sections.IsValidIndex(MaterialIndex))
+		!StaticMesh->GetRenderData()->LODResources.IsValidIndex(LODIndex) ||
+		!StaticMesh->GetRenderData()->LODResources[LODIndex].Sections.IsValidIndex(SectionIndex))
 	{
-		FString Msg = FString::Printf(TEXT("Degenerated static mesh found for LOD %d Material %d. It will be ignored. "), LOD, MaterialIndex);
+		FString Msg = FString::Printf(TEXT("Degenerated static mesh found for LOD %d Material %d. It will be ignored. "), LODIndex, SectionIndex);
 		GenerationContext.Compiler->CompilerLog(FText::FromString(Msg), CurrentNode, EMessageSeverity::Warning);
 		return nullptr;
 	}
@@ -1814,8 +1822,8 @@ mu::MeshPtr ConvertStaticMeshToMutable(const UStaticMesh* StaticMesh, int LOD, i
 	mu::MeshPtr MutableMesh = new mu::Mesh();
 
 	// Vertices
-	int VertexStart = StaticMesh->GetRenderData()->LODResources[LOD].Sections[MaterialIndex].MinVertexIndex;
-	int VertexCount = StaticMesh->GetRenderData()->LODResources[LOD].Sections[MaterialIndex].MaxVertexIndex - VertexStart + 1;
+	int VertexStart = StaticMesh->GetRenderData()->LODResources[LODIndex].Sections[SectionIndex].MinVertexIndex;
+	int VertexCount = StaticMesh->GetRenderData()->LODResources[LODIndex].Sections[SectionIndex].MaxVertexIndex - VertexStart + 1;
 
 	MutableMesh->GetVertexBuffers().SetElementCount(VertexCount);
 	{
@@ -1825,7 +1833,7 @@ mu::MeshPtr ConvertStaticMeshToMutable(const UStaticMesh* StaticMesh, int LOD, i
 
 		// Position buffer
 		{
-			const FPositionVertexBuffer& VertexBuffer = StaticMesh->GetRenderData()->LODResources[LOD].VertexBuffers.PositionVertexBuffer;
+			const FPositionVertexBuffer& VertexBuffer = StaticMesh->GetRenderData()->LODResources[LODIndex].VertexBuffers.PositionVertexBuffer;
 
 			const int ElementSize = 12;
 			const int ChannelCount = 1;
@@ -1844,7 +1852,7 @@ mu::MeshPtr ConvertStaticMeshToMutable(const UStaticMesh* StaticMesh, int LOD, i
 
 		// Tangent buffer
 		{
-			const FStaticMeshVertexBuffer& VertexBuffer = StaticMesh->GetRenderData()->LODResources[LOD].VertexBuffers.StaticMeshVertexBuffer;
+			const FStaticMeshVertexBuffer& VertexBuffer = StaticMesh->GetRenderData()->LODResources[LODIndex].VertexBuffers.StaticMeshVertexBuffer;
 
 			MESH_BUFFER_SEMANTIC Semantics[2];
 			int SemanticIndices[2];
@@ -1883,7 +1891,7 @@ mu::MeshPtr ConvertStaticMeshToMutable(const UStaticMesh* StaticMesh, int LOD, i
 
 		// Texture coordinates
 		{
-			const FStaticMeshVertexBuffer& VertexBuffer = StaticMesh->GetRenderData()->LODResources[LOD].VertexBuffers.StaticMeshVertexBuffer;
+			const FStaticMeshVertexBuffer& VertexBuffer = StaticMesh->GetRenderData()->LODResources[LODIndex].VertexBuffers.StaticMeshVertexBuffer;
 
 			int texChannels = VertexBuffer.GetNumTexCoords();
 			int ChannelCount = texChannels;
@@ -1939,8 +1947,8 @@ mu::MeshPtr ConvertStaticMeshToMutable(const UStaticMesh* StaticMesh, int LOD, i
 
 	// Indices
 	{
-		int IndexStart = StaticMesh->GetRenderData()->LODResources[LOD].Sections[MaterialIndex].FirstIndex;
-		int IndexCount = StaticMesh->GetRenderData()->LODResources[LOD].Sections[MaterialIndex].NumTriangles * 3;
+		int IndexStart = StaticMesh->GetRenderData()->LODResources[LODIndex].Sections[SectionIndex].FirstIndex;
+		int IndexCount = StaticMesh->GetRenderData()->LODResources[LODIndex].Sections[SectionIndex].NumTriangles * 3;
 		MutableMesh->GetIndexBuffers().SetBufferCount(1);
 		MutableMesh->GetIndexBuffers().SetElementCount(IndexCount);
 		MutableMesh->GetFaceBuffers().SetElementCount(IndexCount / 3);
@@ -1970,7 +1978,7 @@ mu::MeshPtr ConvertStaticMeshToMutable(const UStaticMesh* StaticMesh, int LOD, i
 		//}
 		//else
 		{
-			FIndexArrayView Source = StaticMesh->GetRenderData()->LODResources[LOD].IndexBuffer.GetArrayView();
+			FIndexArrayView Source = StaticMesh->GetRenderData()->LODResources[LODIndex].IndexBuffer.GetArrayView();
 			//const uint16* pSource = reinterpret_cast<const uint16*>( StaticMesh->RenderData->LODResources[LOD].IndexBuffer.Indices.GetResourceData() );
 			//pSource += IndexStart;
 			uint16* pDest = reinterpret_cast<uint16*>(MutableMesh->GetIndexBuffers().GetBufferData(0));
@@ -1989,23 +1997,51 @@ mu::MeshPtr ConvertStaticMeshToMutable(const UStaticMesh* StaticMesh, int LOD, i
 
 // Convert a Mesh constant to a mutable format. UniqueTags are the tags that make this Mesh unique that cannot be merged in the cache 
 //  with the exact same Mesh with other tags
-mu::MeshPtr GenerateMutableMesh(const UObject* Mesh, const TSoftClassPtr<UAnimInstance>& AnimInstance, int32 LOD, int32 MaterialIndex, const FString& UniqueTags, FMutableGraphGenerationContext & GenerationContext, const UCustomizableObjectNode* CurrentNode)
+mu::MeshPtr GenerateMutableMesh(const UObject * Mesh, const TSoftClassPtr<UAnimInstance>& AnimInstance, int32 LODIndexConnected, int32 SectionIndexConnected, int32 LODIndex, int32 SectionIndex, const FString& UniqueTags, FMutableGraphGenerationContext & GenerationContext, const UCustomizableObjectNode* CurrentNode)
 {
 	// Get the mesh generation flags to use
 	EMutableMeshConversionFlags CurrentFlags = GenerationContext.MeshGenerationFlags.Last();
 
-	FMutableGraphGenerationContext::FGeneratedMeshData::FKey Key = { Mesh, LOD, GenerationContext.CurrentLOD, MaterialIndex, CurrentFlags, UniqueTags, CurrentNode };
+	FMutableGraphGenerationContext::FGeneratedMeshData::FKey Key = { Mesh, LODIndex, GenerationContext.CurrentLOD, SectionIndex, CurrentFlags, UniqueTags, CurrentNode };
 	mu::MeshPtr MutableMesh = GenerationContext.FindGeneratedMesh(Key);
 	
 	if (!MutableMesh)
 	{
 		if (const USkeletalMesh* SkeletalMesh = Cast<USkeletalMesh>(Mesh))
 		{
-			MutableMesh = ConvertSkeletalMeshToMutable(SkeletalMesh, AnimInstance, LOD, MaterialIndex, GenerationContext, CurrentNode);
+			MutableMesh = ConvertSkeletalMeshToMutable(SkeletalMesh, AnimInstance, LODIndexConnected, SectionIndexConnected, LODIndex, SectionIndex, GenerationContext, CurrentNode);
+
+			FSkeletalMeshModel* ImportedModel = SkeletalMesh->GetImportedModel();
+			
+			if (MutableMesh &&
+				ImportedModel->LODModels.IsValidIndex(LODIndex) &&
+				ImportedModel->LODModels[LODIndex].Sections.IsValidIndex(SectionIndex))
+			{
+				FMeshData MeshData;
+				MeshData.Mesh = Mesh;
+				MeshData.LOD = LODIndex;
+				MeshData.MaterialIndex = SectionIndex;
+				MeshData.Node = CurrentNode;
+				GenerationContext.PinData.GetCurrent().MeshesData.Add(MeshData); // Set::Emplace only supports single element constructors
+			}
 		}
 		else if (const UStaticMesh* StaticMesh = Cast<UStaticMesh>(Mesh))
 		{
-			MutableMesh = ConvertStaticMeshToMutable(StaticMesh, LOD, MaterialIndex, GenerationContext, CurrentNode);
+			MutableMesh = ConvertStaticMeshToMutable(StaticMesh, LODIndex, SectionIndex, GenerationContext, CurrentNode);
+
+			const FStaticMeshRenderData* RenderData = StaticMesh->GetRenderData();
+
+			if (MutableMesh &&
+				RenderData->LODResources.IsValidIndex(LODIndex) &&
+				RenderData->LODResources[LODIndex].Sections.IsValidIndex(SectionIndex))
+			{
+				FMeshData MeshData;
+				MeshData.Mesh = Mesh;
+				MeshData.LOD = LODIndex;
+				MeshData.MaterialIndex = SectionIndex;
+				MeshData.Node = CurrentNode;
+				GenerationContext.PinData.GetCurrent().MeshesData.Add(MeshData); // Set::Emplace only supports single element constructors
+			}
 		}
 		else
 		{
@@ -2018,21 +2054,11 @@ mu::MeshPtr GenerateMutableMesh(const UObject* Mesh, const TSoftClassPtr<UAnimIn
 		}
 	}
 	
-	if (MutableMesh)
-	{
-		FMeshData MeshData;
-		MeshData.Mesh = Mesh;
-		MeshData.LOD = LOD;
-		MeshData.MaterialIndex = MaterialIndex;
-		MeshData.Node = CurrentNode;
-		GenerationContext.PinData.GetCurrent().MeshesData.Add(MeshData); // Set::Emplace only supports single element constructors
-	}
-	
 	return MutableMesh;
 }
 
 
-mu::MeshPtr BuildMorphedMutableMesh(const UEdGraphPin* BaseSourcePin, const FString& MorphTargetName, FMutableGraphGenerationContext & GenerationContext, const FName& RowName)
+mu::MeshPtr BuildMorphedMutableMesh(const UEdGraphPin* BaseSourcePin, const FString& MorphTargetName, FMutableGraphGenerationContext & GenerationContext, const bool bOnlyConnectedLOD, const FName& RowName)
 {
 	check(BaseSourcePin);
 	SCOPED_PIN_DATA(GenerationContext, BaseSourcePin)
@@ -2045,6 +2071,9 @@ mu::MeshPtr BuildMorphedMutableMesh(const UEdGraphPin* BaseSourcePin, const FStr
 		return nullptr;
 	}
 
+	int32 LODIndexConnected = -1; // LOD which the pin is connected to
+	int32 SectionIndexConnected = -1;
+	
 	int32 LODIndex = -1; // Initialization required to remove uninitialized warning.
 	int32 SectionIndex = -1;
 	
@@ -2054,22 +2083,21 @@ mu::MeshPtr BuildMorphedMutableMesh(const UEdGraphPin* BaseSourcePin, const FStr
 	if (const UCustomizableObjectNodeSkeletalMesh* TypedNodeSkeletalMesh = Cast<UCustomizableObjectNodeSkeletalMesh>(Node))
 	{
 		int32 LayoutIndex;
-		TypedNodeSkeletalMesh->GetPinSection(*BaseSourcePin, LODIndex, SectionIndex, LayoutIndex);
+		TypedNodeSkeletalMesh->GetPinSection(*BaseSourcePin, LODIndexConnected, SectionIndexConnected, LayoutIndex);
 		SkeletalMesh = TypedNodeSkeletalMesh->SkeletalMesh;
 	}
 
 	else if (const UCustomizableObjectNodeTable* TypedNodeTable = Cast<UCustomizableObjectNodeTable>(Node))
 	{
-		TypedNodeTable->GetPinLODAndSection(BaseSourcePin, LODIndex, SectionIndex);
+		TypedNodeTable->GetPinLODAndSection(BaseSourcePin, LODIndexConnected, SectionIndexConnected);
 		SkeletalMesh = TypedNodeTable->GetSkeletalMeshAt(BaseSourcePin, RowName);
 	}
 
 	if (SkeletalMesh)
 	{
-		GetEffectiveLODAndSection(GenerationContext, *Node, *SkeletalMesh, LODIndex, SectionIndex);
-
+		GetLODAndSectionForAutomaticLODs(GenerationContext, *Node, *SkeletalMesh, LODIndexConnected, SectionIndexConnected, LODIndex, SectionIndex, bOnlyConnectedLOD);
 		// Get the base mesh
-		mu::MeshPtr BaseSourceMesh = GenerateMutableMesh(SkeletalMesh, TSoftClassPtr<UAnimInstance>(), LODIndex, SectionIndex, FString(), GenerationContext, Node);
+		mu::MeshPtr BaseSourceMesh = GenerateMutableMesh(SkeletalMesh, TSoftClassPtr<UAnimInstance>(), LODIndexConnected, SectionIndexConnected, LODIndex, SectionIndex, FString(), GenerationContext, Node);
 		if (BaseSourceMesh)
 		{
 			// Clone it (it will probably be shared)
@@ -2602,6 +2630,7 @@ mu::NodeMeshPtr GenerateMorphMesh(const UEdGraphPin* Pin,
 	mu::NodeMeshPtr SourceNode,
 	FMutableGraphGenerationContext & GenerationContext,
 	FMutableGraphMeshGenerationData & MeshData,
+	const bool bOnlyConnectedLOD,
 	const FString& TableColumnName = "")
 {
 	SCOPED_PIN_DATA(GenerationContext, Pin)
@@ -2626,7 +2655,7 @@ mu::NodeMeshPtr GenerateMorphMesh(const UEdGraphPin* Pin,
 	}
 	else
 	{
-		mu::NodeMeshPtr NextMorph = GenerateMorphMesh(Pin, TypedNodeMorphs, MorphIndex + 1, SourceNode, GenerationContext, MeshData, TableColumnName);
+		mu::NodeMeshPtr NextMorph = GenerateMorphMesh(Pin, TypedNodeMorphs, MorphIndex + 1, SourceNode, GenerationContext, MeshData, bOnlyConnectedLOD, TableColumnName); // TODO FutureGMT change to a for. This recursion can be problematic with the production cache
 		Result->SetBase(NextMorph);
 	}
 	
@@ -2659,7 +2688,7 @@ mu::NodeMeshPtr GenerateMorphMesh(const UEdGraphPin* Pin,
 				ColumnIndex = Table->AddColumn(StringCast<ANSICHAR>(*ColumnName).Get(), mu::TABLE_COLUMN_TYPE::TCT_MESH);
 			}
 
-			mu::MeshPtr MorphedSourceTableMesh = BuildMorphedMutableMesh(Pin, TypedNodeMorphs[MorphIndex].MorphTargetName, GenerationContext, RowName);
+			mu::MeshPtr MorphedSourceTableMesh = BuildMorphedMutableMesh(Pin, TypedNodeMorphs[MorphIndex].MorphTargetName, GenerationContext, bOnlyConnectedLOD, RowName);
 			Table->SetCell(ColumnIndex, RowIndex, MorphedSourceTableMesh.get());
 		}
 
@@ -2683,7 +2712,7 @@ mu::NodeMeshPtr GenerateMorphMesh(const UEdGraphPin* Pin,
 	}
 	else
 	{
-		MorphedSourceMesh = BuildMorphedMutableMesh(Pin, TypedNodeMorphs[MorphIndex].MorphTargetName, GenerationContext);
+		MorphedSourceMesh = BuildMorphedMutableMesh(Pin, TypedNodeMorphs[MorphIndex].MorphTargetName, GenerationContext, bOnlyConnectedLOD);
 
 		if (MorphedSourceMesh)
 		{
@@ -2759,34 +2788,6 @@ mu::NodeMeshPtr GenerateMorphMesh(const UEdGraphPin* Pin,
 }
 
 
-void GenerateMorphTarget(const UCustomizableObjectNode* Node, const UEdGraphPin* BaseSourcePin, FMutableGraphGenerationContext& GenerationContext, mu::Ptr<mu::NodeMeshMorph> MeshNode, FString MorphName)
-{
-	SCOPED_PIN_DATA(GenerationContext, BaseSourcePin)
-
-	FMutableGraphMeshGenerationData DummyMeshData;
-	mu::NodeMeshPtr BaseSourceMesh = GenerateMutableSourceMesh(BaseSourcePin, GenerationContext, DummyMeshData);
-
-	mu::MeshPtr MorphedSourceMesh = BuildMorphedMutableMesh(BaseSourcePin, MorphName, GenerationContext);
-	if (MorphedSourceMesh)
-	{
-		mu::NodeMeshConstantPtr MorphedSourceMeshNode = new mu::NodeMeshConstant;
-		MorphedSourceMeshNode->SetValue(MorphedSourceMesh);
-		MorphedSourceMeshNode->SetMessageContext(Node);
-
-		mu::NodeMeshMakeMorphPtr Morph = new mu::NodeMeshMakeMorph;
-		Morph->SetBase(BaseSourceMesh.get());
-		Morph->SetTarget(MorphedSourceMeshNode.get());
-		Morph->SetMessageContext(Node);
-
-		MeshNode->SetMorph(Morph);
-	}
-	else
-	{
-		GenerationContext.Compiler->CompilerLog(LOCTEXT("MorphGenerationFailed", "Failed to generate morph target."), Node);
-	}
-}
-
-
 /** Create a default layout. Used when no layout is found. */
 mu::NodeLayoutBlocksPtr CreateDefaultLayout()
 {
@@ -2805,11 +2806,11 @@ mu::NodeLayoutBlocksPtr CreateDefaultLayout()
 }
 
 
-/** Convert a CustomizableObject Source Graph into a mutable source graph  */
-mu::NodeMeshPtr GenerateMutableSourceMesh(const UEdGraphPin * Pin,
-	FMutableGraphGenerationContext & GenerationContext,
-	FMutableGraphMeshGenerationData & MeshData,
-	bool bLinkedToExtendMaterial)
+mu::NodeMeshPtr GenerateMutableSourceMesh(const UEdGraphPin* Pin,
+	FMutableGraphGenerationContext& GenerationContext,
+	FMutableGraphMeshGenerationData& MeshData,
+	const bool bLinkedToExtendMaterial,
+	const bool bOnlyConnectedLOD)
 {
 	check(Pin)
 	RETURN_ON_CYCLE(*Pin, GenerationContext)
@@ -2819,7 +2820,7 @@ mu::NodeMeshPtr GenerateMutableSourceMesh(const UEdGraphPin * Pin,
 
 	UCustomizableObjectNode* Node = CastChecked<UCustomizableObjectNode>(Pin->GetOwningNode());
 
-	const FGeneratedKey Key(reinterpret_cast<void*>(&GenerateMutableSourceMesh), *Pin , *Node, GenerationContext, true);
+	const FGeneratedKey Key(reinterpret_cast<void*>(&GenerateMutableSourceMesh), *Pin , *Node, GenerationContext, true, bOnlyConnectedLOD);
 	if (const FGeneratedData* Generated = GenerationContext.Generated.Find(Key))
 	{
 		MeshData = Generated->meshData;
@@ -2844,19 +2845,19 @@ mu::NodeMeshPtr GenerateMutableSourceMesh(const UEdGraphPin * Pin,
 
 		if (TypedNodeSkel->SkeletalMesh)
 		{
-			int32 LODIndexConnected; // LOD which the pin is connected to
-			int32 LODIndex;
-			int32 SectionIndex;
+			int32 LODIndexConnected = -1; // LOD which the pin is connected to
+			int32 SectionIndexConnected = -1;
+			
+			int32 LODIndex = -1;
+			int32 SectionIndex = -1;
 
 			{
 				int32 LayoutIndex;
-				TypedNodeSkel->GetPinSection(*Pin, LODIndexConnected, SectionIndex, LayoutIndex);
-
-				LODIndex = LODIndexConnected;
+				TypedNodeSkel->GetPinSection(*Pin, LODIndexConnected, SectionIndexConnected, LayoutIndex);
 			}
-				
-			GetEffectiveLODAndSection(GenerationContext, *Node, *TypedNodeSkel->SkeletalMesh, LODIndex, SectionIndex);
 
+			GetLODAndSectionForAutomaticLODs(GenerationContext, *Node, *TypedNodeSkel->SkeletalMesh, LODIndexConnected, SectionIndexConnected, LODIndex, SectionIndex, bOnlyConnectedLOD);
+			
 			// First process the mesh tags that are going to make the mesh unique and affect whether it's repeated in 
 			// the mesh cache or not
 			FString MeshUniqueTags;
@@ -2882,7 +2883,7 @@ mu::NodeMeshPtr GenerateMutableSourceMesh(const UEdGraphPin * Pin,
 
 			FSkeletalMeshModel* ImportedModel = TypedNodeSkel->SkeletalMesh->GetImportedModel();
 			
-			mu::MeshPtr MutableMesh = GenerateMutableMesh(TypedNodeSkel->SkeletalMesh, TypedNodeSkel->AnimInstance, LODIndex, SectionIndex, MeshUniqueTags, GenerationContext, TypedNodeSkel);
+			mu::MeshPtr MutableMesh = GenerateMutableMesh(TypedNodeSkel->SkeletalMesh, TypedNodeSkel->AnimInstance, LODIndexConnected, SectionIndexConnected, LODIndex, SectionIndex, MeshUniqueTags, GenerationContext, TypedNodeSkel);
 			if (MutableMesh)
 			{
 				MeshNode->SetValue(MutableMesh);
@@ -2958,14 +2959,18 @@ mu::NodeMeshPtr GenerateMutableSourceMesh(const UEdGraphPin * Pin,
 					AddTagToMutableMeshUnique(*MutableMesh, MeshTag);
 				}
 
-				MeshData.bHasVertexColors = TypedNodeSkel->SkeletalMesh->GetHasVertexColors();
-				MeshData.NumTexCoordChannels = ImportedModel->LODModels[LODIndex].NumTexCoords;
-				MeshData.MaxBoneIndexTypeSizeBytes = MutableMesh->GetBoneMap().Num() > 256 ? 2 : 1;
-				MeshData.MaxNumBonesPerVertex = ImportedModel->LODModels[LODIndex].GetMaxBoneInfluences();
-				
-				// When mesh data is combined we will get an upper and lower bound of the number of triangles.
-				MeshData.MaxNumTriangles = ImportedModel->LODModels[LODIndex].Sections[SectionIndex].NumTriangles;
-				MeshData.MinNumTriangles = ImportedModel->LODModels[LODIndex].Sections[SectionIndex].NumTriangles;
+				if (ImportedModel->LODModels.IsValidIndex(LODIndex) &&
+					ImportedModel->LODModels[LODIndex].Sections.IsValidIndex(SectionIndex))
+				{
+					MeshData.bHasVertexColors = TypedNodeSkel->SkeletalMesh->GetHasVertexColors();
+					MeshData.NumTexCoordChannels = ImportedModel->LODModels[LODIndex].NumTexCoords;
+					MeshData.MaxBoneIndexTypeSizeBytes = MutableMesh->GetBoneMap().Num() > 256 ? 2 : 1;
+					MeshData.MaxNumBonesPerVertex = ImportedModel->LODModels[LODIndex].GetMaxBoneInfluences();
+					
+					// When mesh data is combined we will get an upper and lower bound of the number of triangles.
+					MeshData.MaxNumTriangles = ImportedModel->LODModels[LODIndex].Sections[SectionIndex].NumTriangles;
+					MeshData.MinNumTriangles = ImportedModel->LODModels[LODIndex].Sections[SectionIndex].NumTriangles;
+				}
 			}
 
 			// Layouts
@@ -2974,6 +2979,9 @@ mu::NodeMeshPtr GenerateMutableSourceMesh(const UEdGraphPin * Pin,
 				const int32 LODIndexLayout = GenerationContext.CurrentAutoLODStrategy == ECustomizableObjectAutomaticLODStrategy::AutomaticFromMesh ?
 					LODIndexConnected :
 					LODIndex;
+				const int32 SectionIndexLayout = GenerationContext.CurrentAutoLODStrategy == ECustomizableObjectAutomaticLODStrategy::AutomaticFromMesh ?
+					SectionIndexConnected :
+					SectionIndex;
 					
 				const int32 NumLayouts = ImportedModel->LODModels[LODIndexLayout].NumTexCoords;
 				MeshNode->SetLayoutCount(NumLayouts);
@@ -2982,7 +2990,7 @@ mu::NodeMeshPtr GenerateMutableSourceMesh(const UEdGraphPin * Pin,
 
 				for (int32 LayoutIndex = 0; LayoutIndex < NumLayouts; ++LayoutIndex)
 				{
-					if (UEdGraphPin* LayoutPin = TypedNodeSkel->GetLayoutPin(LODIndexLayout, SectionIndex, LayoutIndex))
+					if (UEdGraphPin* LayoutPin = TypedNodeSkel->GetLayoutPin(LODIndexLayout, SectionIndexLayout, LayoutIndex))
 					{
 						if (const UEdGraphPin* ConnectedPin = FollowInputPin(*LayoutPin))
 						{
@@ -3006,7 +3014,7 @@ mu::NodeMeshPtr GenerateMutableSourceMesh(const UEdGraphPin * Pin,
 			// Applying Mesh Morph Nodes
 			if (GenerationContext.MeshMorphStack.Num())
 			{
-				MorphResult = GenerateMorphMesh(Pin, GenerationContext.MeshMorphStack, 0, Result, GenerationContext, MeshData);
+				MorphResult = GenerateMorphMesh(Pin, GenerationContext.MeshMorphStack, 0, Result, GenerationContext, MeshData, bOnlyConnectedLOD);
 			}
 		}
 		else
@@ -3035,22 +3043,31 @@ mu::NodeMeshPtr GenerateMutableSourceMesh(const UEdGraphPin * Pin,
 		Result = MeshNode;
 
 		if (TypedNodeStatic->StaticMesh)
-		{
-			// TODO
-			int LOD = 0;
+		{			
+			int32 LODIndex = 0; // TODO MTBL-1474
+			int32 SectionIndex = 0;
 
 			// Find out what material do we need
-			int MaterialIndex = 0;
-			for (; MaterialIndex < TypedNodeStatic->LODs[LOD].Materials.Num(); ++MaterialIndex)
+			[&LODIndex, &SectionIndex, &TypedNodeStatic, &Pin]()
 			{
-				if (TypedNodeStatic->LODs[LOD].Materials[MaterialIndex].MeshPinRef.Get() == Pin)
+				for (; LODIndex < TypedNodeStatic->LODs.Num(); ++LODIndex)
 				{
-					break;
+					for (; SectionIndex < TypedNodeStatic->LODs[LODIndex].Materials.Num(); ++SectionIndex)
+					{
+						if (TypedNodeStatic->LODs[LODIndex].Materials[SectionIndex].MeshPinRef.Get() == Pin)
+						{
+							return;
+						}
+					}
 				}
-			}
-			check(MaterialIndex < TypedNodeStatic->LODs[LOD].Materials.Num());
 
-			mu::MeshPtr MutableMesh = GenerateMutableMesh(TypedNodeStatic->StaticMesh, TSoftClassPtr<UAnimInstance>(), LOD, MaterialIndex, FString(), GenerationContext, TypedNodeStatic);
+				LODIndex = -1;
+				SectionIndex = -1;
+			}();
+			
+			check(SectionIndex < TypedNodeStatic->LODs[LODIndex].Materials.Num());
+
+			mu::MeshPtr MutableMesh = GenerateMutableMesh(TypedNodeStatic->StaticMesh, TSoftClassPtr<UAnimInstance>(), LODIndex, SectionIndex, LODIndex, SectionIndex, FString(), GenerationContext, TypedNodeStatic);
 			if (MutableMesh)
 			{
 				MeshNode->SetValue(MutableMesh);
@@ -3058,7 +3075,7 @@ mu::NodeMeshPtr GenerateMutableSourceMesh(const UEdGraphPin * Pin,
 				// Layouts
 				MeshNode->SetLayoutCount(1);
 
-				if (const UEdGraphPin* ConnectedPin = FollowInputPin(*TypedNodeStatic->LODs[LOD].Materials[MaterialIndex].LayoutPinRef.Get()))
+				if (const UEdGraphPin* ConnectedPin = FollowInputPin(*TypedNodeStatic->LODs[LODIndex].Materials[SectionIndex].LayoutPinRef.Get()))
 				{
 					mu::NodeLayoutPtr LayoutNode = GenerateMutableSourceLayout(ConnectedPin, GenerationContext);
 					MeshNode->SetLayout(0, LayoutNode);
@@ -3084,7 +3101,7 @@ mu::NodeMeshPtr GenerateMutableSourceMesh(const UEdGraphPin * Pin,
 			// Mesh Morph Stack Management
 			FMorphNodeData NewMorphData = { TypedNodeMorph, TypedNodeMorph->MorphTargetName ,TypedNodeMorph->FactorPin(), TypedNodeMorph->MeshPin() };
 			GenerationContext.MeshMorphStack.Push(NewMorphData);
-			Result = GenerateMutableSourceMesh(ConnectedPin, GenerationContext, MeshData);
+			Result = GenerateMutableSourceMesh(ConnectedPin, GenerationContext, MeshData, false, bOnlyConnectedLOD);
 			GenerationContext.MeshMorphStack.Pop(true);
 		}
 		else
@@ -3144,7 +3161,7 @@ mu::NodeMeshPtr GenerateMutableSourceMesh(const UEdGraphPin * Pin,
 
 				if (const UEdGraphPin* MeshConnectedPin = FollowInputPin(*TypedNodeMeshMorphStackApp->GetMeshPin()))
 				{
-					Result = GenerateMutableSourceMesh(MeshConnectedPin, GenerationContext, MeshData);
+					Result = GenerateMutableSourceMesh(MeshConnectedPin, GenerationContext, MeshData, false, bOnlyConnectedLOD);
 				}
 
 				for (int32 MorphIndex = 0; MorphIndex < AddedMorphs; ++MorphIndex)
@@ -3216,7 +3233,7 @@ mu::NodeMeshPtr GenerateMutableSourceMesh(const UEdGraphPin * Pin,
 					if (const UEdGraphPin* ConnectedPin = FollowInputPin(*TypedNodeMeshSwitch->GetElementPin(SelectorIndex)))
 					{
 						FMutableGraphMeshGenerationData ChildMeshData;
-						Result = GenerateMutableSourceMesh(ConnectedPin, GenerationContext, ChildMeshData);
+						Result = GenerateMutableSourceMesh(ConnectedPin, GenerationContext, ChildMeshData, false, bOnlyConnectedLOD);
 						SwitchNode->SetOption(SelectorIndex, Result);
 						MeshData.Combine(ChildMeshData);
 					}
@@ -3241,7 +3258,7 @@ mu::NodeMeshPtr GenerateMutableSourceMesh(const UEdGraphPin * Pin,
 		if (const UEdGraphPin* ConnectedPin = FollowInputPin(*TypedNodeMeshVar->DefaultPin()))
 		{
 			FMutableGraphMeshGenerationData ChildMeshData;
-			mu::NodeMeshPtr ChildNode = GenerateMutableSourceMesh(ConnectedPin, GenerationContext, ChildMeshData);
+			mu::NodeMeshPtr ChildNode = GenerateMutableSourceMesh(ConnectedPin, GenerationContext, ChildMeshData, false, bOnlyConnectedLOD);
 			if (ChildNode)
 			{
 				MeshNode->SetDefaultMesh(ChildNode.get());
@@ -3267,7 +3284,7 @@ mu::NodeMeshPtr GenerateMutableSourceMesh(const UEdGraphPin * Pin,
 			if (const UEdGraphPin* ConnectedPin = FollowInputPin(*VariationPin))
 			{
 				FMutableGraphMeshGenerationData VariationMeshData;
-				mu::NodeMeshPtr ChildNode = GenerateMutableSourceMesh(ConnectedPin, GenerationContext, VariationMeshData);
+				mu::NodeMeshPtr ChildNode = GenerateMutableSourceMesh(ConnectedPin, GenerationContext, VariationMeshData, false, bOnlyConnectedLOD);
 				MeshNode->SetVariationMesh(VariationIndex, ChildNode.get());
 				MeshData.Combine(VariationMeshData);
 			}
@@ -3282,7 +3299,7 @@ mu::NodeMeshPtr GenerateMutableSourceMesh(const UEdGraphPin * Pin,
 		if (const UEdGraphPin* ConnectedPin = FollowInputPin(*TypedNodeGeometry->MeshAPin()))
 		{
 			FMutableGraphMeshGenerationData ChildMeshData;
-			mu::Ptr<mu::NodeMesh> ChildNode = GenerateMutableSourceMesh(ConnectedPin, GenerationContext, ChildMeshData);
+			mu::Ptr<mu::NodeMesh> ChildNode = GenerateMutableSourceMesh(ConnectedPin, GenerationContext, ChildMeshData, false, bOnlyConnectedLOD);
 			if (ChildNode)
 			{
 				MeshNode->SetMeshA(ChildNode.get());
@@ -3301,7 +3318,7 @@ mu::NodeMeshPtr GenerateMutableSourceMesh(const UEdGraphPin * Pin,
 		if (const UEdGraphPin* ConnectedPin = FollowInputPin(*TypedNodeGeometry->MeshBPin()))
 		{
 			FMutableGraphMeshGenerationData ChildMeshData;
-			mu::Ptr<mu::NodeMesh> ChildNode = GenerateMutableSourceMesh(ConnectedPin, GenerationContext, ChildMeshData);
+			mu::Ptr<mu::NodeMesh> ChildNode = GenerateMutableSourceMesh(ConnectedPin, GenerationContext, ChildMeshData, false, bOnlyConnectedLOD);
 			if (ChildNode)
 			{
 				MeshNode->SetMeshB(ChildNode.get());
@@ -3348,7 +3365,7 @@ mu::NodeMeshPtr GenerateMutableSourceMesh(const UEdGraphPin * Pin,
 		if (const UEdGraphPin* ConnectedPin = FollowInputPin(*TypedNodeReshape->BaseMeshPin()))
 		{
 			FMutableGraphMeshGenerationData ChildMeshData;
-			mu::Ptr<mu::NodeMesh> ChildNode = GenerateMutableSourceMesh(ConnectedPin, GenerationContext, ChildMeshData);
+			mu::Ptr<mu::NodeMesh> ChildNode = GenerateMutableSourceMesh(ConnectedPin, GenerationContext, ChildMeshData, false, bOnlyConnectedLOD);
 			if (ChildNode)
 			{
 				MeshNode->SetBaseMesh(ChildNode.get());
@@ -3423,7 +3440,7 @@ mu::NodeMeshPtr GenerateMutableSourceMesh(const UEdGraphPin * Pin,
 		if (const UEdGraphPin* ConnectedPin = FollowInputPin(*TypedNodeReshape->BaseShapePin()))
 		{
 			FMutableGraphMeshGenerationData ChildMeshData;
-			mu::Ptr<mu::NodeMesh> ChildNode = GenerateMutableSourceMesh(ConnectedPin, GenerationContext, ChildMeshData);
+			mu::Ptr<mu::NodeMesh> ChildNode = GenerateMutableSourceMesh(ConnectedPin, GenerationContext, ChildMeshData, false, true);
 	
 			if (ChildNode)
 			{
@@ -3440,7 +3457,7 @@ mu::NodeMeshPtr GenerateMutableSourceMesh(const UEdGraphPin * Pin,
 		if (const UEdGraphPin* ConnectedPin = FollowInputPin(*TypedNodeReshape->TargetShapePin()))
 		{
 			FMutableGraphMeshGenerationData ChildMeshData;
-			mu::Ptr<mu::NodeMesh> ChildNode = GenerateMutableSourceMesh(ConnectedPin, GenerationContext, ChildMeshData);
+			mu::Ptr<mu::NodeMesh> ChildNode = GenerateMutableSourceMesh(ConnectedPin, GenerationContext, ChildMeshData, false, true);
 			
 			if (ChildNode)
 			{
@@ -3480,7 +3497,7 @@ mu::NodeMeshPtr GenerateMutableSourceMesh(const UEdGraphPin * Pin,
 	{
 		if (const UEdGraphPin* ConnectedPin = FollowInputPin(*TypedNode->GetInputMeshPin()))
 		{
-			mu::Ptr<mu::NodeMesh> InputMeshNode = GenerateMutableSourceMesh(ConnectedPin, GenerationContext, MeshData);
+			mu::Ptr<mu::NodeMesh> InputMeshNode = GenerateMutableSourceMesh(ConnectedPin, GenerationContext, MeshData, false, bOnlyConnectedLOD);
 
 			if (TypedNode->PoseAsset && GenerationContext.GetCurrentComponentInfo().RefSkeletalMesh)
 			{
@@ -3549,24 +3566,27 @@ mu::NodeMeshPtr GenerateMutableSourceMesh(const UEdGraphPin * Pin,
 				if (Table)
 				{
 					mu::NodeMeshTablePtr MeshTableNode = new mu::NodeMeshTable();
-
+					
+					int32 LODIndexConnected = -1; // LOD which the pin is connected to
+					int32 SectionIndexConnected = -1;
+					
 					int32 LODIndex = 0;
 					int32 SectionIndex = 0;
 
-					TypedNodeTable->GetPinLODAndSection(Pin, LODIndex, SectionIndex);
+					TypedNodeTable->GetPinLODAndSection(Pin, LODIndexConnected, SectionIndexConnected);
 
 					if (DefaultSkeletalMesh)
 					{
-						GetEffectiveLODAndSection(GenerationContext, *Node, *DefaultSkeletalMesh, LODIndex, SectionIndex);
+						GetLODAndSectionForAutomaticLODs(GenerationContext, *Node, *DefaultSkeletalMesh, LODIndexConnected, SectionIndexConnected, LODIndex, SectionIndex, bOnlyConnectedLOD);
 					}
-
+					
 					// Getting the mutable table mesh column name
 					FString MutableColumnName = TypedNodeTable->GetMutableColumnName(Pin, LODIndex);
 
 					// Generating a new Mesh column if not exists
 					if (Table->FindColumn(StringCast<ANSICHAR>(*MutableColumnName).Get()) == INDEX_NONE)
 					{
-						bSuccess = GenerateTableColumn(TypedNodeTable, Pin, Table, DataTableColumnName, LODIndex, GenerationContext);
+						bSuccess = GenerateTableColumn(TypedNodeTable, Pin, Table, DataTableColumnName, LODIndexConnected, SectionIndexConnected, LODIndex, SectionIndex, bOnlyConnectedLOD, GenerationContext);
 
 						if (!bSuccess)
 						{
@@ -3587,17 +3607,22 @@ mu::NodeMeshPtr GenerateMutableSourceMesh(const UEdGraphPin * Pin,
 
 						if (DefaultSkeletalMesh)
 						{
-							// TODO: this should be made for all the meshes of the Column to support meshes with different values
-							// Filling Mesh Data
-							FSkeletalMeshModel* importModel = DefaultSkeletalMesh->GetImportedModel();
-							MeshData.bHasVertexColors = DefaultSkeletalMesh->GetHasVertexColors();
-							MeshData.NumTexCoordChannels = importModel->LODModels[LODIndex].NumTexCoords;
-							MeshData.MaxBoneIndexTypeSizeBytes = importModel->LODModels[LODIndex].RequiredBones.Num() > 256 ? 2 : 1;
-							MeshData.MaxNumBonesPerVertex = importModel->LODModels[LODIndex].GetMaxBoneInfluences();
-
-							// When mesh data is combined we will get an upper and lower bound of the number of triangles.
-							MeshData.MaxNumTriangles = importModel->LODModels[LODIndex].Sections[SectionIndex].NumTriangles;
-							MeshData.MinNumTriangles = importModel->LODModels[LODIndex].Sections[SectionIndex].NumTriangles;
+							FSkeletalMeshModel* ImportedModel = DefaultSkeletalMesh->GetImportedModel();
+
+							if (ImportedModel->LODModels.IsValidIndex(LODIndex) &&
+								ImportedModel->LODModels[LODIndex].Sections.IsValidIndex(SectionIndex))
+							{
+								// TODO: this should be made for all the meshes of the Column to support meshes with different values
+								// Filling Mesh Data
+								MeshData.bHasVertexColors = DefaultSkeletalMesh->GetHasVertexColors();
+								MeshData.NumTexCoordChannels = ImportedModel->LODModels[LODIndex].NumTexCoords;
+								MeshData.MaxBoneIndexTypeSizeBytes = ImportedModel->LODModels[LODIndex].RequiredBones.Num() > 256 ? 2 : 1;
+								MeshData.MaxNumBonesPerVertex = ImportedModel->LODModels[LODIndex].GetMaxBoneInfluences();
+
+								// When mesh data is combined we will get an upper and lower bound of the number of triangles.
+								MeshData.MaxNumTriangles = ImportedModel->LODModels[LODIndex].Sections[SectionIndex].NumTriangles;
+								MeshData.MinNumTriangles = ImportedModel->LODModels[LODIndex].Sections[SectionIndex].NumTriangles;
+							}
 						}
 
 						TArray<UCustomizableObjectLayout*> Layouts = TypedNodeTable->GetLayouts(Pin);
@@ -3651,7 +3676,7 @@ mu::NodeMeshPtr GenerateMutableSourceMesh(const UEdGraphPin * Pin,
 						// Applying Mesh Morph Nodes
 						if (DefaultSkeletalMesh && GenerationContext.MeshMorphStack.Num())
 						{
-							MorphResult = GenerateMorphMesh(Pin, GenerationContext.MeshMorphStack, 0, Result, GenerationContext, MeshData, MutableColumnName);
+							MorphResult = GenerateMorphMesh(Pin, GenerationContext.MeshMorphStack, 0, Result, GenerationContext, MeshData, bOnlyConnectedLOD, MutableColumnName);
 						}
 					}
 				}
diff --git a/Engine/Plugins/Experimental/Mutable/Source/CustomizableObjectEditor/Private/MuCOE/GenerateMutableSource/GenerateMutableSourceMesh.h b/Engine/Plugins/Experimental/Mutable/Source/CustomizableObjectEditor/Private/MuCOE/GenerateMutableSource/GenerateMutableSourceMesh.h
--- a/Engine/Plugins/Experimental/Mutable/Source/CustomizableObjectEditor/Private/MuCOE/GenerateMutableSource/GenerateMutableSourceMesh.h
+++ b/Engine/Plugins/Experimental/Mutable/Source/CustomizableObjectEditor/Private/MuCOE/GenerateMutableSource/GenerateMutableSourceMesh.h
@@ -9,7 +9,6 @@
 
 class FName;
 class FString;
-
 class UCustomizableObjectNode;
 class UEdGraphPin;
 class UObject;
@@ -18,16 +17,77 @@ class UStaticMesh;
 struct FMutableGraphGenerationContext;
 struct FMutableGraphMeshGenerationData;
 
-mu::MeshPtr ConvertSkeletalMeshToMutable(const USkeletalMesh* InSkeletalMesh, const TSoftClassPtr<class UAnimInstance>& AnimBp, int LOD, int MaterialIndex, FMutableGraphGenerationContext& GenerationContext, const UCustomizableObjectNode* CurrentNode);
 
+/** Returns the corrected LOD and Section Index when using Automatic LOD From Mesh strategy.
+ *
+ * Do not confuse Section Index and Material Index, they are not the same.
+ * 
+ * @param Context 
+ * @param Node 
+ * @param SkeletalMesh 
+ * @param LODIndexConnected Connected pin LOD Index.
+ * @param SectionIndexConnected Connected pin Section Index.
+ * @param OutLODIndex Corrected Skeletal Mesh LOD Index.
+ * @param OutSectionIndex Corrected Skeletal Mesh Section Index.
+ * @param bOnlyConnectedLOD Corrected LOD and Section will unconditionally always be the connected ones.
+ * @return When using Automatic LOD From Mesh, InOutLODIndex and InOutSectionIndex and will return -1 if the section is not found in the currently compiling LOD. */
+void GetLODAndSectionForAutomaticLODs(const FMutableGraphGenerationContext& Context, const UCustomizableObjectNode& Node, const USkeletalMesh& SkeletalMesh,
+	const int32 LODIndexConnected, const int32 SectionIndexConnected,
+	int32& OutLODIndex, int32& OutSectionIndex,
+	bool bOnlyConnectedLOD);
+
+
+/** Converts an Unreal Skeletal Mesh to Mutable Mesh.
+ *
+ * @param InSkeletalMesh Unreal Skeletal Mesh.
+ * @param AnimBp 
+ * @param LODIndexConnected LOD which the pin is connected to.
+ * @param SectionIndexConnected Section which the pin is connected to.
+ * @param LODIndex LOD we are generating. Will be different from LODIndexConnected only when using Automatic LOD From Mesh. 
+ * @param SectionIndex Section we are generating. Will be different from SectionIndexConnected only when using Automatic LOD From Mesh.
+ * @param GenerationContext 
+ * @param CurrentNode 
+ * @return Mutable Mesh. Nullptr if there has been an error. Empty mesh if the Skeletal Mesh does not contain the requested LOD + Section. */
+mu::MeshPtr ConvertSkeletalMeshToMutable(const USkeletalMesh* InSkeletalMesh, const TSoftClassPtr<UAnimInstance>& AnimBp,
+                                         int32 LODIndexConnected,  int32 SectionIndexConnected,
+                                         int32 LODIndex, int32 SectionIndex,
+                                         FMutableGraphGenerationContext& GenerationContext,
+                                         const UCustomizableObjectNode* CurrentNode);
+
+
+mu::MeshPtr ConvertStaticMeshToMutable(const UStaticMesh* StaticMesh, int32 LODIndex, int32 SectionIndex,
+	FMutableGraphGenerationContext& GenerationContext, const UCustomizableObjectNode* CurrentNode);
 
-mu::MeshPtr ConvertStaticMeshToMutable(const UStaticMesh* StaticMesh, int LOD, int MaterialIndex, FMutableGraphGenerationContext& GenerationContext, const UCustomizableObjectNode* CurrentNode);
 
+/** Compiler recursive function. Mutable Mesh.
+ *
+ * @param Mesh Unreal Mesh.
+ * @param AnimBp 
+ * @param LODIndexConnected LOD which the pin is connected to.
+ * @param SectionIndexConnected Section which the pin is connected to.
+ * @param LODIndex LOD we are generating. Will be different from LODIndexConnected only when using Automatic LOD From Mesh. 
+ * @param SectionIndex Section we are generating. Will be different from SectionIndexConnected only when using Automatic LOD From Mesh.
+ * @param MeshUniqueTags 
+ * @param GenerationContext 
+ * @param CurrentNode 
+ * @return Mutable Mesh. Nullptr if there has been an error. Empty mesh if the Skeletal Mesh does not contain the requested LOD + Section. */
+mu::MeshPtr GenerateMutableMesh(const UObject* Mesh, const TSoftClassPtr<UAnimInstance>& AnimBp,
+                                int32 LODIndexConnected, int32 SectionIndexConnected,
+                                int32 LODIndex, int32 SectionIndex, const FString& MeshUniqueTags,
+                                FMutableGraphGenerationContext& GenerationContext, const UCustomizableObjectNode* CurrentNode);
 
-mu::MeshPtr GenerateMutableMesh(const UObject * Mesh, const TSoftClassPtr<class UAnimInstance>& AnimBp, int32 LOD, int32 MaterialIndex, const FString& MeshUniqueTags, FMutableGraphGenerationContext & GenerationContext, const UCustomizableObjectNode* CurrentNode);
 
+mu::MeshPtr BuildMorphedMutableMesh(const UEdGraphPin* BaseSourcePin, const FString& MorphTargetName,
+	FMutableGraphGenerationContext& GenerationContext, bool bOnlyConnectedLOD, const FName& RowName = "");
 
-mu::MeshPtr BuildMorphedMutableMesh(const UEdGraphPin * BaseSourcePin, const FString& MorphTargetName, FMutableGraphGenerationContext & GenerationContext, const FName& RowName = "");
 
-/** Convert a CustomizableObject Source Graph into a mutable source graph. */
-mu::NodeMeshPtr GenerateMutableSourceMesh(const UEdGraphPin* Pin, FMutableGraphGenerationContext& GenerationContext, FMutableGraphMeshGenerationData& MeshData, bool bLinkedToExtendMaterial = false);
+/** Compiler recursive function. Mutable Node Mesh.
+ *
+ * @param Pin 
+ * @param GenerationContext 
+ * @param MeshData 
+ * @param bLinkedToExtendMaterial 
+ * @param bOnlyConnectedLOD Corrected LOD and Section will unconditionally always be the connected ones.
+ * @return  Mutable Mesh Node. */
+mu::NodeMeshPtr GenerateMutableSourceMesh(const UEdGraphPin* Pin, FMutableGraphGenerationContext& GenerationContext,
+	FMutableGraphMeshGenerationData& MeshData, bool bLinkedToExtendMaterial, bool bOnlyConnectedLOD);
diff --git a/Engine/Plugins/Experimental/Mutable/Source/CustomizableObjectEditor/Private/MuCOE/GenerateMutableSource/GenerateMutableSourceModifier.cpp b/Engine/Plugins/Experimental/Mutable/Source/CustomizableObjectEditor/Private/MuCOE/GenerateMutableSource/GenerateMutableSourceModifier.cpp
--- a/Engine/Plugins/Experimental/Mutable/Source/CustomizableObjectEditor/Private/MuCOE/GenerateMutableSource/GenerateMutableSourceModifier.cpp
+++ b/Engine/Plugins/Experimental/Mutable/Source/CustomizableObjectEditor/Private/MuCOE/GenerateMutableSource/GenerateMutableSourceModifier.cpp
@@ -74,7 +74,7 @@ mu::NodeModifierPtr GenerateMutableSourceModifier(const UEdGraphPin * Pin, FMuta
 		if (const UEdGraphPin* ConnectedPin = FollowInputPin(*TypedNodeClipDeform->ClipShapePin()))
 		{
 			FMutableGraphMeshGenerationData DummyMeshData;
-			mu::NodeMeshPtr ClipMesh = GenerateMutableSourceMesh(ConnectedPin, GenerationContext, DummyMeshData);
+			mu::NodeMeshPtr ClipMesh = GenerateMutableSourceMesh(ConnectedPin, GenerationContext, DummyMeshData, false, true);
 
 			ClipNode->SetClipMesh(ClipMesh.get());
 
@@ -117,7 +117,7 @@ mu::NodeModifierPtr GenerateMutableSourceModifier(const UEdGraphPin * Pin, FMuta
 		{
 			FMutableGraphMeshGenerationData DummyMeshData;
 
-			mu::NodeMeshPtr ClipMesh = GenerateMutableSourceMesh(ConnectedPin, GenerationContext, DummyMeshData);
+			mu::NodeMeshPtr ClipMesh = GenerateMutableSourceMesh(ConnectedPin, GenerationContext, DummyMeshData, false, true);
 
 			FPinDataValue* PinData = GenerationContext.PinData.Find(ConnectedPin);
 			for (const FMeshData& MeshData : PinData->MeshesData)
diff --git a/Engine/Plugins/Experimental/Mutable/Source/CustomizableObjectEditor/Private/MuCOE/GenerateMutableSource/GenerateMutableSourceSurface.cpp b/Engine/Plugins/Experimental/Mutable/Source/CustomizableObjectEditor/Private/MuCOE/GenerateMutableSource/GenerateMutableSourceSurface.cpp
--- a/Engine/Plugins/Experimental/Mutable/Source/CustomizableObjectEditor/Private/MuCOE/GenerateMutableSource/GenerateMutableSourceSurface.cpp
+++ b/Engine/Plugins/Experimental/Mutable/Source/CustomizableObjectEditor/Private/MuCOE/GenerateMutableSource/GenerateMutableSourceSurface.cpp
@@ -385,7 +385,7 @@ mu::NodeSurfacePtr GenerateMutableSourceSurface(const UEdGraphPin * Pin, FMutabl
 		if (const UEdGraphPin* ConnectedPin = FollowInputPin(*TypedNodeMat->GetMeshPin()))
 		{
 			FMutableGraphMeshGenerationData MeshData;
-			MeshNode = GenerateMutableSourceMesh(ConnectedPin, GenerationContext, MeshData);
+			MeshNode = GenerateMutableSourceMesh(ConnectedPin, GenerationContext, MeshData, false, false);
 
 			if (MeshNode)
 			{
@@ -439,7 +439,9 @@ mu::NodeSurfacePtr GenerateMutableSourceSurface(const UEdGraphPin * Pin, FMutabl
 					// Generating a column for each modified parameter(texture, color & float) if not exists
 					if (Table && Table->FindColumn(StringCast<ANSICHAR>(*ColumnName).Get()) == INDEX_NONE)
 					{
-						GenerateTableColumn(TypedNodeTable, ConnectedPin, Table, ColumnName, GenerationContext.CurrentLOD, GenerationContext);
+						int32 Dummy = -1; // TODO MTBL-1512
+						bool Dummy2 = false;
+						GenerateTableColumn(TypedNodeTable, ConnectedPin, Table, ColumnName, Dummy, Dummy, GenerationContext.CurrentLOD, Dummy, Dummy2, GenerationContext);
 					}
 
 					// Checking if this material has some parameters modified by the table node linked to it
@@ -1094,7 +1096,7 @@ mu::NodeSurfacePtr GenerateMutableSourceSurface(const UEdGraphPin * Pin, FMutabl
 			FMutableGraphMeshGenerationData MeshData;
 			if (const UEdGraphPin* ConnectedPin = FollowInputPin(*TypedNodeExt->AddMeshPin()))
 			{
-				AddMeshNode = GenerateMutableSourceMesh(ConnectedPin, GenerationContext, MeshData, true);
+				AddMeshNode = GenerateMutableSourceMesh(ConnectedPin, GenerationContext, MeshData, true, false);
 			}
 
 			if (AddMeshNode)
@@ -1256,7 +1258,7 @@ mu::NodeSurfacePtr GenerateMutableSourceSurface(const UEdGraphPin * Pin, FMutabl
 			if (const UEdGraphPin* ConnectedPin = FollowInputPin(*TypedNodeRem->RemoveMeshPin()))
 			{
 				FMutableGraphMeshGenerationData DummyMeshData;
-				RemoveMeshNode = GenerateMutableSourceMesh(ConnectedPin, GenerationContext, DummyMeshData);
+				RemoveMeshNode = GenerateMutableSourceMesh(ConnectedPin, GenerationContext, DummyMeshData, false, true);
 			}
 
 			if (RemoveMeshNode)
@@ -1299,12 +1301,13 @@ mu::NodeSurfacePtr GenerateMutableSourceSurface(const UEdGraphPin * Pin, FMutabl
 			else
 			{
 				FMutableGraphMeshGenerationData DummyMeshData;
-				mu::NodeMeshPtr BaseSourceMesh = GenerateMutableSourceMesh(BaseSourcePin, GenerationContext, DummyMeshData);
+				mu::NodeMeshPtr BaseSourceMesh = GenerateMutableSourceMesh(BaseSourcePin, GenerationContext, DummyMeshData, false, false);
 
 				mu::NodeMeshFragmentPtr MeshFrag = new mu::NodeMeshFragment();
 				MeshFrag->SetMesh(BaseSourceMesh);
 				MeshFrag->SetLayoutOrGroup(TypedNodeRemBlocks->ParentLayoutIndex);
 				MeshFrag->SetFragmentType(mu::NodeMeshFragment::FT_LAYOUT_BLOCKS);
+				MeshFrag->SetMessageContext(TypedNodeRemBlocks);
 
 				// Add the block indices
 				TArray<int> FoundBlockIndices;
@@ -1496,7 +1499,8 @@ mu::NodeSurfacePtr GenerateMutableSourceSurface(const UEdGraphPin * Pin, FMutabl
 			else
 			{
 				//Morph Mesh
-				mu::MeshPtr MorphedSourceMesh = BuildMorphedMutableMesh(BaseSourcePin, TypedNodeMorph->MorphTargetName, GenerationContext);
+				// The Mesh Base Source from the Parent Material will always have all LODs since it is not an auxiliary Mesh (clipping, reshape...).
+				mu::MeshPtr MorphedSourceMesh = BuildMorphedMutableMesh(BaseSourcePin, TypedNodeMorph->MorphTargetName, GenerationContext, false); 
 
 				mu::NodeMeshConstantPtr MorphedSourceMeshNode = new mu::NodeMeshConstant;
 				MorphedSourceMeshNode->SetMessageContext(Node);					
diff --git a/Engine/Plugins/Experimental/Mutable/Source/CustomizableObjectEditor/Private/MuCOE/GenerateMutableSource/GenerateMutableSourceTable.cpp b/Engine/Plugins/Experimental/Mutable/Source/CustomizableObjectEditor/Private/MuCOE/GenerateMutableSource/GenerateMutableSourceTable.cpp
--- a/Engine/Plugins/Experimental/Mutable/Source/CustomizableObjectEditor/Private/MuCOE/GenerateMutableSource/GenerateMutableSourceTable.cpp
+++ b/Engine/Plugins/Experimental/Mutable/Source/CustomizableObjectEditor/Private/MuCOE/GenerateMutableSource/GenerateMutableSourceTable.cpp
@@ -20,7 +20,8 @@
 #define LOCTEXT_NAMESPACE "CustomizableObjectEditor"
 
 
-bool FillTableColumn(const UCustomizableObjectNodeTable* TableNode,	mu::TablePtr MutableTable,	const FString& ColumnName,	const FString& RowName,	const int32 RowIdx,	uint8* CellData, FProperty* Property, const int32 LOD, FMutableGraphGenerationContext& GenerationContext)
+bool FillTableColumn(const UCustomizableObjectNodeTable* TableNode,	mu::TablePtr MutableTable,	const FString& ColumnName,	const FString& RowName,	const int32 RowIdx,	uint8* CellData, FProperty* Property,
+	const int LODIndexConnected, const int32 SectionIndexConnected, int32 LODIndex, int32 SectionIndex, const bool bOnlyConnectedLOD, FMutableGraphGenerationContext& GenerationContext)
 {
 	int32 CurrentColumn;
 
@@ -158,103 +159,108 @@ bool FillTableColumn(const UCustomizableObjectNodeTable* TableNode,	mu::TablePtr
 				return false;
 			}
 
+			GetLODAndSectionForAutomaticLODs(GenerationContext, *TableNode, *SkeletalMesh, LODIndexConnected, SectionIndexConnected, LODIndex, SectionIndex, bOnlyConnectedLOD);
+			
 			// Parameter used for LOD differences
-			int32 CurrentLOD = LOD;
+	
+			if (GenerationContext.CurrentAutoLODStrategy != ECustomizableObjectAutomaticLODStrategy::AutomaticFromMesh || 
+				SectionIndex == SectionIndexConnected)
+			{
+				const int32 NumLODs = SkeletalMesh->GetImportedModel()->LODModels.Num();
 
-			int NumLODs = SkeletalMesh->GetLODInfoArray().Num();
+				if (NumLODs <= LODIndex)
+				{
+					LODIndex = NumLODs - 1;
 
-			if (NumLODs <= CurrentLOD)
+					FString msg = FString::Printf(TEXT("Mesh from column [%s] row [%s] needs LOD %d but has less LODs than the reference mesh. LOD %d will be used instead. This can cause some performance penalties."),
+						*ColumnName, *RowName, LODIndex, LODIndex);
+					GenerationContext.Compiler->CompilerLog(FText::FromString(msg), TableNode);
+				}
+			}
+
+			FSkeletalMeshModel* ImportedModel = SkeletalMesh->GetImportedModel();
+			
+			if (ImportedModel->LODModels.IsValidIndex(LODIndex)) // Ignore error since this Section is empty due to Automatic LODs From Mesh
 			{
-				CurrentLOD = NumLODs - 1;
+				int32 NumSections = ImportedModel->LODModels[LODIndex].Sections.Num();
+				int32 ReferenceNumMaterials = ImportedModel->LODModels[LODIndex].Sections.Num();
 
-				FString msg = FString::Printf(TEXT("Mesh from column [%s] row [%s] needs LOD %d but has less LODs than the reference mesh. LOD %d will be used instead. This can cause some performance penalties."),
-					*ColumnName, *RowName, LOD, CurrentLOD);
-				GenerationContext.Compiler->CompilerLog(FText::FromString(msg), TableNode);
+				if (NumSections != ReferenceNumMaterials)
+				{
+					FString Dif_1 = NumSections > ReferenceNumMaterials ? "more" : "less";
+					FString Dif_2 = NumSections > ReferenceNumMaterials ? "Some will be ignored" : "This can cause some compilation errors.";
+
+					FString msg = FString::Printf(TEXT("Mesh from column [%s] row [%s] has %s Sections than the reference mesh. %s"), *ColumnName, *RowName, *Dif_1, *Dif_2);
+					GenerationContext.Compiler->CompilerLog(FText::FromString(msg), TableNode);
+				}
 			}
 
-			int32 NumMaterials = SkeletalMesh->GetImportedModel()->LODModels[CurrentLOD].Sections.Num();
-			int32 ReferenceNumMaterials = ReferenceSkeletalMesh->GetImportedModel()->LODModels[CurrentLOD].Sections.Num();
+			FString MutableColumnName = TableNode->GenerateSkeletalMeshMutableColumName(ColumnName, LODIndex, SectionIndex);
 
-			if (NumMaterials != ReferenceNumMaterials)
+			CurrentColumn = MutableTable.get()->FindColumn(StringCast<ANSICHAR>(*MutableColumnName).Get());
+
+			if (CurrentColumn == -1)
 			{
-				FString Dif_1 = NumMaterials > ReferenceNumMaterials ? "more" : "less";
-				FString Dif_2 = NumMaterials > ReferenceNumMaterials ? "Some will be ignored" : "This can cause some compilation errors.";
+				CurrentColumn = MutableTable->AddColumn(StringCast<ANSICHAR>(*MutableColumnName).Get(), mu::TABLE_COLUMN_TYPE::TCT_MESH);
+			}
 
-				FString msg = FString::Printf(TEXT("Mesh from column [%s] row [%s] has %s Sections than the reference mesh. %s"), *ColumnName, *RowName, *Dif_1, *Dif_2);
-				GenerationContext.Compiler->CompilerLog(FText::FromString(msg), TableNode);
+			// First process the mesh tags that are going to make the mesh unique and affect whether it's repeated in 
+			// the mesh cache or not
+			FString MeshUniqueTags;
+
+			if (!AnimBPAssetTag.IsEmpty())
+			{
+				MeshUniqueTags += AnimBPAssetTag;
 			}
 
-			int32 Materials = NumMaterials <= ReferenceNumMaterials ? NumMaterials : ReferenceNumMaterials;
+			TArray<FString> ArrayAnimBPTags;
 
-			for (int32 MatIndex = 0; MatIndex < Materials; ++MatIndex)
+			for (const FGameplayTag& Tag : GameplayTags)
 			{
-				FString MutableColumnName = TableNode->GenerateSkeletalMeshMutableColumName(ColumnName, LOD, MatIndex);
+				MeshUniqueTags += GenerateGameplayTag(Tag.ToString());
+			}
 
-				CurrentColumn = MutableTable.get()->FindColumn(StringCast<ANSICHAR>(*MutableColumnName).Get());
+			//TODO: Add AnimBp physics to Tables.
+			mu::MeshPtr MutableMesh = GenerateMutableMesh(SkeletalMesh, TSoftClassPtr<UAnimInstance>(), LODIndexConnected, SectionIndexConnected, LODIndex, SectionIndex, MeshUniqueTags, GenerationContext, TableNode);
 
-				if (CurrentColumn == -1)
-				{
-					CurrentColumn = MutableTable->AddColumn(StringCast<ANSICHAR>(*MutableColumnName).Get(), mu::TABLE_COLUMN_TYPE::TCT_MESH);
-				}
+			if (MutableMesh)
+			{
+ 				if (SkeletalMesh->GetPhysicsAsset() && MutableMesh->GetPhysicsBody() && MutableMesh->GetPhysicsBody()->GetBodyCount())
+				{	
+					TSoftObjectPtr<UPhysicsAsset> PhysicsAsset = SkeletalMesh->GetPhysicsAsset();
+					GenerationContext.PhysicsAssetMap.Add(PhysicsAsset.ToString(), PhysicsAsset);
+					FString PhysicsAssetTag = FString("__PhysicsAsset:") + PhysicsAsset.ToString();
 
-				// First process the mesh tags that are going to make the mesh unique and affect whether it's repeated in 
-				// the mesh cache or not
-				FString MeshUniqueTags;
+					AddTagToMutableMeshUnique(*MutableMesh, PhysicsAssetTag);
+				}
 
 				if (!AnimBPAssetTag.IsEmpty())
 				{
-					MeshUniqueTags += AnimBPAssetTag;
+					AddTagToMutableMeshUnique(*MutableMesh, AnimBPAssetTag);
 				}
 
-				TArray<FString> ArrayAnimBPTags;
-
 				for (const FGameplayTag& Tag : GameplayTags)
 				{
-					MeshUniqueTags += GenerateGameplayTag(Tag.ToString());
+					AddTagToMutableMeshUnique(*MutableMesh, GenerateGameplayTag(Tag.ToString()));
 				}
 
-				//TODO: Add AnimBp physics to Tables.
-				mu::MeshPtr MutableMesh = GenerateMutableMesh(SkeletalMesh, TSoftClassPtr<UAnimInstance>(), CurrentLOD, MatIndex, MeshUniqueTags, GenerationContext, TableNode);
+				AddSocketTagsToMesh(SkeletalMesh, MutableMesh, GenerationContext);
 
-				if (MutableMesh)
+				if (UCustomizableObjectSystem::GetInstance()->IsMutableAnimInfoDebuggingEnabled())
 				{
- 					if (SkeletalMesh->GetPhysicsAsset() && MutableMesh->GetPhysicsBody() && MutableMesh->GetPhysicsBody()->GetBodyCount())
-					{	
-						TSoftObjectPtr<UPhysicsAsset> PhysicsAsset = SkeletalMesh->GetPhysicsAsset();
-						GenerationContext.PhysicsAssetMap.Add(PhysicsAsset.ToString(), PhysicsAsset);
-						FString PhysicsAssetTag = FString("__PhysicsAsset:") + PhysicsAsset.ToString();
-
-						AddTagToMutableMeshUnique(*MutableMesh, PhysicsAssetTag);
-					}
-
-					if (!AnimBPAssetTag.IsEmpty())
-					{
-						AddTagToMutableMeshUnique(*MutableMesh, AnimBPAssetTag);
-					}
-
-					for (const FGameplayTag& Tag : GameplayTags)
-					{
-						AddTagToMutableMeshUnique(*MutableMesh, GenerateGameplayTag(Tag.ToString()));
-					}
-
-					AddSocketTagsToMesh(SkeletalMesh, MutableMesh, GenerationContext);
-
-					if (UCustomizableObjectSystem::GetInstance()->IsMutableAnimInfoDebuggingEnabled())
-					{
-						FString MeshPath;
-						SkeletalMesh->GetOuter()->GetPathName(nullptr, MeshPath);
-						FString MeshTag = FString("__MeshPath:") + MeshPath;
-						AddTagToMutableMeshUnique(*MutableMesh, MeshTag);
-					}
-
-					MutableTable->SetCell(CurrentColumn, RowIdx, MutableMesh.get());
-				}
-				else
-				{
-					FString msg = FString::Printf(TEXT("Error converting skeletal mesh LOD %d, Material %d from column [%s] row [%s] to mutable."),
-						LOD, MatIndex, *ColumnName, *RowName);
-					GenerationContext.Compiler->CompilerLog(FText::FromString(msg), TableNode);
+					FString MeshPath;
+					SkeletalMesh->GetOuter()->GetPathName(nullptr, MeshPath);
+					FString MeshTag = FString("__MeshPath:") + MeshPath;
+					AddTagToMutableMeshUnique(*MutableMesh, MeshTag);
 				}
+
+				MutableTable->SetCell(CurrentColumn, RowIdx, MutableMesh.get());
+			}
+			else
+			{
+				FString msg = FString::Printf(TEXT("Error converting skeletal mesh LOD %d, Section %d from column [%s] row [%s] to mutable."),
+					LODIndex, SectionIndex, *ColumnName, *RowName);
+				GenerationContext.Compiler->CompilerLog(FText::FromString(msg), TableNode);
 			}
 		}
 
@@ -278,9 +284,8 @@ bool FillTableColumn(const UCustomizableObjectNodeTable* TableNode,	mu::TablePtr
 				return false;
 			}
 
-			// Static Meshes only support LOD 0 at the moment (TODO)
-			//int32 CurrentLOD = LOD;
-			int32 CurrentLOD = 0;
+			// Parameter used for LOD differences
+			int32 CurrentLOD = LODIndex;
 
 			int NumLODs = StaticMesh->GetRenderData()->LODResources.Num();
 
@@ -289,7 +294,7 @@ bool FillTableColumn(const UCustomizableObjectNodeTable* TableNode,	mu::TablePtr
 				CurrentLOD = NumLODs - 1;
 
 				FString msg = FString::Printf(TEXT("Mesh from column [%s] row [%s] needs LOD %d but has less LODs than the reference mesh. LOD %d will be used instead. This can cause some performance penalties."),
-					*ColumnName, *RowName, LOD, CurrentLOD);
+					*ColumnName, *RowName, LODIndex, CurrentLOD);
 				GenerationContext.Compiler->CompilerLog(FText::FromString(msg), TableNode);
 			}
 
@@ -305,32 +310,27 @@ bool FillTableColumn(const UCustomizableObjectNodeTable* TableNode,	mu::TablePtr
 				GenerationContext.Compiler->CompilerLog(FText::FromString(msg), TableNode);
 			}
 
-			int32 Materials = NumMaterials <= ReferenceNumMaterials ? NumMaterials : ReferenceNumMaterials;
+			FString MutableColumnName = TableNode->GenerateStaticMeshMutableColumName(ColumnName, SectionIndex);
 
-			for (int32 MatIndex = 0; MatIndex < Materials; ++MatIndex)
-			{
-				FString MutableColumnName = TableNode->GenerateStaticMeshMutableColumName(ColumnName, MatIndex);
-
-				CurrentColumn = MutableTable.get()->FindColumn(StringCast<ANSICHAR>(*MutableColumnName).Get());
+			CurrentColumn = MutableTable.get()->FindColumn(StringCast<ANSICHAR>(*MutableColumnName).Get());
 
-				if (CurrentColumn == -1)
-				{
-					CurrentColumn = MutableTable->AddColumn(StringCast<ANSICHAR>(*MutableColumnName).Get(), mu::TABLE_COLUMN_TYPE::TCT_MESH);
-				}
+			if (CurrentColumn == -1)
+			{
+				CurrentColumn = MutableTable->AddColumn(StringCast<ANSICHAR>(*MutableColumnName).Get(), mu::TABLE_COLUMN_TYPE::TCT_MESH);
+			}
 
-				mu::MeshPtr MutableMesh = GenerateMutableMesh(StaticMesh, TSoftClassPtr<UAnimInstance>(), CurrentLOD, MatIndex, FString(), GenerationContext, TableNode);
+			mu::MeshPtr MutableMesh = GenerateMutableMesh(StaticMesh, TSoftClassPtr<UAnimInstance>(), CurrentLOD, SectionIndex, CurrentLOD, SectionIndex, FString(), GenerationContext, TableNode); // TODO GMT
 
-				if (MutableMesh)
-				{
-					MutableTable->SetCell(CurrentColumn, RowIdx, MutableMesh.get());
-				}
-				else
-				{
-					FString msg = FString::Printf(TEXT("Error converting skeletal mesh LOD %d, Material %d from column [%s] row [%s] to mutable."),
-						LOD, MatIndex, *ColumnName, *RowName);
+			if (MutableMesh)
+			{
+				MutableTable->SetCell(CurrentColumn, RowIdx, MutableMesh.get());
+			}
+			else
+			{
+				FString msg = FString::Printf(TEXT("Error converting skeletal mesh LOD %d, Section %d from column [%s] row [%s] to mutable."),
+					LODIndex, SectionIndex, *ColumnName, *RowName);
 
-					GenerationContext.Compiler->CompilerLog(FText::FromString(msg), TableNode);
-				}
+				GenerationContext.Compiler->CompilerLog(FText::FromString(msg), TableNode);
 			}
 		}
 
@@ -701,7 +701,8 @@ bool FillTableColumn(const UCustomizableObjectNodeTable* TableNode,	mu::TablePtr
 }
 
 
-bool GenerateTableColumn(const UCustomizableObjectNodeTable* TableNode, const UEdGraphPin* Pin, mu::TablePtr MutableTable, const FString& DataTableColumnName, const int32 LOD, FMutableGraphGenerationContext& GenerationContext)
+bool GenerateTableColumn(const UCustomizableObjectNodeTable* TableNode, const UEdGraphPin* Pin, mu::TablePtr MutableTable, const FString& DataTableColumnName,
+	const int32 LODIndexConnected, const int32 SectionIndexConnected, const int32 LODIndex, const int32 SectionIndex, const bool bOnlyConnectedLOD, FMutableGraphGenerationContext& GenerationContext)
 {
 	SCOPED_PIN_DATA(GenerationContext, Pin)
 
@@ -734,7 +735,8 @@ bool GenerateTableColumn(const UCustomizableObjectNodeTable* TableNode, const UE
 
 			if (CellData)
 			{
-				bool bColumnGenerated = FillTableColumn(TableNode, MutableTable, DataTableColumnName, RowNames[RowIndex].ToString(), RowIndex, CellData, ColumnProperty, LOD, GenerationContext);
+				bool bColumnGenerated = FillTableColumn(TableNode, MutableTable, DataTableColumnName, RowNames[RowIndex].ToString(), RowIndex, CellData, ColumnProperty,
+					LODIndexConnected, SectionIndexConnected, LODIndex, SectionIndex, bOnlyConnectedLOD, GenerationContext);
 				
 				if (!bSuccess)
 				{
diff --git a/Engine/Plugins/Experimental/Mutable/Source/CustomizableObjectEditor/Private/MuCOE/GenerateMutableSource/GenerateMutableSourceTable.h b/Engine/Plugins/Experimental/Mutable/Source/CustomizableObjectEditor/Private/MuCOE/GenerateMutableSource/GenerateMutableSourceTable.h
--- a/Engine/Plugins/Experimental/Mutable/Source/CustomizableObjectEditor/Private/MuCOE/GenerateMutableSource/GenerateMutableSourceTable.h
+++ b/Engine/Plugins/Experimental/Mutable/Source/CustomizableObjectEditor/Private/MuCOE/GenerateMutableSource/GenerateMutableSourceTable.h
@@ -11,9 +11,48 @@ class UCustomizableObjectNodeTable;
 class UEdGraphPin;
 struct FMutableGraphGenerationContext;
 
-bool GenerateTableColumn(const UCustomizableObjectNodeTable* TableNode, const UEdGraphPin* Pin, mu::TablePtr MutableTable, const FString& DataTableColumnName, const int32 LOD, FMutableGraphGenerationContext& GenerationContext);
 
+/** 
+ *
+ * @param TableNode 
+ * @param Pin 
+ * @param MutableTable 
+ * @param DataTableColumnName 
+ * @param LODIndexConnected LOD which the pin is connected to.
+ * @param SectionIndexConnected Section which the pin is connected to.
+ * @param LODIndex LOD we are generating. Will be different from LODIndexConnected only when using Automatic LOD From Mesh. 
+ * @param SectionIndex Section we are generating. Will be different from SectionIndexConnected only when using Automatic LOD From Mesh.
+ * @param bOnlyConnectedLOD Corrected LOD and Section will unconditionally always be the connected ones.
+ * @param GenerationContext 
+ * @return  */
+bool GenerateTableColumn(const UCustomizableObjectNodeTable* TableNode, const UEdGraphPin* Pin, mu::TablePtr MutableTable, const FString& DataTableColumnName,
+	int32 LODIndexConnected, int32 SectionIndexConnected,
+	int32 LODIndex, int32 SectionIndex,
+	bool bOnlyConnectedLOD,
+	FMutableGraphGenerationContext& GenerationContext);
+
+
+/**
+ *
+ * @param TableNode 
+ * @param MutableTable 
+ * @param ColumnName 
+ * @param RowName 
+ * @param RowIdx 
+ * @param CellData 
+ * @param Property 
+ * @param LODIndexConnected LOD which the pin is connected to.
+ * @param SectionIndexConnected Section which the pin is connected to.
+ * @param LODIndex LOD we are generating. Will be different from LODIndexConnected only when using Automatic LOD From Mesh. 
+ * @param SectionIndex Section we are generating. Will be different from SectionIndexConnected only when using Automatic LOD From Mesh.
+ * @param bOnlyConnectedLOD Corrected LOD and Section will unconditionally always be the connected ones.
+ * @param GenerationContext 
+ * @return  */
 bool FillTableColumn(const UCustomizableObjectNodeTable* TableNode, mu::TablePtr MutableTable, const FString& ColumnName,
-	const FString& RowName, const int32 RowIdx, uint8* CellData, FProperty* Property, const int32 LOD, FMutableGraphGenerationContext& GenerationContext);
+	const FString& RowName, int32 RowIdx, uint8* CellData, FProperty* Property,
+	int32 LODIndexConnected, int32 SectionIndexConnected, int32 LODIndex, int32 SectionIndex,
+	bool bOnlyConnectedLOD,
+	FMutableGraphGenerationContext& GenerationContext);
+
 
 mu::TablePtr GenerateMutableSourceTable(const FString& TableName, const UEdGraphPin* Pin, FMutableGraphGenerationContext& GenerationContext);
