[Mutable] Fix missing header in GenerateMutableSource.h

#rnx

diff --git a/Engine/Plugins/Experimental/Mutable/Source/CustomizableObjectEditor/Private/MuCOE/GenerateMutableSource/GenerateMutableSource.h b/Engine/Plugins/Experimental/Mutable/Source/CustomizableObjectEditor/Private/MuCOE/GenerateMutableSource/GenerateMutableSource.h
--- a/Engine/Plugins/Experimental/Mutable/Source/CustomizableObjectEditor/Private/MuCOE/GenerateMutableSource/GenerateMutableSource.h
+++ b/Engine/Plugins/Experimental/Mutable/Source/CustomizableObjectEditor/Private/MuCOE/GenerateMutableSource/GenerateMutableSource.h
@@ -3,6 +3,7 @@
 #pragma once
 
 #include "Animation/Skeleton.h"
+#include "Animation/AnimInstance.h"
 #include "Engine/SkeletalMesh.h"
 #include "MuCO/CustomizableObject.h"
 #include "MuCOE/ExtensionDataCompilerInterface.h"
