[Mutable] Fix missing header in GenerateMutableSourceMesh.h

#rnx

diff --git a/Engine/Plugins/Experimental/Mutable/Source/CustomizableObjectEditor/Private/MuCOE/GenerateMutableSource/GenerateMutableSourceMesh.h b/Engine/Plugins/Experimental/Mutable/Source/CustomizableObjectEditor/Private/MuCOE/GenerateMutableSource/GenerateMutableSourceMesh.h
--- a/Engine/Plugins/Experimental/Mutable/Source/CustomizableObjectEditor/Private/MuCOE/GenerateMutableSource/GenerateMutableSourceMesh.h
+++ b/Engine/Plugins/Experimental/Mutable/Source/CustomizableObjectEditor/Private/MuCOE/GenerateMutableSource/GenerateMutableSourceMesh.h
@@ -4,7 +4,7 @@
 
 #include "MuT/NodeModifierMeshClipWithMesh.h"
 #include "MuT/Table.h"
-
+#include "Animation/AnimInstance.h"
 #include "UObject/SoftObjectPtr.h"
 
 class FName;
