[mutable] Prevent concurrent operation execution if we are memory constrained.
Track memory in streaming buffers.
Track memory for constant image loading.
Make interpolate operation "in-place".
Removed unused operation "interpolate-3".
#rnx

diff --git a/Engine/Plugins/Experimental/Mutable/Source/CustomizableObject/Classes/MuCO/CustomizableObject.h b/Engine/Plugins/Experimental/Mutable/Source/CustomizableObject/Classes/MuCO/CustomizableObject.h
--- a/Engine/Plugins/Experimental/Mutable/Source/CustomizableObject/Classes/MuCO/CustomizableObject.h
+++ b/Engine/Plugins/Experimental/Mutable/Source/CustomizableObject/Classes/MuCO/CustomizableObject.h
@@ -1439,7 +1439,7 @@ private:
 	// This is a manual version number for the binary blobs in this asset.
 	// Increasing it invalidates all the previously compiled models.
 	// Warning: If while merging code both versions have changed, take the highest+1.
-	static const int32 CurrentSupportedVersion = 389;
+	static const int32 CurrentSupportedVersion = 390;
 
 public:
 
diff --git a/Engine/Plugins/Experimental/Mutable/Source/MutableRuntime/Private/MuR/CodeRunner.cpp b/Engine/Plugins/Experimental/Mutable/Source/MutableRuntime/Private/MuR/CodeRunner.cpp
--- a/Engine/Plugins/Experimental/Mutable/Source/MutableRuntime/Private/MuR/CodeRunner.cpp
+++ b/Engine/Plugins/Experimental/Mutable/Source/MutableRuntime/Private/MuR/CodeRunner.cpp
@@ -944,7 +944,8 @@ namespace mu
 
 			int32 MipsToSkip = item.ExecutionOptions;
             Ptr<const Image> Source;
-            program.GetConstant(cat, Source, MipsToSkip);
+			program.GetConstant(cat, Source, MipsToSkip,
+				[this](int32 x, int32 y, int32 m, EImageFormat f) { return CreateImage(x,y,m,f); } );
 
 			// Assume the ROM has been loaded previously in a task generated at IssueOp
 			check(Source);
@@ -3794,29 +3795,29 @@ namespace mu
                     {
 						int32 LevelCount = FMath::Max(pMin->GetLODCount(), pMax->GetLODCount());
 						
-						Ptr<Image> pNew = CreateImage( pMin->GetSizeX(), pMin->GetSizeY(), LevelCount, pMin->GetFormat() );
+						Ptr<Image> pNew = CloneOrTakeOver(pMin);
 
 						// Be defensive: ensure image sizes match.
-						if (pMin->GetSize() != pMax->GetSize())
+						if (pNew->GetSize() != pMax->GetSize())
 						{
 							MUTABLE_CPUPROFILER_SCOPE(ImageResize_ForInterpolate);
-							Ptr<Image> Resized = CreateImage(pMin->GetSizeX(), pMin->GetSizeY(), pMax->GetLODCount(), pMax->GetFormat(), EInitializationType::NotInitialized);
+							Ptr<Image> Resized = CreateImage(pNew->GetSizeX(), pNew->GetSizeY(), pMax->GetLODCount(), pMax->GetFormat(), EInitializationType::NotInitialized);
 							ImageResizeLinear(Resized.get(), 0, pMax.get());
 							Release(pMax);
 							pMax = Resized;
 						}
 
-						if (pMin->GetLODCount() != LevelCount)
+						if (pNew->GetLODCount() != LevelCount)
 						{
 							MUTABLE_CPUPROFILER_SCOPE(Mipmap_ForInterpolate);
 
-							ImagePtr pDest = CreateImage(pMin->GetSizeX(), pMin->GetSizeY(), LevelCount, pMin->GetFormat());
+							ImagePtr pDest = CreateImage(pNew->GetSizeX(), pNew->GetSizeY(), LevelCount, pNew->GetFormat());
 
 							FMipmapGenerationSettings settings{};
-							ImageMipmap(m_pSystem->WorkingMemoryManager, m_pSettings->ImageCompressionQuality, pDest.get(), pMin.get(), LevelCount, settings);
+							ImageMipmap(m_pSystem->WorkingMemoryManager, m_pSettings->ImageCompressionQuality, pDest.get(), pNew.get(), LevelCount, settings);
 
-							Release(pMin);
-							pMin = pDest;
+							Release(pNew);
+							pNew = pDest;
 						}
 
 						if (pMax->GetLODCount() != LevelCount)
@@ -3832,9 +3833,8 @@ namespace mu
 							pMax = pDest;
 						}
 
-                        ImageInterpolate( pNew.get(), pMin.get(), pMax.get(), bifactor );
+                        ImageInterpolate( pNew.get(), pMax.get(), bifactor );
 
-						Release(pMin);
 						Release(pMax);
 						StoreImage(item, pNew);
 					}
@@ -3859,55 +3859,6 @@ namespace mu
             break;
         }
 
-        case OP_TYPE::IM_INTERPOLATE3:
-        {
-			OP::ImageInterpolate3Args args = pModel->GetPrivate()->m_program.GetOpArgs<OP::ImageInterpolate3Args>(item.At);
-            switch (item.Stage)
-            {
-            case 0:
-                AddOp( FScheduledOp( item.At, item, 1),
-                        FScheduledOp::FromOpAndOptions( args.factor1, item, 0),
-                        FScheduledOp::FromOpAndOptions( args.factor2, item, 0),
-                        FScheduledOp( args.target0, item),
-                        FScheduledOp( args.target1, item),
-                        FScheduledOp( args.target2, item) );
-                break;
-
-            case 1:
-            {
-           		MUTABLE_CPUPROFILER_SCOPE(IM_INTERPOLATE3_1)
-
-                // \TODO Optimise limit cases
-                float factor1 = LoadScalar(FScheduledOp::FromOpAndOptions(args.factor1, item, 0));
-                float factor2 = LoadScalar(FScheduledOp::FromOpAndOptions(args.factor2, item, 0));
-
-				Ptr<const Image> Target0 = LoadImage( FCacheAddress(args.target0,item) );
-				Ptr<const Image> Target1 = LoadImage( FCacheAddress(args.target1,item) );
-				Ptr<const Image> Target2 = LoadImage( FCacheAddress(args.target2,item) );
-
-				Ptr<Image> Result;
-				
-				if (Target0)
-				{
-					// \TODO Optimise by doing in-place if possible
-					Result = CreateImageLike(Target0.get(), EInitializationType::NotInitialized);
-					ImageInterpolate(Result.get(), Target0.get(), Target1.get(), Target2.get(), factor1, factor2);
-				}
-
-				Release(Target0);
-				Release(Target1);
-				Release(Target2);
-				StoreImage(item, Result);
-				break;
-            }
-
-            default:
-                check(false);
-            }
-
-            break;
-        }
-
         case OP_TYPE::IM_SATURATE:
         {
 			OP::ImageSaturateArgs args = pModel->GetPrivate()->m_program.GetOpArgs<OP::ImageSaturateArgs>(item.At);
@@ -6393,18 +6344,6 @@ namespace mu
 			break;
 		}
 
-		case OP_TYPE::IM_INTERPOLATE3:
-		{
-			OP::ImageInterpolate3Args args = program.GetOpArgs<OP::ImageInterpolate3Args>(item.At);
-			switch (item.Stage)
-			{
-			case 0: AddOp(FScheduledOp(item.At, item, 1), FScheduledOp(args.target0, item)); break;
-			case 1: StoreValidDesc(item); break;
-			default: check(false);
-			}
-			break;
-		}
-
 		case OP_TYPE::IM_SATURATE:
 		{
 			OP::ImageSaturateArgs args = program.GetOpArgs<OP::ImageSaturateArgs>(item.At);
diff --git a/Engine/Plugins/Experimental/Mutable/Source/MutableRuntime/Private/MuR/CodeRunner.h b/Engine/Plugins/Experimental/Mutable/Source/MutableRuntime/Private/MuR/CodeRunner.h
--- a/Engine/Plugins/Experimental/Mutable/Source/MutableRuntime/Private/MuR/CodeRunner.h
+++ b/Engine/Plugins/Experimental/Mutable/Source/MutableRuntime/Private/MuR/CodeRunner.h
@@ -37,9 +37,7 @@ namespace  mu
 	class RangeIndex;
 
 
-    //---------------------------------------------------------------------------------------------
-    //! Code execution of the mutable virtual machine.
-    //---------------------------------------------------------------------------------------------
+    /** Code execution of the mutable virtual machine. */
     class CodeRunner
     {
     public:
@@ -50,8 +48,7 @@ namespace  mu
 
     protected:
 
-        //! Type of data sometimes stored in the code runner heap to pass info between
-        //! operation stages.
+        /** Type of data sometimes stored in the code runner heap to pass info between operation stages. */
         struct FScheduledOpData
         {
 			union
@@ -126,7 +123,7 @@ namespace  mu
 		TArray< FScheduledOpData > m_heapData;
 		TArray< FImageDesc > m_heapImageDesc;
 
-		//! Only used for correct mip skipping with external images. Is the LOD for which the image is build.
+		/** Only used for correct mip skipping with external images. It is the LOD for which the image is build. */
 		int32 ImageLOD;
 
 	public:
@@ -371,7 +368,10 @@ namespace  mu
 	protected:
 
 		/** Strategy to choose the order of execution of operations. */
-		EExecutionStrategy ExecutionStrategy = EExecutionStrategy::None; 
+		EExecutionStrategy ExecutionStrategy = EExecutionStrategy::None;
+
+		/** If this flag is enabled, issued operation stage that use tasks will be executed in the mutable thread instead of in a generic worker thread. */
+		bool bForceSerialTaskExecution = false;
 
 		/** List of pending operations that we don't know if they cannot be run yet because of dependencies. */
 		TArray< FTask > ClosedTasks;
@@ -382,127 +382,130 @@ namespace  mu
 		/** For every op, up to what stage it has been scheduled to run. */
 		CodeContainer<uint8> ScheduledStagePerOp;
 
+		/** List of tasks that are ready to run concurrently. */
+		TArray< TSharedPtr<FIssuedTask> > IssuedTasksOnHold;
+
 		/** List of tasks that have been set to run concurrently and their completion is unknown. */
 		TArray< TSharedPtr<FIssuedTask> > IssuedTasks;
 
 	public:
 
-		bool LoadBool(const FCacheAddress& From)
+		inline bool LoadBool(const FCacheAddress& From)
 		{
 			return m_pSystem->WorkingMemoryManager.CurrentInstanceCache->GetBool(From);
 		}
 
-		float LoadInt(const FCacheAddress& From)
+		inline float LoadInt(const FCacheAddress& From)
 		{
 			return m_pSystem->WorkingMemoryManager.CurrentInstanceCache->GetInt(From);
 		}
 
-		float LoadScalar(const FCacheAddress& From)
+		inline float LoadScalar(const FCacheAddress& From)
 		{
 			return m_pSystem->WorkingMemoryManager.CurrentInstanceCache->GetScalar(From);
 		}
 
-		FVector4f LoadColor(const FCacheAddress& From)
+		inline FVector4f LoadColor(const FCacheAddress& From)
 		{
 			return m_pSystem->WorkingMemoryManager.CurrentInstanceCache->GetColour(From);
 		}
 
-		Ptr<const String> LoadString(const FCacheAddress& From)
+		inline Ptr<const String> LoadString(const FCacheAddress& From)
 		{
 			return m_pSystem->WorkingMemoryManager.CurrentInstanceCache->GetString(From);
 		}
 
-		FProjector LoadProjector(const FCacheAddress& From)
+		inline FProjector LoadProjector(const FCacheAddress& From)
 		{
 			return m_pSystem->WorkingMemoryManager.CurrentInstanceCache->GetProjector(From);
 		}
 
-		Ptr<const Mesh> LoadMesh(const FCacheAddress& From)
+		inline Ptr<const Mesh> LoadMesh(const FCacheAddress& From)
 		{
 			return m_pSystem->WorkingMemoryManager.LoadMesh(From);
 		}
 
-		Ptr<const Image> LoadImage(const FCacheAddress& From)
+		inline Ptr<const Image> LoadImage(const FCacheAddress& From)
 		{
 			return m_pSystem->WorkingMemoryManager.LoadImage(From);
 		}
 
-		Ptr<const Layout> LoadLayout(const FCacheAddress& From)
+		inline Ptr<const Layout> LoadLayout(const FCacheAddress& From)
 		{
 			return m_pSystem->WorkingMemoryManager.CurrentInstanceCache->GetLayout(From);
 		}
 
-		Ptr<const Instance> LoadInstance(const FCacheAddress& From)
+		inline Ptr<const Instance> LoadInstance(const FCacheAddress& From)
 		{
 			return m_pSystem->WorkingMemoryManager.CurrentInstanceCache->GetInstance(From);
 		}
 
-		Ptr<const ExtensionData> LoadExtensionData(const FCacheAddress& From)
+		inline Ptr<const ExtensionData> LoadExtensionData(const FCacheAddress& From)
 		{
 			return m_pSystem->WorkingMemoryManager.CurrentInstanceCache->GetExtensionData(From);
 		}
 
-		void StoreValidDesc(const FCacheAddress& To)
+		inline void StoreValidDesc(const FCacheAddress& To)
 		{
 			m_pSystem->WorkingMemoryManager.CurrentInstanceCache->SetValidDesc(To);
 		}
 
-		void StoreBool(const FCacheAddress& To, bool Value)
+		inline void StoreBool(const FCacheAddress& To, bool Value)
 		{
 			m_pSystem->WorkingMemoryManager.CurrentInstanceCache->SetBool(To, Value);
 		}
 
-		void StoreInt(const FCacheAddress& To, int32 Value)
+		inline void StoreInt(const FCacheAddress& To, int32 Value)
 		{
 			m_pSystem->WorkingMemoryManager.CurrentInstanceCache->SetInt(To, Value);
 		}
 
-		void StoreScalar(const FCacheAddress& To, float Value)
+		inline void StoreScalar(const FCacheAddress& To, float Value)
 		{
 			m_pSystem->WorkingMemoryManager.CurrentInstanceCache->SetScalar(To, Value);
 		}
 
-		void StoreString(const FCacheAddress& To, Ptr<const String> Value)
+		inline void StoreString(const FCacheAddress& To, Ptr<const String> Value)
 		{
 			m_pSystem->WorkingMemoryManager.CurrentInstanceCache->SetString(To, Value);
 		}
 
-		void StoreColor(const FCacheAddress& To, const FVector4f& Value)
+		inline void StoreColor(const FCacheAddress& To, const FVector4f& Value)
 		{
 			m_pSystem->WorkingMemoryManager.CurrentInstanceCache->SetColour(To, Value);
 		}
 
-		void StoreProjector(const FCacheAddress& To, const FProjector& Value)
+		inline void StoreProjector(const FCacheAddress& To, const FProjector& Value)
 		{
 			m_pSystem->WorkingMemoryManager.CurrentInstanceCache->SetProjector(To, Value);
 		}
 
-		void StoreMesh(const FCacheAddress& To, Ptr<const Mesh> Resource)
+		inline void StoreMesh(const FCacheAddress& To, Ptr<const Mesh> Resource)
 		{
 			m_pSystem->WorkingMemoryManager.StoreMesh(To, Resource);
 		}
 
-		void StoreImage(const FCacheAddress& To, Ptr<const Image> Resource)
+		inline void StoreImage(const FCacheAddress& To, Ptr<const Image> Resource)
 		{
 			m_pSystem->WorkingMemoryManager.StoreImage(To, Resource);
 		}
 
-		void StoreLayout(const FCacheAddress& To, Ptr<const Layout> Resource)
+		inline void StoreLayout(const FCacheAddress& To, Ptr<const Layout> Resource)
 		{
 			m_pSystem->WorkingMemoryManager.CurrentInstanceCache->SetLayout(To, Resource);
 		}
 
-		void StoreInstance(const FCacheAddress& To, Ptr<const Instance> Resource)
+		inline void StoreInstance(const FCacheAddress& To, Ptr<const Instance> Resource)
 		{
 			m_pSystem->WorkingMemoryManager.CurrentInstanceCache->SetInstance(To, Resource);
 		}
 
-		void StoreExtensionData(const FCacheAddress& To, Ptr<const ExtensionData> Resource)
+		inline void StoreExtensionData(const FCacheAddress& To, Ptr<const ExtensionData> Resource)
 		{
 			m_pSystem->WorkingMemoryManager.CurrentInstanceCache->SetExtensionData(To, Resource);
 		}
 
-		Ptr<Image> CreateImage(uint32 SizeX, uint32 SizeY, uint32 Lods, EImageFormat Format, EInitializationType Init = EInitializationType::Black)
+		inline Ptr<Image> CreateImage(uint32 SizeX, uint32 SizeY, uint32 Lods, EImageFormat Format, EInitializationType Init = EInitializationType::Black)
 		{
 			Ptr<Image> Result = m_pSystem->WorkingMemoryManager.CreateImage(SizeX, SizeY, Lods, Format, Init);
 			return MoveTemp(Result);
@@ -515,38 +518,38 @@ namespace  mu
 		}
 
 		/** Ref will be nulled and relesed in any case. */
-		Ptr<Image> CloneOrTakeOver(Ptr<const Image>& Ref)
+		inline Ptr<Image> CloneOrTakeOver(Ptr<const Image>& Ref)
 		{
 			Ptr<Image> Result = m_pSystem->WorkingMemoryManager.CloneOrTakeOver(Ref);
 			return MoveTemp(Result);
 		}
 
-		void Release(Ptr<const Image>& Resource)
+		inline void Release(Ptr<const Image>& Resource)
 		{
 			return m_pSystem->WorkingMemoryManager.Release(Resource);
 		}
 
-		void Release(Ptr<Image>& Resource)
+		inline void Release(Ptr<Image>& Resource)
 		{
 			return m_pSystem->WorkingMemoryManager.Release(Resource);
 		}
 
-		UE_NODISCARD Ptr<Mesh> CreateMesh(int32 BudgetReserveSize = 0)
+		UE_NODISCARD inline Ptr<Mesh> CreateMesh(int32 BudgetReserveSize = 0)
 		{
 			return m_pSystem->WorkingMemoryManager.CreateMesh(BudgetReserveSize);
 		}
 
-		UE_NODISCARD Ptr<Mesh> CloneOrTakeOver(Ptr<const Mesh>& Ref)
+		UE_NODISCARD inline Ptr<Mesh> CloneOrTakeOver(Ptr<const Mesh>& Ref)
 		{
 			return m_pSystem->WorkingMemoryManager.CloneOrTakeOver(Ref);
 		}
 
-		void Release(Ptr<const Mesh>& Resource)
+		inline void Release(Ptr<const Mesh>& Resource)
 		{
 			m_pSystem->WorkingMemoryManager.Release(Resource);
 		}
 
-		void Release(Ptr<Mesh>& Resource)
+		inline void Release(Ptr<Mesh>& Resource)
 		{
 			m_pSystem->WorkingMemoryManager.Release(Resource);
 		}
@@ -606,6 +609,31 @@ namespace  mu
 		/** Update debug stats. */
 		void UpdateTraces();
 
+		/** */
+		inline bool ShouldIssueTask() const
+		{
+			// Can we afford to delay issued tasks?
+			bool bCanDelayTasks = IssuedTasks.Num() > 0 || OpenTasks.Num() > 0;
+			if (!bCanDelayTasks)
+			{
+				return true;
+			}
+			else
+			{
+				// We could wait. Let's see if we have enough memory to issue tasks anyway.
+				bool bHaveEnoughMemory = !m_pSystem->WorkingMemoryManager.IsMemoryBudgetFull();
+				if (bHaveEnoughMemory)
+				{
+					return  true;
+				}
+			}
+
+			return false;
+		}
+
+		/** */
+		void LaunchIssuedTask(const TSharedPtr<FIssuedTask>& TaskToIssue, bool& bOutFailed);
+
     };
 
 }
diff --git a/Engine/Plugins/Experimental/Mutable/Source/MutableRuntime/Private/MuR/CodeRunnerTasked.cpp b/Engine/Plugins/Experimental/Mutable/Source/MutableRuntime/Private/MuR/CodeRunnerTasked.cpp
--- a/Engine/Plugins/Experimental/Mutable/Source/MutableRuntime/Private/MuR/CodeRunnerTasked.cpp
+++ b/Engine/Plugins/Experimental/Mutable/Source/MutableRuntime/Private/MuR/CodeRunnerTasked.cpp
@@ -101,6 +101,8 @@ namespace mu
 				break;
 			}
 
+			m_pSystem->WorkingMemoryManager.TrackedBudgetBytes_Stream -= o.m_streamBuffer.GetAllocatedSize();
+
 			o.m_romIndex = -1;
 			o.m_streamBuffer.Empty();
 		}
@@ -110,82 +112,15 @@ namespace mu
 	TRACE_DECLARE_INT_COUNTER(MutableRuntime_OpenTask, TEXT("MutableRuntime/OpenTask"));
 	TRACE_DECLARE_INT_COUNTER(MutableRuntime_ClosedTasks, TEXT("MutableRuntime/ClosedTasks"));
 	TRACE_DECLARE_INT_COUNTER(MutableRuntime_IssuedTasks, TEXT("MutableRuntime/IssuedTasks"));
-	TRACE_DECLARE_INT_COUNTER(MutableRuntime_PCache_Keep, TEXT("MutableRuntime/PCache/Keep"));
-	TRACE_DECLARE_INT_COUNTER(MutableRuntime_PCache_Burn, TEXT("MutableRuntime/PCache/Burn"));
-	TRACE_DECLARE_INT_COUNTER(MutableRuntime_PCache_Update, TEXT("MutableRuntime/PCache/UpdateCount"));
-	TRACE_DECLARE_INT_COUNTER(MutableRuntime_PCache_UpdateSize, TEXT("MutableRuntime/PCache/UpdateSize"));
+	TRACE_DECLARE_INT_COUNTER(MutableRuntime_IssuedTasksOnHold, TEXT("MutableRuntime/IssuedHoldTasks"));
 
 	void CodeRunner::UpdateTraces()
 	{
-#if UE_MUTABLE_ENABLE_SLOW_TRACES && !UE_BUILD_SHIPPING && !UE_BUILD_TEST
-
-		//// Code Runner status
-		//TRACE_COUNTER_SET(MutableRuntime_OpenTask, OpenTasks.Num());
-		//TRACE_COUNTER_SET(MutableRuntime_ClosedTasks, ClosedTasks.Num());
-		//TRACE_COUNTER_SET(MutableRuntime_IssuedTasks, IssuedTasks.Num());
-
-		//// Program cache status
-		//{
-		//	CodeContainer< TPair<int32, Ptr<const RefCounted>> >::iterator It = GetMemory().m_resources.begin();
-		//	int32 NumKeep = 0;
-		//	int32 NumBurn = 0;
-		//	for (; It.IsValid(); ++It)
-		//	{
-		//		int32 State = (*It).Key;
-		//		if (State == 1)
-		//		{
-		//			++NumKeep;
-		//		}
-		//		else if (State == 2)
-		//		{
-		//			++NumBurn;
-		//		}
-		//	}
-		//	TRACE_COUNTER_SET(MutableRuntime_PCache_Keep, NumKeep);
-		//	TRACE_COUNTER_SET(MutableRuntime_PCache_Burn, NumBurn);
-		//}
-
-		//// Program cache types
-		//{
-		//	int32 NumUpdateCacheResources = 0;
-		//	int32 UpdateCacheResourcesSize = 0;
-		//	CodeContainer<int32>::iterator It = GetMemory().m_opHitCount.begin();
-		//	for (; It.IsValid(); ++It)
-		//	{
-		//		int32 Count = *It;
-		//		if ( (Count > 0x0fffff)
-		//			&& 
-		//			GetMemory().IsValid(It.get_address()))
-		//		{
-		//			++NumUpdateCacheResources;
-
-		//			OP_TYPE OpType = m_pModel->GetPrivate()->m_program.GetOpType(It.get_address().At);
-		//			DATATYPE DataType = GetOpDataType(OpType);
-		//			switch (DataType)
-		//			{
-		//			case DATATYPE::DT_MESH:
-		//			{
-		//				const Mesh* Resource = reinterpret_cast<const Mesh*>(GetMemory().m_resources[It.get_address()].Value.get());
-		//				UpdateCacheResourcesSize += Resource->GetDataSize();
-		//				break;
-		//			}
-		//			case DATATYPE::DT_IMAGE:
-		//			{
-		//				const Image* Resource = reinterpret_cast<const Image*>(GetMemory().m_resources[It.get_address()].Value.get());
-		//				UpdateCacheResourcesSize += Resource->GetDataSize();
-		//				break;
-		//			}
-		//			default:
-		//				check(false);
-		//				break;
-		//			}
-		//		}
-		//	}
-		//	TRACE_COUNTER_SET(MutableRuntime_PCache_Update, NumUpdateCacheResources);
-		//	TRACE_COUNTER_SET(MutableRuntime_PCache_UpdateSize, UpdateCacheResourcesSize);
-		//}
-
-#endif
+		// Code Runner status
+		TRACE_COUNTER_SET(MutableRuntime_OpenTask, OpenTasks.Num());
+		TRACE_COUNTER_SET(MutableRuntime_ClosedTasks, ClosedTasks.Num());
+		TRACE_COUNTER_SET(MutableRuntime_IssuedTasks, IssuedTasks.Num());
+		TRACE_COUNTER_SET(MutableRuntime_IssuedTasksOnHold, IssuedTasksOnHold.Num());
 	}
 
 
@@ -266,7 +201,6 @@ namespace mu
 		{   0,  20,   0,   0 },	// IM_BLANKLAYOUT
 		{   0,   0, -20,   0 },	// IM_COMPOSE
 		{   0,   0, -20,   0 },	// IM_INTERPOLATE
-		{   0,   0, -30,   0 },	// IM_INTERPOLATE3
 		{   0,   0,   0,   0 },	// IM_SATURATE
 		{   0,   0,   0,   0 },	// IM_LUMINANCE
 		{   0,   0,   0,   0 },	// IM_SWIZZLE
@@ -357,6 +291,74 @@ namespace mu
 	}
 
 
+	//---------------------------------------------------------------------------------------------
+	void CodeRunner::LaunchIssuedTask( const TSharedPtr<FIssuedTask>& TaskToIssue, bool& bOutFailed )
+	{
+		bool bFailed = false;
+		bool bHasWork = TaskToIssue->Prepare(this, bFailed);
+		if (bFailed)
+		{
+			bUnrecoverableError = true;
+			return;
+		}
+
+		// Launch it
+		if (bHasWork)
+		{
+			if (bForceSerialTaskExecution)
+			{
+				TaskToIssue->Event = {};
+				TaskToIssue->DoWork();
+			}
+			else
+			{
+#ifdef MUTABLE_USE_NEW_TASKGRAPH
+				TaskToIssue->Event = UE::Tasks::Launch(TEXT("MutableCore_Task"),
+					[TaskToIssue]() { TaskToIssue->DoWork(); },
+					UE::Tasks::ETaskPriority::Inherit);
+#else
+
+				ENamedThreads::Type Priority;
+				switch (CoreRunnerTaskPriority)
+				{
+				default:
+				case 0:
+					Priority = ENamedThreads::AnyThread;
+					break;
+				case 1:
+					Priority = ENamedThreads::AnyHiPriThreadHiPriTask;
+					break;
+				case 2:
+					Priority = ENamedThreads::AnyHiPriThreadNormalTask;
+					break;
+				case 3:
+					Priority = ENamedThreads::AnyNormalThreadHiPriTask;
+					break;
+				case 4:
+					Priority = ENamedThreads::AnyNormalThreadNormalTask;
+					break;
+				case 5:
+					Priority = ENamedThreads::AnyBackgroundHiPriTask;
+					break;
+				case 6:
+					Priority = ENamedThreads::AnyBackgroundThreadNormalTask;
+					break;
+				}
+
+				TaskToIssue->Event = FFunctionGraphTask::CreateAndDispatchWhenReady(
+					[TaskToIssue]() { TaskToIssue->DoWork(); },
+					TStatId{},
+					nullptr,
+					Priority);
+#endif
+			}
+		}
+
+		// Remember it for later processing.
+		IssuedTasks.Add(TaskToIssue);
+	}
+
+
     //---------------------------------------------------------------------------------------------
     void CodeRunner::Run()
     {
@@ -417,43 +419,41 @@ namespace mu
 				}
 			}
 
-			UpdateTraces();
-
 			while (!OpenTasks.IsEmpty())
 			{
 				// Get a new task to run
 				FScheduledOp item;
 				switch (ExecutionStrategy)
 				{
-				case EExecutionStrategy::MinimizeMemory:
-				{
-					// TODO: This should be done when an operation is added to the OpenTasks array instead of every time.
-					int32 BestOp = 0;
-					int8 BestDelta = TNumericLimits<int8>::Max();
-					int32 OpenOpCount = OpenTasks.Num();
-					const FProgram& Program = m_pModel->GetPrivate()->m_program;
-					for (int32 OpIndex = 0; OpIndex < OpenOpCount; ++OpIndex)
-					{
-						const FScheduledOp& Candidate = OpenTasks[OpIndex];
-						int32 OpDelta = GetOpEstimatedMemoryDelta(Candidate, Program);
-
-						if (OpDelta < BestDelta)
-						{
-							BestOp = OpIndex;
-							BestDelta = OpDelta;
-
-							// Shortcut: If we are freeing memory, we don't care how much: it is already the best op
-							if (BestDelta < 0)
-							{
-								break;
-							}
-						}
-					}
-
-					item = OpenTasks[BestOp];
-					OpenTasks.RemoveAtSwap(BestOp,1,false);
-					break;
-				}
+				//case EExecutionStrategy::MinimizeMemory:
+				//{
+				//	// TODO: This should be done when an operation is added to the OpenTasks array instead of every time.
+				//	int32 BestOp = 0;
+				//	int8 BestDelta = TNumericLimits<int8>::Max();
+				//	int32 OpenOpCount = OpenTasks.Num();
+				//	const FProgram& Program = m_pModel->GetPrivate()->m_program;
+				//	for (int32 OpIndex = 0; OpIndex < OpenOpCount; ++OpIndex)
+				//	{
+				//		const FScheduledOp& Candidate = OpenTasks[OpIndex];
+				//		int32 OpDelta = GetOpEstimatedMemoryDelta(Candidate, Program);
+
+				//		if (OpDelta < BestDelta)
+				//		{
+				//			BestOp = OpIndex;
+				//			BestDelta = OpDelta;
+
+				//			// Shortcut: If we are freeing memory, we don't care how much: it is already the best op
+				//			if (BestDelta < 0)
+				//			{
+				//				break;
+				//			}
+				//		}
+				//	}
+
+				//	item = OpenTasks[BestOp];
+				//	OpenTasks.RemoveAtSwap(BestOp,1,false);
+				//	break;
+				//}
 
 				case EExecutionStrategy::None:
 				default:
@@ -481,69 +481,20 @@ namespace mu
 				TSharedPtr<FIssuedTask> IssuedTask = IssueOp(item);
 				if (IssuedTask)
 				{
-					bool bFailed = false;
-					bool hasWork = IssuedTask->Prepare(this, bFailed);
-					if (bFailed)
-					{
-						bUnrecoverableError = true;
-						return;
-					}
-
-					// Launch it
-					if (hasWork)
+					if (ShouldIssueTask())
 					{
-						// Optional hack to ensure deterministic single threading.
-						if (bForceSingleThread)
+						bool bFailed = false;
+						LaunchIssuedTask(IssuedTask, bFailed);
+						if (bFailed)
 						{
-							IssuedTask->Event = {};
-							IssuedTask->DoWork();
-						}
-						else
-						{
-#ifdef MUTABLE_USE_NEW_TASKGRAPH
-							IssuedTask->Event = UE::Tasks::Launch(TEXT("MutableCore_Task"),
-								[IssuedTask]() { IssuedTask->DoWork(); },
-								UE::Tasks::ETaskPriority::Inherit );
-#else
-
-							ENamedThreads::Type Priority;
-							switch (CoreRunnerTaskPriority)
-							{
-							default:
-							case 0:
-								Priority = ENamedThreads::AnyThread;
-								break;
-							case 1:
-								Priority = ENamedThreads::AnyHiPriThreadHiPriTask;
-								break;
-							case 2:
-								Priority = ENamedThreads::AnyHiPriThreadNormalTask;
-								break;
-							case 3:
-								Priority = ENamedThreads::AnyNormalThreadHiPriTask;
-								break;
-							case 4:
-								Priority = ENamedThreads::AnyNormalThreadNormalTask;
-								break;
-							case 5:
-								Priority = ENamedThreads::AnyBackgroundHiPriTask;
-								break;
-							case 6:
-								Priority = ENamedThreads::AnyBackgroundThreadNormalTask;
-								break;								
-							}
-							
-							IssuedTask->Event = FFunctionGraphTask::CreateAndDispatchWhenReady(
-								[IssuedTask]() { IssuedTask->DoWork(); },
-								TStatId{},
-								nullptr,
-								Priority);
-#endif
+							bUnrecoverableError = true;
+							return;
 						}
 					}
-					
-					// Remember it for later processing.
-					IssuedTasks.Add(IssuedTask);
+					else
+					{
+						IssuedTasksOnHold.Add(IssuedTask);
+					}
 				}
 				else
 				{
@@ -552,15 +503,12 @@ namespace mu
 
 					if (ScheduledStagePerOp[item] == item.Stage + 1)
 					{
-						// We completed everything that was requested, clear it otherwise if needed
-						// again it is not going to be rebuilt.
-						// \TODO: track rebuilds.
+						// We completed everything that was requested, clear it. Otherwise if needed again it is not going to be rebuilt.
+						// \TODO: track operations that are run more than once?.
 						ScheduledStagePerOp[item] = 0;
 					}
 				}
 
-				UpdateTraces();
-
 				if (bProfile)
 				{
 					++NumRunOps;
@@ -568,13 +516,28 @@ namespace mu
 				}
 			}
 
+			UpdateTraces();
+
+			// Look for tasks on hold and see if we can launch them
+			while (IssuedTasksOnHold.Num() && ShouldIssueTask())
+			{
+				TSharedPtr<FIssuedTask> TaskToIssue = IssuedTasksOnHold.Pop(false);
+
+				bool bFailed = false;
+				LaunchIssuedTask(TaskToIssue, bFailed);
+				if (bFailed)
+				{
+					bUnrecoverableError = true;
+					return;
+				}
+			}
+
 			// Look for completed streaming ops and complete the rom loading
 			for (FRomLoadOp& o : m_romLoadOps)
 			{
 				if (o.m_romIndex>=0 && m_pSystem->StreamInterface->IsReadCompleted(o.m_streamID))
 				{
 					CompleteRomLoadOp(o);
-					UpdateTraces();
 				}
 			}
 
@@ -607,9 +570,10 @@ namespace mu
 					++Index;
 				}
 
-				UpdateTraces();
 			}
 
+			UpdateTraces();
+
 			// Debug: Did we dead-lock?
 			bool bDeadLock = !(OpenTasks.Num() || IssuedTasks.Num() || !ClosedTasks.Num() || bSomeWasReady);
 			if (bDeadLock)
@@ -665,8 +629,6 @@ namespace mu
 					// We should never reach this, since it would mean we deadlocked.
 					//check(false);
 				}
-
-				UpdateTraces();
 			}
 		}
 
@@ -2101,9 +2063,8 @@ namespace mu
 
 			if (uint32(op->m_streamBuffer.Num()) < RomSize)
 			{
-				// This could happen in 32-bit platforms
-				check(RomSize < std::numeric_limits<size_t>::max());
 				op->m_streamBuffer.SetNum((size_t)RomSize);
+				Runner->m_pSystem->WorkingMemoryManager.TrackedBudgetBytes_Stream += op->m_streamBuffer.GetAllocatedSize();
 			}
 
 			uint32 RomId = program.m_roms[RomIndex].Id;
@@ -2276,8 +2237,6 @@ namespace mu
 
 			// Free roms if necessary
 			{
-				MUTABLE_CPUPROFILER_SCOPE(FreeingRoms);
-
 				Runner->m_pSystem->WorkingMemoryManager.MarkRomUsed(RomIndex, Runner->m_pModel);
 				Runner->m_pSystem->WorkingMemoryManager.EnsureBudgetBelow(RomSize);
 			}
@@ -2302,9 +2261,8 @@ namespace mu
 
 			if (uint32(op->m_streamBuffer.Num()) < RomSize)
 			{
-				// This could happen in 32-bit platforms
-				check(RomSize < std::numeric_limits<size_t>::max());
 				op->m_streamBuffer.SetNum(RomSize);
+				Runner->m_pSystem->WorkingMemoryManager.TrackedBudgetBytes_Stream += op->m_streamBuffer.GetAllocatedSize();
 			}
 
 			uint32 RomId = program.m_roms[RomIndex].Id;
@@ -2353,8 +2311,10 @@ namespace mu
 			Runner->m_pSystem->WorkingMemoryManager.MarkRomUsed(RomIndex, Runner->m_pModel);
 			--ModelCache.PendingOpsPerRom[RomIndex];
 
-			if (DebugRom && (DebugRomAll || RomIndex==DebugRomIndex))
+			if (DebugRom && (DebugRomAll || RomIndex == DebugRomIndex))
+			{
 				UE_LOG(LogMutableCore, Log, TEXT("FLoadImageRomsTask::Complete rom %d, now peding ops is %d."), RomIndex, ModelCache.PendingOpsPerRom[RomIndex]);
+			}
 		}
 	}
 
@@ -2510,16 +2470,19 @@ namespace mu
 			int32 LODIndexCount = program.m_constantImages[ImageIndex].LODCount - ReallySkip;
 			check(LODIndexCount > 0);
 
-			bool bAnyMissing = false;
-			for (int32 i=0; i<LODIndexCount; ++i)
-			{
-				uint32 LODIndex = program.m_constantImageLODIndices[LODIndexIndex+i];
-				if ( !program.m_constantImageLODs[LODIndex].Value )
-				{
-					bAnyMissing = true;
-					break;
-				}
-			}
+			// We always need to follow this path, or roms may not be protected for long enough and might be unloaded 
+			// because of memory budget contraints.
+			bool bAnyMissing = true;
+			//bool bAnyMissing = false;
+			//for (int32 i=0; i<LODIndexCount; ++i)
+			//{
+			//	uint32 LODIndex = program.m_constantImageLODIndices[LODIndexIndex+i];
+			//	if ( !program.m_constantImageLODs[LODIndex].Value )
+			//	{
+			//		bAnyMissing = true;
+			//		break;
+			//	}
+			//}
 
 			if (bAnyMissing)
 			{
diff --git a/Engine/Plugins/Experimental/Mutable/Source/MutableRuntime/Private/MuR/ModelPrivate.h b/Engine/Plugins/Experimental/Mutable/Source/MutableRuntime/Private/MuR/ModelPrivate.h
--- a/Engine/Plugins/Experimental/Mutable/Source/MutableRuntime/Private/MuR/ModelPrivate.h
+++ b/Engine/Plugins/Experimental/Mutable/Source/MutableRuntime/Private/MuR/ModelPrivate.h
@@ -22,7 +22,6 @@ namespace mu
 {
 
 	/** Used to debug and log. */
-	constexpr bool bForceSingleThread = false;
 	constexpr bool DebugRom = false;
 	constexpr bool DebugRomAll = false;
 	constexpr int32 DebugRomIndex = 44;
@@ -645,7 +644,8 @@ namespace mu
 
 
         //! Get a constant image, assuming it is fully loaded. The image constant will be composed with lodaded mips if necessary.
-        void GetConstant( int32 ConstantIndex, ImagePtrConst& res, int32 MipsToSkip) const
+		template <typename CreateImageFunc>
+        void GetConstant( int32 ConstantIndex, ImagePtrConst& res, int32 MipsToSkip, const CreateImageFunc& CreateImage) const
         {
 			int32 ReallySkippedLODs = FMath::Min(m_constantImages[ConstantIndex].LODCount - 1, MipsToSkip);
 			int32 FirstLODIndexIndex = m_constantImages[ConstantIndex].FirstIndex;
@@ -669,7 +669,7 @@ namespace mu
 			{
 				MUTABLE_CPUPROFILER_SCOPE(ComposeConstantImage);
 
-				Ptr<Image> Result = new Image( CurrentMip->GetSizeX(), CurrentMip->GetSizeY(), FinalLODs, CurrentMip->GetFormat() );
+				Ptr<Image> Result = CreateImage( CurrentMip->GetSizeX(), CurrentMip->GetSizeY(), FinalLODs, CurrentMip->GetFormat() );
 				Result->m_flags = CurrentMip->m_flags;
 
 				// Some non-block pixel formats require separate memory size calculation
diff --git a/Engine/Plugins/Experimental/Mutable/Source/MutableRuntime/Private/MuR/OpImageInterpolate.h b/Engine/Plugins/Experimental/Mutable/Source/MutableRuntime/Private/MuR/OpImageInterpolate.h
--- a/Engine/Plugins/Experimental/Mutable/Source/MutableRuntime/Private/MuR/OpImageInterpolate.h
+++ b/Engine/Plugins/Experimental/Mutable/Source/MutableRuntime/Private/MuR/OpImageInterpolate.h
@@ -12,54 +12,49 @@ namespace mu
 	//!
 	//---------------------------------------------------------------------------------------------
     inline void ImageInterpolate( Image* pDest,
-                                  const Image* pA,
                                   const Image* pB,
                                   float factor )
 	{
 		MUTABLE_CPUPROFILER_SCOPE(ImageInterpolate)
 
-		check(pDest && pA && pB);
-        check( pA->GetSizeX() == pDest->GetSizeX() );
-        check( pA->GetSizeY() == pDest->GetSizeY() );
-        check( pA->GetFormat() == pDest->GetFormat() );
-        check( pA->GetSizeX() == pB->GetSizeX() );
-        check( pA->GetSizeY() == pB->GetSizeY() );
-        check( pA->GetFormat() == pB->GetFormat() );
+		check(pDest && pB);
+        check(pDest->GetSizeX() == pB->GetSizeX() );
+        check(pDest->GetSizeY() == pB->GetSizeY() );
+        check(pDest->GetFormat() == pB->GetFormat() );
 
 		// Clamp the factor
 		factor = FMath::Max( 0.0f, FMath::Min( 1.0f, factor ) );
 
-        uint8_t* pDestBuf = pDest->GetData();
-        const uint8_t* pABuf = pA->GetData();
-        const uint8_t* pBBuf = pB->GetData();
+        uint8* pDestBuf = pDest->GetData();
+        const uint8* pBBuf = pB->GetData();
 
 		// Generic implementation
-		int pixelCount = (int)pA->CalculatePixelCount();
+		int32 PixelCount = (int32)pDest->CalculatePixelCount();
 
-		switch ( pA->GetFormat() )
+		switch ( pDest->GetFormat() )
 		{
 		case EImageFormat::IF_L_UBYTE:
 		{
-            uint32_t w_8 = (uint32_t)(factor*255);
-			for ( int i=0; i<pixelCount; ++i )
+            uint32 w_8 = (uint32)(factor*255);
+			for ( int i=0; i< PixelCount; ++i )
 			{
-                uint32_t a_8 = *pABuf++;
-                uint32_t b_8 = *pBBuf++;
-                uint32_t i_16 = a_8 * (255-w_8) + b_8 * w_8;
-                *pDestBuf++ = (uint8_t) ( i_16>>8 );
+                uint32 a_8 = *pDestBuf;
+                uint32 b_8 = *pBBuf++;
+                uint32 i_16 = a_8 * (255-w_8) + b_8 * w_8;
+                *pDestBuf++ = (uint8) ( i_16>>8 );
 			}
 			break;
 		}
 
 		case EImageFormat::IF_RGB_UBYTE:
 		{
-            uint32_t w_8 = (uint32_t)(factor*255);
-			for ( int i=0; i<pixelCount*3; ++i )
+            uint32 w_8 = (uint32)(factor*255);
+			for ( int i=0; i< PixelCount *3; ++i )
 			{
-                uint32_t a_8 = *pABuf++;
-                uint32_t b_8 = *pBBuf++;
-                uint32_t i_16 = a_8 * (255-w_8) + b_8 * w_8;
-                *pDestBuf++ = (uint8_t) ( i_16>>8 );
+                uint32 a_8 = *pDestBuf;
+                uint32 b_8 = *pBBuf++;
+                uint32 i_16 = a_8 * (255-w_8) + b_8 * w_8;
+                *pDestBuf++ = (uint8) ( i_16>>8 );
 			}
 			break;
 		}
@@ -67,13 +62,13 @@ namespace mu
         case EImageFormat::IF_BGRA_UBYTE:
         case EImageFormat::IF_RGBA_UBYTE:
         {
-            uint32_t w_8 = (uint32_t)(factor*255);
-			for ( int i=0; i<pixelCount*4; ++i )
+            uint32 w_8 = (uint32)(factor*255);
+			for ( int i=0; i< PixelCount *4; ++i )
 			{
-                uint32_t a_8 = *pABuf++;
-                uint32_t b_8 = *pBBuf++;
-                uint32_t i_16 = a_8 * (255-w_8) + b_8 * w_8;
-                *pDestBuf++ = (uint8_t)(i_16>>8);
+                uint32 a_8 = *pDestBuf;
+                uint32 b_8 = *pBBuf++;
+                uint32 i_16 = a_8 * (255-w_8) + b_8 * w_8;
+                *pDestBuf++ = (uint8)(i_16>>8);
 			}
 			break;
 		}
@@ -84,83 +79,4 @@ namespace mu
 
 	}
 
-
-	//---------------------------------------------------------------------------------------------
-	//!
-	//---------------------------------------------------------------------------------------------
-	inline void ImageInterpolate( Image* Dest, const Image* p0, const Image* p1, const Image* p2, float factor1, float factor2 )
-	{
-		MUTABLE_CPUPROFILER_SCOPE(ImageInterpolate2)
-
-		check( p0->GetSizeX() == p1->GetSizeX() );
-		check( p0->GetSizeY() == p1->GetSizeY() );
-		check( p0->GetFormat() == p1->GetFormat() );
-		check( p0->GetSizeX() == p2->GetSizeX() );
-		check( p0->GetSizeY() == p2->GetSizeY() );
-		check( p0->GetFormat() == p2->GetFormat() );
-
-		// Clamp the factors
-		factor1 = FMath::Clamp(factor1, 0.0f, 1.0f );
-		factor2 = FMath::Min( 1.0f-factor1, factor2 );
-		float factor0 = 1.0f-factor1-factor2;
-
-        uint8* DestBuf = Dest->GetData();
-        const uint8* p0Buf = p0->GetData();
-        const uint8* p1Buf = p1->GetData();
-        const uint8* p2Buf = p2->GetData();
-
-		// Generic implementation
-		int pixelCount = (int)p0->CalculatePixelCount();
-
-        uint32 w0_8 = (uint32)(factor0*255);
-        uint32 w1_8 = (uint32)(factor1*255);
-        uint32 w2_8 = (uint32)(factor2*255);
-
-		switch ( p0->GetFormat() )
-		{
-		case EImageFormat::IF_L_UBYTE:
-		{
-			for ( int i=0; i<pixelCount; ++i )
-			{
-                uint32 t0_8 = *p0Buf++;
-                uint32 t1_8 = *p1Buf++;
-                uint32 t2_8 = *p2Buf++;
-                uint32 i_16 = t0_8 * w0_8 + t1_8 * w1_8 + t2_8 * w2_8;
-                *DestBuf++ = (uint8) ( i_16>>8 );
-			}
-			break;
-		}
-
-		case EImageFormat::IF_RGB_UBYTE:
-		{
-			for ( int i=0; i<pixelCount*3; ++i )
-			{
-                uint32 t0_8 = *p0Buf++;
-                uint32 t1_8 = *p1Buf++;
-                uint32 t2_8 = *p2Buf++;
-                uint32 i_16 = t0_8 * w0_8 + t1_8 * w1_8 + t2_8 * w2_8;
-                *DestBuf++ = (uint8) ( i_16>>8 );
-			}
-			break;
-		}
-
-        case EImageFormat::IF_BGRA_UBYTE:
-        case EImageFormat::IF_RGBA_UBYTE:
-        {
-			for ( int i=0; i<pixelCount*4; ++i )
-			{
-                uint32 t0_8 = *p0Buf++;
-                uint32 t1_8 = *p1Buf++;
-                uint32 t2_8 = *p2Buf++;
-                uint32 i_16 = t0_8 * w0_8 + t1_8 * w1_8 + t2_8 * w2_8;
-                *DestBuf++ = (uint8) (i_16>>8);
-			}
-			break;
-		}
-
-		default:
-			check(false);
-		}
-	}
-
 }
diff --git a/Engine/Plugins/Experimental/Mutable/Source/MutableRuntime/Private/MuR/Operations.cpp b/Engine/Plugins/Experimental/Mutable/Source/MutableRuntime/Private/MuR/Operations.cpp
--- a/Engine/Plugins/Experimental/Mutable/Source/MutableRuntime/Private/MuR/Operations.cpp
+++ b/Engine/Plugins/Experimental/Mutable/Source/MutableRuntime/Private/MuR/Operations.cpp
@@ -86,7 +86,6 @@ namespace mu
 		{ DT_IMAGE,			true,   { 1, 1, 1, 1, 1, 1, 1, 1, 1, 1, 1, 1, 0, 0, 0, 0, 0, 0, 0, 0, 0, 0, 0, 0, 0, 0, 0, 0, 0 }		},	// IM_BLANKLAYOUT
 		{ DT_IMAGE,			true,	{ 1, 1, 1, 1, 1, 1, 1, 1, 1, 1, 1, 0, 0, 0, 1, 1, 1, 1, 1, 1, 1, 1, 0, 0, 1, 1, 0, 0, 0 }		},	// IM_COMPOSE
 		{ DT_IMAGE,			true,	{ 0, 1, 1, 1, 0, 0, 0, 0, 0, 0, 0, 0, 0, 0, 0, 0, 0, 0, 0, 0, 0, 0, 0, 0, 0, 0, 0, 0, 0 }		},	// IM_INTERPOLATE
-		{ DT_IMAGE,			true,	{ 0, 1, 1, 1, 0, 0, 0, 0, 0, 0, 0, 0, 0, 0, 0, 0, 0, 0, 0, 0, 0, 0, 0, 0, 0, 0, 0, 0, 0 }		},	// IM_INTERPOLATE3
 		{ DT_IMAGE,			true,	{ 0, 1, 1, 0, 0, 0, 0, 0, 0, 0, 0, 0, 0, 0, 0, 0, 0, 0, 0, 0, 0, 0, 0, 0, 0, 0, 0, 0, 0 }		},	// IM_SATURATE
 		{ DT_IMAGE,			true,	{ 0, 1, 1, 0, 0, 0, 0, 0, 0, 0, 0, 0, 0, 0, 0, 0, 0, 0, 0, 0, 0, 0, 0, 0, 0, 0, 0, 0, 0 }		},	// IM_LUMINANCE
 		{ DT_IMAGE,			true,	{ 0, 1, 1, 1, 0, 0, 0, 0, 0, 0, 0, 0, 0, 0, 0, 0, 0, 0, 0, 0, 0, 0, 0, 0, 0, 0, 0, 0, 0 }		},	// IM_SWIZZLE
@@ -276,14 +275,6 @@ namespace mu
             }
             break;
 
-        case OP_TYPE::IM_INTERPOLATE3:
-            f(&op.args.ImageInterpolate3.factor1 );
-            f(&op.args.ImageInterpolate3.factor2 );
-            f(&op.args.ImageInterpolate3.target0 );
-            f(&op.args.ImageInterpolate3.target1 );
-            f(&op.args.ImageInterpolate3.target2 );
-            break;
-
         case OP_TYPE::IM_SATURATE:
             f(&op.args.ImageSaturate.base );
             f(&op.args.ImageSaturate.factor );
@@ -694,17 +685,6 @@ namespace mu
             break;
         }
 
-        case OP_TYPE::IM_INTERPOLATE3:
-        {
-			OP::ImageInterpolate3Args args = program.GetOpArgs<OP::ImageInterpolate3Args>(at);
-            f(args.factor1 );
-            f(args.factor2 );
-            f(args.target0 );
-            f(args.target1 );
-            f(args.target2 );
-            break;
-        }
-
         case OP_TYPE::IM_SWIZZLE:
         {
 			OP::ImageSwizzleArgs args = program.GetOpArgs<OP::ImageSwizzleArgs>(at);
diff --git a/Engine/Plugins/Experimental/Mutable/Source/MutableRuntime/Private/MuR/Operations.h b/Engine/Plugins/Experimental/Mutable/Source/MutableRuntime/Private/MuR/Operations.h
--- a/Engine/Plugins/Experimental/Mutable/Source/MutableRuntime/Private/MuR/Operations.h
+++ b/Engine/Plugins/Experimental/Mutable/Source/MutableRuntime/Private/MuR/Operations.h
@@ -170,9 +170,6 @@ namespace mu
         //! Interpolate between 2 images taken from a row of targets (2 consecutive targets).
         IM_INTERPOLATE,
 
-        //! Interpolate between 3 images.
-        IM_INTERPOLATE3,
-
         //! Change the saturation of the image.
         IM_SATURATE,
 
diff --git a/Engine/Plugins/Experimental/Mutable/Source/MutableRuntime/Private/MuR/System.cpp b/Engine/Plugins/Experimental/Mutable/Source/MutableRuntime/Private/MuR/System.cpp
--- a/Engine/Plugins/Experimental/Mutable/Source/MutableRuntime/Private/MuR/System.cpp
+++ b/Engine/Plugins/Experimental/Mutable/Source/MutableRuntime/Private/MuR/System.cpp
@@ -43,7 +43,6 @@ namespace mu
 
 	TRACE_DECLARE_INT_COUNTER(MutableRuntime_LiveInstances,		TEXT("MutableRuntime/LiveInstances"));
 	TRACE_DECLARE_INT_COUNTER(MutableRuntime_Updates,			TEXT("MutableRuntime/Updates"));
-	TRACE_DECLARE_INT_COUNTER(MutableRuntime_StreamingBytes,	TEXT("MutableRuntime/StreamingBytes"));
 
 	TRACE_DECLARE_INT_COUNTER(MutableRuntime_MemInternal,	TEXT("MutableRuntime/MemInternal"));
 	TRACE_DECLARE_INT_COUNTER(MutableRuntime_MemTemp,		TEXT("MutableRuntime/MemTemp"));
@@ -169,9 +168,10 @@ namespace mu
 		check(m_pD->WorkingMemoryManager.TempMeshes.IsEmpty());
 
 		m_pD->WorkingMemoryManager.TrackedBudgetBytes_Pooled = 0;
-		m_pD->WorkingMemoryManager.TrackedBudgetBytes_Temp = 0;
+		//m_pD->WorkingMemoryManager.TrackedBudgetBytes_Temp = 0;
 		m_pD->WorkingMemoryManager.TrackedBudgetBytes_Cached = 0;
 		m_pD->WorkingMemoryManager.TrackedBudgetBytes_Rom = 0;
+		m_pD->WorkingMemoryManager.TrackedBudgetBytes_Stream = 0;
 	}
 
 	
@@ -933,7 +933,7 @@ namespace mu
 			// Data reserved for ongoing streaming operations
 			for ( const CodeRunner::FRomLoadOp& Op : CurrentRunner->m_romLoadOps )
 			{
-				StreamingBufferBytes += Op.m_streamBuffer.Num();
+				StreamingBufferBytes += Op.m_streamBuffer.GetAllocatedSize();
 			}
 
 			InternalBytes += CurrentRunner->GetInternalMemoryBytes();
@@ -976,13 +976,13 @@ namespace mu
 			+ StreamingBufferBytes
 			+ RomBytes;
 
-		UE_LOG(LogMutableCore, Log, TEXT("Mem KB: Internal %7d   Temp %7d (%7d)   Pool %7d (%7d)   Cache0+1 %7d (%7d) (%7d)   Rom %7d (%7d)   Stream %7d   Total %7d / %7d"),
+		UE_LOG(LogMutableCore, Log, TEXT("Mem KB: Internal %7d   Temp %7d    Pool %7d (%7d)   Cache0+1 %7d (%7d) (%7d)   Rom %7d (%7d)   Stream %7d (%7d)   Total %7d / %7d"),
 			InternalBytes / 1024,
-			TempBytes / 1024, TrackedBudgetBytes_Temp / 1024,
+			TempBytes / 1024, //TrackedBudgetBytes_Temp / 1024,
 			PoolBytes / 1024, TrackedBudgetBytes_Pooled / 1024,
 			CacheBytes / 1024, TrackedCacheBytes / 1024, TrackedBudgetBytes_Cached / 1024,
 			RomBytes / 1024, TrackedBudgetBytes_Rom / 1024,
-			StreamingBufferBytes / 1024,
+			StreamingBufferBytes / 1024, TrackedBudgetBytes_Stream / 1024,
 			TotalBytes / 1024,
 			BudgetBytes / 1024
 			);
@@ -1030,6 +1030,52 @@ namespace mu
     }
 
 
+	//---------------------------------------------------------------------------------------------
+	int64 FWorkingMemoryManager::GetCurrentMemoryBytes() const
+	{
+		MUTABLE_CPUPROFILER_SCOPE(GetCurrentMemoryBytes);
+
+		int32 RomBytes = TrackedBudgetBytes_Rom; // GetRomBytes();
+		int32 PoolBytes = TrackedBudgetBytes_Pooled; // GetPooledBytes();
+		int32 TempBytes = GetTempBytes(); // TrackedBudgetBytes_Temp + GetTempMeshesBytes(); 
+		int32 TrackedCacheBytes = TrackedBudgetBytes_Cached; // GetTrackedCacheBytes();
+		int32 TrackedStreamBytes = TrackedBudgetBytes_Stream;
+
+		int64 TotalBytes = 0;
+		TotalBytes += TempBytes;
+		TotalBytes += RomBytes;
+		TotalBytes += PoolBytes;
+		TotalBytes += TrackedCacheBytes;
+		TotalBytes += TrackedStreamBytes;
+
+		TRACE_COUNTER_SET(MutableRuntime_MemTemp, TempBytes);
+		TRACE_COUNTER_SET(MutableRuntime_MemPool, PoolBytes);
+		TRACE_COUNTER_SET(MutableRuntime_MemCache, TrackedCacheBytes);
+		TRACE_COUNTER_SET(MutableRuntime_MemRom, RomBytes);
+		TRACE_COUNTER_SET(MutableRuntime_MemStream, TrackedStreamBytes);
+		TRACE_COUNTER_SET(MutableRuntime_MemTotal, TotalBytes);
+		TRACE_COUNTER_SET(MutableRuntime_MemBudget, BudgetBytes);
+
+		return TotalBytes;
+	}
+
+
+	//---------------------------------------------------------------------------------------------
+	bool FWorkingMemoryManager::IsMemoryBudgetFull() const
+	{
+		// If we have 0 budget it means we have unlimited budget
+		if (BudgetBytes == 0)
+		{
+			return false;
+		}
+
+		uint64 CurrentBytes = GetCurrentMemoryBytes();
+		uint64 BudgetThresholdBytes = (BudgetBytes * 9) / 10;
+		
+		return CurrentBytes > BudgetThresholdBytes;
+	}
+
+
     //---------------------------------------------------------------------------------------------
 	bool FWorkingMemoryManager::EnsureBudgetBelow( uint64 AdditionalMemory )
     {
@@ -1039,29 +1085,14 @@ namespace mu
 		if (BudgetBytes == 0)
 		{
 			return true;
-		}
+		}		
 
-		uint32 RomBytes = 0;
-		uint32 PoolBytes = 0;
-		uint32 TempBytes = 0;
-		uint32 TrackedCacheBytes = 0;
-		{
-			MUTABLE_CPUPROFILER_SCOPE(EnsureBudgetBelow_Measure);
-			RomBytes = TrackedBudgetBytes_Rom; // GetRomBytes();
-			PoolBytes = TrackedBudgetBytes_Pooled; // GetPooledBytes();
-			TempBytes = TrackedBudgetBytes_Temp + GetTempMeshesBytes(); // GetTempBytes();
-			TrackedCacheBytes = TrackedBudgetBytes_Cached; // GetTrackedCacheBytes();
-		}
-
-		uint64 TotalBytes = 0;
-		TotalBytes += TempBytes;
-		TotalBytes += RomBytes;
-		TotalBytes += PoolBytes;
-		TotalBytes += TrackedCacheBytes;
+		int64 TotalBytes = GetCurrentMemoryBytes();
 
 		// Add the extra memory that we are trying to allocate when we return.
 		TotalBytes += AdditionalMemory;
 
+
         bool bFinished = TotalBytes <= BudgetBytes;
 
 		// Try to free pooled resources first
@@ -1399,16 +1430,9 @@ namespace mu
 
 		}
 
-		TRACE_COUNTER_SET(MutableRuntime_MemTemp, TempBytes);
-		TRACE_COUNTER_SET(MutableRuntime_MemPool, PoolBytes);
-		TRACE_COUNTER_SET(MutableRuntime_MemCache, TrackedCacheBytes);
-		TRACE_COUNTER_SET(MutableRuntime_MemRom, RomBytes);
-		TRACE_COUNTER_SET(MutableRuntime_MemTotal, TotalBytes);
-		TRACE_COUNTER_SET(MutableRuntime_MemBudget, BudgetBytes);
-
 		if (!bFinished)
 		{
-			uint64 ExcessBytes = TotalBytes - BudgetBytes;
+			int64 ExcessBytes = TotalBytes - BudgetBytes;
 
 			if (ExcessBytes > BudgetExcessBytes)
 			{
diff --git a/Engine/Plugins/Experimental/Mutable/Source/MutableRuntime/Private/MuR/SystemPrivate.h b/Engine/Plugins/Experimental/Mutable/Source/MutableRuntime/Private/MuR/SystemPrivate.h
--- a/Engine/Plugins/Experimental/Mutable/Source/MutableRuntime/Private/MuR/SystemPrivate.h
+++ b/Engine/Plugins/Experimental/Mutable/Source/MutableRuntime/Private/MuR/SystemPrivate.h
@@ -1107,16 +1107,17 @@ namespace mu
         };
 
 		/** Maximum working memory that mutable should be using. */
-		uint64 BudgetBytes = 0;
+		int64 BudgetBytes = 0;
 
 		/** Maximum excess memory reached suring the current operation. */
-		uint64 BudgetExcessBytes = 0;
+		int64 BudgetExcessBytes = 0;
 
-		/** . */
-		uint32 TrackedBudgetBytes_Rom = 0;
-		uint32 TrackedBudgetBytes_Temp = 0;
-		uint32 TrackedBudgetBytes_Pooled = 0;
-		uint32 TrackedBudgetBytes_Cached = 0;
+		/** Signed ints to account for imprecisions especially with temp.*/
+		int32 TrackedBudgetBytes_Rom = 0;
+		// int32 TrackedBudgetBytes_Temp = 0;
+		int32 TrackedBudgetBytes_Pooled = 0;
+		int32 TrackedBudgetBytes_Cached = 0;
+		int32 TrackedBudgetBytes_Stream = 0;
 
 		/** This value is used to track the order of loading of roms. */
         uint64 RomTick = 0;
@@ -1148,6 +1149,12 @@ namespace mu
 		* Return true if it succeeded, false otherwise.
 		*/
 		bool EnsureBudgetBelow(uint64 AdditionalMemory);
+		
+		/** Return true if the memory budget is 90% full. */
+		bool IsMemoryBudgetFull() const;
+
+		/** Calculate the current usage of memory as used to calculate the budget. */
+		int64 GetCurrentMemoryBytes() const;
 
         /** Register that a specific rom has been requested and update the heuristics to keep it in memory. */
         void MarkRomUsed( int32 RomIndex, const TSharedPtr<const Model>& );
@@ -1194,7 +1201,7 @@ namespace mu
 			Ptr<Image> Result = new Image(SizeX, SizeY, Lods, Format, Init);
 
 			TempImages.Add(Result);
-			TrackedBudgetBytes_Temp += Result->GetDataSize();
+			//TrackedBudgetBytes_Temp += Result->GetDataSize();
 			return Result;
 		}
 
@@ -1206,7 +1213,7 @@ namespace mu
 			int32 Removed = TempImages.RemoveSingle(Resource);
 			if (Removed)
 			{
-				TrackedBudgetBytes_Temp -= Resource->GetDataSize();
+				//TrackedBudgetBytes_Temp -= Resource->GetDataSize();
 			}
 			
 			check(!TempImages.Contains(Resource));
@@ -1246,7 +1253,7 @@ namespace mu
 			int32 Removed = TempImages.RemoveSingle(Resource);
 			if (Removed)
 			{
-				TrackedBudgetBytes_Temp -= ResourceDataSize;
+				//TrackedBudgetBytes_Temp -= ResourceDataSize;
 			}
 
 			check(!TempImages.Contains(Resource));
@@ -1285,7 +1292,7 @@ namespace mu
 			int32 Removed = TempImages.RemoveSingle(Resource);
 			if (Removed)
 			{
-				TrackedBudgetBytes_Temp -= ResourceDataSize;
+				//TrackedBudgetBytes_Temp -= ResourceDataSize;
 			}
 
 			check(!TempImages.Contains(Resource));
@@ -1444,7 +1451,7 @@ namespace mu
 			if (!bTakeOwnership && Result->IsUnique())
 			{
 				TempImages.Add(Result);
-				TrackedBudgetBytes_Temp += ResultDataSize;
+				//TrackedBudgetBytes_Temp += ResultDataSize;
 			}
 
 			return Result;
@@ -1459,7 +1466,7 @@ namespace mu
 				int32 Removed = TempImages.RemoveSingle(Resource);
 				if (Removed)
 				{
-					TrackedBudgetBytes_Temp -= ResourceDataSize;
+					//TrackedBudgetBytes_Temp -= ResourceDataSize;
 				}
 
 				check(!TempImages.Contains(Resource));
diff --git a/Engine/Plugins/Experimental/Mutable/Source/MutableRuntime/Public/MuR/Model.h b/Engine/Plugins/Experimental/Mutable/Source/MutableRuntime/Public/MuR/Model.h
--- a/Engine/Plugins/Experimental/Mutable/Source/MutableRuntime/Public/MuR/Model.h
+++ b/Engine/Plugins/Experimental/Mutable/Source/MutableRuntime/Public/MuR/Model.h
@@ -13,7 +13,7 @@
 //! data structures. Compiled models are not necessarily compatible when the runtime is updated,
 //! so this version number can be used externally to verify this. It is not used internally, and
 //! serializing models from different versions than this runtime will probably result in a crash.
-#define MUTABLE_COMPILED_MODEL_CODE_VERSION		uint32( 69 )
+#define MUTABLE_COMPILED_MODEL_CODE_VERSION		uint32( 70 )
 #define MUTABLE_PARAMETERS_VERSION              uint32( 2 )
 
 
diff --git a/Engine/Plugins/Experimental/Mutable/Source/MutableTools/Private/MuT/AST.cpp b/Engine/Plugins/Experimental/Mutable/Source/MutableTools/Private/MuT/AST.cpp
--- a/Engine/Plugins/Experimental/Mutable/Source/MutableTools/Private/MuT/AST.cpp
+++ b/Engine/Plugins/Experimental/Mutable/Source/MutableTools/Private/MuT/AST.cpp
@@ -1136,10 +1136,6 @@ FImageDesc ASTOpFixed::GetImageDesc( bool returnBestOption, FGetImageDescContext
         res = GetImageDesc(op.args.ImageInterpolate.targets[0], returnBestOption, context );
         break;
 
-    case OP_TYPE::IM_INTERPOLATE3:
-        res = GetImageDesc( op.args.ImageInterpolate3.target0, returnBestOption, context );
-        break;
-
     case OP_TYPE::IM_PLAINCOLOUR:
         res.m_format = EImageFormat(op.args.ImagePlainColour.format);
         res.m_size[0] = op.args.ImagePlainColour.size[0];
@@ -1516,22 +1512,6 @@ bool ASTOpFixed::IsImagePlainConstant(FVector4f& colour) const
         res = children[op.args.ImagePlainColour.colour]->IsColourConstant( colour );
         break;
 
-    case OP_TYPE::IM_INTERPOLATE3:
-        res = children[op.args.ImageInterpolate3.target0]->IsColourConstant( colour );
-        if (res)
-        {
-			FVector4f baseColour;
-            res = children[op.args.ImageInterpolate3.target1]->IsColourConstant( colour );
-            res &= (colour==baseColour);
-        }
-        if (res)
-        {
-			FVector4f baseColour;
-            res = children[op.args.ImageInterpolate3.target2]->IsColourConstant( colour );
-            res &= (colour==baseColour);
-        }
-        break;
-
     default:
         // TODO: Improve this test with more operations
         //check( false );
@@ -1642,13 +1622,6 @@ mu::Ptr<ImageSizeExpression> ASTOpFixed::GetImageSizeExpression() const
         }
         break;
 
-    case OP_TYPE::IM_INTERPOLATE3:
-        if ( children[op.args.ImageInterpolate3.target0] )
-        {
-            pRes = children[op.args.ImageInterpolate3.target0].child()->GetImageSizeExpression();
-        }
-        break;
-
     case OP_TYPE::IM_SATURATE:
         if ( children[op.args.ImageSaturate.base] )
         {
diff --git a/Engine/Plugins/Experimental/Mutable/Source/MutableTools/Private/MuT/ASTOpImagePixelFormat.cpp b/Engine/Plugins/Experimental/Mutable/Source/MutableTools/Private/MuT/ASTOpImagePixelFormat.cpp
--- a/Engine/Plugins/Experimental/Mutable/Source/MutableTools/Private/MuT/ASTOpImagePixelFormat.cpp
+++ b/Engine/Plugins/Experimental/Mutable/Source/MutableTools/Private/MuT/ASTOpImagePixelFormat.cpp
@@ -477,29 +477,6 @@ namespace mu
 			break;
 		}
 
-
-		case OP_TYPE::IM_INTERPOLATE3:
-		{
-			// This op doesn't support compressed formats
-			if (!bIsCompressedFormat)
-			{
-				// We move the format down all the paths
-				auto newOp = mu::Clone<ASTOpFixed>(at);
-
-				auto top0 = newOp->children[newOp->op.args.ImageInterpolate3.target0].child();
-				newOp->SetChild(newOp->op.args.ImageInterpolate3.target0, Visit(top0, currentFormatOp));
-
-				auto top1 = newOp->children[newOp->op.args.ImageInterpolate3.target1].child();
-				newOp->SetChild(newOp->op.args.ImageInterpolate3.target1, Visit(top1, currentFormatOp));
-
-				auto top2 = newOp->children[newOp->op.args.ImageInterpolate3.target2].child();
-				newOp->SetChild(newOp->op.args.ImageInterpolate3.target2, Visit(top2, currentFormatOp));
-
-				newAt = newOp;
-			}
-			break;
-		}
-
 		case OP_TYPE::IM_LAYER:
 		{
 			if (GetOpDesc(at->GetOpType()).supportedBasePixelFormats[(size_t)format])
diff --git a/Engine/Plugins/Experimental/Mutable/Source/MutableTools/Private/MuT/CodeGenerator.h b/Engine/Plugins/Experimental/Mutable/Source/MutableTools/Private/MuT/CodeGenerator.h
--- a/Engine/Plugins/Experimental/Mutable/Source/MutableTools/Private/MuT/CodeGenerator.h
+++ b/Engine/Plugins/Experimental/Mutable/Source/MutableTools/Private/MuT/CodeGenerator.h
@@ -62,7 +62,6 @@ namespace mu
 	class NodeImageReference;
 	class NodeImageFormat;
 	class NodeImageGradient;
-	class NodeImageInterpolate3;
 	class NodeImageInterpolate;
 	class NodeImageInvert;
 	class NodeImageLayer;
@@ -427,7 +426,6 @@ namespace mu
 		void GenerateImage_LayerColour(const FImageGenerationOptions&, FImageGenerationResult&, const NodeImageLayerColour*);
 		void GenerateImage_Resize(const FImageGenerationOptions&, FImageGenerationResult&, const NodeImageResize*);
 		void GenerateImage_PlainColour(const FImageGenerationOptions&, FImageGenerationResult&, const NodeImagePlainColour*);
-		void GenerateImage_Interpolate3(const FImageGenerationOptions&, FImageGenerationResult&, const NodeImageInterpolate3*);
 		void GenerateImage_Project(const FImageGenerationOptions&, FImageGenerationResult&, const NodeImageProject*);
 		void GenerateImage_Mipmap(const FImageGenerationOptions&, FImageGenerationResult&, const NodeImageMipmap*);
 		void GenerateImage_Switch(const FImageGenerationOptions&, FImageGenerationResult&, const NodeImageSwitch*);
diff --git a/Engine/Plugins/Experimental/Mutable/Source/MutableTools/Private/MuT/CodeGenerator_Image.cpp b/Engine/Plugins/Experimental/Mutable/Source/MutableTools/Private/MuT/CodeGenerator_Image.cpp
--- a/Engine/Plugins/Experimental/Mutable/Source/MutableTools/Private/MuT/CodeGenerator_Image.cpp
+++ b/Engine/Plugins/Experimental/Mutable/Source/MutableTools/Private/MuT/CodeGenerator_Image.cpp
@@ -62,8 +62,6 @@
 #include "MuT/NodeImageGradient.h"
 #include "MuT/NodeImageGradientPrivate.h"
 #include "MuT/NodeImageInterpolate.h"
-#include "MuT/NodeImageInterpolate3.h"
-#include "MuT/NodeImageInterpolate3Private.h"
 #include "MuT/NodeImageInterpolatePrivate.h"
 #include "MuT/NodeImageInvert.h"
 #include "MuT/NodeImageInvertPrivate.h"
@@ -168,7 +166,6 @@ namespace mu
 			case NodeImage::EType::LayerColour: GenerateImage_LayerColour(Options, result, static_cast<const NodeImageLayerColour*>(Node)); break;
 			case NodeImage::EType::Resize: GenerateImage_Resize(Options, result, static_cast<const NodeImageResize*>(Node)); break;
 			case NodeImage::EType::PlainColour: GenerateImage_PlainColour(Options, result, static_cast<const NodeImagePlainColour*>(Node)); break;
-			case NodeImage::EType::Interpolate3: GenerateImage_Interpolate3(Options, result, static_cast<const NodeImageInterpolate3*>(Node)); break;
 			case NodeImage::EType::Project: GenerateImage_Project(Options, result, static_cast<const NodeImageProject*>(Node)); break;
 			case NodeImage::EType::Mipmap: GenerateImage_Mipmap(Options, result, static_cast<const NodeImageMipmap*>(Node)); break;
 			case NodeImage::EType::Switch: GenerateImage_Switch(Options, result, static_cast<const NodeImageSwitch*>(Node)); break;
@@ -693,98 +690,6 @@ namespace mu
     }
 
 
-    //---------------------------------------------------------------------------------------------
-	void CodeGenerator::GenerateImage_Interpolate3(const FImageGenerationOptions& Options, FImageGenerationResult& result, const NodeImageInterpolate3* InNode)
-	{
-		const NodeImageInterpolate3::Private& node = *InNode->GetPrivate();
-
-		MUTABLE_CPUPROFILER_SCOPE(NodeImageInterpolate3);
-
-        Ptr<ASTOpFixed> op = new ASTOpFixed();
-        op->op.type = OP_TYPE::IM_INTERPOLATE3;
-
-        // Factors
-        if ( node.m_pFactor1.get() )
-        {
-            op->SetChild( op->op.args.ImageInterpolate3.factor1, Generate( node.m_pFactor1 ));
-        }
-        else
-        {
-            // This argument is required
-            op->SetChild( op->op.args.ImageInterpolate3.factor1,
-                          GenerateMissingScalarCode(TEXT("Interpolation factor 1"), 0.3f, node.m_errorContext ));
-        }
-
-        if ( node.m_pFactor2.get() )
-        {
-            op->SetChild( op->op.args.ImageInterpolate3.factor2, Generate( node.m_pFactor2 ));
-        }
-        else
-        {
-            // This argument is required
-            op->SetChild( op->op.args.ImageInterpolate3.factor2,
-                          GenerateMissingScalarCode(TEXT("Interpolation factor 2"), 0.3f, node.m_errorContext ));
-        }
-
-        // Target 0
-        Ptr<ASTOp> target;
-        if ( node.m_pTarget0 )
-        {
-			FImageGenerationResult BaseResult;
-			GenerateImage(Options, BaseResult, node.m_pTarget0);
-			target = BaseResult.op;
-        }
-        else
-        {
-            // This argument is required
-            target = GenerateMissingImageCode(TEXT("Interpolation target 0"), EImageFormat::IF_RGB_UBYTE, node.m_errorContext, Options);
-        }
-        target = GenerateImageFormat( target, EImageFormat::IF_RGB_UBYTE );
-        target = GenerateImageSize
-                ( target, FImageSize((uint16)m_imageState.Last().m_imageRect.size[0],(uint16)m_imageState.Last().m_imageRect.size[1]) );
-        op->SetChild( op->op.args.ImageInterpolate3.target0, target);
-
-        // Target 1
-        if ( node.m_pTarget1 )
-        {
-			FImageGenerationResult BaseResult;
-			GenerateImage(Options, BaseResult, node.m_pTarget1);
-			target = BaseResult.op;
-		}
-        else
-        {
-            // This argument is required
-            target = GenerateMissingImageCode(TEXT("Interpolation target 1"), EImageFormat::IF_RGB_UBYTE, node.m_errorContext, Options);
-        }
-        target = GenerateImageFormat( target, EImageFormat::IF_RGB_UBYTE );
-        target = GenerateImageSize
-                ( target, FImageSize( (uint16)m_imageState.Last().m_imageRect.size[0],
-                                      (uint16)m_imageState.Last().m_imageRect.size[1]) );
-        op->SetChild( op->op.args.ImageInterpolate3.target1, target);
-
-        // Target 2
-        if ( node.m_pTarget2 )
-        {
-			FImageGenerationResult BaseResult;
-			GenerateImage(Options, BaseResult, node.m_pTarget2);
-			target = BaseResult.op;
-		}
-        else
-        {
-            // This argument is required
-            target = GenerateMissingImageCode(TEXT("Interpolation target 2"), EImageFormat::IF_RGB_UBYTE, node.m_errorContext, Options);
-        }
-        target = GenerateImageFormat( target, EImageFormat::IF_RGB_UBYTE );
-        target = GenerateImageSize
-                ( target, FImageSize( (uint16)m_imageState.Last().m_imageRect.size[0],
-                                      (uint16)m_imageState.Last().m_imageRect.size[1] ) );
-        op->SetChild( op->op.args.ImageInterpolate3.target2, target);
-
-
-        result.op = op;
-    }
-
-
 	//---------------------------------------------------------------------------------------------
 	void CodeGenerator::GenerateImage_Swizzle(const FImageGenerationOptions& Options, FImageGenerationResult& result, const NodeImageSwizzle* InNode)
 	{
diff --git a/Engine/Plugins/Experimental/Mutable/Source/MutableTools/Private/MuT/CodeOptimiser_Semantic.cpp b/Engine/Plugins/Experimental/Mutable/Source/MutableTools/Private/MuT/CodeOptimiser_Semantic.cpp
--- a/Engine/Plugins/Experimental/Mutable/Source/MutableTools/Private/MuT/CodeOptimiser_Semantic.cpp
+++ b/Engine/Plugins/Experimental/Mutable/Source/MutableTools/Private/MuT/CodeOptimiser_Semantic.cpp
@@ -655,31 +655,9 @@ namespace mu
             break;
         }
 
-
-        case OP_TYPE::IM_INTERPOLATE3:
-        {
-            // We move the op down all the paths
-            auto newOp = mu::Clone<ASTOpFixed>(at);
-
-            auto top0 = newOp->children[newOp->op.args.ImageInterpolate3.target0].child();
-            newOp->SetChild(newOp->op.args.ImageInterpolate3.target0,
-                            Visit(top0, currentCropOp));
-
-            auto top1 = newOp->children[newOp->op.args.ImageInterpolate3.target1].child();
-            newOp->SetChild(newOp->op.args.ImageInterpolate3.target1,
-                            Visit(top1, currentCropOp));
-
-            auto top2 = newOp->children[newOp->op.args.ImageInterpolate3.target2].child();
-            newOp->SetChild(newOp->op.args.ImageInterpolate3.target2,
-                            Visit(top2, currentCropOp));
-
-            newAt = newOp;
-            break;
-        }
-
-        default:
-            break;
-        }
+		default:
+			break;
+		}
 
         // end on line, replace with crop
         if (at==newAt && at!=m_initialSource && !skipSinking)
@@ -1144,24 +1122,6 @@ namespace mu
 				break;
 			}
 
-
-			case OP_TYPE::IM_INTERPOLATE3:
-			{
-				Ptr<ASTOpFixed> newOp = mu::Clone<ASTOpFixed>(at);
-
-				Ptr<ASTOp> top0 = newOp->children[newOp->op.args.ImageInterpolate3.target0].child();
-				newOp->SetChild(newOp->op.args.ImageInterpolate3.target0, Visit(top0, currentSinkingOp));
-
-				Ptr<ASTOp> top1 = newOp->children[newOp->op.args.ImageInterpolate3.target1].child();
-				newOp->SetChild(newOp->op.args.ImageInterpolate3.target1, Visit(top1, currentSinkingOp));
-
-				Ptr<ASTOp> top2 = newOp->children[newOp->op.args.ImageInterpolate3.target2].child();
-				newOp->SetChild(newOp->op.args.ImageInterpolate3.target2, Visit(top2, currentSinkingOp));
-
-				newAt = newOp;
-				break;
-			}
-
 			case OP_TYPE::IM_MULTILAYER:
 			{
 				Ptr<ASTOpImageMultiLayer> nop = mu::Clone<ASTOpImageMultiLayer>(at);
@@ -1477,35 +1437,6 @@ namespace mu
                 break;
             }
 
-            case OP_TYPE::IM_INTERPOLATE3:
-            {
-                Ptr<ASTOpFixed> newOp = mu::Clone<ASTOpFixed>(sourceAt);
-
-                {
-                    Ptr<ASTOpFixed> sourceOp = mu::Clone<ASTOpFixed>(this);
-                    auto targetAt = newOp->children[ newOp->op.args.ImageInterpolate3.target0 ].child();
-                    sourceOp->SetChild( sourceOp->op.args.ImageResize.source, targetAt);
-                    newOp->SetChild( newOp->op.args.ImageInterpolate3.target0, sourceOp );
-                }
-
-                {
-                    Ptr<ASTOpFixed> sourceOp = mu::Clone<ASTOpFixed>(this);
-                    auto targetAt = newOp->children[ newOp->op.args.ImageInterpolate3.target1 ].child();
-                    sourceOp->SetChild( sourceOp->op.args.ImageResize.source, targetAt);
-                    newOp->SetChild( newOp->op.args.ImageInterpolate3.target1, sourceOp );
-                }
-
-                {
-                    Ptr<ASTOpFixed> sourceOp = mu::Clone<ASTOpFixed>(this);
-                    auto targetAt = newOp->children[ newOp->op.args.ImageInterpolate3.target2 ].child();
-                    sourceOp->SetChild( sourceOp->op.args.ImageResize.source, targetAt);
-                    newOp->SetChild( newOp->op.args.ImageInterpolate3.target2, sourceOp );
-                }
-
-                at = newOp;
-                break;
-            }
-
             case OP_TYPE::IM_LAYER:
             {
                 Ptr<ASTOpImageLayer> newOp = mu::Clone<ASTOpImageLayer>(sourceAt);
diff --git a/Engine/Plugins/Experimental/Mutable/Source/MutableTools/Private/MuT/ErrorLog.cpp b/Engine/Plugins/Experimental/Mutable/Source/MutableTools/Private/MuT/ErrorLog.cpp
--- a/Engine/Plugins/Experimental/Mutable/Source/MutableTools/Private/MuT/ErrorLog.cpp
+++ b/Engine/Plugins/Experimental/Mutable/Source/MutableTools/Private/MuT/ErrorLog.cpp
@@ -249,7 +249,6 @@ namespace mu
 		TEXT("IM_BLANKLAYOUT   "),
 		TEXT("IM_COMPOSE       "),
 		TEXT("IM_INTERPOLATE   "),
-		TEXT("IM_INTERPOLATE3  "),
 		TEXT("IM_SATURATE      "),
 		TEXT("IM_LUMINANCE     "),
 		TEXT("IM_SWIZZLE       "),
diff --git a/Engine/Plugins/Experimental/Mutable/Source/MutableTools/Private/MuT/ImageDescGenerator.cpp b/Engine/Plugins/Experimental/Mutable/Source/MutableTools/Private/MuT/ImageDescGenerator.cpp
--- a/Engine/Plugins/Experimental/Mutable/Source/MutableTools/Private/MuT/ImageDescGenerator.cpp
+++ b/Engine/Plugins/Experimental/Mutable/Source/MutableTools/Private/MuT/ImageDescGenerator.cpp
@@ -18,7 +18,6 @@
 #include "MuT/NodeImageConstantPrivate.h"
 #include "MuT/NodeImageFormatPrivate.h"
 #include "MuT/NodeImageGradientPrivate.h"
-#include "MuT/NodeImageInterpolate3Private.h"
 #include "MuT/NodeImageInterpolatePrivate.h"
 #include "MuT/NodeImageInvertPrivate.h"
 #include "MuT/NodeImageLayerColourPrivate.h"
@@ -256,30 +255,6 @@ namespace mu
     }
 
 
-    //---------------------------------------------------------------------------------------------
-    Ptr<ASTOp> ImageDescGenerator::Visit( const NodeImageInterpolate3::Private& node )
-    {
-        m_desc = MUTABLE_MISSING_IMAGE_DESC;
-
-        if ( node.m_pTarget0 )
-        {
-            Generate( *node.m_pTarget0->GetBasePrivate() );
-        }
-        else if ( node.m_pTarget1 )
-        {
-            Generate( *node.m_pTarget1->GetBasePrivate() );
-        }
-        else if ( node.m_pTarget2 )
-        {
-            Generate( *node.m_pTarget2->GetBasePrivate() );
-        }
-
-        m_desc.m_format = EImageFormat::IF_RGB_UBYTE;
-
-        return 0;
-    }
-
-
 	//---------------------------------------------------------------------------------------------
     Ptr<ASTOp> ImageDescGenerator::Visit(const NodeImageSwizzle::Private& node)
 	{
diff --git a/Engine/Plugins/Experimental/Mutable/Source/MutableTools/Private/MuT/ImageDescGenerator.h b/Engine/Plugins/Experimental/Mutable/Source/MutableTools/Private/MuT/ImageDescGenerator.h
--- a/Engine/Plugins/Experimental/Mutable/Source/MutableTools/Private/MuT/ImageDescGenerator.h
+++ b/Engine/Plugins/Experimental/Mutable/Source/MutableTools/Private/MuT/ImageDescGenerator.h
@@ -13,7 +13,6 @@
 #include "MuT/NodeImageFormat.h"
 #include "MuT/NodeImageGradient.h"
 #include "MuT/NodeImageInterpolate.h"
-#include "MuT/NodeImageInterpolate3.h"
 #include "MuT/NodeImageInvert.h"
 #include "MuT/NodeImageLayer.h"
 #include "MuT/NodeImageLayerColour.h"
@@ -52,7 +51,6 @@ namespace mu
                                // Standard operations from models
                                public Visitor<NodeImageConstant::Private, Ptr<ASTOp>, true>,
                                public Visitor<NodeImageInterpolate::Private, Ptr<ASTOp>, true>,
-                               public Visitor<NodeImageInterpolate3::Private, Ptr<ASTOp>, true>,
                                public Visitor<NodeImageSaturate::Private, Ptr<ASTOp>, true>,
                                public Visitor<NodeImageLuminance::Private, Ptr<ASTOp>, true>,
                                public Visitor<NodeImageLayer::Private, Ptr<ASTOp>, true>,
@@ -84,7 +82,6 @@ namespace mu
 
         Ptr<ASTOp> Visit( const NodeImageConstant::Private& );
         Ptr<ASTOp> Visit( const NodeImageInterpolate::Private& );
-        Ptr<ASTOp> Visit( const NodeImageInterpolate3::Private& );
         Ptr<ASTOp> Visit( const NodeImageSaturate::Private& );
         Ptr<ASTOp> Visit( const NodeImageLuminance::Private& );
         Ptr<ASTOp> Visit( const NodeImageLayer::Private& );
diff --git a/Engine/Plugins/Experimental/Mutable/Source/MutableTools/Private/MuT/NodeImage.cpp b/Engine/Plugins/Experimental/Mutable/Source/MutableTools/Private/MuT/NodeImage.cpp
--- a/Engine/Plugins/Experimental/Mutable/Source/MutableTools/Private/MuT/NodeImage.cpp
+++ b/Engine/Plugins/Experimental/Mutable/Source/MutableTools/Private/MuT/NodeImage.cpp
@@ -13,7 +13,6 @@
 #include "MuT/NodeImageFormat.h"
 #include "MuT/NodeImageGradient.h"
 #include "MuT/NodeImageInterpolate.h"
-#include "MuT/NodeImageInterpolate3.h"
 #include "MuT/NodeImageInvert.h"
 #include "MuT/NodeImageLayer.h"
 #include "MuT/NodeImageLayerColour.h"
@@ -84,7 +83,7 @@ namespace mu
 		case 20 :  return NodeImageLayerColour::StaticUnserialise( arch ); break;
 		case 21 :  return NodeImageResize::StaticUnserialise( arch ); break;
 		case 22 :  return NodeImagePlainColour::StaticUnserialise( arch ); break;
-		case 23 :  return NodeImageInterpolate3::StaticUnserialise( arch ); break;
+		//case 23 :  return NodeImageInterpolate3::StaticUnserialise( arch ); break;
         case 24 :  return NodeImageProject::StaticUnserialise( arch ); break;
         case 25 :  return NodeImageMipmap::StaticUnserialise( arch ); break;
         case 26 :  return NodeImageSwitch::StaticUnserialise( arch ); break;
diff --git a/Engine/Plugins/Experimental/Mutable/Source/MutableTools/Private/MuT/NodeImageInterpolate3.cpp b/Engine/Plugins/Experimental/Mutable/Source/MutableTools/Private/MuT/NodeImageInterpolate3.cpp
deleted file mode 100644
diff --git a/Engine/Plugins/Experimental/Mutable/Source/MutableTools/Private/MuT/NodeImageInterpolate3Private.h b/Engine/Plugins/Experimental/Mutable/Source/MutableTools/Private/MuT/NodeImageInterpolate3Private.h
deleted file mode 100644
diff --git a/Engine/Plugins/Experimental/Mutable/Source/MutableTools/Public/MuT/NodeImageInterpolate3.h b/Engine/Plugins/Experimental/Mutable/Source/MutableTools/Public/MuT/NodeImageInterpolate3.h
deleted file mode 100644
