[mutable] Fix editor compilation broken in previous submit.
#rnx

diff --git a/Engine/Plugins/Experimental/Mutable/Source/CustomizableObjectEditor/Private/MuCOE/SMutableConstantsWidget.cpp b/Engine/Plugins/Experimental/Mutable/Source/CustomizableObjectEditor/Private/MuCOE/SMutableConstantsWidget.cpp
--- a/Engine/Plugins/Experimental/Mutable/Source/CustomizableObjectEditor/Private/MuCOE/SMutableConstantsWidget.cpp
+++ b/Engine/Plugins/Experimental/Mutable/Source/CustomizableObjectEditor/Private/MuCOE/SMutableConstantsWidget.cpp
@@ -1150,7 +1150,10 @@ void SMutableConstantsWidget::LoadConstantImages()
 	for (int32 ImageIndex = 0; ImageIndex < ConstantsCount; ImageIndex++)
 	{
 		TSharedPtr<FMutableConstantImageElement> ConstantImageElement =MakeShared<FMutableConstantImageElement>();
-		MutableProgramPtr->GetConstant(ImageIndex, ConstantImageElement->ImagePtr, 0);
+
+		MutableProgramPtr->GetConstant(ImageIndex, ConstantImageElement->ImagePtr, 0,
+			[this](int32 x, int32 y, int32 m, mu::EImageFormat f) { return new mu::Image(x, y, m, f); });
+
 		ConstantImageElement->IndexOnSourceVector = ImageIndex;
 		
 		ConstantImagesAccumulatedSize += ConstantImageElement->ImagePtr->GetDataSize();
