[mutable] Fix "Loading Assets" message being always displayed regardles of compilation type.

diff --git a/Engine/Plugins/Experimental/Mutable/Source/CustomizableObjectEditor/Private/MuCOE/CustomizableObjectCompiler.cpp b/Engine/Plugins/Experimental/Mutable/Source/CustomizableObjectEditor/Private/MuCOE/CustomizableObjectCompiler.cpp
--- a/Engine/Plugins/Experimental/Mutable/Source/CustomizableObjectEditor/Private/MuCOE/CustomizableObjectCompiler.cpp
+++ b/Engine/Plugins/Experimental/Mutable/Source/CustomizableObjectEditor/Private/MuCOE/CustomizableObjectCompiler.cpp
@@ -731,9 +731,9 @@ void FCustomizableObjectCompiler::AddCachedReferencers(const FName& PathName, TA
 }
 
 
-void FCustomizableObjectCompiler::LaunchMutableCompile(bool ShowNotification)
+void FCustomizableObjectCompiler::LaunchMutableCompile(bool bShowNotification)
 {
-	if(!Options.bSilentCompilation && ShowNotification)
+	if(bShowNotification)
 	{
 		AddCompileNotification(LOCTEXT("CustomizableObjectCompileInProgress", "Compiling"));
 	}
@@ -746,14 +746,14 @@ void FCustomizableObjectCompiler::LaunchMutableCompile(bool ShowNotification)
 }
 
 
-void FCustomizableObjectCompiler::SaveCODerivedData(bool ShowNotification)
+void FCustomizableObjectCompiler::SaveCODerivedData(bool bShowNotification)
 {
 	if (!SaveDDTask.IsValid())
 	{
 		return;
 	}
 
-	if (!Options.bSilentCompilation && ShowNotification)
+	if (bShowNotification)
 	{
 		AddCompileNotification(LOCTEXT("SavingCustomizableObjectDerivedData", "Saving Data"));
 	}
@@ -1333,7 +1333,7 @@ void FCustomizableObjectCompiler::CompileInternal(UCustomizableObject* Object, c
 				// will do the remaining steps (convert textures asynchronously and launch the Mutable compile thread)
 				CompilationLaunchPending = true;
 
-				if (!Options.bSilentCompilation && PendingTexturesToLoad)
+				if (PendingTexturesToLoad)
 				{
 					AddCompileNotification(LOCTEXT("ConvertingToMutableTexture", "Converting textures"));
 				}
@@ -1522,6 +1522,11 @@ void FCustomizableObjectCompiler::ForceFinishBeforeStartCompilation(UCustomizabl
 
 void FCustomizableObjectCompiler::AddCompileNotification(const FText& CompilationStep) const
 {
+	if (Options.bSilentCompilation)
+	{
+		return;
+	}
+
 	const FText Text = CurrentObject ? FText::FromString(FString::Printf(TEXT("Compiling %s"), *CurrentObject->GetName())) : LOCTEXT("CustomizableObjectCompileInProgressNotification", "Compiling Customizable Object");
 	
 	FCustomizableObjectEditorLogger::CreateLog(Text)
diff --git a/Engine/Plugins/Experimental/Mutable/Source/CustomizableObjectEditor/Private/MuCOE/CustomizableObjectCompiler.h b/Engine/Plugins/Experimental/Mutable/Source/CustomizableObjectEditor/Private/MuCOE/CustomizableObjectCompiler.h
--- a/Engine/Plugins/Experimental/Mutable/Source/CustomizableObjectEditor/Private/MuCOE/CustomizableObjectCompiler.h
+++ b/Engine/Plugins/Experimental/Mutable/Source/CustomizableObjectEditor/Private/MuCOE/CustomizableObjectCompiler.h
@@ -165,15 +165,15 @@ private:
 	void AddCachedReferencers(const FName& PathName, TArray<FName>& ArrayReferenceNames);
 
 	/** Launches the compile task in another thread when compiling a CO in the editor
-	* param ShowNotification [in] whether to show the compiling CO notification or not
+	* @param bShowNotification [in] whether to show the compiling CO notification or not
 	* @return nothing */
-	void LaunchMutableCompile(bool ShowNotification);
+	void LaunchMutableCompile(bool bShowNotification);
 
 	/** Launches the save derived data task in another thread after compiling a CO in the
 	* editor
-	* param ShowNotification [in] whether to show the saving DD notification or not
+	* @param bShowNotification [in] whether to show the saving DD notification or not
 	* @return nothing */
-	void SaveCODerivedData(bool ShowNotification);
+	void SaveCODerivedData(bool bShowNotification);
 
 	/** When compiling a CO in the editor, flag to know when the Unreal textures have been converted to Mutable textures */
 	bool PendingTexturesToLoad;
