[mutable] Minor cleanup of streaming data save code.
#rnx

diff --git a/Engine/Plugins/Experimental/Mutable/Source/CustomizableObjectEditor/Private/MuCOE/CustomizableObjectCompileRunnable.cpp b/Engine/Plugins/Experimental/Mutable/Source/CustomizableObjectEditor/Private/MuCOE/CustomizableObjectCompileRunnable.cpp
--- a/Engine/Plugins/Experimental/Mutable/Source/CustomizableObjectEditor/Private/MuCOE/CustomizableObjectCompileRunnable.cpp
+++ b/Engine/Plugins/Experimental/Mutable/Source/CustomizableObjectEditor/Private/MuCOE/CustomizableObjectCompileRunnable.cpp
@@ -229,8 +229,8 @@ uint32 FCustomizableObjectSaveDDRunnable::Run()
 		if (bFilesDeleted)
 		{
 			// Create file writers...
-			FArchive* ModelMemoryWriter = FileManager.CreateFileWriter(*CompileDataFullFileName);
-			FArchive* StreamableMemoryWriter = FileManager.CreateFileWriter(*StreamableDataFullFileName);
+			TUniquePtr<FArchive> ModelMemoryWriter( FileManager.CreateFileWriter(*CompileDataFullFileName) );
+			TUniquePtr<FArchive> StreamableMemoryWriter( FileManager.CreateFileWriter(*StreamableDataFullFileName) );
 			check(ModelMemoryWriter);
 			check(StreamableMemoryWriter);
 
@@ -239,13 +239,13 @@ uint32 FCustomizableObjectSaveDDRunnable::Run()
 			*StreamableMemoryWriter << CustomizableObjectHeader;
 
 			// Serialize Customizable Object's Data to disk
-			ModelMemoryWriter->Serialize(reinterpret_cast<void*>(Bytes.GetData()), Bytes.Num() * sizeof(uint8));
+			ModelMemoryWriter->Serialize(Bytes.GetData(), Bytes.Num() * sizeof(uint8));
 			Bytes.Empty();
 
 			// Serialize mu::Model and streamable resources
 			*ModelMemoryWriter << bModelSerialized;
 
-			FUnrealMutableModelBulkStreamer Streamer(ModelMemoryWriter, StreamableMemoryWriter);
+			FUnrealMutableModelBulkStreamer Streamer(ModelMemoryWriter.Get(), StreamableMemoryWriter.Get());
 			mu::Model::Serialise(Model.Get(), Streamer);
 
 			// Save to disk
@@ -254,9 +254,11 @@ uint32 FCustomizableObjectSaveDDRunnable::Run()
 
 			ModelMemoryWriter->Close();
 			StreamableMemoryWriter->Close();
-
-			delete ModelMemoryWriter;
-			delete StreamableMemoryWriter;
+		}
+		else
+		{
+			// Remove old data if there.
+			Model.Reset();
 		}
 	}
 
