[Mutable] Add Wrap, ClampToEdge and ClampToBlack adressing modes to ImageTransform operations.

- Modified transform applied to image to be more usable with new address modes, now it is equivalent to a center rotation followed by scale followed by a translation.
- The translation is scaled to always be within the ranage -1..1 inside the texture bounds.

#rnx

diff --git a/Engine/Plugins/Experimental/Mutable/Source/CustomizableObject/Classes/MuCO/CustomizableObject.h b/Engine/Plugins/Experimental/Mutable/Source/CustomizableObject/Classes/MuCO/CustomizableObject.h
--- a/Engine/Plugins/Experimental/Mutable/Source/CustomizableObject/Classes/MuCO/CustomizableObject.h
+++ b/Engine/Plugins/Experimental/Mutable/Source/CustomizableObject/Classes/MuCO/CustomizableObject.h
@@ -1439,7 +1439,7 @@ private:
 	// This is a manual version number for the binary blobs in this asset.
 	// Increasing it invalidates all the previously compiled models.
 	// Warning: If while merging code both versions have changed, take the highest+1.
-	static const int32 CurrentSupportedVersion = 391;
+	static const int32 CurrentSupportedVersion = 392;
 
 public:
 
diff --git a/Engine/Plugins/Experimental/Mutable/Source/CustomizableObjectEditor/Classes/MuCOE/Nodes/CustomizableObjectNodeTextureTransform.h b/Engine/Plugins/Experimental/Mutable/Source/CustomizableObjectEditor/Classes/MuCOE/Nodes/CustomizableObjectNodeTextureTransform.h
--- a/Engine/Plugins/Experimental/Mutable/Source/CustomizableObjectEditor/Classes/MuCOE/Nodes/CustomizableObjectNodeTextureTransform.h
+++ b/Engine/Plugins/Experimental/Mutable/Source/CustomizableObjectEditor/Classes/MuCOE/Nodes/CustomizableObjectNodeTextureTransform.h
@@ -11,12 +11,23 @@ class FArchive;
 class UCustomizableObjectNodeRemapPins;
 class UObject;
 
+UENUM()
+enum class ETextureTransformAddressMode
+{
+	Wrap,
+	ClampToEdge,
+	ClampToBlack
+};
+
 UCLASS()
 class CUSTOMIZABLEOBJECTEDITOR_API UCustomizableObjectNodeTextureTransform : public UCustomizableObjectNode
 {
 public:
 	GENERATED_BODY()
 
+	UPROPERTY(EditAnywhere, Category = CustomizableObject)
+	ETextureTransformAddressMode AddressMode = ETextureTransformAddressMode::Wrap;
+
 	// Begin EdGraphNode interface
 	FText GetNodeTitle(ENodeTitleType::Type TittleType) const override;
 	FLinearColor GetNodeTitleColor() const override;
diff --git a/Engine/Plugins/Experimental/Mutable/Source/CustomizableObjectEditor/Private/MuCOE/GenerateMutableSource/GenerateMutableSourceImage.cpp b/Engine/Plugins/Experimental/Mutable/Source/CustomizableObjectEditor/Private/MuCOE/GenerateMutableSource/GenerateMutableSourceImage.cpp
--- a/Engine/Plugins/Experimental/Mutable/Source/CustomizableObjectEditor/Private/MuCOE/GenerateMutableSource/GenerateMutableSourceImage.cpp
+++ b/Engine/Plugins/Experimental/Mutable/Source/CustomizableObjectEditor/Private/MuCOE/GenerateMutableSource/GenerateMutableSourceImage.cpp
@@ -791,6 +791,17 @@ mu::NodeImagePtr GenerateMutableSourceImage(const UEdGraphPin* Pin, FMutableGrap
 			mu::NodeScalarPtr RotationNode = GenerateMutableSourceFloat(RotationPin, GenerationContext);
 			TransformNode->SetRotation( RotationNode ); 
 		}
+
+		TransformNode->SetAddressMode(Invoke([&]() 
+		{
+			switch (TypedNodeTransform->AddressMode)
+			{
+			case ETextureTransformAddressMode::Wrap:		 return mu::EAddressMode::Wrap;
+			case ETextureTransformAddressMode::ClampToEdge:  return mu::EAddressMode::ClampToEdge;
+			case ETextureTransformAddressMode::ClampToBlack: return mu::EAddressMode::ClampToBlack;
+			default: { check(false); return mu::EAddressMode::None; }
+			}
+		}));
 	}
 
 	else if (const UCustomizableObjectNodeTextureSaturate* TypedNodeSaturate = Cast<UCustomizableObjectNodeTextureSaturate>(Node))
diff --git a/Engine/Plugins/Experimental/Mutable/Source/CustomizableObjectEditor/Private/MuCOE/GenerateMutableSource/GenerateMutableSourceSurface.cpp b/Engine/Plugins/Experimental/Mutable/Source/CustomizableObjectEditor/Private/MuCOE/GenerateMutableSource/GenerateMutableSourceSurface.cpp
--- a/Engine/Plugins/Experimental/Mutable/Source/CustomizableObjectEditor/Private/MuCOE/GenerateMutableSource/GenerateMutableSourceSurface.cpp
+++ b/Engine/Plugins/Experimental/Mutable/Source/CustomizableObjectEditor/Private/MuCOE/GenerateMutableSource/GenerateMutableSourceSurface.cpp
@@ -688,7 +688,7 @@ mu::NodeSurfacePtr GenerateMutableSourceSurface(const UEdGraphPin * Pin, FMutabl
 
 							mu::NodeImageMipmapPtr MipmapImage = new mu::NodeImageMipmap();
 							MipmapImage->SetSource(LastImage.get());
-							MipmapImage->SetMipmapGenerationSettings(mu::EMipmapFilterType::MFT_SimpleAverage, mu::EAddressMode::AM_NONE, 1.0f, false);
+							MipmapImage->SetMipmapGenerationSettings(mu::EMipmapFilterType::MFT_SimpleAverage, mu::EAddressMode::None, 1.0f, false);
 
 							MipmapImage->SetMessageContext(Node);
 							LastImage = MipmapImage;
@@ -724,7 +724,7 @@ mu::NodeSurfacePtr GenerateMutableSourceSurface(const UEdGraphPin * Pin, FMutabl
 
 									mu::NodeImageMipmapPtr NormalCompositeMipmapImage = new mu::NodeImageMipmap();
 									NormalCompositeMipmapImage->SetSource(CompositeNormalImage);
-									NormalCompositeMipmapImage->SetMipmapGenerationSettings(mu::EMipmapFilterType::MFT_SimpleAverage, mu::EAddressMode::AM_NONE, 1.0f, true);
+									NormalCompositeMipmapImage->SetMipmapGenerationSettings(mu::EMipmapFilterType::MFT_SimpleAverage, mu::EAddressMode::None, 1.0f, true);
 
 									CompositedImage->SetNormal(NormalCompositeMipmapImage);
 								}
diff --git a/Engine/Plugins/Experimental/Mutable/Source/MutableRuntime/Private/MuR/CodeRunner.cpp b/Engine/Plugins/Experimental/Mutable/Source/MutableRuntime/Private/MuR/CodeRunner.cpp
--- a/Engine/Plugins/Experimental/Mutable/Source/MutableRuntime/Private/MuR/CodeRunner.cpp
+++ b/Engine/Plugins/Experimental/Mutable/Source/MutableRuntime/Private/MuR/CodeRunner.cpp
@@ -1101,7 +1101,6 @@ namespace mu
                 MeshDifference(Result.get(), pBase.get(), pTarget.get(),
                                NumChannels, Semantics.GetData(), SemanticIndices.GetData(),
                                bIgnoreTextureCoords != 0, bOutSuccess);
-			
 				Release(pBase);
 				Release(pTarget);
 
@@ -4729,14 +4728,14 @@ namespace mu
             case 0:
 			{
 				const TArray<FScheduledOp, TInlineAllocator<6>> Deps = {
-						FScheduledOp( Args.base, item),
-						FScheduledOp( Args.offsetX, item),
-						FScheduledOp( Args.offsetY, item),
-						FScheduledOp( Args.scaleX, item),
-						FScheduledOp( Args.scaleY, item),
-						FScheduledOp( Args.rotation, item) };
+						FScheduledOp(Args.base, item),
+						FScheduledOp(Args.offsetX, item),
+						FScheduledOp(Args.offsetY, item),
+						FScheduledOp(Args.scaleX, item),
+						FScheduledOp(Args.scaleY, item),
+						FScheduledOp(Args.rotation, item) };
 
-                AddOp( FScheduledOp( item.At, item, 1), Deps );
+                AddOp(FScheduledOp(item.At, item, 1), Deps);
 
 				break;
 			}
@@ -4744,35 +4743,35 @@ namespace mu
             {
             	MUTABLE_CPUPROFILER_SCOPE(IM_TRANSFORM_1)
             		
-                Ptr<const Image> pBaseImage = LoadImage( FCacheAddress(Args.base, item) );
+                Ptr<const Image> pBaseImage = LoadImage(FCacheAddress(Args.base, item));
                 
                 const FVector2f Offset = FVector2f(
-                        Args.offsetX ? LoadScalar( FCacheAddress(Args.offsetX, item) ) : 0.0f,
-                        Args.offsetY ? LoadScalar( FCacheAddress(Args.offsetY, item) ) : 0.0f );
+                        Args.offsetX ? LoadScalar(FCacheAddress(Args.offsetX, item)) : 0.0f,
+                        Args.offsetY ? LoadScalar(FCacheAddress(Args.offsetY, item)) : 0.0f);
 
                 const FVector2f Scale = FVector2f(
-                        Args.scaleX ? LoadScalar( FCacheAddress(Args.scaleX, item) ) : 1.0f,
-                        Args.scaleY ? LoadScalar( FCacheAddress(Args.scaleY, item) ) : 1.0f );
+                        Args.scaleX ? LoadScalar(FCacheAddress(Args.scaleX, item)) : 1.0f,
+                        Args.scaleY ? LoadScalar(FCacheAddress(Args.scaleY, item)) : 1.0f);
 
 				// Map Range 0-1 to a full rotation
-                const float Rotation = LoadScalar( FCacheAddress(Args.rotation, item) ) * UE_TWO_PI;
+                const float Rotation = LoadScalar(FCacheAddress(Args.rotation, item)) * UE_TWO_PI;
 			
 				EImageFormat BaseFormat = pBaseImage->GetFormat();
 				int32 BaseLODs = pBaseImage->GetLODCount();
 				FImageSize BaseSize = pBaseImage->GetSize();
 				FImageSize SampleImageSize =  FImageSize(
-					static_cast<uint16>( FMath::Clamp( FMath::FloorToInt( float(BaseSize.X) * FMath::Abs(Scale.X)), 2, BaseSize.X) ),
-					static_cast<uint16>( FMath::Clamp( FMath::FloorToInt( float(BaseSize.Y) * FMath::Abs(Scale.Y)), 2, BaseSize.Y) ) );
+					static_cast<uint16>(FMath::Clamp(FMath::FloorToInt(float(BaseSize.X) * FMath::Abs(Scale.X)), 2, BaseSize.X)),
+					static_cast<uint16>(FMath::Clamp(FMath::FloorToInt(float(BaseSize.Y) * FMath::Abs(Scale.Y)), 2, BaseSize.Y)));
 	
 				Ptr<Image> pSampleImage = CreateImage(SampleImageSize[0], SampleImageSize[1], BaseLODs, BaseFormat, EInitializationType::NotInitialized);
-				ImageResizeLinear( pSampleImage.get(), 0, pBaseImage.get() );
+				ImageResizeLinear(pSampleImage.get(), 0, pBaseImage.get());
 				Release(pBaseImage);
 
 				Ptr<Image> Result = CreateImage(BaseSize.X, BaseSize.Y, 1, BaseFormat);
-				ImageTransform( Result.get(), pSampleImage.get(), Offset, Scale, Rotation );
+				ImageTransform(Result.get(), pSampleImage.get(), Offset, Scale, Rotation, static_cast<EAddressMode>(Args.AddressMode));
 
 				Release(pSampleImage);
-				StoreImage( item, Result );
+				StoreImage(item, Result);
 
                 break;
             }
diff --git a/Engine/Plugins/Experimental/Mutable/Source/MutableRuntime/Private/MuR/ImagePrivate.h b/Engine/Plugins/Experimental/Mutable/Source/MutableRuntime/Private/MuR/ImagePrivate.h
--- a/Engine/Plugins/Experimental/Mutable/Source/MutableRuntime/Private/MuR/ImagePrivate.h
+++ b/Engine/Plugins/Experimental/Mutable/Source/MutableRuntime/Private/MuR/ImagePrivate.h
@@ -15,13 +15,14 @@ namespace mu
 	MUTABLE_DEFINE_ENUM_SERIALISABLE( ESamplingMethod );
 	MUTABLE_DEFINE_ENUM_SERIALISABLE( EMinFilterMethod );
 	MUTABLE_DEFINE_ENUM_SERIALISABLE( EImageFormat );
+	MUTABLE_DEFINE_ENUM_SERIALISABLE( EAddressMode );
 
 
 	struct FMipmapGenerationSettings
 	{
 		float m_sharpenFactor = 0.0f;
 		EMipmapFilterType m_filterType = EMipmapFilterType::MFT_SimpleAverage;
-		EAddressMode m_addressMode = EAddressMode::AM_NONE;
+		EAddressMode m_addressMode = EAddressMode::None;
 		bool m_ditherMipmapAlpha = false;
 
 		void Serialise( OutputArchive& arch ) const
diff --git a/Engine/Plugins/Experimental/Mutable/Source/MutableRuntime/Private/MuR/OpImageMipmap.cpp b/Engine/Plugins/Experimental/Mutable/Source/MutableRuntime/Private/MuR/OpImageMipmap.cpp
--- a/Engine/Plugins/Experimental/Mutable/Source/MutableRuntime/Private/MuR/OpImageMipmap.cpp
+++ b/Engine/Plugins/Experimental/Mutable/Source/MutableRuntime/Private/MuR/OpImageMipmap.cpp
@@ -140,7 +140,7 @@ namespace mu
 			return kernel;
 		}
 
-		template<int32 PIXEL_SIZE, typename CHANNEL_TYPE, EAddressMode AD_MODE = EAddressMode::AM_NONE>
+		template<int32 PIXEL_SIZE, typename CHANNEL_TYPE, EAddressMode AD_MODE = EAddressMode::Wrap>
 		class ImageAccessor
 		{
 			const FIntVector2 m_dim;
@@ -155,18 +155,18 @@ namespace mu
 
 			uint8 operator()(int32 x, int32 y, int32 c) const
 			{
-				if constexpr (AD_MODE == EAddressMode::AM_WRAP)
+				if constexpr (AD_MODE == EAddressMode::Wrap)
 				{
 					// (y & (m_dim.Y - 1)) * m_dim.Y + (x & (m_dim.X - 1))
 					return m_pImageBuf[((y % m_dim.Y) * m_dim.X + (x % m_dim.X)) * PIXEL_SIZE + c];
 				}
-				else if (AD_MODE == EAddressMode::AM_CLAMP)
+				else if (AD_MODE == EAddressMode::ClampToEdge)
 				{
 					return m_pImageBuf[(mu::clamp(y, 0, m_dim.Y - 1) * m_dim.X +
 						mu::clamp(x, 0, m_dim.X - 1))
 						* PIXEL_SIZE + c];
 				}
-				else if (AD_MODE == EAddressMode::AM_BLACK_BORDER)
+				else if (AD_MODE == EAddressMode::ClampToBlack)
 				{
 					return ((x < m_dim.X) && (y < m_dim.Y) && (x >= 0) && (y >= 0))
 						? m_pImageBuf[(y * m_dim.X + x) * PIXEL_SIZE + c]
@@ -251,7 +251,7 @@ namespace mu
 			for (; mips >= 0; --mips)
 			{
 				ImageAccessor<PIXEL_SIZE, uint8, AD_MODE> sourceSampler(pSource, sourceSize);
-				ImageAccessor<PIXEL_SIZE, uint8, EAddressMode::AM_NONE> sourceSamplerAddressModeNone(pSource, sourceSize);
+				ImageAccessor<PIXEL_SIZE, uint8, EAddressMode::None> sourceSamplerAddressModeNone(pSource, sourceSize);
 
 				FIntVector2 destSize = FIntVector2(FMath::DivideAndRoundUp(sourceSize.X, 2), FMath::DivideAndRoundUp(sourceSize.Y, 2));
 
@@ -458,27 +458,27 @@ namespace mu
 		{
 			switch (settings.m_addressMode)
 			{
-			case EAddressMode::AM_CLAMP:
+			case EAddressMode::ClampToEdge:
 			{
-				GenerateMipmapUint8Sharpen<PIXEL_SIZE, EAddressMode::AM_CLAMP>(
+				GenerateMipmapUint8Sharpen<PIXEL_SIZE, EAddressMode::ClampToEdge>(
 					mips, pSource, Dest, sourceSize, settings.m_sharpenFactor);
 				break;
 			}
-			case EAddressMode::AM_WRAP:
+			case EAddressMode::Wrap:
 			{
-				GenerateMipmapUint8Sharpen<PIXEL_SIZE, EAddressMode::AM_WRAP>(
+				GenerateMipmapUint8Sharpen<PIXEL_SIZE, EAddressMode::Wrap>(
 					mips, pSource, Dest, sourceSize, settings.m_sharpenFactor);
 				break;
 			}
-			case EAddressMode::AM_BLACK_BORDER:
+			case EAddressMode::ClampToBlack:
 			{
-				GenerateMipmapUint8Sharpen<PIXEL_SIZE, EAddressMode::AM_BLACK_BORDER>(
+				GenerateMipmapUint8Sharpen<PIXEL_SIZE, EAddressMode::ClampToBlack>(
 					mips, pSource, Dest, sourceSize, settings.m_sharpenFactor);
 				break;
 			}
 			default:
 			{
-				GenerateMipmapUint8Sharpen<PIXEL_SIZE, EAddressMode::AM_NONE>(
+				GenerateMipmapUint8Sharpen<PIXEL_SIZE, EAddressMode::None>(
 					mips, pSource, Dest, sourceSize, settings.m_sharpenFactor);
 				break;
 			}
diff --git a/Engine/Plugins/Experimental/Mutable/Source/MutableRuntime/Private/MuR/OpImageTransform.h b/Engine/Plugins/Experimental/Mutable/Source/MutableRuntime/Private/MuR/OpImageTransform.h
--- a/Engine/Plugins/Experimental/Mutable/Source/MutableRuntime/Private/MuR/OpImageTransform.h
+++ b/Engine/Plugins/Experimental/Mutable/Source/MutableRuntime/Private/MuR/OpImageTransform.h
@@ -3,93 +3,13 @@
 #pragma once
 
 #include "MuR/ImagePrivate.h"
-#include "Async/ParallelFor.h"
 #include "MuR/MutableTrace.h"
 
-#include "Math/Vector.h"
-#include "Math/IntVector.h"
+#include "MuR/OpImageTransformPrivate.h"
 
 namespace mu
 {
-
-    //---------------------------------------------------------------------------------------------
-    //! Image transform, translate, scale, rotate 
-    //---------------------------------------------------------------------------------------------
-
-    template<int32 NC>
-	inline void ImageTransform( Image* pDest, const Image* pBase, FVector2f Offset, FVector2f Scale, float RotationRad)
-    {
-		FIntVector2 BaseSize = FIntVector2(pBase->GetSizeX(), pBase->GetSizeY());
-		FIntVector2 DestSize = FIntVector2(pDest->GetSizeX(), pDest->GetSizeY());
-
-		const FVector2f BaseSizeF = FVector2f(BaseSize.X, BaseSize.Y);
-		const FVector2f DestSizeF = FVector2f(DestSize.X, DestSize.Y);
-
-		float SinRot, CosRot;
-		FMath::SinCos(&SinRot, &CosRot, RotationRad);
-
-		const FVector2f DestNormFactor = FVector2f(1.0f) / DestSizeF;
-		const FVector2f BaseNormFactor = FVector2f(1.0f) / BaseSizeF;
-
-		Scale.X = FMath::IsNearlyZero(Scale.X) ? UE_SMALL_NUMBER : Scale.X;
-		Scale.Y = FMath::IsNearlyZero(Scale.Y) ? UE_SMALL_NUMBER : Scale.Y;
-
-		FVector2f ScaleReciprocal =  FVector2f(1.0f) / Scale; 
-
-		for (int32 Y = 0; Y < DestSize.Y; ++Y)
-		{
-			for (int32 X = 0; X < DestSize.X; ++X)
-			{
-				FVector2f DestUV = (FVector2f(X, Y) + 0.5f) * DestNormFactor - 0.5f;	
-
-				DestUV = FVector2f( CosRot * DestUV.X + SinRot * DestUV.Y,
-									CosRot * DestUV.Y - SinRot * DestUV.X);
-				DestUV *= ScaleReciprocal;
-				DestUV += -Offset + 0.5f;
-
-				const FVector2f Coord = FVector2f(DestUV.X - FMath::Floor(DestUV.X), DestUV.Y - FMath::Floor(DestUV.Y)) * BaseSizeF;
-
-				const FVector2f CoordFloor = FVector2f( FMath::Floor( Coord.X ), FMath::Floor( Coord.Y ) );
-				const FVector2f CoordFract = Coord - CoordFloor;
-
-				const FVector2f UV00 = CoordFloor * BaseNormFactor;
-				const FVector2f UV01 = ( CoordFloor + FVector2f(0.0f, 1.0f) ) * BaseNormFactor;
-				const FVector2f UV10 = ( CoordFloor + FVector2f(1.0f, 0.0f) ) * BaseNormFactor;
-				const FVector2f UV11 = ( CoordFloor + FVector2f(1.0f) ) * BaseNormFactor;
-
-				const FVector2f Coord00 =
-						FVector2f( UV00.X - FMath::Floor(UV00.X), UV00.Y - FMath::Floor(UV00.Y) ) * BaseSizeF - FVector2f(0.5f);
-				const FVector2f Coord01 = 
-						FVector2f( UV01.X - FMath::Floor(UV01.X), UV01.Y - FMath::Floor(UV01.Y) ) * BaseSizeF - FVector2f(0.5f);
-				const FVector2f Coord10 =
-						FVector2f( UV10.X - FMath::Floor(UV10.X), UV10.Y - FMath::Floor(UV10.Y) ) * BaseSizeF - FVector2f(0.5f);
-				const FVector2f Coord11 =
-						FVector2f( UV11.X - FMath::Floor(UV11.X), UV11.Y - FMath::Floor(UV11.Y) ) * BaseSizeF - FVector2f(0.5f);
-
-				const FIntVector2 S00 = FIntVector2(Coord00.X, Coord00.Y);
-				const FIntVector2 S01 = FIntVector2(Coord01.X, Coord01.Y);
-				const FIntVector2 S10 = FIntVector2(Coord10.X, Coord10.Y);
-				const FIntVector2 S11 = FIntVector2(Coord11.X, Coord11.Y);
-
-				for (int32 Channel = 0; Channel < NC; ++Channel)
-				{
-					const float V00 = float(pBase->GetData()[(S00.Y*BaseSize.X + S00.X) * NC + Channel]);
-					const float V01 = float(pBase->GetData()[(S01.Y*BaseSize.X + S01.X) * NC + Channel]);
-					const float V10 = float(pBase->GetData()[(S10.Y*BaseSize.X + S10.X) * NC + Channel]);
-					const float V11 = float(pBase->GetData()[(S11.Y*BaseSize.X + S11.X) * NC + Channel]);
-
-					const float V = FMath::Lerp(
-							FMath::Lerp(V00, V10, CoordFract.X), 
-						    FMath::Lerp(V01, V11, CoordFract.X), CoordFract.Y);
-
-					pDest->GetData()[(Y*DestSize.X + X)*NC + Channel] = static_cast<uint8>(V);
-				}	
-			}
-		}
-    }
-    
-	inline void ImageTransform( Image* pDestImage, const Image* pImage, FVector2f Offset, FVector2f Scale, float RotationRad)
-
+	void ImageTransform(Image* pDestImage, const Image* pImage, FVector2f Offset, FVector2f Scale, float RotationRad, EAddressMode AddressMode)
     {
 		MUTABLE_CPUPROFILER_SCOPE(ImageTransform)
 
@@ -98,20 +18,62 @@ namespace mu
 		
 		case EImageFormat::IF_L_UBYTE:
 		{
-			ImageTransform<1>( pDestImage, pImage, Offset, Scale, RotationRad );
+			switch (AddressMode)
+			{
+			case EAddressMode::ClampToBlack:
+				ImageTransformImpl<1, EAddressMode::ClampToBlack>(pDestImage, pImage, Offset, Scale, RotationRad);
+				break;
+			case EAddressMode::ClampToEdge:
+				ImageTransformImpl<1, EAddressMode::ClampToEdge>(pDestImage, pImage, Offset, Scale, RotationRad);
+				break;
+			case EAddressMode::Wrap:
+				ImageTransformImpl<1, EAddressMode::Wrap>(pDestImage, pImage, Offset, Scale, RotationRad);
+				break;
+			default:
+				check(false);
+				break;
+			}
 			break;
 		}
 		
 		case EImageFormat::IF_RGB_UBYTE:
 		{
-			ImageTransform<3>( pDestImage, pImage, Offset, Scale, RotationRad );
+			switch (AddressMode)
+			{
+			case EAddressMode::ClampToBlack:
+				ImageTransformImpl<3, EAddressMode::ClampToBlack>(pDestImage, pImage, Offset, Scale, RotationRad);
+				break;
+			case EAddressMode::ClampToEdge:
+				ImageTransformImpl<3, EAddressMode::ClampToEdge>(pDestImage, pImage, Offset, Scale, RotationRad);
+				break;
+			case EAddressMode::Wrap:
+				ImageTransformImpl<3, EAddressMode::Wrap>(pDestImage, pImage, Offset, Scale, RotationRad);
+				break;
+			default:
+				check(false);
+				break;
+			}
 			break;
 		}
 
 		case EImageFormat::IF_BGRA_UBYTE:
 		case EImageFormat::IF_RGBA_UBYTE:
 		{
-			ImageTransform<4>( pDestImage, pImage, Offset, Scale, RotationRad );
+			switch (AddressMode)
+			{
+			case EAddressMode::ClampToBlack:
+				ImageTransformImpl<4, EAddressMode::ClampToBlack>(pDestImage, pImage, Offset, Scale, RotationRad);
+				break;
+			case EAddressMode::ClampToEdge:
+				ImageTransformImpl<4, EAddressMode::ClampToEdge>(pDestImage, pImage, Offset, Scale, RotationRad);
+				break;
+			case EAddressMode::Wrap:
+				ImageTransformImpl<4, EAddressMode::Wrap>(pDestImage, pImage, Offset, Scale, RotationRad);
+				break;
+			default:
+				check(false);
+				break;
+			}
 			break;
 		}
 
diff --git a/Engine/Plugins/Experimental/Mutable/Source/MutableRuntime/Private/MuR/OpImageTransformPrivate.cpp b/Engine/Plugins/Experimental/Mutable/Source/MutableRuntime/Private/MuR/OpImageTransformPrivate.cpp
new file mode 100644
--- /dev/null
+++ b/Engine/Plugins/Experimental/Mutable/Source/MutableRuntime/Private/MuR/OpImageTransformPrivate.cpp
@@ -0,0 +1,184 @@
+// Copyright Epic Games, Inc. All Rights Reserved.
+
+#pragma once
+
+#include "OpImageTransformPrivate.h"
+#include "Math/Vector.h"
+#include "Math/IntVector.h"
+
+namespace mu
+{
+
+    //---------------------------------------------------------------------------------------------
+    //! Image transform, translate, scale, rotate 
+    //---------------------------------------------------------------------------------------------
+
+    template<int32 NC, EAddressMode AddressMode>
+	void ImageTransformImpl(Image* pDest, const Image* pBase, FVector2f Offset, FVector2f Scale, float RotationRad)
+    {
+		using FUInt16Vector2 = UE::Math::TIntVector2<uint16>;
+		struct FPixelData
+		{
+			alignas(8) uint16 Data[NC];
+		};
+
+		FUInt16Vector2 BaseSize = FUInt16Vector2(pBase->GetSizeX(), pBase->GetSizeY());
+		FUInt16Vector2 DestSize = FUInt16Vector2(pDest->GetSizeX(), pDest->GetSizeY());
+
+		const FVector2f BaseSizeF = FVector2f(BaseSize.X, BaseSize.Y);
+		const FVector2f DestSizeF = FVector2f(DestSize.X, DestSize.Y);
+
+		float SinRot, CosRot;
+		FMath::SinCos(&SinRot, &CosRot, RotationRad);
+
+		const FVector2f DestNormFactor = FVector2f(1.0f) / DestSizeF;
+		const FVector2f BaseNormFactor = FVector2f(1.0f) / BaseSizeF;
+
+		Scale.X = FMath::IsNearlyZero(Scale.X) ? UE_SMALL_NUMBER : Scale.X;
+		Scale.Y = FMath::IsNearlyZero(Scale.Y) ? UE_SMALL_NUMBER : Scale.Y;
+
+		const FVector2f ScaleReciprocal = FVector2f(1.0f) / Scale; 
+		const FVector2f ScaledOffset = Offset * ScaleReciprocal;
+
+		for (int32 Y = 0; Y < DestSize.Y; ++Y)
+		{
+			for (int32 X = 0; X < DestSize.X; ++X)
+			{
+				FVector2f DestUv = (FVector2f(X, Y) + 0.5f) * DestNormFactor - 0.5f;	
+
+				DestUv -= Offset;
+				DestUv = FVector2f(CosRot * DestUv.X + SinRot * DestUv.Y,
+								   CosRot * DestUv.Y - SinRot * DestUv.X);
+				DestUv *= ScaleReciprocal;
+				DestUv = DestUv + 0.5f;
+
+				// TODO: Make the bloack edge antialiazed.
+				if constexpr (AddressMode == EAddressMode::ClampToBlack)
+				{
+					if ((DestUv.X < 0.0f) | (DestUv.X > 1.0f) | (DestUv.Y < 0.0f) | (DestUv.Y > 1.0f))
+					{
+						for (uint32 Channel = 0; Channel < NC; ++Channel)
+						{
+							pDest->GetData()[(Y*DestSize.X + X)*NC + Channel] = 0;
+						}
+
+						continue;
+					}
+				}
+
+				const FVector2f Uv = Invoke([&]() -> FVector2f
+				{
+					if constexpr (AddressMode == EAddressMode::Wrap)
+					{
+						return FVector2f(FMath::Frac(DestUv.X), FMath::Frac(DestUv.Y));
+					}
+					else if (AddressMode == EAddressMode::ClampToBlack || AddressMode == EAddressMode::ClampToEdge)
+					{
+						return FVector2f(FMath::Clamp(DestUv.X, 0.0f, 1.0f), FMath::Clamp(DestUv.Y, 0.0f, 1.0f));
+					}
+					else
+					{
+						return DestUv;
+					}
+				});
+
+				auto SampleImageBilinear = [&](FVector2f Uv, FUInt16Vector2 Size, const uint8* DataPtr)
+				{
+					auto ComputeInterpolator = [](float T) -> uint16
+					{
+						return static_cast<uint16>(255.0f * T);
+					};
+
+					const FVector2f SizeF(Size.X, Size.Y);
+
+					const FVector2f CoordsF = FVector2f(
+						FMath::Clamp(Uv.X * SizeF.X - 0.5f, 0.0f, SizeF.X - 1.0f),
+						FMath::Clamp(Uv.Y * SizeF.Y - 0.5f, 0.0f, SizeF.Y - 1.0f));
+
+					const FUInt16Vector2 Frac = FUInt16Vector2(
+						ComputeInterpolator(FMath::Frac(CoordsF.X)),
+						ComputeInterpolator(FMath::Frac(CoordsF.Y)));
+
+					const FIntVector2 Coords = FIntVector2(CoordsF.X, CoordsF.Y);
+					const FIntVector2 CoordsPlusOne = Invoke([&]() -> FIntVector2
+					{	
+						if constexpr (AddressMode == EAddressMode::Wrap)
+						{
+							return FIntVector2(Coords.X + 1 < Size.X ? Coords.X + 1 : 0, 
+											   Coords.Y + 1 < Size.Y ? Coords.Y + 1 : 0);
+						}
+						else if (AddressMode == EAddressMode::ClampToBlack || AddressMode == EAddressMode::ClampToEdge)
+						{
+							return FIntVector2(FMath::Min(Size.X - 1, Coords.X + 1), 
+											   FMath::Min(Size.Y - 1, Coords.Y + 1));
+						}
+					});
+						
+					uint8 const * const Pixel00Ptr = DataPtr + (Coords.Y        * Size.X + Coords.X)        * NC;
+					uint8 const * const Pixel10Ptr = DataPtr + (Coords.Y        * Size.X + CoordsPlusOne.X) * NC;
+					uint8 const * const Pixel01Ptr = DataPtr + (CoordsPlusOne.Y * Size.X + Coords.X)        * NC;
+					uint8 const * const Pixel11Ptr = DataPtr + (CoordsPlusOne.Y * Size.X + CoordsPlusOne.X) * NC;
+
+					auto LoadPixel = [](const uint8* Ptr) -> FPixelData
+					{
+						FPixelData Result;
+						if constexpr (NC == 4)
+						{
+							const uint32 PackedData = *reinterpret_cast<const uint32*>(Ptr);
+							
+							Result.Data[0] = static_cast<uint16>((PackedData >> (8 * 0)) & 0xFF);
+							Result.Data[1] = static_cast<uint16>((PackedData >> (8 * 1)) & 0xFF);
+							Result.Data[2] = static_cast<uint16>((PackedData >> (8 * 2)) & 0xFF);
+							Result.Data[3] = static_cast<uint16>((PackedData >> (8 * 3)) & 0xFF);
+						}
+						else
+						{
+							for (int32 C = 0; C < NC; ++C)
+							{
+								Result.Data[C] = static_cast<uint16>(Ptr[C]);
+							}
+						}
+
+						return Result;
+					};
+
+					FPixelData PixelData00 = LoadPixel(Pixel00Ptr);
+					FPixelData PixelData10 = LoadPixel(Pixel10Ptr);
+					FPixelData PixelData01 = LoadPixel(Pixel01Ptr);
+					FPixelData PixelData11 = LoadPixel(Pixel11Ptr);
+
+					FPixelData FilteredPixelData;
+				
+					for (int32 C = 0; C < NC; ++C)
+					{
+						const uint16 LerpY0 = ((PixelData10.Data[C] * Frac.X) + PixelData00.Data[C] * (255 - Frac.X)) / 255;
+						const uint16 LerpY1 = ((PixelData11.Data[C] * Frac.X) + PixelData01.Data[C] * (255 - Frac.X)) / 255;
+						FilteredPixelData.Data[C] = ((LerpY1 * Frac.Y) + LerpY0*(255 - Frac.Y)) / 255;
+					} 
+
+					return FilteredPixelData;
+				};
+
+				FPixelData PixelData = SampleImageBilinear(Uv, BaseSize, pBase->GetData());
+				static_assert(NC <= 4);
+
+				uint8 * const DestData = pDest->GetData();
+				for (uint32 Channel = 0; Channel < NC; ++Channel)
+				{
+					DestData[(Y*DestSize.X + X)*NC + Channel] = PixelData.Data[Channel];
+				}
+			}
+		}
+    }
+
+	// Force instantiation of used templates. A link error will occur if any of them is missing.
+	template void ImageTransformImpl<1, EAddressMode::ClampToBlack>(Image*, const Image*, FVector2f, FVector2f, float);
+	template void ImageTransformImpl<1, EAddressMode::ClampToEdge>(Image*, const Image*, FVector2f, FVector2f, float);
+	template void ImageTransformImpl<1, EAddressMode::Wrap>(Image*, const Image*, FVector2f, FVector2f, float);
+	template void ImageTransformImpl<3, EAddressMode::ClampToBlack>(Image*, const Image*, FVector2f, FVector2f, float);
+	template void ImageTransformImpl<3, EAddressMode::ClampToEdge>(Image*, const Image*, FVector2f, FVector2f, float);
+	template void ImageTransformImpl<3, EAddressMode::Wrap>(Image*, const Image*, FVector2f, FVector2f, float);
+	template void ImageTransformImpl<4, EAddressMode::ClampToBlack>(Image*, const Image*, FVector2f, FVector2f, float);
+	template void ImageTransformImpl<4, EAddressMode::ClampToEdge>(Image*, const Image*, FVector2f, FVector2f, float);
+	template void ImageTransformImpl<4, EAddressMode::Wrap>(Image*, const Image*, FVector2f, FVector2f, float);
+}
diff --git a/Engine/Plugins/Experimental/Mutable/Source/MutableRuntime/Private/MuR/OpImageTransformPrivate.h b/Engine/Plugins/Experimental/Mutable/Source/MutableRuntime/Private/MuR/OpImageTransformPrivate.h
new file mode 100644
--- /dev/null
+++ b/Engine/Plugins/Experimental/Mutable/Source/MutableRuntime/Private/MuR/OpImageTransformPrivate.h
@@ -0,0 +1,29 @@
+// Copyright Epic Games, Inc. All Rights Reserved.
+
+#pragma once
+
+#include "MuR/ImagePrivate.h"
+
+#include "Math/Vector.h"
+#include "Math/IntVector.h"
+
+namespace mu
+{
+
+    //---------------------------------------------------------------------------------------------
+    //! Image transform, translate, scale, rotate 
+    //---------------------------------------------------------------------------------------------
+	template<int32, EAddressMode> 
+	void ImageTransformImpl(Image*, const Image*, FVector2f, FVector2f, float);
+
+	// extern template could be used if for some reason we need to provide the full declaration here. 
+	//extern template FORCENOINLINE void ImageTransformImpl<1, EAddressMode::ClampToBlack>(Image*, const Image*, FVector2f, FVector2f, float);
+	//extern template FORCENOINLINE void ImageTransformImpl<1, EAddressMode::ClampToEdge>(Image*, const Image*, FVector2f, FVector2f, float);
+	//extern template FORCENOINLINE void ImageTransformImpl<1, EAddressMode::Wrap>(Image*, const Image*, FVector2f, FVector2f, float);
+	//extern template FORCENOINLINE void ImageTransformImpl<3, EAddressMode::ClampToBlack>(Image*, const Image*, FVector2f, FVector2f, float);
+	//extern template FORCENOINLINE void ImageTransformImpl<3, EAddressMode::ClampToEdge>(Image*, const Image*, FVector2f, FVector2f, float);
+	//extern template FORCENOINLINE void ImageTransformImpl<3, EAddressMode::Wrap>(Image*, const Image*, FVector2f, FVector2f, float);
+	//extern template FORCENOINLINE void ImageTransformImpl<4, EAddressMode::ClampToBlack>(Image*, const Image*, FVector2f, FVector2f, float);
+	//extern template FORCENOINLINE void ImageTransformImpl<4, EAddressMode::ClampToEdge>(Image*, const Image*, FVector2f, FVector2f, float);
+	//extern template FORCENOINLINE void ImageTransformImpl<4, EAddressMode::Wrap>(Image*, const Image*, FVector2f, FVector2f, float);
+}
diff --git a/Engine/Plugins/Experimental/Mutable/Source/MutableRuntime/Private/MuR/Operations.h b/Engine/Plugins/Experimental/Mutable/Source/MutableRuntime/Private/MuR/Operations.h
--- a/Engine/Plugins/Experimental/Mutable/Source/MutableRuntime/Private/MuR/Operations.h
+++ b/Engine/Plugins/Experimental/Mutable/Source/MutableRuntime/Private/MuR/Operations.h
@@ -821,6 +821,8 @@ namespace mu
 			ADDRESS scaleX = 0;
 			ADDRESS scaleY = 0;
 			ADDRESS rotation = 0;
+
+			uint32 AddressMode = 0;
 		};
 
         //-------------------------------------------------------------------------------------
diff --git a/Engine/Plugins/Experimental/Mutable/Source/MutableRuntime/Public/MuR/Image.h b/Engine/Plugins/Experimental/Mutable/Source/MutableRuntime/Public/MuR/Image.h
--- a/Engine/Plugins/Experimental/Mutable/Source/MutableRuntime/Public/MuR/Image.h
+++ b/Engine/Plugins/Experimental/Mutable/Source/MutableRuntime/Public/MuR/Image.h
@@ -128,12 +128,10 @@ namespace mu
 
 	enum class EAddressMode
 	{
-		AM_NONE,
-		AM_WRAP,
-		AM_CLAMP,
-		AM_BLACK_BORDER,
-		_AM_COUNT
-
+		None,
+		Wrap,
+		ClampToEdge,
+		ClampToBlack,
 	};
 
 	enum class EMipmapFilterType
@@ -174,6 +172,7 @@ namespace mu
 		TotalAreaHeuristic,
 		MaxValue
 	};
+
 	static_assert(uint32(EMinFilterMethod::MaxValue) <= (1 << 3), "EMinFilterMethod enum cannot hold more than 8 values");
 
 	struct FImageFormatData
diff --git a/Engine/Plugins/Experimental/Mutable/Source/MutableTools/Private/MuT/ASTOpImageMipmap.h b/Engine/Plugins/Experimental/Mutable/Source/MutableTools/Private/MuT/ASTOpImageMipmap.h
--- a/Engine/Plugins/Experimental/Mutable/Source/MutableTools/Private/MuT/ASTOpImageMipmap.h
+++ b/Engine/Plugins/Experimental/Mutable/Source/MutableTools/Private/MuT/ASTOpImageMipmap.h
@@ -33,7 +33,7 @@ namespace mu
 
 		//! Mipmap generation settings. 
 		float SharpenFactor = 0.0f;
-		EAddressMode AddressMode = EAddressMode::AM_NONE;
+		EAddressMode AddressMode = EAddressMode::None;
 		EMipmapFilterType FilterType = EMipmapFilterType::MFT_Unfiltered;
 		bool DitherMipmapAlpha = false;
 
diff --git a/Engine/Plugins/Experimental/Mutable/Source/MutableTools/Private/MuT/ASTOpImageTransform.cpp b/Engine/Plugins/Experimental/Mutable/Source/MutableTools/Private/MuT/ASTOpImageTransform.cpp
--- a/Engine/Plugins/Experimental/Mutable/Source/MutableTools/Private/MuT/ASTOpImageTransform.cpp
+++ b/Engine/Plugins/Experimental/Mutable/Source/MutableTools/Private/MuT/ASTOpImageTransform.cpp
@@ -46,7 +46,8 @@ namespace mu
 				offsetY == other->offsetY &&
 				scaleX == other->scaleX &&
 				scaleY == other->scaleY &&
-				rotation == other->rotation;
+				rotation == other->rotation &&
+				AddressMode == other->AddressMode;
 		}
 		return false;
 	}
@@ -62,6 +63,7 @@ namespace mu
 		hash_combine(res, scaleX.child().get());
 		hash_combine(res, scaleY.child().get());
 		hash_combine(res, rotation.child().get());
+		hash_combine(res, std::hash<uint32>()(static_cast<uint32>(AddressMode)));
 		return res;
 	}
 
@@ -76,6 +78,7 @@ namespace mu
 		n->scaleX = mapChild(scaleX.child());
 		n->scaleY = mapChild(scaleY.child());
 		n->rotation = mapChild(rotation.child());
+		n->AddressMode = AddressMode;
 		return n;
 	}
 
@@ -99,6 +102,7 @@ namespace mu
 		if (!linkedAddress)
 		{
 			OP::ImageTransformArgs Args;
+			FMemory::Memzero(Args);
 
 			Args.base = base ? base->linkedAddress : 0;
 			Args.offsetX = offsetX ? offsetX->linkedAddress : 0;
@@ -106,6 +110,8 @@ namespace mu
 			Args.scaleX = scaleX ? scaleX->linkedAddress : 0;
 			Args.scaleY = scaleY ? scaleY->linkedAddress : 0;
 			Args.rotation = rotation ? rotation->linkedAddress : 0;
+			Args.AddressMode = static_cast<uint32>(AddressMode);
+
 
 			linkedAddress = (OP::ADDRESS)program.m_opAddress.Num();
 			program.m_opAddress.Add((uint32_t)program.m_byteCode.Num());
diff --git a/Engine/Plugins/Experimental/Mutable/Source/MutableTools/Private/MuT/ASTOpImageTransform.h b/Engine/Plugins/Experimental/Mutable/Source/MutableTools/Private/MuT/ASTOpImageTransform.h
--- a/Engine/Plugins/Experimental/Mutable/Source/MutableTools/Private/MuT/ASTOpImageTransform.h
+++ b/Engine/Plugins/Experimental/Mutable/Source/MutableTools/Private/MuT/ASTOpImageTransform.h
@@ -28,6 +28,8 @@ struct FProgram;
 		ASTChild scaleY;
 		ASTChild rotation;
 
+		EAddressMode AddressMode = EAddressMode::Wrap;
+
 	public:
 
 		ASTOpImageTransform();
diff --git a/Engine/Plugins/Experimental/Mutable/Source/MutableTools/Private/MuT/CodeGenerator.cpp b/Engine/Plugins/Experimental/Mutable/Source/MutableTools/Private/MuT/CodeGenerator.cpp
--- a/Engine/Plugins/Experimental/Mutable/Source/MutableTools/Private/MuT/CodeGenerator.cpp
+++ b/Engine/Plugins/Experimental/Mutable/Source/MutableTools/Private/MuT/CodeGenerator.cpp
@@ -1636,11 +1636,11 @@ namespace mu
 								int mipsY = (int)ceilf(logf((float)blockSizeY / finfo.m_pixelsPerBlockY) / logf(2.0f));
 								mop->BlockLevels = (uint8_t)FMath::Max(mipsX, mipsY);
 
-								// Not important for the end of the mip tail?
-								mop->AddressMode = EAddressMode::AM_CLAMP;
-								mop->FilterType = EMipmapFilterType::MFT_SimpleAverage;
-								mop->SharpenFactor = 0;
-								mop->DitherMipmapAlpha = false;
+							// Not important for the end of the mip tail?
+							mop->AddressMode = EAddressMode::ClampToEdge;
+							mop->FilterType = EMipmapFilterType::MFT_SimpleAverage;
+							mop->SharpenFactor = 0;
+							mop->DitherMipmapAlpha = false;
 
 								imageAd = mop;
 							}
diff --git a/Engine/Plugins/Experimental/Mutable/Source/MutableTools/Private/MuT/CodeGenerator_Image.cpp b/Engine/Plugins/Experimental/Mutable/Source/MutableTools/Private/MuT/CodeGenerator_Image.cpp
--- a/Engine/Plugins/Experimental/Mutable/Source/MutableTools/Private/MuT/CodeGenerator_Image.cpp
+++ b/Engine/Plugins/Experimental/Mutable/Source/MutableTools/Private/MuT/CodeGenerator_Image.cpp
@@ -609,6 +609,7 @@ namespace mu
  		op->scaleX = ScaleX ? ScaleX : ScaleY;
 		op->scaleY = ScaleY ? ScaleY : ScaleX;
 		op->rotation = Rotation; 
+		op->AddressMode = node.AddressMode;
 
         // Base image
         Ptr<ASTOp> base;
diff --git a/Engine/Plugins/Experimental/Mutable/Source/MutableTools/Private/MuT/NodeImageTransform.cpp b/Engine/Plugins/Experimental/Mutable/Source/MutableTools/Private/MuT/NodeImageTransform.cpp
--- a/Engine/Plugins/Experimental/Mutable/Source/MutableTools/Private/MuT/NodeImageTransform.cpp
+++ b/Engine/Plugins/Experimental/Mutable/Source/MutableTools/Private/MuT/NodeImageTransform.cpp
@@ -144,5 +144,16 @@ namespace mu
 		m_pD->m_pRotation = pNode;
 	}
 
+	//---------------------------------------------------------------------------------------------
+	void NodeImageTransform::SetAddressMode(EAddressMode AddressMode)
+	{
+		m_pD->AddressMode = AddressMode;
+	}
+
+	//---------------------------------------------------------------------------------------------
+	EAddressMode NodeImageTransform::GetAddressMode() const
+	{
+		return m_pD->AddressMode;
+	}
 }
 
diff --git a/Engine/Plugins/Experimental/Mutable/Source/MutableTools/Private/MuT/NodeImageTransformPrivate.h b/Engine/Plugins/Experimental/Mutable/Source/MutableTools/Private/MuT/NodeImageTransformPrivate.h
--- a/Engine/Plugins/Experimental/Mutable/Source/MutableTools/Private/MuT/NodeImageTransformPrivate.h
+++ b/Engine/Plugins/Experimental/Mutable/Source/MutableTools/Private/MuT/NodeImageTransformPrivate.h
@@ -3,12 +3,12 @@
 
 #pragma once
 
+#include "MuR/ImagePrivate.h"
 #include "MuT/NodeImagePrivate.h"
 #include "MuT/NodeImageTransform.h"
 #include "MuT/AST.h"
 #include "MuT/NodeScalar.h"
 
-
 namespace mu
 {
 
@@ -30,10 +30,12 @@ namespace mu
 		NodeScalarPtr m_pScaleY;
 		NodeScalarPtr m_pRotation;
 
+		EAddressMode AddressMode = EAddressMode::Wrap;
+
 		//!
 		void Serialise( OutputArchive& arch ) const
 		{
-            uint32_t ver = 0;
+            uint32 ver = 1;
 			arch << ver;
 
 			arch << m_pBase;
@@ -42,14 +44,15 @@ namespace mu
 			arch << m_pScaleX;
 			arch << m_pScaleY;
 			arch << m_pRotation;
+			arch << static_cast<uint32>(AddressMode);
 		}
 
 		//!
 		void Unserialise( InputArchive& arch )
 		{
-            uint32_t ver;
+            uint32 ver;
 			arch >> ver;
-			check(ver==0);
+			check(ver <= 1);
 
 			arch >> m_pBase;
 			arch >> m_pOffsetX;
@@ -57,6 +60,14 @@ namespace mu
 			arch >> m_pScaleX;
 			arch >> m_pScaleY;
 			arch >> m_pRotation;
+
+			uint32 AddressModeValue = static_cast<uint32>(EAddressMode::Wrap);
+			if (ver == 1)
+			{
+				arch >> AddressModeValue;
+			}
+
+			AddressMode = static_cast<EAddressMode>(AddressModeValue);
 		}
 	};
 
diff --git a/Engine/Plugins/Experimental/Mutable/Source/MutableTools/Public/MuT/NodeImageTransform.h b/Engine/Plugins/Experimental/Mutable/Source/MutableTools/Public/MuT/NodeImageTransform.h
--- a/Engine/Plugins/Experimental/Mutable/Source/MutableTools/Public/MuT/NodeImageTransform.h
+++ b/Engine/Plugins/Experimental/Mutable/Source/MutableTools/Public/MuT/NodeImageTransform.h
@@ -7,7 +7,6 @@
 #include "MuT/Node.h"
 #include "MuT/NodeImage.h"
 
-
 namespace mu
 {
 
@@ -23,6 +22,7 @@ namespace mu
 	class InputArchive;
 	class OutputArchive;
 
+	enum class EAddressMode;
 
 	//! Node that multiplies the colors of an image, channel by channel.
 	//! \ingroup model
@@ -72,6 +72,9 @@ namespace mu
 		void SetScaleX( NodeScalarPtr pNode );
 		void SetScaleY( NodeScalarPtr pNode );
 		void SetRotation( NodeScalarPtr pNode );
+
+		EAddressMode GetAddressMode() const;
+		void SetAddressMode(EAddressMode AddressMode);
 		//-----------------------------------------------------------------------------------------
 		// Interface pattern
 		//-----------------------------------------------------------------------------------------
