#rnx

diff --git a/Engine/Plugins/Experimental/Mutable/Source/CustomizableObject/Classes/MuCO/CustomizableInstanceLODManagement.h b/Engine/Plugins/Experimental/Mutable/Source/CustomizableObject/Classes/MuCO/CustomizableInstanceLODManagement.h
--- a/Engine/Plugins/Experimental/Mutable/Source/CustomizableObject/Classes/MuCO/CustomizableInstanceLODManagement.h
+++ b/Engine/Plugins/Experimental/Mutable/Source/CustomizableObject/Classes/MuCO/CustomizableInstanceLODManagement.h
@@ -5,6 +5,9 @@
 
 #include "CustomizableInstanceLODManagement.generated.h"
 
+typedef TMap<const class UCustomizableObjectInstance*, class FMutableUpdateCandidate> FMutableInstanceUpdateMap;
+
+#define NEW_MUTABLE_PLUGIN_QUEUE_SYSTEM
 
 // This is an abstract base class, override it to create a new Instance LOD management system and register with UCustomizableObjectSystem::SetInstanceLODManagement
 UCLASS(Blueprintable, BlueprintType)
@@ -16,6 +19,8 @@ public:
 	virtual ~UCustomizableInstanceLODManagementBase() {};
 
 	// WARNING! The following methods must be overriden in derived classes
+	virtual void UpdateInstanceDistsAndLODs(FMutableInstanceUpdateMap& InOutRequestedUpdates) { check(0); };
+	// Remove the following method after CL 25973936 is propagated to all streams
 	virtual void UpdateInstanceDistsAndLODs() { check(0); };
 	virtual int32 GetNumGeneratedInstancesLimitFullLODs() const { check(0); return 0; };
 	virtual int32 GetNumGeneratedInstancesLimitLOD1() const { check(0); return 0; };
@@ -34,7 +39,7 @@ public:
 	UCustomizableInstanceLODManagement();
 	virtual ~UCustomizableInstanceLODManagement();
 
-	virtual void UpdateInstanceDistsAndLODs() override;
+	virtual void UpdateInstanceDistsAndLODs(FMutableInstanceUpdateMap& InOutRequestedUpdates) override;
 
 	virtual int32 GetNumGeneratedInstancesLimitFullLODs() const override;
 	virtual int32 GetNumGeneratedInstancesLimitLOD1() const override;
diff --git a/Engine/Plugins/Experimental/Mutable/Source/CustomizableObject/Classes/MuCO/CustomizableObjectInstance.h b/Engine/Plugins/Experimental/Mutable/Source/CustomizableObject/Classes/MuCO/CustomizableObjectInstance.h
--- a/Engine/Plugins/Experimental/Mutable/Source/CustomizableObject/Classes/MuCO/CustomizableObjectInstance.h
+++ b/Engine/Plugins/Experimental/Mutable/Source/CustomizableObject/Classes/MuCO/CustomizableObjectInstance.h
@@ -51,6 +51,10 @@ namespace ESkeletalMeshState
 };
 
 
+// Priority for the mutable update queue, Low is the normal distance-based priority, High is normally used for discards and Mid for LOD downgrades
+enum class EQueuePriorityType : uint8 { High, Med, Med_Low, Low };
+
+
 /** Result of all the checks just before beginning an update. */
 enum class EUpdateRequired : uint8
 {
@@ -121,6 +125,7 @@ class CUSTOMIZABLEOBJECT_API UCustomizableObjectInstance : public UObject
 	GENERATED_BODY()
 
 	friend UCustomizableInstancePrivateData;
+	friend FMutableUpdateCandidate;
 
 public:
 	UCustomizableObjectInstance();
@@ -285,7 +290,7 @@ private:
 	 * Once the update reaches this function, the update has been considered started and must complete all the update flow.
 	 * Starting at this function, all Update code paths must end up in FinishUpdateGlobal!
 	 *
-	 * @param bIsCloseDistTick true iff called from the tick.
+	 * @param bIsCloseDistTick true if and only if called from the tick.
 	 */
 	void DoUpdateSkeletalMesh(bool bIsCloseDistTick, bool bOnlyUpdateIfNotGenerated, bool bIgnoreCloseDist, bool bForceHighPriority, const EUpdateRequired* OptionalUpdateRequired, FInstanceUpdateDelegate* UpdateCallback);
 
@@ -307,6 +312,9 @@ public:
 	// Releases all the mutable resources this instance holds, should only be called when it is not going to be used any more.
 	void ReleaseMutableResources(bool bCalledFromBeginDestroy);
 
+	/** Returns the priority an update issued by DoUpdateSkeletalMesh would get in the current instance configuration and state */
+	EQueuePriorityType GetUpdatePriority(bool bIsCloseDistTick, bool bOnlyUpdateIfNotGenerated, bool bIgnoreCloseDist, bool bForceHighPriority) const;
+
 	// Returns de description texture (ex: color bar) for this parameter and DescIndex
 	// This will only be valid if bBuildParameterDecorations was set to true before the last update.
 	UFUNCTION(BlueprintCallable, Category = CustomizableObjectInstance, meta = (DeprecatedFunction, DeprecationMessage = "Parameter decorations have been removed. This method will be removed in future versions."))
@@ -632,7 +640,11 @@ public:
 	int32 GetNumLODsAvailable() const;
 
 	UE_DEPRECATED(5.2, "Use SetRequestedLODs instead.")
-	void SetMinMaxLODToLoad(int32 NewMinLOD = 0, int32 NewMaxLOD = INT32_MAX, bool LimitLODUpgrades = true);
+	void SetMinMaxLODToLoad(FMutableInstanceUpdateMap& InOutRequestedUpdates, int32 NewMinLOD = 0, int32 NewMaxLOD = INT32_MAX, bool bLimitLODUpgrades = true);
+
+	// Remove the following method after CL 25973936 is propagated to all streams
+	UE_DEPRECATED(5.2, "Use SetRequestedLODs instead.")
+	void SetMinMaxLODToLoad(int32 NewMinLOD = 0, int32 NewMaxLOD = INT32_MAX, bool bLimitLODUpgrades = true);
 
 	/** Return the Min LOD this Instance is using (from the beginning of an update. If an update fails this value will be incorrect). */
 	int32 GetCurrentMinLOD() const;
@@ -647,7 +659,11 @@ public:
 	 * Requires mutable.EnableOnlyGenerateRequestedLODs and CurrentInstanceLODManagement->IsOnlyGenerateRequestedLODLevelsEnabled() to be true.
 	 * @param InMinLOD - MinLOD to generate.
 	 * @param InMaxLOD - MaxLOD to generate.
-	 * @param InRequestedLODsPerComponent - Array with bitmasks of requested LODs per component with range from [0 .. CO->GetComponentCount()]. */
+	 * @param InRequestedLODsPerComponent - Array with bitmasks of requested LODs per component with range from [0 .. CO->GetComponentCount()].
+	 * @param InOutRequestedUpdates - Map from Instance to Update data that stores a request for the Instance to be updated, which will be either processed or discarded by priority (to be rerequested the next tick) */
+	void SetRequestedLODs(int32 InMinLOD, int32 InMaxLOD, const TArray<uint16>& InRequestedLODsPerComponent, FMutableInstanceUpdateMap& InOutRequestedUpdates);
+
+	// Remove the following method after CL 25973936 is propagated to all streams
 	void SetRequestedLODs(int32 InMinLOD, int32 InMaxLOD, const TArray<uint16>& InRequestedLODsPerComponent);
 
 	const TArray<uint16>& GetRequestedLODsPerComponent() const;
@@ -771,6 +787,55 @@ private:
 	bool bBuildParameterRelevancy_DEPRECATED = false;
 };
 
+
+class FMutableUpdateCandidate
+{
+public:
+	// The Instance to possibly update
+	UCustomizableObjectInstance* CustomizableObjectInstance;
+	EQueuePriorityType Priority = EQueuePriorityType::Med;
+
+	// These are the LODs that would be applied if this candidate is chosen
+	int32 MinLOD = 0;
+	int32 MaxLOD = INT32_MAX;
+
+	/** Array of RequestedLODs per component to generate if this candidate is chosen */
+	TArray<uint16> RequestedLODLevels;
+
+	FMutableUpdateCandidate(UCustomizableObjectInstance* InCustomizableObjectInstance) : CustomizableObjectInstance(InCustomizableObjectInstance) {}
+
+	FMutableUpdateCandidate(const UCustomizableObjectInstance* InCustomizableObjectInstance, const int32 InMinLOD, const int32 InMaxLOD,
+		const TArray<uint16>& InRequestedLODLevels) :
+		CustomizableObjectInstance(const_cast<UCustomizableObjectInstance*>(InCustomizableObjectInstance)), MinLOD(InMinLOD), MaxLOD(InMaxLOD),
+		RequestedLODLevels(InRequestedLODLevels) {}
+
+	bool HasBeenIssued()
+	{
+		return bHasBeenIssued;
+	}
+
+	void Issue()
+	{
+		bHasBeenIssued = true;
+	}
+
+	void ApplyLODUpdateParamsToInstance()
+	{
+		CustomizableObjectInstance->Descriptor.MinLOD = MinLOD;
+		CustomizableObjectInstance->Descriptor.MaxLOD = MaxLOD;
+
+		CustomizableObjectInstance->Descriptor.RequestedLODLevels = RequestedLODLevels;
+
+		CustomizableObjectInstance->UpdateDescriptorRuntimeHash.UpdateMinMaxLOD(MinLOD, MaxLOD);
+		CustomizableObjectInstance->UpdateDescriptorRuntimeHash.UpdateRequestedLODs(CustomizableObjectInstance->Descriptor.RequestedLODLevels);
+	}
+
+private:
+	/** If true it means that DoUpdateSkeletalMesh has decided this update should be performed, if false it should be ignored. Just used for consistency checks */
+	bool bHasBeenIssued = false;
+};
+
+
 #if WITH_EDITOR
 CUSTOMIZABLEOBJECT_API void CopyTextureProperties(UTexture2D* Texture, const UTexture2D* SourceTexture);
 #endif	
diff --git a/Engine/Plugins/Experimental/Mutable/Source/CustomizableObject/Classes/MuCO/CustomizableObjectSystem.h b/Engine/Plugins/Experimental/Mutable/Source/CustomizableObject/Classes/MuCO/CustomizableObjectSystem.h
--- a/Engine/Plugins/Experimental/Mutable/Source/CustomizableObject/Classes/MuCO/CustomizableObjectSystem.h
+++ b/Engine/Plugins/Experimental/Mutable/Source/CustomizableObject/Classes/MuCO/CustomizableObjectSystem.h
@@ -185,7 +185,8 @@ public:
 
 
 /** End a Customizable Object Instance Update. All code paths of an update have to end here. */
-void FinishUpdateGlobal(UCustomizableObjectInstance* Instance, EUpdateResult UpdateResult, FInstanceUpdateDelegate* UpdateCallback, const mu::ParametersPtrConst Parameters = nullptr, const FDescriptorRuntimeHash InUpdatedHash = FDescriptorRuntimeHash());
+void FinishUpdateGlobal(UCustomizableObjectInstance* Instance, EUpdateResult UpdateResult, FInstanceUpdateDelegate* UpdateCallback, const FCustomizableObjectInstanceDescriptor* InstanceDescriptor = nullptr, const FDescriptorRuntimeHash InUpdatedHash = FDescriptorRuntimeHash());
+void FinishUpdateGlobal(UCustomizableObjectInstance* Instance, EUpdateResult UpdateResult, FInstanceUpdateDelegate* UpdateCallback, mu::ParametersPtrConst Parameters, const FDescriptorRuntimeHash InUpdatedHash = FDescriptorRuntimeHash());
 
 
 UCLASS(Blueprintable, BlueprintType)
@@ -395,6 +396,9 @@ private:
 	// If there is an on-going operation, advance it.
 	void AdvanceCurrentOperation();
 
+	void DiscardInstances();
+	void ReleaseInstanceIDs();
+
 	// TODO: Can we move this to the editor module?
 #if WITH_EDITOR
 
diff --git a/Engine/Plugins/Experimental/Mutable/Source/CustomizableObject/Private/MuCO/CustomizableInstanceLODManagement.cpp b/Engine/Plugins/Experimental/Mutable/Source/CustomizableObject/Private/MuCO/CustomizableInstanceLODManagement.cpp
--- a/Engine/Plugins/Experimental/Mutable/Source/CustomizableObject/Private/MuCO/CustomizableInstanceLODManagement.cpp
+++ b/Engine/Plugins/Experimental/Mutable/Source/CustomizableObject/Private/MuCO/CustomizableInstanceLODManagement.cpp
@@ -94,7 +94,7 @@ void UpdateCameraToInstancesDistance(const FVector CameraPosition)
 #endif
 
 
-void UCustomizableInstanceLODManagement::UpdateInstanceDistsAndLODs()
+void UCustomizableInstanceLODManagement::UpdateInstanceDistsAndLODs(FMutableInstanceUpdateMap& InOutRequestedUpdates)
 {
 	int32 NumGeneratedInstancesLimitLODs = GetNumGeneratedInstancesLimitFullLODs();
 	if (NumGeneratedInstancesLimitLODs > 0 || (IsOnlyUpdateCloseCustomizableObjectsEnabled() && IsOnlyGenerateRequestedLODLevelsEnabled()))
@@ -213,12 +213,12 @@ void UCustomizableInstanceLODManagement::UpdateInstanceDistsAndLODs()
 				if (!bAlreadyReachedFixedLOD && SortedInstances[i]->GetMinSquareDistToPlayer() < DistanceForFixedLODSquared)
 				{
 					RequestedLODs.Init(MAX_uint16, SortedInstances[i]->GetNumComponents());
-					SortedInstances[i]->SetRequestedLODs(0, MAX_int32, RequestedLODs);
+					SortedInstances[i]->SetRequestedLODs(0, MAX_int32, RequestedLODs, InOutRequestedUpdates);
 				}
 				else
 				{
 					RequestedLODs.Init(MAX_uint16, SortedInstances[i]->GetNumComponents());
-					SortedInstances[i]->SetRequestedLODs(2, 2, RequestedLODs);
+					SortedInstances[i]->SetRequestedLODs(2, 2, RequestedLODs, InOutRequestedUpdates);
 					bAlreadyReachedFixedLOD = true;
 				}
 			}
@@ -228,12 +228,12 @@ void UCustomizableInstanceLODManagement::UpdateInstanceDistsAndLODs()
 				if (!bAlreadyReachedFixedLOD && SortedInstances[i]->GetMinSquareDistToPlayer() < DistanceForFixedLODSquared)
 				{
 					RequestedLODs.Init(MAX_uint16, SortedInstances[i]->GetNumComponents());
-					SortedInstances[i]->SetRequestedLODs(1, MAX_int32, RequestedLODs);
+					SortedInstances[i]->SetRequestedLODs(1, MAX_int32, RequestedLODs, InOutRequestedUpdates);
 				}
 				else
 				{
 					RequestedLODs.Init(MAX_uint16, SortedInstances[i]->GetNumComponents());
-					SortedInstances[i]->SetRequestedLODs(2, MAX_int32, RequestedLODs);
+					SortedInstances[i]->SetRequestedLODs(2, MAX_int32, RequestedLODs, InOutRequestedUpdates);
 					bAlreadyReachedFixedLOD = true;
 				}
 			}
@@ -241,7 +241,7 @@ void UCustomizableInstanceLODManagement::UpdateInstanceDistsAndLODs()
 			for (int32 i = NumGeneratedInstancesLimitLODs + NumGeneratedInstancesLimitLOD1; i < NumGeneratedInstancesLimitLODs + NumGeneratedInstancesLimitLOD1 + NumGeneratedInstancesLimitLOD2 && i < SortedInstances.Num(); ++i)
 			{
 				RequestedLODs.Init(MAX_uint16, SortedInstances[i]->GetNumComponents());
-				SortedInstances[i]->SetRequestedLODs(2, MAX_int32, RequestedLODs);
+				SortedInstances[i]->SetRequestedLODs(2, MAX_int32, RequestedLODs, InOutRequestedUpdates);
 			}
 
 			for (int32 i = NumGeneratedInstancesLimitLODs + NumGeneratedInstancesLimitLOD1 + NumGeneratedInstancesLimitLOD2; i < SortedInstances.Num(); ++i)
@@ -255,7 +255,7 @@ void UCustomizableInstanceLODManagement::UpdateInstanceDistsAndLODs()
 			for (int32 i = 0; i < SortedInstances.Num(); ++i)
 			{
 				RequestedLODs.Init(MAX_uint16, SortedInstances[i]->GetNumComponents());
-				SortedInstances[i]->SetRequestedLODs(2, MAX_int32, RequestedLODs);
+				SortedInstances[i]->SetRequestedLODs(2, MAX_int32, RequestedLODs, InOutRequestedUpdates);
 			}
 		}
 	}
@@ -375,7 +375,7 @@ void UCustomizableInstanceLODManagement::UpdateInstanceDistsAndLODs()
 					It.Value.MaxLOD = It.Value.MinLOD;
 				}
 
-				It.Key->SetRequestedLODs(It.Value.MinLOD, It.Value.MaxLOD, It.Value.RequestedLODsPerComponent);
+				It.Key->SetRequestedLODs(It.Value.MinLOD, It.Value.MaxLOD, It.Value.RequestedLODsPerComponent, InOutRequestedUpdates);
 			}
 		}
 	}
@@ -428,4 +428,3 @@ bool UCustomizableInstanceLODManagement::IsOnlyUpdateCloseCustomizableObjectsEna
 {
 	return CVarOnlyUpdateCloseCustomizableObjects.GetValueOnGameThread();
 }
-
diff --git a/Engine/Plugins/Experimental/Mutable/Source/CustomizableObject/Private/MuCO/CustomizableInstancePrivateData.h b/Engine/Plugins/Experimental/Mutable/Source/CustomizableObject/Private/MuCO/CustomizableInstancePrivateData.h
--- a/Engine/Plugins/Experimental/Mutable/Source/CustomizableObject/Private/MuCO/CustomizableInstancePrivateData.h
+++ b/Engine/Plugins/Experimental/Mutable/Source/CustomizableObject/Private/MuCO/CustomizableInstancePrivateData.h
@@ -226,8 +226,9 @@ public:
 	/** See FCustomizableObjectInstanceDescriptor::ReloadParameters(...). */
 	void ReloadParameters(UCustomizableObjectInstance* Public, bool bInvalidatePreviousData = false);
 
-	void TickUpdateCloseCustomizableObjects(UCustomizableObjectInstance& Public);
-	void UpdateInstanceIfNotGenerated(UCustomizableObjectInstance& Public);
+	/** If bOnlyUpdatePrioForUpdates is true, Instances that need an update won't be modified. Just their entry in InOutRequestedUpdates will get its priority updated  */
+	void TickUpdateCloseCustomizableObjects(UCustomizableObjectInstance& Public, const bool bOnlyUpdatePrioForUpdates, FMutableInstanceUpdateMap& InOutRequestedUpdates);
+	void UpdateInstanceIfNotGenerated(UCustomizableObjectInstance& Public, FMutableInstanceUpdateMap& InOutRequestedUpdates);
 
 	// Returns true if success (?)
 	bool UpdateSkeletalMesh_PostBeginUpdate0(UCustomizableObjectInstance* Public, const TSharedPtr<FMutableOperationData>& OperationData);
diff --git a/Engine/Plugins/Experimental/Mutable/Source/CustomizableObject/Private/MuCO/CustomizableObjectInstance.cpp b/Engine/Plugins/Experimental/Mutable/Source/CustomizableObject/Private/MuCO/CustomizableObjectInstance.cpp
--- a/Engine/Plugins/Experimental/Mutable/Source/CustomizableObject/Private/MuCO/CustomizableObjectInstance.cpp
+++ b/Engine/Plugins/Experimental/Mutable/Source/CustomizableObject/Private/MuCO/CustomizableObjectInstance.cpp
@@ -461,7 +461,6 @@ void UCustomizableObjectInstance::PostLoad()
 	Super::PostLoad();
 
 	Descriptor.ReloadParameters();
-	SetRequestedLODs(Descriptor.MinLOD, Descriptor.MaxLOD, Descriptor.RequestedLODLevels);
 }
 
 
@@ -567,8 +566,6 @@ mu::ParametersPtr UCustomizableInstancePrivateData::GetParameters(UCustomizableO
 	{	
 		return nullptr;
 	}
-
-	InstanceUpdateFlags(*Public); // TODO Move somewhere else.
 		
 	return Public->GetDescriptor().GetParameters();
 }
@@ -611,7 +608,7 @@ void UCustomizableObjectInstance::SetObject(UCustomizableObject* InObject)
 {
 	Descriptor.SetCustomizableObject(*InObject);
 	PrivateData->ReloadParameters(this, true);
-	SetRequestedLODs(Descriptor.MinLOD, Descriptor.MaxLOD, Descriptor.RequestedLODLevels);
+	//SetRequestedLODs(Descriptor.MinLOD, Descriptor.MaxLOD, Descriptor.RequestedLODLevels);
 }
 
 
@@ -826,6 +823,43 @@ EUpdateRequired UCustomizableObjectInstance::IsUpdateRequired(bool bIsCloseDistT
 }
 
 
+EQueuePriorityType UCustomizableObjectInstance::GetUpdatePriority(bool bIsCloseDistTick, bool bOnlyUpdateIfNotGenerated, bool bIgnoreCloseDist, bool bForceHighPriority) const
+{
+	const bool bIsGenerated = GetPrivate()->HasCOInstanceFlags(Generated);
+	const bool bShouldUpdateLODs = GetPrivate()->HasCOInstanceFlags(PendingLODsUpdate);
+	const bool bIsDowngradeLODUpdate = GetPrivate()->HasCOInstanceFlags(PendingLODsDowngrade);
+	const bool bIsPlayerOrNearIt = GetPrivate()->HasCOInstanceFlags(UsedByPlayerOrNearIt);
+
+	EQueuePriorityType Priority = EQueuePriorityType::Low;
+	if (bForceHighPriority)
+	{
+		Priority = EQueuePriorityType::High;
+	}
+	else if (!bIsGenerated || !HasAnySkeletalMesh())
+	{
+		Priority = EQueuePriorityType::Med;
+	}
+	else if (bShouldUpdateLODs && bIsDowngradeLODUpdate)
+	{
+		Priority = EQueuePriorityType::Med_Low;
+	}
+	else if (bIsPlayerOrNearIt && bShouldUpdateLODs && !bIsDowngradeLODUpdate)
+	{
+		Priority = EQueuePriorityType::High;
+	}
+	else if (bShouldUpdateLODs && !bIsDowngradeLODUpdate)
+	{
+		Priority = EQueuePriorityType::Med;
+	}
+	else if (bIsPlayerOrNearIt)
+	{
+		Priority = EQueuePriorityType::High;
+	}
+
+	return Priority;
+}
+
+
 void UCustomizableObjectInstance::DoUpdateSkeletalMesh(bool bIsCloseDistTick, bool bOnlyUpdateIfNotGenerated, bool bIgnoreCloseDist, bool bForceHighPriority, const EUpdateRequired* OptionalUpdateRequired, FInstanceUpdateDelegate* UpdateCallback)
 {
 	MUTABLE_CPUPROFILER_SCOPE(UCustomizableInstancePrivateData::DoUpdateSkeletalMesh);
@@ -844,40 +878,13 @@ void UCustomizableObjectInstance::DoUpdateSkeletalMesh(bool bIsCloseDistTick, bo
 	}		
 	case EUpdateRequired::Update:
 	{
-		const bool bIsGenerated = GetPrivate()->HasCOInstanceFlags(Generated);
-		const bool bShouldUpdateLODs = GetPrivate()->HasCOInstanceFlags(PendingLODsUpdate);	
-		const bool bIsDowngradeLODUpdate = GetPrivate()->HasCOInstanceFlags(PendingLODsDowngrade);
-		const bool bIsPlayerOrNearIt = GetPrivate()->HasCOInstanceFlags(UsedByPlayerOrNearIt);
-
-		FMutableQueueElem::EQueuePriorityType Priority = FMutableQueueElem::EQueuePriorityType::Low;
-		if (bForceHighPriority)
-		{
-			Priority = FMutableQueueElem::EQueuePriorityType::High;
-		}
-		else if (!bIsGenerated || !HasAnySkeletalMesh())
-		{
-			Priority = FMutableQueueElem::EQueuePriorityType::Med;
-		}
-		else if (bShouldUpdateLODs && bIsDowngradeLODUpdate)
-		{
-			Priority = FMutableQueueElem::EQueuePriorityType::Med_Low;
-		}
-		else if (bIsPlayerOrNearIt && bShouldUpdateLODs && !bIsDowngradeLODUpdate)
-		{
-			Priority = FMutableQueueElem::EQueuePriorityType::High;
-		}
-		else if (bShouldUpdateLODs && !bIsDowngradeLODUpdate)
-		{
-			Priority = FMutableQueueElem::EQueuePriorityType::Med;
-		}
-		else if (bIsPlayerOrNearIt)
-		{
-			Priority = FMutableQueueElem::EQueuePriorityType::High;
-		}
+		EQueuePriorityType Priority = GetUpdatePriority(bIsCloseDistTick, bOnlyUpdateIfNotGenerated, bIgnoreCloseDist, bForceHighPriority);
+		const bool bShouldUpdateLODs = GetPrivate()->HasCOInstanceFlags(PendingLODsUpdate);
 
 		{
 			const uint32 InstanceId = GetUniqueID();
 			const float Distance = FMath::Sqrt(GetPrivate()->LastMinSquareDistFromComponentToPlayer);
+			const bool bIsPlayerOrNearIt = GetPrivate()->HasCOInstanceFlags(UsedByPlayerOrNearIt);
 			UE_LOG(LogMutable, Verbose, TEXT("Started UpdateSkeletalMesh Async. of Instance %d with priority %d at dist %f bIsPlayerOrNearIt=%d, frame=%d"), InstanceId, static_cast<int32>(Priority), Distance, bIsPlayerOrNearIt, GFrameNumber);				
 		}
 
@@ -901,7 +908,7 @@ void UCustomizableObjectInstance::DoUpdateSkeletalMesh(bool bIsCloseDistTick, bo
 		GetPrivate()->SetCOInstanceFlags(Generated); // Will be done in UpdateSkeletalMesh_PostBeginUpdate
 
 		// Do not do work after calling InitUpdateSkeletalMesh. This function can optimize an update and fully complete it before even exiting its scope.
-		System->GetPrivate()->InitUpdateSkeletalMesh(*this, Priority, UpdateCallback);
+		System->GetPrivate()->InitUpdateSkeletalMesh(*this, Priority, bIsCloseDistTick, UpdateCallback);
 		break;
 	}
 
@@ -919,23 +926,72 @@ void UCustomizableObjectInstance::DoUpdateSkeletalMesh(bool bIsCloseDistTick, bo
 }
 
 
-void UCustomizableInstancePrivateData::TickUpdateCloseCustomizableObjects(UCustomizableObjectInstance& Public)
+void UCustomizableInstancePrivateData::TickUpdateCloseCustomizableObjects(UCustomizableObjectInstance& Public, const bool bOnlyUpdatePrioForUpdates, FMutableInstanceUpdateMap& InOutRequestedUpdates)
 {
 	const EUpdateRequired UpdateRequired = Public.IsUpdateRequired(true, true, false);
 	if (UpdateRequired != EUpdateRequired::NoUpdate) // Since this is done in the tick, avoid starting an update that we know for sure that would not be performed. Once started it has some performance implications that we want to avoid.
 	{
-		Public.DoUpdateSkeletalMesh(true, true, false, false, &UpdateRequired, nullptr);
+		if (!bOnlyUpdatePrioForUpdates || UpdateRequired == EUpdateRequired::Discard)
+		{
+			Public.DoUpdateSkeletalMesh(true, true, false, false, &UpdateRequired, nullptr);
+			InOutRequestedUpdates.Remove(&Public);
+		}
+		else if (UpdateRequired == EUpdateRequired::Update)
+		{
+			EQueuePriorityType Priority = Public.GetUpdatePriority(true, true, false, false);
+
+			FMutableUpdateCandidate* UpdateCandidate = InOutRequestedUpdates.Find(&Public);
+
+			if (UpdateCandidate)
+			{
+				ensure(HasCOInstanceFlags(PendingLODsUpdate | PendingLODsDowngrade));
+
+				UpdateCandidate->Priority = Priority;
+				UpdateCandidate->Issue();
+			}
+			else
+			{
+				FMutableUpdateCandidate Candidate(&Public);
+				Candidate.Priority = Priority;
+				Candidate.Issue();
+				InOutRequestedUpdates.Add(&Public, Candidate);
+			}
+		}
+		else
+		{
+			check(false);
+		}
+	}
+	else
+	{
+		InOutRequestedUpdates.Remove(&Public);
 	}
 
 	ClearCOInstanceFlags(PendingLODsUpdate | PendingLODsDowngrade);
 }
 
 
-void UCustomizableInstancePrivateData::UpdateInstanceIfNotGenerated(UCustomizableObjectInstance& Public)
+void UCustomizableInstancePrivateData::UpdateInstanceIfNotGenerated(UCustomizableObjectInstance& Public, FMutableInstanceUpdateMap& InOutRequestedUpdates)
 {
 	if (!HasCOInstanceFlags(Generated))
 	{
 		Public.DoUpdateSkeletalMesh(false, true, false, false, nullptr, nullptr);
+
+		EQueuePriorityType Priority = Public.GetUpdatePriority(true, true, false, false);
+		FMutableUpdateCandidate* UpdateCandidate = InOutRequestedUpdates.Find(&Public);
+
+		if (UpdateCandidate)
+		{
+			UpdateCandidate->Priority = Priority;
+			UpdateCandidate->Issue();
+		}
+		else
+		{
+			FMutableUpdateCandidate Candidate(&Public);
+			Candidate.Priority = Priority;
+			Candidate.Issue();
+			InOutRequestedUpdates.Add(&Public, Candidate);
+		}
 	}
 }
 
@@ -6316,6 +6372,12 @@ void UCustomizableObjectInstance::SetMinSquareDistToPlayer(float NewValue)
 }
 
 
+void UCustomizableObjectInstance::SetMinMaxLODToLoad(FMutableInstanceUpdateMap& InOutRequestedUpdates, int32 NewMinLOD, int32 NewMaxLOD, bool bLimitLODUpgrades)
+{
+	SetRequestedLODs(NewMinLOD, NewMaxLOD, Descriptor.RequestedLODLevels, InOutRequestedUpdates);
+}
+
+
 void UCustomizableObjectInstance::SetMinMaxLODToLoad(int32 NewMinLOD, int32 NewMaxLOD, bool bLimitLODUpgrades)
 {
 	SetRequestedLODs(NewMinLOD, NewMaxLOD, Descriptor.RequestedLODLevels);
@@ -6365,7 +6427,9 @@ FAutoConsoleVariableRef CVarMutableIgnoreMinMaxLOD(
 	TEXT("The limits on the number of LODs to generate will be ignored."));
 #endif
 
-void UCustomizableObjectInstance::SetRequestedLODs(int32 InMinLOD, int32 InMaxLOD, const TArray<uint16>& InRequestedLODsPerComponent)
+
+void UCustomizableObjectInstance::SetRequestedLODs(int32 InMinLOD, int32 InMaxLOD, const TArray<uint16>& InRequestedLODsPerComponent, 
+	                                               FMutableInstanceUpdateMap& InOutRequestedUpdates)
 {
 	check(PrivateData);
 	if (!PrivateData->HasCOInstanceFlags(LODsStreamingEnabled))
@@ -6382,6 +6446,8 @@ void UCustomizableObjectInstance::SetRequestedLODs(int32 InMinLOD, int32 InMaxLO
 	}
 #endif
 
+	FMutableUpdateCandidate MutableUpdateCandidate(this);
+
 	// Clamp Min LOD
 	InMinLOD = FMath::Min(FMath::Max(InMinLOD, static_cast<int32>(PrivateData->FirstLODAvailable)), PrivateData->FirstLODAvailable + PrivateData->NumMaxLODsToStream);
 
@@ -6394,8 +6460,9 @@ void UCustomizableObjectInstance::SetRequestedLODs(int32 InMinLOD, int32 InMaxLO
 	PrivateData->SetCOInstanceFlags(bIsDowngradeLODUpdate ? PendingLODsDowngrade : ECONone);
 
 	// Save the new LODs
-	Descriptor.MinLOD = InMinLOD;
-	Descriptor.MaxLOD = InMaxLOD;
+	MutableUpdateCandidate.MinLOD = InMinLOD;
+	MutableUpdateCandidate.MaxLOD = InMaxLOD;
+	MutableUpdateCandidate.RequestedLODLevels = Descriptor.GetRequestedLODLevels();
 
 	bool bUpdateRequestedLODs = false;
 	if (UCustomizableObjectSystem::GetInstance()->IsOnlyGenerateRequestedLODsEnabled())
@@ -6415,9 +6482,9 @@ void UCustomizableObjectInstance::SetRequestedLODs(int32 InMinLOD, int32 InMaxLO
 			for (int32 ComponentIndex = 0; ComponentIndex < ComponentCount; ++ComponentIndex)
 			{
 				// Generate at least the MaxLOD
-				int32 RequestedLODs = 1 << Descriptor.MaxLOD;
+				int32 RequestedLODs = 1 << MutableUpdateCandidate.MaxLOD;
 
-				for (int32 LODIndex = 0; LODIndex < Descriptor.MaxLOD; ++LODIndex)
+				for (int32 LODIndex = 0; LODIndex < MutableUpdateCandidate.MaxLOD; ++LODIndex)
 				{
 					// Add previously generated LODs if the descriptor have not changed
 					if (!bIgnoreGeneratedLODs && GeneratedLODsPerComponent[ComponentIndex] & (1 << LODIndex))
@@ -6428,46 +6495,44 @@ void UCustomizableObjectInstance::SetRequestedLODs(int32 InMinLOD, int32 InMaxLO
 					// Add new requested LODs that fall within the range
 					if (InRequestedLODsPerComponent[ComponentIndex] & (1 << LODIndex))
 					{
-						if (LODIndex >= Descriptor.MinLOD)
+						if (LODIndex >= MutableUpdateCandidate.MinLOD)
 						{
 							RequestedLODs |= (1 << LODIndex);
 						}
 						else
 						{
 							// RequestedLOD is out of range, generate the MinLOD instead;
-							RequestedLODs |= (1 << Descriptor.MinLOD);
+							RequestedLODs |= (1 << MutableUpdateCandidate.MinLOD);
 						}
 					}
 				}
 
-				bUpdateRequestedLODs |= RequestedLODs != Descriptor.RequestedLODLevels[ComponentIndex];
+				bool RequestedLODsDifferent = !Descriptor.RequestedLODLevels.IsValidIndex(ComponentIndex) || 
+					                          RequestedLODs != Descriptor.RequestedLODLevels[ComponentIndex];
+				bUpdateRequestedLODs |= RequestedLODsDifferent;
 
 				// Save new RequestedLODs
-				Descriptor.RequestedLODLevels[ComponentIndex] = RequestedLODs;
+				MutableUpdateCandidate.RequestedLODLevels[ComponentIndex] = RequestedLODs;
 			}
 		}
 	}
 
 	if (bMinMaxLODChanged || bUpdateRequestedLODs)
 	{
-		UpdateDescriptorRuntimeHash.UpdateMinMaxLOD(InMinLOD, InMaxLOD);
-		UpdateDescriptorRuntimeHash.UpdateRequestedLODs(Descriptor.RequestedLODLevels);
+		// TODO: Remove this flag as it will become redundant with the new InOutRequestedUpdates system
+		PrivateData->SetCOInstanceFlags(PendingLODsUpdate);
 
-		// Update the Queued operation data or request a new update
-		const FMutableQueueElem* QueueElem = UCustomizableObjectSystem::GetInstance()->GetPrivate()->MutableOperationQueue.Get(this);
-		if (QueueElem && QueueElem->Operation && QueueElem->Operation->Type == FMutableOperation::EOperationType::Update)
-		{
-			QueueElem->Operation->InstanceDescriptorRuntimeHash.UpdateMinMaxLOD(InMinLOD, InMaxLOD);
-			QueueElem->Operation->InstanceDescriptorRuntimeHash.UpdateRequestedLODs(Descriptor.RequestedLODLevels);
-		}
-		else
-		{
-			PrivateData->SetCOInstanceFlags(PendingLODsUpdate);
-		}
+		InOutRequestedUpdates.Add(this, MutableUpdateCandidate);
 	}
 }
 
 
+void UCustomizableObjectInstance::SetRequestedLODs(int32 InMinLOD, int32 InMaxLOD, const TArray<uint16>& InRequestedLODsPerComponent)
+{
+	// Remove this method after CL 25973936 is propagated to all streams
+}
+
+
 const TArray<uint16>& UCustomizableObjectInstance::GetRequestedLODsPerComponent() const
 {
 	return Descriptor.RequestedLODLevels;
diff --git a/Engine/Plugins/Experimental/Mutable/Source/CustomizableObject/Private/MuCO/CustomizableObjectSystem.cpp b/Engine/Plugins/Experimental/Mutable/Source/CustomizableObject/Private/MuCO/CustomizableObjectSystem.cpp
--- a/Engine/Plugins/Experimental/Mutable/Source/CustomizableObject/Private/MuCO/CustomizableObjectSystem.cpp
+++ b/Engine/Plugins/Experimental/Mutable/Source/CustomizableObject/Private/MuCO/CustomizableObjectSystem.cpp
@@ -102,98 +102,90 @@ TAutoConsoleVariable<bool> CVarEnableImageCache(
 	ECVF_Scalability);
 
 
-bool FMutableQueue::IsEmpty() const
+bool FMutablePendingInstanceWork::ArePendingUpdatesEmpty() const
 {
-	return Array.Num() == 0;
+	return PendingInstanceUpdates.Num() == 0;
 }
 
 
-int FMutableQueue::Num() const
+int32 FMutablePendingInstanceWork::Num() const
 {
-	return Array.Num();
+	return PendingInstanceUpdates.Num() + PendingInstanceDiscards.Num() + PendingIDsToRelease.Num() + NumLODUpdatesLastTick;
 }
 
 
-void FMutableQueue::Enqueue(const FMutableQueueElem& TaskToEnqueue)
+void FMutablePendingInstanceWork::SetLODUpdatesLastTick(int32 NumLODUpdates)
 {
-	for (FMutableQueueElem& QueueElem : Array)
+	NumLODUpdatesLastTick = NumLODUpdates;
+}
+
+
+void FMutablePendingInstanceWork::AddUpdate(const FMutablePendingInstanceUpdate& UpdateToAdd)
+{
+	if (FMutablePendingInstanceUpdate* ExistingUpdate = PendingInstanceUpdates.Find(UpdateToAdd.CustomizableObjectInstance))
 	{
-		check(QueueElem.Operation);
-		check(TaskToEnqueue.Operation);
-		if ((QueueElem.Operation->Type == FMutableOperation::EOperationType::Update || QueueElem.Operation->Type == FMutableOperation::EOperationType::Discard) &&
-			QueueElem.Operation->CustomizableObjectInstance == TaskToEnqueue.Operation->CustomizableObjectInstance)
-		{
-			UCustomizableObjectInstance* Instance = QueueElem.Operation->CustomizableObjectInstance.Get();
-			FInstanceUpdateDelegate* Callback = &QueueElem.Operation->UpdateCallback;
-			FinishUpdateGlobal(Instance, EUpdateResult::ErrorReplaced, Callback, QueueElem.Operation->GetParameters());
+		FMutablePendingInstanceUpdate TaskToEnqueue = UpdateToAdd;
+
+		FInstanceUpdateDelegate* Callback = ExistingUpdate->Callback;
+		FinishUpdateGlobal(ExistingUpdate->CustomizableObjectInstance.Get(), EUpdateResult::ErrorReplaced, Callback, &ExistingUpdate->InstanceDescriptor);
+
+		TaskToEnqueue.PriorityType = FMath::Min(ExistingUpdate->PriorityType, UpdateToAdd.PriorityType);
+		TaskToEnqueue.SecondsAtUpdate = FMath::Min(ExistingUpdate->SecondsAtUpdate, UpdateToAdd.SecondsAtUpdate);
 		
-			QueueElem.Operation = TaskToEnqueue.Operation;
-			QueueElem.PriorityType = FMath::Min(QueueElem.PriorityType, TaskToEnqueue.PriorityType);
-			QueueElem.Priority = FMath::Min(QueueElem.Priority, TaskToEnqueue.Priority);
-			
-			return;
-		}
+		PendingInstanceUpdates.Remove(ExistingUpdate->CustomizableObjectInstance);
+		PendingInstanceUpdates.Add(TaskToEnqueue);
+	}
+	else
+	{
+		PendingInstanceUpdates.Add(UpdateToAdd);
 	}
 
-	Array.HeapPush(TaskToEnqueue);
+	if (FMutablePendingInstanceDiscard* ExistingDiscard = PendingInstanceDiscards.Find(UpdateToAdd.CustomizableObjectInstance))
+	{
+		FinishUpdateGlobal(ExistingDiscard->CustomizableObjectInstance.Get(), EUpdateResult::ErrorReplaced, nullptr, nullptr);
+
+		PendingInstanceDiscards.Remove(ExistingDiscard->CustomizableObjectInstance);
+	}
 }
 
 
-void FMutableQueue::Dequeue(FMutableQueueElem* DequeuedTask)
+void FMutablePendingInstanceWork::RemoveUpdate(const TWeakObjectPtr<UCustomizableObjectInstance>& Instance)
 {
-	Array.HeapPop(*DequeuedTask);
+	PendingInstanceUpdates.Remove(Instance);
 }
 
 
-void FMutableQueue::ChangePriorities()
+const FMutablePendingInstanceUpdate* FMutablePendingInstanceWork::GetUpdate(const TWeakObjectPtr<UCustomizableObjectInstance>& Instance) const
 {
-	for (FMutableQueueElem& Elem : Array)
-	{
-		check(Elem.Operation);
-		if (Elem.Operation->CustomizableObjectInstance.IsValid())
-		{
-			if (UCustomizableInstancePrivateData* CustomizableObjectInstancePrivateData = Elem.Operation->CustomizableObjectInstance->GetPrivate())
-			{
-				Elem.Priority = CustomizableObjectInstancePrivateData->MinSquareDistFromComponentToPlayer;
-			}
-		}
-	}
+	return PendingInstanceUpdates.Find(Instance);
 }
 
 
-void FMutableQueue::UpdatePriority(const UCustomizableObjectInstance* Instance)
+void FMutablePendingInstanceWork::AddDiscard(const FMutablePendingInstanceDiscard& TaskToEnqueue)
 {
-	for (FMutableQueueElem& Elem : Array)
+	if (FMutablePendingInstanceUpdate* ExistingUpdate = PendingInstanceUpdates.Find(TaskToEnqueue.CustomizableObjectInstance.Get()))
 	{
-		check(Elem.Operation);
-		if (Elem.Operation->CustomizableObjectInstance == Instance)
-		{
-			check(Instance->GetPrivate() != nullptr);
-			Elem.Priority = Instance->GetPrivate()->MinSquareDistFromComponentToPlayer;
-			break;
-		}
+		FinishUpdateGlobal(ExistingUpdate->CustomizableObjectInstance.Get(), EUpdateResult::ErrorReplaced, ExistingUpdate->Callback, &ExistingUpdate->InstanceDescriptor);
+		PendingInstanceUpdates.Remove(ExistingUpdate->CustomizableObjectInstance);
 	}
+
+	PendingInstanceDiscards.Add(TaskToEnqueue);
+
+	TaskToEnqueue.CustomizableObjectInstance->GetPrivate()->SetCOInstanceFlags(Updating);
 }
 
 
-const FMutableQueueElem* FMutableQueue::Get(const UCustomizableObjectInstance* Instance) const
+void FMutablePendingInstanceWork::AddIDRelease(mu::Instance::ID IDToRelease)
 {
-	for (const FMutableQueueElem& Elem : Array)
-	{
-		check(Elem.Operation);
-		if (Elem.Operation->CustomizableObjectInstance == Instance)
-		{
-			return &Elem;
-		}
-	}
-
-	return nullptr;
+	PendingIDsToRelease.Add(IDToRelease);
 }
 
 
-void FMutableQueue::Sort()
+void FMutablePendingInstanceWork::RemoveAllUpdatesAndDiscardsAndReleases()
 {
-	Array.Heapify();
+	PendingInstanceUpdates.Empty();
+	PendingInstanceDiscards.Empty();
+	PendingIDsToRelease.Empty();
 }
 
 
@@ -226,7 +218,6 @@ FString UCustomizableObjectSystem::GetPluginVersion() const
 }
 
 
-
 void UCustomizableObjectSystem::LogShowData(bool bFullInfo, bool ShowMaterialInfo) const
 {
 	LogInformationUtil::ResetCounters();
@@ -427,11 +418,7 @@ void UCustomizableObjectSystem::BeginDestroy()
 
 		Private->CurrentInstanceBeingUpdated = nullptr;
 
-		while (!Private->MutableOperationQueue.IsEmpty())
-		{
-			FMutableQueueElem AuxOperation;
-			Private->MutableOperationQueue.Dequeue(&AuxOperation);
-		}
+		Private->MutablePendingInstanceWork.RemoveAllUpdatesAndDiscardsAndReleases();
 
 		FCustomizableObjectSystemPrivate::SSystem = nullptr;
 
@@ -505,18 +492,13 @@ void FCustomizableObjectSystemPrivate::UpdateStats()
 		int CountLOD1 = 0;
 		int CountLOD2 = 0;
 		int CountTotal = 0;
-		int CountPending = 0;
+
 		for (TObjectIterator<UCustomizableObjectInstance> CustomizableObjectInstance; CustomizableObjectInstance; ++CustomizableObjectInstance)
 		{
 			if ( IsValidChecked(*CustomizableObjectInstance) && CustomizableObjectInstance->GetPrivate())
 			{
 				++CountTotal;
 
-				if (CustomizableObjectInstance->SkeletalMeshStatus == ESkeletalMeshState::AsyncUpdatePending)
-				{
-					++CountPending;
-				}
-
 				for (int32 ComponentIndex = 0; ComponentIndex < CustomizableObjectInstance->SkeletalMeshes.Num(); ++ComponentIndex)
 				{
 					if (CustomizableObjectInstance->SkeletalMeshes[ComponentIndex] && CustomizableObjectInstance->SkeletalMeshes[ComponentIndex]->GetResourceForRendering())
@@ -542,9 +524,9 @@ void FCustomizableObjectSystemPrivate::UpdateStats()
 
 		NumInstances = CountLOD0 + CountLOD1 + CountLOD2;
 		TotalInstances = CountTotal;
-		NumPendingInstances = CountPending;
+		NumPendingInstances = MutablePendingInstanceWork.Num();
 		//SET_DWORD_STAT(STAT_MutableSkeletalMeshResourceMemory, Size / 1024.f);
-		SET_DWORD_STAT(STAT_MutablePendingInstanceUpdates, MutableOperationQueue.Num());
+		SET_DWORD_STAT(STAT_MutablePendingInstanceUpdates, MutablePendingInstanceWork.Num());
 
 		uint64 Size = 0;
 		uint32 CountAllocated = 0;
@@ -591,7 +573,7 @@ void FCustomizableObjectSystemPrivate::UpdateStats()
 		if (LogBenchmarkUtil::isLoggingActive())
 		{
 			LogBenchmarkUtil::updateStat("customizable_objects", (int32)CountAllocatedSkeletalMesh);
-			LogBenchmarkUtil::updateStat("pending_instance_updates", MutableOperationQueue.Num());
+			LogBenchmarkUtil::updateStat("pending_instance_updates", MutablePendingInstanceWork.Num());
 			LogBenchmarkUtil::updateStat("allocated_textures", (int32)CountAllocated);
 			LogBenchmarkUtil::updateStat("texture_resource_memory", (long double)Size / 1048576.0L);
 			LogBenchmarkUtil::updateStat("texture_generated_memory", (long double)SizeGenerated / 1048576.0L);
@@ -602,7 +584,7 @@ void FCustomizableObjectSystemPrivate::UpdateStats()
 		SET_DWORD_STAT(STAT_MutableNumInstancesLOD1, CountLOD1);
 		SET_DWORD_STAT(STAT_MutableNumInstancesLOD2, CountLOD2);
 		SET_DWORD_STAT(STAT_MutableNumAllocatedSkeletalMeshes, CountAllocatedSkeletalMesh);
-		SET_DWORD_STAT(STAT_MutablePendingInstanceUpdates, MutableOperationQueue.Num());
+		SET_DWORD_STAT(STAT_MutablePendingInstanceUpdates, MutablePendingInstanceWork.Num());
 		SET_DWORD_STAT(STAT_MutableNumAllocatedTextures, CountAllocated);
 		SET_DWORD_STAT(STAT_MutableTextureResourceMemory, Size / 1024.f);
 		SET_DWORD_STAT(STAT_MutableTextureGeneratedMemory, SizeGenerated / 1024.f);
@@ -796,6 +778,13 @@ void FinishUpdateGlobal(UCustomizableObjectInstance* Instance, EUpdateResult Upd
 }
 
 
+void FinishUpdateGlobal(UCustomizableObjectInstance* Instance, EUpdateResult UpdateResult, FInstanceUpdateDelegate* UpdateCallback, const FCustomizableObjectInstanceDescriptor* InstanceDescriptor, const FDescriptorRuntimeHash InUpdatedHash)
+{
+	FinishUpdateGlobal(Instance, UpdateResult, UpdateCallback, InstanceDescriptor ? InstanceDescriptor->GetParameters() : nullptr, 
+		InUpdatedHash);
+}
+
+
 /** Update the given Instance Skeletal Meshes and call its callbacks. */
 void UpdateSkeletalMesh(UCustomizableObjectInstance& CustomizableObjectInstance, const FDescriptorRuntimeHash& UpdatedDescriptorRuntimeHash, EUpdateResult UpdateResult, FInstanceUpdateDelegate* UpdateCallback, mu::ParametersPtrConst Parameters = nullptr)
 {
@@ -862,19 +851,15 @@ void UpdateSkeletalMesh(UCustomizableObjectInstance& CustomizableObjectInstance,
 }
 
 
-void FCustomizableObjectSystemPrivate::InitUpdateSkeletalMesh(UCustomizableObjectInstance& Instance, FMutableQueueElem::EQueuePriorityType Priority, FInstanceUpdateDelegate* UpdateCallback)
+void FCustomizableObjectSystemPrivate::GetMipStreamingConfig(const UCustomizableObjectInstance& Instance, bool& bOutNeverStream, int32& OutMipsToSkip) const
 {
-	MUTABLE_CPUPROFILER_SCOPE(FCustomizableObjectSystemPrivate::InitUpdateSkeletalMesh);
-
-	check(IsInGameThread());
-
 	const FString CurrentState = Instance.GetCurrentState();
 	const FParameterUIData* State = Instance.GetCustomizableObject()->StateUIDataMap.Find(CurrentState);
 
 	// \TODO: This should be controllable independently
-	const bool bNeverStream = State ? State->TextureCompressionStrategy != ETextureCompressionStrategy::None : false;
-	bool bUseMipmapStreaming = !bNeverStream;
-	int32 MipsToSkip = 0; // 0 means all mips
+	bOutNeverStream = State ? State->TextureCompressionStrategy != ETextureCompressionStrategy::None : false;
+	bool bUseMipmapStreaming = !bOutNeverStream;
+	OutMipsToSkip = 0; // 0 means all mips
 
 #if PLATFORM_SUPPORTS_TEXTURE_STREAMING
 	if (!IStreamingManager::Get().IsTextureStreamingEnabled())
@@ -887,17 +872,27 @@ void FCustomizableObjectSystemPrivate::InitUpdateSkeletalMesh(UCustomizableObjec
 
 	if (bUseMipmapStreaming && EnableMutableProgressiveMipStreaming)
 	{
-		MipsToSkip = 255; // This means skip all possible mips until only UTexture::GetStaticMinTextureResidentMipCount() are left
+		OutMipsToSkip = 255; // This means skip all possible mips until only UTexture::GetStaticMinTextureResidentMipCount() are left
 	}
+}
+
+
+void FCustomizableObjectSystemPrivate::InitUpdateSkeletalMesh(UCustomizableObjectInstance& Instance, EQueuePriorityType Priority, bool bIsCloseDistTick, FInstanceUpdateDelegate* UpdateCallback)
+{
+	MUTABLE_CPUPROFILER_SCOPE(FCustomizableObjectSystemPrivate::InitUpdateSkeletalMesh);
+
+	check(IsInGameThread());
+
+	bool bNeverStream = false;
+	int32 MipsToSkip = 0;
+
+	GetMipStreamingConfig(Instance, bNeverStream, MipsToSkip);
 	
 	const FDescriptorRuntimeHash UpdateDescriptorHash = Instance.GetUpdateDescriptorRuntimeHash();
 
-	if (const FMutableQueueElem* QueueElem = MutableOperationQueue.Get(&Instance))
+	if (const FMutablePendingInstanceUpdate* QueueElem = MutablePendingInstanceWork.GetUpdate(&Instance))
 	{
-		if (const TSharedPtr<FMutableOperation>& QueuedOperation = QueueElem->Operation;
-			QueuedOperation &&
-			QueuedOperation->Type == FMutableOperation::EOperationType::Update &&
-			UpdateDescriptorHash.IsSubset(QueuedOperation->InstanceDescriptorRuntimeHash))
+		if (UpdateDescriptorHash.IsSubset(FDescriptorRuntimeHash(QueueElem->InstanceDescriptor)))
 		{
 			FinishUpdateGlobal(&Instance, EUpdateResult::ErrorOptimized, UpdateCallback);			
 			return; // The the requested update is equal to the last enqueued update.
@@ -905,7 +900,6 @@ void FCustomizableObjectSystemPrivate::InitUpdateSkeletalMesh(UCustomizableObjec
 	}
 
 	if (CurrentMutableOperation &&
-		CurrentMutableOperation->Type == FMutableOperation::EOperationType::Update &&
 		UpdateDescriptorHash.IsSubset(CurrentMutableOperation->InstanceDescriptorRuntimeHash))
 	{
 		FinishUpdateGlobal(&Instance, EUpdateResult::ErrorOptimized, UpdateCallback);			
@@ -925,11 +919,14 @@ void FCustomizableObjectSystemPrivate::InitUpdateSkeletalMesh(UCustomizableObjec
 	{	
 		Instance.SkeletalMeshStatus = ESkeletalMeshState::AsyncUpdatePending;
 
-		const TSharedPtr<FMutableOperation> Operation = MakeShared<FMutableOperation>(FMutableOperation::CreateInstanceUpdate(&Instance, bNeverStream, MipsToSkip, UpdateCallback));
-		MutableOperationQueue.Enqueue(FMutableQueueElem::Create(Operation, Priority, Instance.GetPrivate()->MinSquareDistFromComponentToPlayer));
-
 		// Cache Texture Parameters being used during the update:
-		GetImageProviderChecked()->CacheImages(*Operation->GetParameters());
+		GetImageProviderChecked()->CacheImages(*Instance.GetPrivate()->GetParameters(&Instance));
+
+		if (!bIsCloseDistTick) // When called from bIsCloseDistTick, the update operation is directly processed without going to a queue first
+		{
+			const FMutablePendingInstanceUpdate InstanceUpdate(&Instance, Priority, UpdateCallback, bNeverStream, MipsToSkip);
+			MutablePendingInstanceWork.AddUpdate(InstanceUpdate);
+		}
 	}
 }
 
@@ -940,9 +937,8 @@ void FCustomizableObjectSystemPrivate::InitDiscardResourcesSkeletalMesh(UCustomi
 
 	if (InCustomizableObjectInstance && InCustomizableObjectInstance->IsValidLowLevel())
 	{
-		const TSharedRef<FMutableOperation> Operation = MakeShared<FMutableOperation>(FMutableOperation::CreateInstanceDiscard(InCustomizableObjectInstance));
 		check(InCustomizableObjectInstance->GetPrivate() != nullptr);
-		MutableOperationQueue.Enqueue(FMutableQueueElem::Create(Operation, FMutableQueueElem::EQueuePriorityType::High, InCustomizableObjectInstance->GetPrivate()->MinSquareDistFromComponentToPlayer));
+		MutablePendingInstanceWork.AddDiscard(FMutablePendingInstanceDiscard(InCustomizableObjectInstance));
 	}
 }
 
@@ -951,8 +947,7 @@ void FCustomizableObjectSystemPrivate::InitInstanceIDRelease(mu::Instance::ID ID
 {
 	check(IsInGameThread());
 
-	const TSharedRef<FMutableOperation> Operation = MakeShared<FMutableOperation>(FMutableOperation::CreateInstanceIDRelease(IDToRelease));
-	MutableOperationQueue.Enqueue(FMutableQueueElem::Create(Operation, FMutableQueueElem::EQueuePriorityType::High, 0));
+	MutablePendingInstanceWork.AddIDRelease(IDToRelease);
 }
 
 
@@ -2269,28 +2264,14 @@ namespace impl
 	}
 
 
-	/** Enqueue the release ID operation in the Mutable queue */
-	void Task_Game_ReleaseInstanceID(TSharedPtr<FMutableOperation> Operation)
-	{
-		MUTABLE_CPUPROFILER_SCOPE(Task_Game_ReleaseInstanceID)
-
-		check(Operation);
-		check(Operation->Type == FMutableOperation::EOperationType::IDRelease);
-
-		Task_Game_ReleaseInstanceID(Operation->IDToRelease);
-	}
-
-
 	/** "Start Update" */
 	void Task_Game_StartUpdate(TSharedPtr<FMutableOperation> Operation)
 	{
 		MUTABLE_CPUPROFILER_SCOPE(Task_Game_StartUpdate)
 
 		UCustomizableObjectSystem::GetInstance()->GetPrivate()->MutableTaskGraph.AllowLaunchingMutableTaskLowPriority(false, false);
-
 		
 		check(Operation);
-		check(Operation->Type == FMutableOperation::EOperationType::Update);
 
 		UCustomizableObjectSystem* System = UCustomizableObjectSystem::GetInstance();
 		check(System != nullptr);
@@ -2448,7 +2429,8 @@ namespace impl
 		CurrentOperationData->UpdateCallback = Operation->UpdateCallback;
 		CustomizableObject->GetLowPriorityTextureNames(CurrentOperationData->LowPriorityTextures);
 
-		if (System->IsOnlyGenerateRequestedLODsEnabled() && System->CurrentInstanceLODManagement->IsOnlyGenerateRequestedLODLevelsEnabled() && !CandidateInstancePrivateData->HasCOInstanceFlags(ForceGenerateAllLODs))
+		if (System->IsOnlyGenerateRequestedLODsEnabled() && System->CurrentInstanceLODManagement->IsOnlyGenerateRequestedLODLevelsEnabled() && 
+			!Operation->bForceGenerateAllLODs)
 		{
 			CurrentOperationData->RequestedLODs = Operation->InstanceDescriptorRuntimeHash.GetRequestedLODs();
 		}
@@ -2535,77 +2517,12 @@ void UCustomizableObjectSystem::AdvanceCurrentOperation()
 	// - we have an ongoing operations
 	// - we have no pending work for the ongoing operation
 	// - so we are starting it.
-	switch (Private->CurrentMutableOperation->Type)
 	{
-		case FMutableOperation::EOperationType::Discard:
-		{
-			MUTABLE_CPUPROFILER_SCOPE(OperationDiscard);
-
-			// \TODO: Discards could be done in any case, concurrently with update operations. Should they be
-			// in their own "queue"?
-
-			UCustomizableObjectInstance* COI = Private->CurrentMutableOperation->CustomizableObjectInstance.Get();
-			UCustomizableInstancePrivateData* COIPrivateData = COI ? COI->GetPrivate() : nullptr;
-
-			// Only discard resources if the instance is still out range (it could have got closer to the player since the task was queued)
-			if (!CurrentInstanceLODManagement->IsOnlyUpdateCloseCustomizableObjectsEnabled() ||
-				!COI ||
-				((COIPrivateData != nullptr) &&
-				 (COIPrivateData->LastMinSquareDistFromComponentToPlayer > FMath::Square(CurrentInstanceLODManagement->GetOnlyUpdateCloseCustomizableObjectsDist()))
-				)
-			   )
-			{
-				if (COI && COI->IsValidLowLevel())
-				{
-					check(COIPrivateData != nullptr);
-					COIPrivateData->DiscardResourcesAndSetReferenceSkeletalMesh(COI);
-					COIPrivateData->ClearCOInstanceFlags(Updating);
-					COI->SkeletalMeshStatus = ESkeletalMeshState::Correct;
-				}
-			}
-			else
-			{
-				check(COIPrivateData != nullptr);
-				COIPrivateData->ClearCOInstanceFlags(Updating);
-			}
-
-			if (COI && !COI->HasAnySkeletalMesh())
-			{
-				// To solve the problem in the Mutable demo where PIE just after editor start made all instances appear as reference mesh until editor restart
-				check(COIPrivateData != nullptr);
-				COIPrivateData->ClearCOInstanceFlags(Generated);
-			}
-
-			ClearCurrentMutableOperation();
-
-			break;
-		}
-
-		case FMutableOperation::EOperationType::Update:
-		{
-			MUTABLE_CPUPROFILER_SCOPE(OperationUpdate);
-
-			// Start the first task of the update process. See namespace impl comments above.
-			impl::Task_Game_StartUpdate(Private->CurrentMutableOperation);
-			break;
-		}
-
-		case FMutableOperation::EOperationType::IDRelease:
-		{
-			MUTABLE_CPUPROFILER_SCOPE(OperationIDRelease);
-
-			impl::Task_Game_ReleaseInstanceID(Private->CurrentMutableOperation);
-
-			ClearCurrentMutableOperation();
-
-			break;
-		}
+		MUTABLE_CPUPROFILER_SCOPE(OperationUpdate);
 
-		default:
-			check(false);
-			break;
+		// Start the first task of the update process. See namespace impl comments above.
+		impl::Task_Game_StartUpdate(Private->CurrentMutableOperation);
 	}
-
 }
 
 
@@ -2613,7 +2530,7 @@ bool UCustomizableObjectSystem::Tick(float DeltaTime)
 {
 	MUTABLE_CPUPROFILER_SCOPE(UCustomizableObjectSystem::Tick)
 	
-	// Building instances is not enable in servers. If at some point relevant collision or animation data is necessary for server logic this will need to be changed.
+	// Building instances is not enabled in servers. If at some point relevant collision or animation data is necessary for server logic this will need to be changed.
 #if UE_SERVER
 	return true;
 #endif
@@ -2648,9 +2565,10 @@ bool UCustomizableObjectSystem::Tick(float DeltaTime)
 	if (!Private->CurrentMutableOperation)
 	{
 		// Reset the instance relevancy
-		// \TODO: Review
-		// TODO: This should be done only when requiring a new job. And for the current operation instance, in case it needs to be cancelled.
-		CurrentInstanceLODManagement->UpdateInstanceDistsAndLODs();
+		// The RequestedUpdates only refer to LOD changes. User Customization and discards are handled separately
+		FMutableInstanceUpdateMap RequestedLODUpdates;
+		
+		CurrentInstanceLODManagement->UpdateInstanceDistsAndLODs(RequestedLODUpdates);
 
 		for (TObjectIterator<UCustomizableObjectInstance> CustomizableObjectInstance; CustomizableObjectInstance; ++CustomizableObjectInstance)
 		{
@@ -2660,21 +2578,150 @@ bool UCustomizableObjectSystem::Tick(float DeltaTime)
 
 				if (ObjectInstancePrivateData->HasCOInstanceFlags(UsedByComponentInPlay))
 				{
-					ObjectInstancePrivateData->TickUpdateCloseCustomizableObjects(**CustomizableObjectInstance);
+					ObjectInstancePrivateData->TickUpdateCloseCustomizableObjects(**CustomizableObjectInstance, true, RequestedLODUpdates);
 				}
 				else if (ObjectInstancePrivateData->HasCOInstanceFlags(UsedByComponent))
 				{
-					ObjectInstancePrivateData->UpdateInstanceIfNotGenerated(**CustomizableObjectInstance);
+					ensure(!RequestedLODUpdates.Contains(*CustomizableObjectInstance));
+					ObjectInstancePrivateData->UpdateInstanceIfNotGenerated(**CustomizableObjectInstance, RequestedLODUpdates);
+				}
+				else
+				{
+					ensure(!RequestedLODUpdates.Contains(*CustomizableObjectInstance));
 				}
 
 				ObjectInstancePrivateData->ClearCOInstanceFlags((ECOInstanceFlags)(UsedByComponent | UsedByComponentInPlay | PendingLODsUpdate)); // TODO MTBL-391: Makes no sense to clear it here, what if an update is requested before we set it back to true
 			}
+			else
+			{
+				ensure(!RequestedLODUpdates.Contains(*CustomizableObjectInstance));
+			}
 		}
 
-		// Update the queue. TODO: This should be done only when requiring a new job. And for the current operation instance, in case it needs to be cancelled.
+		TSharedPtr<FMutableOperation> FoundOperation;
+
 		{
-			Private->MutableOperationQueue.ChangePriorities();
+			// Look for the highest priority update between the pending updates and the LOD Requested Updates
+			EQueuePriorityType MaxPriorityFound = EQueuePriorityType::Low;
+			double MaxSquareDistanceFound = TNumericLimits<double>::Max();
+			double MinTimeFound = TNumericLimits<double>::Max();
+			const FMutablePendingInstanceUpdate* PendingInstanceUpdateFound = nullptr;
+			FMutableUpdateCandidate* LODUpdateCandidateFound = nullptr;
+
+			// Look for the highest priority Pending Update
+			for (auto Iterator = Private->MutablePendingInstanceWork.GetUpdateIterator(); Iterator; ++Iterator)
+			{
+				FMutablePendingInstanceUpdate& PendingUpdate = *Iterator;
 
+				if (PendingUpdate.CustomizableObjectInstance.IsValid())
+				{
+					EQueuePriorityType PriorityType = PendingUpdate.CustomizableObjectInstance->GetUpdatePriority(true, true, false, false);
+					
+					if (PendingUpdate.PriorityType <= MaxPriorityFound)
+					{
+						const double MinSquareDistFromComponentToPlayer = PendingUpdate.CustomizableObjectInstance->GetPrivate()->MinSquareDistFromComponentToPlayer;
+						
+						if (MinSquareDistFromComponentToPlayer < MaxSquareDistanceFound ||
+							(MinSquareDistFromComponentToPlayer == MaxSquareDistanceFound && PendingUpdate.SecondsAtUpdate < MinTimeFound))
+						{
+							MaxPriorityFound = PriorityType;
+							MaxSquareDistanceFound = MinSquareDistFromComponentToPlayer;
+							MinTimeFound = PendingUpdate.SecondsAtUpdate;
+							PendingInstanceUpdateFound = &PendingUpdate;
+							LODUpdateCandidateFound = nullptr;
+						}
+					}
+				}
+				else
+				{
+					Iterator.RemoveCurrent();
+				}
+			}
+
+			// Look for a higher priority LOD update
+			for (TPair<const UCustomizableObjectInstance*, FMutableUpdateCandidate>& LODUpdateTuple : RequestedLODUpdates)
+			{
+				const UCustomizableObjectInstance* Instance = LODUpdateTuple.Key;
+
+				if (Instance)
+				{
+					FMutableUpdateCandidate& LODUpdateCandidate = LODUpdateTuple.Value;
+					ensure(LODUpdateCandidate.HasBeenIssued());
+
+					if (LODUpdateCandidate.Priority <= MaxPriorityFound)
+					{
+						if (LODUpdateCandidate.CustomizableObjectInstance->GetPrivate()->MinSquareDistFromComponentToPlayer < MaxSquareDistanceFound)
+						{
+							MaxPriorityFound = LODUpdateCandidate.Priority;
+							MaxSquareDistanceFound = LODUpdateCandidate.CustomizableObjectInstance->GetPrivate()->MinSquareDistFromComponentToPlayer;
+							PendingInstanceUpdateFound = nullptr;
+							LODUpdateCandidateFound = &LODUpdateCandidate;
+						}
+					}
+				}
+			}
+
+			Private->MutablePendingInstanceWork.SetLODUpdatesLastTick(RequestedLODUpdates.Num());
+
+			// If the chosen LODUpdate has the same instance as a PendingUpdate, choose the PendingUpdate to apply both the LOD update
+			// and customization change
+			if (LODUpdateCandidateFound)
+			{
+				if (const FMutablePendingInstanceUpdate *PendingUpdateWithSameInstance = Private->MutablePendingInstanceWork.GetUpdate(LODUpdateCandidateFound->CustomizableObjectInstance))
+				{
+					PendingInstanceUpdateFound = PendingUpdateWithSameInstance;
+					LODUpdateCandidateFound = nullptr;
+
+					// In the processing of the PendingUpdate just below, it will add the LODUpdate's LOD params
+				}
+			}
+
+			if (PendingInstanceUpdateFound)
+			{
+				check(!LODUpdateCandidateFound);
+
+				UCustomizableObjectInstance* PendingInstance = PendingInstanceUpdateFound->CustomizableObjectInstance.Get();
+				ensure(PendingInstance);
+
+				// Maybe there's a LODUpdate that has the same instance, merge both updates as an optimization
+				FMutableUpdateCandidate* LODUpdateWithSameInstance = RequestedLODUpdates.Find(PendingInstance);
+				
+				if (LODUpdateWithSameInstance)
+				{
+					LODUpdateWithSameInstance->ApplyLODUpdateParamsToInstance();
+				}
+
+				// No need to do a DoUpdateSkeletalMesh, as it was already done when adding the pending update
+
+				FoundOperation = MakeShared<FMutableOperation>(FMutableOperation::CreateInstanceUpdate(PendingInstance,
+					PendingInstanceUpdateFound->bNeverStream, PendingInstanceUpdateFound->MipsToSkip, PendingInstanceUpdateFound->Callback));
+
+				Private->MutablePendingInstanceWork.RemoveUpdate(PendingInstanceUpdateFound->CustomizableObjectInstance);
+			}
+			else if(LODUpdateCandidateFound)
+			{
+				// Commit the LOD changes
+				LODUpdateCandidateFound->ApplyLODUpdateParamsToInstance();
+
+				// No need to check if a PendingUpdate has the same instance, it would already have been caught in the 
+				// previous if statements
+
+				// LOD updates are detected automatically and have never had a DoUpdateSkeletalMesh, so call it without enqueing a new update
+				UCustomizableInstancePrivateData* ObjectInstancePrivateData = LODUpdateCandidateFound->CustomizableObjectInstance->GetPrivate();
+				ObjectInstancePrivateData->TickUpdateCloseCustomizableObjects(*LODUpdateCandidateFound->CustomizableObjectInstance, false, RequestedLODUpdates);
+
+				bool bNeverStream = false;
+				int32 MipsToSkip = 0;
+
+				ensure(LODUpdateCandidateFound->CustomizableObjectInstance);
+				GetPrivate()->GetMipStreamingConfig(*LODUpdateCandidateFound->CustomizableObjectInstance, bNeverStream, MipsToSkip);
+
+				FoundOperation = MakeShared<FMutableOperation>(FMutableOperation::CreateInstanceUpdate(LODUpdateCandidateFound->CustomizableObjectInstance,
+					bNeverStream, MipsToSkip, nullptr));
+			}
+		}
+
+		{
 			for (TObjectIterator<UCustomizableObjectInstance> CustomizableObjectInstance; CustomizableObjectInstance; ++CustomizableObjectInstance)
 			{
 				if (IsValidChecked(*CustomizableObjectInstance) && CustomizableObjectInstance->GetPrivate())
@@ -2683,19 +2730,28 @@ bool UCustomizableObjectSystem::Tick(float DeltaTime)
 					CustomizableObjectInstance->GetPrivate()->MinSquareDistFromComponentToPlayer = FLT_MAX;
 				}
 			}
-
-			Private->MutableOperationQueue.Sort();
 		}
 
 		// Update the streaming limit if it has changed. It is safe to do this now.
 		Private->UpdateMemoryLimit();
 
+		// Free memory before starting the new update
+		DiscardInstances();
+		ReleaseInstanceIDs();
+
 		// Decide the next mutable operation to perform.
-		if (!Private->MutableOperationQueue.IsEmpty())
+		if (FoundOperation.IsValid())
 		{
-			FMutableQueueElem AuxElem;
-			Private->MutableOperationQueue.Dequeue(&AuxElem);
-			Private->CurrentMutableOperation = AuxElem.Operation;
+			Private->CurrentMutableOperation = FoundOperation;
+
+			UCustomizableObjectInstance* FoundInstance = Private->CurrentMutableOperation->CustomizableObjectInstance.Get();
+			
+			if (FoundInstance->GetPrivate()->HasCOInstanceFlags(ForceGenerateAllLODs))
+			{
+				Private->CurrentMutableOperation->bForceGenerateAllLODs = true;
+			}
+
+			Private->CurrentMutableOperation->CustomizableObjectInstance->GetPrivate()->InstanceUpdateFlags(*FoundInstance);
 		}
 		else
 		{
@@ -2734,6 +2790,89 @@ bool UCustomizableObjectSystem::Tick(float DeltaTime)
 }
 
 
+TAutoConsoleVariable<int32> CVarMaxNumInstancesToDiscardPerTick(
+	TEXT("mutable.MaxNumInstancesToDiscardPerTick"),
+	30,
+	TEXT("The maximum number of stale instances that will be discarded per tick by Mutable."),
+	ECVF_Scalability);
+
+
+void UCustomizableObjectSystem::DiscardInstances()
+{
+	// Handle instance discards
+	int32 NumInstancesDiscarded = 0;
+	const int32 DiscardLimitPerTick = CVarMaxNumInstancesToDiscardPerTick.GetValueOnGameThread();
+
+	for (auto Iterator = Private->MutablePendingInstanceWork.GetDiscardIterator(); Iterator && NumInstancesDiscarded < DiscardLimitPerTick; ++Iterator)
+	{
+		MUTABLE_CPUPROFILER_SCOPE(OperationDiscard);
+
+		UCustomizableObjectInstance* COI = Iterator->CustomizableObjectInstance.Get();
+
+		if (COI)
+		{
+			UCustomizableInstancePrivateData* COIPrivateData = COI ? COI->GetPrivate() : nullptr;
+
+			// Only discard resources if the instance is still out range (it could have got closer to the player since the task was queued)
+			if (!CurrentInstanceLODManagement->IsOnlyUpdateCloseCustomizableObjectsEnabled() ||
+				!COI ||
+				((COIPrivateData != nullptr) &&
+					(COIPrivateData->LastMinSquareDistFromComponentToPlayer > FMath::Square(CurrentInstanceLODManagement->GetOnlyUpdateCloseCustomizableObjectsDist()))
+					)
+				)
+			{
+				if (COI && COI->IsValidLowLevel())
+				{
+					check(COIPrivateData != nullptr);
+					COIPrivateData->DiscardResourcesAndSetReferenceSkeletalMesh(COI);
+					COIPrivateData->ClearCOInstanceFlags(Updating);
+					COI->SkeletalMeshStatus = ESkeletalMeshState::Correct;
+				}
+			}
+			else
+			{
+				check(COIPrivateData != nullptr);
+				COIPrivateData->ClearCOInstanceFlags(Updating);
+			}
+
+			if (COI && !COI->HasAnySkeletalMesh())
+			{
+				// To solve the problem in the Mutable demo where PIE just after editor start made all instances appear as reference mesh until editor restart
+				check(COIPrivateData != nullptr);
+				COIPrivateData->ClearCOInstanceFlags(Generated);
+			}
+		}
+
+		Iterator.RemoveCurrent();
+		NumInstancesDiscarded++;
+	}
+}
+
+
+TAutoConsoleVariable<int32> CVarMaxNumInstanceIDsToReleasePerTick(
+	TEXT("mutable.MaxNumInstanceIDsToReleasePerTick"),
+	30,
+	TEXT("The maximum number of stale instances IDs that will be released per tick by Mutable."),
+	ECVF_Scalability);
+
+
+void UCustomizableObjectSystem::ReleaseInstanceIDs()
+{
+	// Handle ID discards
+	int32 NumIDsReleased = 0;
+	const int32 IDReleaseLimitPerTick = CVarMaxNumInstanceIDsToReleasePerTick.GetValueOnGameThread();
+
+	for (auto Iterator = Private->MutablePendingInstanceWork.GetIDsToReleaseIterator();
+		Iterator && NumIDsReleased < IDReleaseLimitPerTick; ++Iterator)
+	{
+		impl::Task_Game_ReleaseInstanceID(*Iterator);
+
+		Iterator.RemoveCurrent();
+		NumIDsReleased++;
+	}
+}
+
+
 TArray<FCustomizableObjectExternalTexture> UCustomizableObjectSystem::GetTextureParameterValues()
 {
 	TArray<FCustomizableObjectExternalTexture> Result;
@@ -2782,7 +2921,6 @@ FMutableOperation FMutableOperation::CreateInstanceUpdate(UCustomizableObjectIns
 	check(InCustomizableObjectInstance->GetCustomizableObject() != nullptr);
 
 	FMutableOperation Op;
-	Op.Type = EOperationType::Update;
 	Op.bNeverStream = bInNeverStream;
 	Op.MipsToSkip = InMipsToSkip;
 	Op.CustomizableObjectInstance = InCustomizableObjectInstance;
@@ -2790,6 +2928,7 @@ FMutableOperation FMutableOperation::CreateInstanceUpdate(UCustomizableObjectIns
 	Op.bStarted = false;
 	Op.bBuildParameterRelevancy = InCustomizableObjectInstance->GetBuildParameterRelevancy();
 	Op.Parameters = InCustomizableObjectInstance->GetPrivate()->GetParameters(InCustomizableObjectInstance);
+
 	if (UpdateCallback)
 	{
 		Op.UpdateCallback = *UpdateCallback;		
@@ -2807,31 +2946,6 @@ FMutableOperation FMutableOperation::CreateInstanceUpdate(UCustomizableObjectIns
 }
 
 
-FMutableOperation FMutableOperation::CreateInstanceDiscard(UCustomizableObjectInstance* InCustomizableObjectInstance)
-{
-	check(InCustomizableObjectInstance != nullptr);
-	check(InCustomizableObjectInstance->GetPrivate() != nullptr);
-
-	FMutableOperation Op;
-	Op.Type = EOperationType::Discard;
-	Op.CustomizableObjectInstance = InCustomizableObjectInstance;
-
-	Op.CustomizableObjectInstance->GetPrivate()->SetCOInstanceFlags(Updating);
-
-	return Op;
-}
-
-
-FMutableOperation FMutableOperation::CreateInstanceIDRelease(mu::Instance::ID IDToRelease)
-{
-	FMutableOperation Op;
-	Op.Type = EOperationType::IDRelease;
-	Op.IDToRelease = IDToRelease;
-
-	return Op;
-}
-
-
 int32 UCustomizableObjectSystem::GetNumInstances() const
 {
 	check(Private != nullptr);
diff --git a/Engine/Plugins/Experimental/Mutable/Source/CustomizableObject/Private/MuCO/CustomizableObjectSystemPrivate.h b/Engine/Plugins/Experimental/Mutable/Source/CustomizableObject/Private/MuCO/CustomizableObjectSystemPrivate.h
--- a/Engine/Plugins/Experimental/Mutable/Source/CustomizableObject/Private/MuCO/CustomizableObjectSystemPrivate.h
+++ b/Engine/Plugins/Experimental/Mutable/Source/CustomizableObject/Private/MuCO/CustomizableObjectSystemPrivate.h
@@ -45,31 +45,11 @@ public:
 
 
 	static FMutableOperation CreateInstanceUpdate(UCustomizableObjectInstance* COInstance, bool bInNeverStream, int32 MipsToSkip, const FInstanceUpdateDelegate* UpdateCallback);
-	static FMutableOperation CreateInstanceDiscard(UCustomizableObjectInstance* COInstance);
-	static FMutableOperation CreateInstanceIDRelease(mu::Instance::ID);
-
-
-	enum class EOperationType
-	{
-		// Create mutable resources for a new instance or to update an exiting one.
-		Update,
-
-		// Discard the resources of an instance.
-		Discard,
-
-		// Release the instance ID and all the temp data associated with it. Usually used with the LiveUpdateMode
-		IDRelease
-
-		// Attention! If any new operation type is added, make sure to review FMutableQueue::Enqueue in CustomizableObjectSystem.cpp and
-		// modify it if new operations of this type should override older ones to the same instance. The default behavior is to enqueue 
-		// the new ones so that they are executed after the old ones.
-	};
-
-	// Type of the operation
-	EOperationType Type;
 
 	bool bStarted = false;
 
+	bool bForceGenerateAllLODs = false;
+
 	bool bNeverStream = false;
 
 	/** If this is non-zero, the instance images will be generated with a certain amount of mips starting from the smallest. Otherwise, the full images will be generated. 
@@ -108,31 +88,49 @@ public:
 };
 
 
-
-struct FMutableQueueElem
+struct FMutablePendingInstanceUpdate
 {
-	// Priority for mutable update queue, Low is the normal distance-based priority, High is normally used for discards and Mid for LOD downgrades
-	enum EQueuePriorityType { High, Med, Med_Low, Low };
+	EQueuePriorityType PriorityType = EQueuePriorityType::Low;
+
+	TWeakObjectPtr<UCustomizableObjectInstance> CustomizableObjectInstance;
+	FCustomizableObjectInstanceDescriptor InstanceDescriptor;
+
+	double SecondsAtUpdate = 0;
+	FInstanceUpdateDelegate* Callback = nullptr;
+	bool bNeverStream = false;
 
-	EQueuePriorityType PriorityType;
-	float Priority;
-	TSharedPtr<FMutableOperation> Operation;
-	bool bIsDiscardResources;
+	/** If this is non-zero, the instance images will be generated with a certain amount of mips starting from the smallest. Otherwise, the full images will be generated.
+	 * This is used for both Update and UpdateImage operations.
+	 */
+	int32 MipsToSkip = 0;
 
-	static FMutableQueueElem Create(TSharedPtr<FMutableOperation> InOperation, EQueuePriorityType NewPriorityType, double NewPriority)
+	FMutablePendingInstanceUpdate(UCustomizableObjectInstance* InCustomizableObjectInstance)
 	{
-		FMutableQueueElem MutableQueueElem;
+		check(InCustomizableObjectInstance);
+		CustomizableObjectInstance = InCustomizableObjectInstance;
+	}
 
-		MutableQueueElem.PriorityType = NewPriorityType;
-		MutableQueueElem.Priority = NewPriority;
-		check(InOperation);
-		MutableQueueElem.Operation = InOperation;
-		MutableQueueElem.bIsDiscardResources = (InOperation->Type==FMutableOperation::EOperationType::Discard);
+	FMutablePendingInstanceUpdate(UCustomizableObjectInstance* InCustomizableObjectInstance, EQueuePriorityType NewPriorityType, 
+		                          FInstanceUpdateDelegate* InCallback, bool bInNeverStream, int32 InMipsToSkip)
+	{
+		PriorityType = NewPriorityType;
+
+		check(InCustomizableObjectInstance);
+		CustomizableObjectInstance = InCustomizableObjectInstance;
+		InstanceDescriptor = CustomizableObjectInstance->GetDescriptor();
 
-		return MutableQueueElem;
+		SecondsAtUpdate = FPlatformTime::Seconds();
+		Callback = InCallback;
+		bNeverStream = bInNeverStream;
+		MipsToSkip = InMipsToSkip;
+	}
+
+	friend bool operator ==(const FMutablePendingInstanceUpdate& A, const FMutablePendingInstanceUpdate& B)
+	{
+		return A.CustomizableObjectInstance.HasSameIndexAndSerialNumber(B.CustomizableObjectInstance);
 	}
 
-	friend bool operator <(const FMutableQueueElem& A, const FMutableQueueElem& B)
+	friend bool operator <(const FMutablePendingInstanceUpdate& A, const FMutablePendingInstanceUpdate& B)
 	{
 		if (A.PriorityType < B.PriorityType)
 		{
@@ -144,45 +142,128 @@ struct FMutableQueueElem
 		}
 		else
 		{
-			if (A.bIsDiscardResources && !B.bIsDiscardResources)
-			{
-				return true;
-			}
-			else if (!A.bIsDiscardResources && B.bIsDiscardResources)
-			{
-				return false;
-			}
-
-			return A.Priority < B.Priority;
+			return A.SecondsAtUpdate < B.SecondsAtUpdate;
 		}
 	}
 };
 
 
-/** Instances operations queue.
+inline uint32 GetTypeHash(const FMutablePendingInstanceUpdate& Update)
+{
+	return GetTypeHash(Update.CustomizableObjectInstance.GetWeakPtrTypeHash());
+}
+
+
+struct FPendingInstanceUpdateKeyFuncs : BaseKeyFuncs<FMutablePendingInstanceUpdate, TWeakObjectPtr<UCustomizableObjectInstance>>
+{
+	FORCEINLINE static const TWeakObjectPtr<UCustomizableObjectInstance>& GetSetKey(const FMutablePendingInstanceUpdate& PendingUpdate)
+	{
+		return PendingUpdate.CustomizableObjectInstance;
+	}
+
+	FORCEINLINE static bool Matches(const TWeakObjectPtr<UCustomizableObjectInstance>& A, const TWeakObjectPtr<UCustomizableObjectInstance>& B)
+	{
+		return A.HasSameIndexAndSerialNumber(B);
+	}
+
+	FORCEINLINE static uint32 GetKeyHash(const TWeakObjectPtr<UCustomizableObjectInstance>& Identifier)
+	{
+		return GetTypeHash(Identifier.GetWeakPtrTypeHash());
+	}
+};
+
+
+struct FMutablePendingInstanceDiscard
+{
+	TWeakObjectPtr<UCustomizableObjectInstance> CustomizableObjectInstance;
+
+	FMutablePendingInstanceDiscard(UCustomizableObjectInstance* InCustomizableObjectInstance)
+	{
+		CustomizableObjectInstance = InCustomizableObjectInstance;
+	}
+
+	friend bool operator ==(const FMutablePendingInstanceDiscard& A, const FMutablePendingInstanceDiscard& B)
+	{
+		return A.CustomizableObjectInstance.HasSameIndexAndSerialNumber(B.CustomizableObjectInstance);
+	}
+};
+
+
+inline uint32 GetTypeHash(const FMutablePendingInstanceDiscard& Discard)
+{
+	return GetTypeHash(Discard.CustomizableObjectInstance.GetWeakPtrTypeHash());
+}
+
+
+struct FPendingInstanceDiscardKeyFuncs : BaseKeyFuncs<FMutablePendingInstanceUpdate, TWeakObjectPtr<UCustomizableObjectInstance>>
+{
+	FORCEINLINE static const TWeakObjectPtr<UCustomizableObjectInstance>& GetSetKey(const FMutablePendingInstanceDiscard& PendingDiscard)
+	{
+		return PendingDiscard.CustomizableObjectInstance;
+	}
+
+	FORCEINLINE static bool Matches(const TWeakObjectPtr<UCustomizableObjectInstance>& A, const TWeakObjectPtr<UCustomizableObjectInstance>& B)
+	{
+		return A.HasSameIndexAndSerialNumber(B);
+	}
+
+	FORCEINLINE static uint32 GetKeyHash(const TWeakObjectPtr<UCustomizableObjectInstance>& Identifier)
+	{
+		return GetTypeHash(Identifier.GetWeakPtrTypeHash());
+	}
+};
+
+
+/** Instance updates queue.
  *
- * The queue will only contain a single operation per UCustomizableObjectInstance.
+ * The queues will only contain a single operation per UCustomizableObjectInstance.
  * If there is already an operation it will be replaced. */
-class FMutableQueue
+class FMutablePendingInstanceWork
 {
-	TArray<FMutableQueueElem> Array;
+	TSet<FMutablePendingInstanceUpdate, FPendingInstanceUpdateKeyFuncs> PendingInstanceUpdates;
+
+	TSet<FMutablePendingInstanceDiscard, FPendingInstanceDiscardKeyFuncs> PendingInstanceDiscards;
+
+	TSet<mu::Instance::ID> PendingIDsToRelease;
+
+	int32 NumLODUpdatesLastTick = 0;
 
 public:
-	bool IsEmpty() const;
+	// Returns true if there are no pending instance updates. Doesn't take into account discards.
+	bool ArePendingUpdatesEmpty() const;
 
-	int Num() const;
+	// Returns the number of pending instance updates, LOD Updates, discards and releases last tick.
+	int32 Num() const;
 
-	void Enqueue(const FMutableQueueElem& TaskToEnqueue);
+	void SetLODUpdatesLastTick(int32 NumLODUpdates);
 
-	void Dequeue(FMutableQueueElem* DequeuedTask);
+	// Adds a new instance update
+	void AddUpdate(const FMutablePendingInstanceUpdate& UpdateToAdd);
 
-	void ChangePriorities();
+	// Removes an instance update
+	void RemoveUpdate(const TWeakObjectPtr<UCustomizableObjectInstance>& Instance);
 
-	void UpdatePriority(const UCustomizableObjectInstance* Instance);
+	const FMutablePendingInstanceUpdate* GetUpdate(const TWeakObjectPtr<UCustomizableObjectInstance>& Instance) const;
 
-	const FMutableQueueElem* Get(const UCustomizableObjectInstance* Instance) const;
-	
-	void Sort();
+	TSet<FMutablePendingInstanceUpdate, FPendingInstanceUpdateKeyFuncs>::TIterator GetUpdateIterator()
+	{
+		return PendingInstanceUpdates.CreateIterator();
+	}
+
+	TSet<FMutablePendingInstanceDiscard, FPendingInstanceDiscardKeyFuncs>::TIterator GetDiscardIterator()
+	{
+		return PendingInstanceDiscards.CreateIterator();
+	}
+
+	TSet<mu::Instance::ID>::TIterator GetIDsToReleaseIterator()
+	{
+		return PendingIDsToRelease.CreateIterator();
+	}
+
+	void AddDiscard(const FMutablePendingInstanceDiscard& TaskToEnqueue);
+	void AddIDRelease(mu::Instance::ID IDToRelease);
+
+	void RemoveAllUpdatesAndDiscardsAndReleases();
 };
 
 
@@ -551,7 +632,8 @@ public:
 	// from the game thread, and it is read from the Mutable thread only while updating the instance.
 	TArray<mu::RESOURCE_ID> ProtectedObjectCachedImages;
 
-	FMutableQueue MutableOperationQueue;
+	// The pending instance updates, discards or releases
+	FMutablePendingInstanceWork MutablePendingInstanceWork;
 
 	// Queue of game-thread tasks that need to be executed for the current operation
 	TQueue<FMutableTask> PendingTasks;
@@ -673,13 +755,15 @@ public:
 
 
 	// Init the async Skeletal Mesh creation/update
-	void InitUpdateSkeletalMesh(UCustomizableObjectInstance& Public, FMutableQueueElem::EQueuePriorityType Priority, FInstanceUpdateDelegate* UpdateCallback = nullptr);
+	void InitUpdateSkeletalMesh(UCustomizableObjectInstance& Public, EQueuePriorityType Priority, bool bIsCloseDistTick, FInstanceUpdateDelegate* UpdateCallback = nullptr);
 		
 	// Init an async and safe release of the UE and Mutable resources used by the instance without actually destroying the instance, for example if it's very far away
 	void InitDiscardResourcesSkeletalMesh(UCustomizableObjectInstance* InCustomizableObjectInstance);
 
 	// Init the async release of a Mutable Core Instance ID and all the temp resources associated with it
 	void InitInstanceIDRelease(mu::Instance::ID IDToRelease);
+
+	void GetMipStreamingConfig(const UCustomizableObjectInstance& Instance, bool& bOutNeverStream, int32& OutMipsToSkip) const;
 	
 	bool IsReplaceDiscardedWithReferenceMeshEnabled() const { return bReplaceDiscardedWithReferenceMesh; }
 	void SetReplaceDiscardedWithReferenceMeshEnabled(bool bIsEnabled) { bReplaceDiscardedWithReferenceMesh = bIsEnabled; }
diff --git a/Engine/Plugins/Experimental/Mutable/Source/CustomizableObject/Public/MuCO/CustomizableObjectInstanceDescriptor.h b/Engine/Plugins/Experimental/Mutable/Source/CustomizableObject/Public/MuCO/CustomizableObjectInstanceDescriptor.h
--- a/Engine/Plugins/Experimental/Mutable/Source/CustomizableObject/Public/MuCO/CustomizableObjectInstanceDescriptor.h
+++ b/Engine/Plugins/Experimental/Mutable/Source/CustomizableObject/Public/MuCO/CustomizableObjectInstanceDescriptor.h
@@ -16,6 +16,9 @@ class UCustomizableObject;
 class UCustomizableObjectInstance;
 class FDescriptorHash;
 class FDescriptorRuntimeHash;
+class FMutableUpdateCandidate;
+
+typedef TMap<const UCustomizableObjectInstance*, FMutableUpdateCandidate> FMutableInstanceUpdateMap;
 
 namespace mu
 {
@@ -332,6 +335,9 @@ struct CUSTOMIZABLEOBJECT_API FCustomizableObjectInstanceDescriptor
 	/** See FMultilayerProjector::UpdateVirtualLayer. */
 	void MultilayerProjectorUpdateVirtualLayer(const FName& ProjectorParamName, const FName& Id, const FMultilayerProjectorVirtualLayer& Layer);
 
+	/** Return a Mutable Core object containing all parameters. */
+	mu::Ptr<mu::Parameters> GetParameters() const;
+
 private:
 
 	UPROPERTY()
@@ -375,9 +381,6 @@ private:
 	UPROPERTY()
 	TMap<FName, FMultilayerProjector> MultilayerProjectors;
 
-	/** Return a Mutable Core object containing all parameters. */
-	mu::Ptr<mu::Parameters> GetParameters() const;
-
 	void CreateParametersLookupTable();
 	
 	// Friends
@@ -386,6 +389,7 @@ private:
 	friend UCustomizableObjectInstance;
 	friend UCustomizableInstancePrivateData;
 	friend FMultilayerProjector;
+	friend FMutableUpdateCandidate;
 };
 
 
