[Mutable] Fix crash on RETURN_ON_UNCOMPILED_CO when throwing a warning.

#rnx

diff --git a/Engine/Plugins/Experimental/Mutable/Source/CustomizableObject/Private/MuCO/CustomizableObjectInstanceDescriptor.cpp b/Engine/Plugins/Experimental/Mutable/Source/CustomizableObject/Private/MuCO/CustomizableObjectInstanceDescriptor.cpp
--- a/Engine/Plugins/Experimental/Mutable/Source/CustomizableObject/Private/MuCO/CustomizableObjectInstanceDescriptor.cpp
+++ b/Engine/Plugins/Experimental/Mutable/Source/CustomizableObject/Private/MuCO/CustomizableObjectInstanceDescriptor.cpp
@@ -1119,7 +1119,7 @@ bool FCustomizableObjectInstanceDescriptor::HasAnyParameters() const
 #define RETURN_ON_UNCOMPILED_CO(CustomizableObject, ErrorMessage) \
 	if (!CustomizableObject->IsCompiled()) \
 	{ \
-		FString AdditionalLoggingInfo = FString::Printf(TEXT("Calling function: %hs.  %s"), __FUNCTION__, *ErrorMessage); \
+		FString AdditionalLoggingInfo = FString::Printf(TEXT("Calling function: %hs.  %s"), __FUNCTION__, ErrorMessage); \
 		CustomizableObject->AddUncompiledCOWarning(AdditionalLoggingInfo);\
 		return; \
 	} \
@@ -1129,7 +1129,7 @@ bool FCustomizableObjectInstanceDescriptor::HasAnyParameters() const
 #define RETURN_ON_UNCOMPILED_CO(CustomizableObject, ErrorMessage) \
 	if (!ensureMsgf(CustomizableObject->IsCompiled(), TEXT("Customizable Object (%s) was not compiled."), *GetNameSafe(CustomizableObject))) \
 	{ \
-		FString AdditionalLoggingInfo = FString::Printf(TEXT("Calling function: %hs.  %s"), __FUNCTION__, *ErrorMessage); \
+		FString AdditionalLoggingInfo = FString::Printf(TEXT("Calling function: %hs.  %s"), __FUNCTION__, ErrorMessage); \
 		CustomizableObject->AddUncompiledCOWarning(AdditionalLoggingInfo);\
 		return; \
 	} \
@@ -1472,7 +1472,7 @@ void FCustomizableObjectInstanceDescriptor::SetBoolParameterSelectedOption(const
 void FCustomizableObjectInstanceDescriptor::SetVectorParameterSelectedOption(const FString& VectorParamName, const FLinearColor& VectorValue)
 {
 	check(CustomizableObject);
-	RETURN_ON_UNCOMPILED_CO(CustomizableObject, FString("Error: Cannot set Int parameter "));
+	RETURN_ON_UNCOMPILED_CO(CustomizableObject, TEXT("Error: Cannot set Int parameter "));
 
 	const int32 ParameterIndexInObject = CustomizableObject->FindParameter(VectorParamName);
 	const int32 ParameterIndexInInstance = FindVectorParameterNameIndex(VectorParamName);
