[mutable] Ensure more temporary image memory is tracked for the budget control.
Fix bug with image magnify on non-blank images.
#rnx

diff --git a/Engine/Plugins/Experimental/Mutable/Source/CustomizableObject/Private/MuCO/CustomizableObjectInstance.cpp b/Engine/Plugins/Experimental/Mutable/Source/CustomizableObject/Private/MuCO/CustomizableObjectInstance.cpp
--- a/Engine/Plugins/Experimental/Mutable/Source/CustomizableObject/Private/MuCO/CustomizableObjectInstance.cpp
+++ b/Engine/Plugins/Experimental/Mutable/Source/CustomizableObject/Private/MuCO/CustomizableObjectInstance.cpp
@@ -3195,6 +3195,8 @@ FTexturePlatformData* UCustomizableInstancePrivateData::MutableCreateImagePlatfo
 		//return nullptr;
 	}
 
+	mu::FImageOperator ImOp = mu::FImageOperator::GetDefault();
+
 	EPixelFormat PlatformFormat = PF_Unknown;
 	switch (MutableFormat)
 	{
@@ -3225,17 +3227,17 @@ FTexturePlatformData* UCustomizableInstancePrivateData::MutableCreateImagePlatfo
 		// Cannot prepare texture if it's not in the right format, this can happen if mutable is in debug mode or in case of bugs
 		UE_LOG(LogMutable, Warning, TEXT("Building instance: a texture was generated in an unsupported format, it will be converted to Unreal with a performance penalty."));
 
-		switch (mu::GetImageFormatData(MutableFormat).m_channels)
+		switch (mu::GetImageFormatData(MutableFormat).Channels)
 		{
 		case 1:
 			PlatformFormat = PF_R8;
-			MutableImage = ImagePixelFormat(0, MutableImage.get(), mu::EImageFormat::IF_L_UBYTE);
+			MutableImage = ImOp.ImagePixelFormat(0, MutableImage.get(), mu::EImageFormat::IF_L_UBYTE);
 			break;
 		case 2:
 		case 3:
 		case 4:
 			PlatformFormat = PF_R8G8B8A8;
-			MutableImage = ImagePixelFormat(0, MutableImage.get(), mu::EImageFormat::IF_RGBA_UBYTE);
+			MutableImage = ImOp.ImagePixelFormat(0, MutableImage.get(), mu::EImageFormat::IF_RGBA_UBYTE);
 			break;
 		default: 
 			// Absolutely worst case
@@ -3388,10 +3390,12 @@ void UCustomizableInstancePrivateData::ConvertImage(UTexture2D* Texture, mu::Ima
 	// Extract a single channel, if requested.
 	if (ExtractChannel >= 0)
 	{
-		MutableImage = mu::ImagePixelFormat( 4, MutableImage.get(), mu::EImageFormat::IF_RGBA_UBYTE );
+		mu::FImageOperator ImOp = mu::FImageOperator::GetDefault();
+
+		MutableImage = ImOp.ImagePixelFormat( 4, MutableImage.get(), mu::EImageFormat::IF_RGBA_UBYTE );
 
 		uint8_t Channel = uint8_t( FMath::Clamp(ExtractChannel,0,3) );
-		MutableImage = mu::ImageSwizzle( mu::EImageFormat::IF_L_UBYTE, &MutableImage, &Channel );
+		MutableImage = ImOp.ImageSwizzle( mu::EImageFormat::IF_L_UBYTE, &MutableImage, &Channel );
 		MutableFormat = mu::EImageFormat::IF_L_UBYTE;
 	}
 
@@ -3401,7 +3405,7 @@ void UCustomizableInstancePrivateData::ConvertImage(UTexture2D* Texture, mu::Ima
 		UE_LOG(LogMutable, Warning, TEXT("Building instance: a texture was generated in RGB format, which is slow to convert to Unreal."));
 
 		// Expand the image.
-		mu::ImagePtr Converted = new mu::Image(MutableImage->GetSizeX(), MutableImage->GetSizeY(), MutableImage->GetLODCount(), mu::EImageFormat::IF_RGBA_UBYTE);
+		mu::ImagePtr Converted = new mu::Image(MutableImage->GetSizeX(), MutableImage->GetSizeY(), MutableImage->GetLODCount(), mu::EImageFormat::IF_RGBA_UBYTE, mu::EInitializationType::NotInitialized);
 
 		for (int LODIndex = 0; LODIndex < Converted->GetLODCount(); ++LODIndex)
 		{
@@ -3427,7 +3431,7 @@ void UCustomizableInstancePrivateData::ConvertImage(UTexture2D* Texture, mu::Ima
 		MUTABLE_CPUPROFILER_SCOPE(Swizzle);
 		// Swizzle the image.
 		// \TODO: Raise a warning?
-		mu::ImagePtr Converted = new mu::Image(MutableImage->GetSizeX(), MutableImage->GetSizeY(), 1, mu::EImageFormat::IF_RGBA_UBYTE);
+		mu::ImagePtr Converted = new mu::Image(MutableImage->GetSizeX(), MutableImage->GetSizeY(), 1, mu::EImageFormat::IF_RGBA_UBYTE, mu::EInitializationType::NotInitialized);
 		int32 pixelCount = MutableImage->GetSizeX() * MutableImage->GetSizeY();
 
 		const uint8_t* pSource = MutableImage->GetData();
diff --git a/Engine/Plugins/Experimental/Mutable/Source/CustomizableObject/Private/MuCO/CustomizableObjectMipDataProvider.cpp b/Engine/Plugins/Experimental/Mutable/Source/CustomizableObject/Private/MuCO/CustomizableObjectMipDataProvider.cpp
--- a/Engine/Plugins/Experimental/Mutable/Source/CustomizableObject/Private/MuCO/CustomizableObjectMipDataProvider.cpp
+++ b/Engine/Plugins/Experimental/Mutable/Source/CustomizableObject/Private/MuCO/CustomizableObjectMipDataProvider.cpp
@@ -196,7 +196,7 @@ namespace impl
 					UE_LOG(LogMutable, Warning, TEXT("Mutable generated an incomplete mip chain for image."));
 
 					// Force the right number of mips. The missing data will be black.
-					mu::Ptr<mu::Image> NewImage = new mu::Image(Image->GetSizeX(), Image->GetSizeY(), FullMipCount, Image->GetFormat());
+					mu::Ptr<mu::Image> NewImage = new mu::Image(Image->GetSizeX(), Image->GetSizeY(), FullMipCount, Image->GetFormat(), mu::EInitializationType::Black);
 					check(NewImage);
 					if (NewImage->GetDataSize() >= Image->GetDataSize())
 					{
diff --git a/Engine/Plugins/Experimental/Mutable/Source/CustomizableObject/Private/MuCO/CustomizableObjectSystem.cpp b/Engine/Plugins/Experimental/Mutable/Source/CustomizableObject/Private/MuCO/CustomizableObjectSystem.cpp
--- a/Engine/Plugins/Experimental/Mutable/Source/CustomizableObject/Private/MuCO/CustomizableObjectSystem.cpp
+++ b/Engine/Plugins/Experimental/Mutable/Source/CustomizableObject/Private/MuCO/CustomizableObjectSystem.cpp
@@ -1566,7 +1566,7 @@ namespace impl
 				const int32 MipSizeY = FMath::Max(Image.FullImageSizeY >> MipsToSkip, 1);
 				if (MipsToSkip > 0 && CustomizableObjectSystemPrivateData->EnableSkipGenerateResidentMips != 0 && OperationData->LowPriorityTextures.Find(Image.Name) != INDEX_NONE)
 				{
-					Image.Image = new mu::Image(MipSizeX, MipSizeY, FullLODCount - MipsToSkip, ImageDesc.m_format);
+					Image.Image = new mu::Image(MipSizeX, MipSizeY, FullLODCount - MipsToSkip, ImageDesc.m_format, mu::EInitializationType::Black);
 				}
 				else
 				{
@@ -1588,7 +1588,7 @@ namespace impl
 				{
 					// Generate a correctly-sized but empty image instead, to avoid crashes.
 					UE_LOG(LogMutable, Warning, TEXT("Mutable generated a wrongly-sized image %d."), Image.ImageID);
-					Image.Image = new mu::Image(MipSizeX, MipSizeY, FullLODCount - MipsToSkip, Image.Image->GetFormat());
+					Image.Image = new mu::Image(MipSizeX, MipSizeY, FullLODCount - MipsToSkip, Image.Image->GetFormat(), mu::EInitializationType::Black);
 				}
 
 				// We need one mip or the complete chain. Otherwise there was a bug.
@@ -1606,7 +1606,7 @@ namespace impl
 					UE_LOG(LogMutable, Warning, TEXT("Mutable generated an incomplete mip chain for image %d."), Image.ImageID);
 
 					// Force the right number of mips. The missing data will be black.
-					mu::Ptr<mu::Image> NewImage = new mu::Image(Image.Image->GetSizeX(), Image.Image->GetSizeY(), FullMipCount, Image.Image->GetFormat());
+					mu::Ptr<mu::Image> NewImage = new mu::Image(Image.Image->GetSizeX(), Image.Image->GetSizeY(), FullMipCount, Image.Image->GetFormat(), mu::EInitializationType::Black);
 					check(NewImage);
 					if (NewImage->GetDataSize() >= Image.Image->GetDataSize())
 					{
diff --git a/Engine/Plugins/Experimental/Mutable/Source/CustomizableObject/Private/MuCO/UnrealMutableImageProvider.cpp b/Engine/Plugins/Experimental/Mutable/Source/CustomizableObject/Private/MuCO/UnrealMutableImageProvider.cpp
--- a/Engine/Plugins/Experimental/Mutable/Source/CustomizableObject/Private/MuCO/UnrealMutableImageProvider.cpp
+++ b/Engine/Plugins/Experimental/Mutable/Source/CustomizableObject/Private/MuCO/UnrealMutableImageProvider.cpp
@@ -35,19 +35,19 @@ namespace
 			break;
 		}
 
-		pResult = new mu::Image(SizeX, SizeY, LODs, MutableFormat);
-
 		// If not locked ReadOnly the Texture Source's FGuid can change, invalidating the texture's caching/shaders
 		// making shader compile and cook times increase
 		const uint8* pSource = Texture->Source.LockMipReadOnly(MipmapsToSkip);
 		if (pSource)
 		{
+			pResult = new mu::Image(SizeX, SizeY, LODs, MutableFormat, mu::EInitializationType::NotInitialized);
 			FMemory::Memcpy(pResult->GetData(), pSource, pResult->GetDataSize());
 			Texture->Source.UnlockMip(MipmapsToSkip);
 		}
 		else
 		{
 			check(false);
+			pResult = new mu::Image(SizeX, SizeY, LODs, MutableFormat, mu::EInitializationType::Black);
 		}
 
 
@@ -72,20 +72,20 @@ namespace
 			break;
 		}
 
-		pResult = new mu::Image(SizeX, SizeY, LODs, MutableFormat);
-
 		// If not locked ReadOnly the Texture Source's FGuid can change, invalidating the texture's caching/shaders
 		// making shader compile and cook times increase
 		const void* pSource = Texture->GetPlatformData()->Mips[MipmapsToSkip].BulkData.LockReadOnly();
 
 		if (pSource)
 		{
+			pResult = new mu::Image(SizeX, SizeY, LODs, MutableFormat, mu::EInitializationType::NotInitialized);
 			FMemory::Memcpy(pResult->GetData(), pSource, pResult->GetDataSize());
 			Texture->GetPlatformData()->Mips[MipmapsToSkip].BulkData.Unlock();
 		}
 		else
 		{
 			check(false);
+			pResult = new mu::Image(SizeX, SizeY, LODs, MutableFormat, mu::EInitializationType::Black);
 		}
 
 #endif
@@ -227,7 +227,7 @@ mu::EImageFormat GetMutablePixelFormat(EPixelFormat InTextureFormat)
 			int SizeX = TextureToLoad->GetSizeX() >> MipIndex;
 			int SizeY = TextureToLoad->GetSizeY() >> MipIndex;
 
-			Image = new mu::Image(SizeX, SizeY, LODs, MutImageFormat);
+			Image = new mu::Image(SizeX, SizeY, LODs, MutImageFormat, mu::EInitializationType::NotInitialized);
 			MutImageDataSize = Image->GetDataSize();
 
 			// In a packaged game the bulk data has to be loaded
@@ -504,7 +504,7 @@ void FUnrealMutableImageProvider::CacheImage(const mu::EXTERNAL_IMAGE_ID& Id, bo
 					case UCustomizableSystemImageProvider::ValueType::Raw:
 					{
 						FIntVector desc = Provider->GetTextureParameterValueSize(Id);
-						pResult = new mu::Image(desc[0], desc[1], 1, mu::EImageFormat::IF_RGBA_UBYTE);
+						pResult = new mu::Image(desc[0], desc[1], 1, mu::EImageFormat::IF_RGBA_UBYTE, mu::EInitializationType::Black);
 						Provider->GetTextureParameterValueData(Id, pResult->GetData());
 						break;
 					}
@@ -659,7 +659,7 @@ mu::ImagePtr FUnrealMutableImageProvider::CreateDummy()
 	uint8_t colours[checkerTileCount][4] = { { 255, 255, 0, 0 },  { 0, 0, 255, 0 } };
 #endif
 
-	mu::ImagePtr pResult = new mu::Image(size, size, DUMMY_IMAGE_DESC.m_lods, DUMMY_IMAGE_DESC.m_format);
+	mu::ImagePtr pResult = new mu::Image(size, size, DUMMY_IMAGE_DESC.m_lods, DUMMY_IMAGE_DESC.m_format, mu::EInitializationType::NotInitialized);
 
 	uint8_t* pData = pResult->GetData();
 	for (int x = 0; x < size; ++x)
diff --git a/Engine/Plugins/Experimental/Mutable/Source/CustomizableObjectEditor/Private/MuCOE/SMutableConstantsWidget.cpp b/Engine/Plugins/Experimental/Mutable/Source/CustomizableObjectEditor/Private/MuCOE/SMutableConstantsWidget.cpp
--- a/Engine/Plugins/Experimental/Mutable/Source/CustomizableObjectEditor/Private/MuCOE/SMutableConstantsWidget.cpp
+++ b/Engine/Plugins/Experimental/Mutable/Source/CustomizableObjectEditor/Private/MuCOE/SMutableConstantsWidget.cpp
@@ -1152,7 +1152,7 @@ void SMutableConstantsWidget::LoadConstantImages()
 		TSharedPtr<FMutableConstantImageElement> ConstantImageElement =MakeShared<FMutableConstantImageElement>();
 
 		MutableProgramPtr->GetConstant(ImageIndex, ConstantImageElement->ImagePtr, 0,
-			[this](int32 x, int32 y, int32 m, mu::EImageFormat f) { return new mu::Image(x, y, m, f); });
+			[this](int32 x, int32 y, int32 m, mu::EImageFormat f, mu::EInitializationType i) { return new mu::Image(x, y, m, f, i); });
 
 		ConstantImageElement->IndexOnSourceVector = ImageIndex;
 		
diff --git a/Engine/Plugins/Experimental/Mutable/Source/CustomizableObjectEditor/Private/MuCOE/UnrealToMutableTextureConversionUtils.cpp b/Engine/Plugins/Experimental/Mutable/Source/CustomizableObjectEditor/Private/MuCOE/UnrealToMutableTextureConversionUtils.cpp
--- a/Engine/Plugins/Experimental/Mutable/Source/CustomizableObjectEditor/Private/MuCOE/UnrealToMutableTextureConversionUtils.cpp
+++ b/Engine/Plugins/Experimental/Mutable/Source/CustomizableObjectEditor/Private/MuCOE/UnrealToMutableTextureConversionUtils.cpp
@@ -251,7 +251,7 @@ TTuple<mu::ImagePtr, EUnrealToMutableConversionError> ConvertTextureUnrealToMuta
 	switch (MutableCompatibleFormat)
 	{
 	case ERawImageFormat::G8:
-		Image = new mu::Image(SizeX, SizeY, LODs, mu::EImageFormat::IF_L_UBYTE);
+		Image = new mu::Image(SizeX, SizeY, LODs, mu::EImageFormat::IF_L_UBYTE, mu::EInitializationType::NotInitialized);
 		FMemory::Memcpy(Image->GetData(), TempImage1.RawData.GetData(), Image->GetDataSize());
 		break;
 
@@ -272,12 +272,12 @@ TTuple<mu::ImagePtr, EUnrealToMutableConversionError> ConvertTextureUnrealToMuta
 		// TODO: If we ever manage to get Pixel Format data on cook compilation time, remove the code that sets bHasAlphaChannel and just use Texture->HasAlphaChannel() here. Currently unreliable, it always returns EPixelFormat::PF_Unknown when cooking, which returns always "false" to HasAlphaChannel().
 		if (bHasAlphaChannel)
 		{
-			Image = new mu::Image(SizeX, SizeY, LODs, mu::EImageFormat::IF_RGBA_UBYTE);
+			Image = new mu::Image(SizeX, SizeY, LODs, mu::EImageFormat::IF_RGBA_UBYTE, mu::EInitializationType::NotInitialized);
 			FMemory::Memcpy(Image->GetData(), TempImage1.RawData.GetData(), Image->GetDataSize());
 		}
 		else
 		{
-			Image = new mu::Image(SizeX, SizeY, LODs, mu::EImageFormat::IF_RGB_UBYTE);
+			Image = new mu::Image(SizeX, SizeY, LODs, mu::EImageFormat::IF_RGB_UBYTE, mu::EInitializationType::NotInitialized);
 			// Manual copy
 			const uint8* DataSource = TempImage1.RawData.GetData();
 			uint8* DataDest = Image->GetData();
diff --git a/Engine/Plugins/Experimental/Mutable/Source/MutableRuntime/Private/MuR/CodeRunner.cpp b/Engine/Plugins/Experimental/Mutable/Source/MutableRuntime/Private/MuR/CodeRunner.cpp
--- a/Engine/Plugins/Experimental/Mutable/Source/MutableRuntime/Private/MuR/CodeRunner.cpp
+++ b/Engine/Plugins/Experimental/Mutable/Source/MutableRuntime/Private/MuR/CodeRunner.cpp
@@ -25,19 +25,14 @@
 #include "MuR/OpImageBinarise.h"
 #include "MuR/OpImageBlend.h"
 #include "MuR/OpImageColourMap.h"
-#include "MuR/OpImageCompose.h"
-#include "MuR/OpImageCrop.h"
-#include "MuR/OpImageDifference.h"
 #include "MuR/OpImageDisplace.h"
 #include "MuR/OpImageGradient.h"
 #include "MuR/OpImageInterpolate.h"
 #include "MuR/OpImageInvert.h"
 #include "MuR/OpImageLuminance.h"
-#include "MuR/OpImageMipmap.h"
 #include "MuR/OpImageNormalCombine.h"
 #include "MuR/OpImageProject.h"
 #include "MuR/OpImageRasterMesh.h"
-#include "MuR/OpImageResize.h"
 #include "MuR/OpImageSaturate.h"
 #include "MuR/OpImageTransform.h"
 #include "MuR/OpLayoutPack.h"
@@ -944,8 +939,10 @@ namespace mu
 
 			int32 MipsToSkip = item.ExecutionOptions;
             Ptr<const Image> Source;
-			program.GetConstant(cat, Source, MipsToSkip,
-				[this](int32 x, int32 y, int32 m, EImageFormat f) { return CreateImage(x,y,m,f); } );
+			program.GetConstant(cat, Source, MipsToSkip, [this](int32 x, int32 y, int32 m, EImageFormat f, EInitializationType i)
+				{
+					return CreateImage(x, y, m, f, i);
+				});
 
 			// Assume the ROM has been loaded previously in a task generated at IssueOp
 			check(Source);
@@ -2987,6 +2984,8 @@ namespace mu
     {
 		MUTABLE_CPUPROFILER_SCOPE(RunCode_Image);
 
+		FImageOperator ImOp = MakeImageOperator(this);
+
 		OP_TYPE type = pModel->GetPrivate()->m_program.GetOpType(item.At);
 		switch (type)
         {
@@ -3126,7 +3125,7 @@ namespace mu
 						Ptr<Image> Formatted = CreateImage( New->GetSizeX(), New->GetSizeY(), New->GetLODCount(), BaseFormat, EInitializationType::NotInitialized );
 
 						bool bSuccess = false;
-						ImagePixelFormat(bSuccess, m_pSettings->ImageCompressionQuality, Formatted.get(), New.get());
+						ImOp.ImagePixelFormat(bSuccess, m_pSettings->ImageCompressionQuality, Formatted.get(), New.get());
 						check(bSuccess); // Decompression cannot fail
 
 						Release(New);
@@ -3187,7 +3186,7 @@ namespace mu
 						Ptr<Image> Formatted = CreateImage(Blended->GetSizeX(), Blended->GetSizeY(), Blended->GetLODCount(), Base->GetFormat(), EInitializationType::NotInitialized);
 
 						bool bSuccess = false;
-						ImagePixelFormat(bSuccess, m_pSettings->ImageCompressionQuality, Formatted.get(), Blended.get());
+						ImOp.ImagePixelFormat(bSuccess, m_pSettings->ImageCompressionQuality, Formatted.get(), Blended.get());
 						check(bSuccess);
 
 						Release(Blended);
@@ -3201,7 +3200,7 @@ namespace mu
 						MUTABLE_CPUPROFILER_SCOPE(ImageResize_BlendedFixForMultilayer);
 
 						Ptr<Image> Resized = CreateImage(ResultSize[0], ResultSize[1], Blended->GetLODCount(), Blended->GetFormat(), EInitializationType::NotInitialized);
-						ImageResizeLinear(Resized.get(), 0, Blended.get());
+						ImOp.ImageResizeLinear(Resized.get(), 0, Blended.get());
 						Release(Blended);
 						Blended = Resized;
 					}
@@ -3250,7 +3249,7 @@ namespace mu
 							MUTABLE_CPUPROFILER_SCOPE(ImageResize_MaskFixForMultilayer);
 
 							Ptr<Image> Resized = CreateImage(ResultSize[0], ResultSize[1], Mask->GetLODCount(), Mask->GetFormat(), EInitializationType::NotInitialized);
-							ImageResizeLinear(Resized.get(), 0, Mask.get());
+							ImOp.ImageResizeLinear(Resized.get(), 0, Mask.get());
 							Release(Mask);
 							Mask = Resized;
 						}
@@ -3401,17 +3400,17 @@ namespace mu
 					MUTABLE_CPUPROFILER_SCOPE(ImageNormalComposite_EmergencyFix);
 
 					int levelCount = Base->GetLODCount();
-					ImagePtr pDest = CreateImage(Normal->GetSizeX(), Normal->GetSizeY(), levelCount, Normal->GetFormat());
+					ImagePtr Dest = CreateImage(Normal->GetSizeX(), Normal->GetSizeY(), levelCount, Normal->GetFormat(), EInitializationType::NotInitialized);
 
 					FMipmapGenerationSettings mipSettings{};
-					ImageMipmap(m_pSystem->WorkingMemoryManager,m_pSettings->ImageCompressionQuality, pDest.get(), Normal.get(), levelCount, mipSettings);
+					ImOp.ImageMipmap(m_pSettings->ImageCompressionQuality, Dest.get(), Normal.get(), levelCount, mipSettings);
 
 					Release(Normal);
-					Normal = pDest;
+					Normal = Dest;
 				}
 
 
-                ImagePtr Result = CreateImage( Base->GetSizeX(), Base->GetSizeY(), Base->GetLODCount(), Base->GetFormat() );
+                ImagePtr Result = CreateImage( Base->GetSizeX(), Base->GetSizeY(), Base->GetLODCount(), Base->GetFormat(), EInitializationType::NotInitialized);
 				ImageNormalComposite(Result.get(), Base.get(), Normal.get(), args.mode, args.power);
 
 				Release(Base);
@@ -3516,7 +3515,7 @@ namespace mu
                 if ( Base->GetSize()!=DestSize )
                 {
 					Ptr<Image> Result = CreateImage(DestSize[0], DestSize[1], Base->GetLODCount(), Base->GetFormat(), EInitializationType::NotInitialized);
-					ImageResizeLinear( Result.get(), m_pSettings->ImageCompressionQuality, Base.get() );
+					ImOp.ImageResizeLinear( Result.get(), m_pSettings->ImageCompressionQuality, Base.get());
 					Release(Base);
 
                     // If the source image had mips, generate them as well for the resized image.
@@ -3527,12 +3526,11 @@ namespace mu
                     if (sourceHasMips)
                     {
                         int levelCount = Image::GetMipmapCount( Result->GetSizeX(), Result->GetSizeY() );
-                        Ptr<Image> Mipmapped = CreateImage( Result->GetSizeX(), Result->GetSizeY(), levelCount, Result->GetFormat() );
+                        Ptr<Image> Mipmapped = CreateImage( Result->GetSizeX(), Result->GetSizeY(), levelCount, Result->GetFormat(), EInitializationType::NotInitialized);
 
 						FMipmapGenerationSettings mipSettings{};
 
-                        ImageMipmap( m_pSystem->WorkingMemoryManager, m_pSettings->ImageCompressionQuality,
-                                     Mipmapped.get(), Result.get(), levelCount, mipSettings );
+						ImOp.ImageMipmap( m_pSettings->ImageCompressionQuality, Mipmapped.get(), Result.get(), levelCount, mipSettings );
 
 						Release(Result);
 						Result = Mipmapped;
@@ -3613,9 +3611,9 @@ namespace mu
 					//ReducedBlockSizeInPixels.X = BlockSizeInPixels.X >> MipsToSkip;
 					//ReducedBlockSizeInPixels.Y = BlockSizeInPixels.Y >> MipsToSkip;
 					//const FImageFormatData& FormatData = GetImageFormatData((EImageFormat)args.format);
-					//int MinBlockSize = FMath::Max(FormatData.m_pixelsPerBlockX, FormatData.m_pixelsPerBlockY);
-					//ReducedBlockSizeInPixels.X = FMath::Max<int32>(ReducedBlockSizeInPixels.X, FormatData.m_pixelsPerBlockX);
-					//ReducedBlockSizeInPixels.Y = FMath::Max<int32>(ReducedBlockSizeInPixels.Y, FormatData.m_pixelsPerBlockY);
+					//int MinBlockSize = FMath::Max(FormatData.PixelsPerBlockX, FormatData.PixelsPerBlockY);
+					//ReducedBlockSizeInPixels.X = FMath::Max<int32>(ReducedBlockSizeInPixels.X, FormatData.PixelsPerBlockX);
+					//ReducedBlockSizeInPixels.Y = FMath::Max<int32>(ReducedBlockSizeInPixels.Y, FormatData.PixelsPerBlockY);
 					//FIntPoint ReducedImageSizeInPixels = SizeInBlocks * ReducedBlockSizeInPixels;
 
 					// This method simply reduces the size and assumes all the other operations will handle degeenrate cases.
@@ -3640,8 +3638,9 @@ namespace mu
                     }
                 }
 
-                ImagePtr pNew = CreateImage(ImageSizeInPixels.X, ImageSizeInPixels.Y, MipsToGenerate, EImageFormat(args.format) );
-                StoreImage( item, pNew );
+				// It needs to be initialized in case it has gaps.
+                ImagePtr New = CreateImage(ImageSizeInPixels.X, ImageSizeInPixels.Y, MipsToGenerate, EImageFormat(args.format), EInitializationType::Black );
+                StoreImage( item, New );
                 break;
             }
 
@@ -3801,7 +3800,7 @@ namespace mu
 						{
 							MUTABLE_CPUPROFILER_SCOPE(ImageResize_ForInterpolate);
 							Ptr<Image> Resized = CreateImage(pNew->GetSizeX(), pNew->GetSizeY(), pMax->GetLODCount(), pMax->GetFormat(), EInitializationType::NotInitialized);
-							ImageResizeLinear(Resized.get(), 0, pMax.get());
+							ImOp.ImageResizeLinear(Resized.get(), 0, pMax.get());
 							Release(pMax);
 							pMax = Resized;
 						}
@@ -3810,10 +3809,10 @@ namespace mu
 						{
 							MUTABLE_CPUPROFILER_SCOPE(Mipmap_ForInterpolate);
 
-							ImagePtr pDest = CreateImage(pNew->GetSizeX(), pNew->GetSizeY(), LevelCount, pNew->GetFormat());
+							ImagePtr pDest = CreateImage(pNew->GetSizeX(), pNew->GetSizeY(), LevelCount, pNew->GetFormat(), EInitializationType::NotInitialized);
 
 							FMipmapGenerationSettings settings{};
-							ImageMipmap(m_pSystem->WorkingMemoryManager, m_pSettings->ImageCompressionQuality, pDest.get(), pNew.get(), LevelCount, settings);
+							ImOp.ImageMipmap(m_pSettings->ImageCompressionQuality, pDest.get(), pNew.get(), LevelCount, settings);
 
 							Release(pNew);
 							pNew = pDest;
@@ -3823,10 +3822,10 @@ namespace mu
 						{
 							MUTABLE_CPUPROFILER_SCOPE(Mipmap_ForInterpolate);
 
-							ImagePtr pDest = CreateImage(pMax->GetSizeX(), pMax->GetSizeY(), LevelCount, pMax->GetFormat());
+							ImagePtr pDest = CreateImage(pMax->GetSizeX(), pMax->GetSizeY(), LevelCount, pMax->GetFormat(), EInitializationType::NotInitialized);
 
 							FMipmapGenerationSettings settings{};
-							ImageMipmap(m_pSystem->WorkingMemoryManager, m_pSettings->ImageCompressionQuality, pDest.get(), pMax.get(), LevelCount, settings);
+							ImOp.ImageMipmap(m_pSettings->ImageCompressionQuality, pDest.get(), pMax.get(), LevelCount, settings);
 
 							Release(pMax);
 							pMax = pDest;
@@ -3965,7 +3964,7 @@ namespace mu
 				{
 					MUTABLE_CPUPROFILER_SCOPE(ImageResize_ForColourmap);
 					Ptr<Image> Resized = CreateImage(Source->GetSizeX(), Source->GetSizeY(), 1, Mask->GetFormat(), EInitializationType::NotInitialized);
-					ImageResizeLinear(Resized.get(), 0, Mask.get());
+					ImOp.ImageResizeLinear(Resized.get(), 0, Mask.get());
 					Release(Mask);
 					Mask = Resized;
 				}
@@ -4012,7 +4011,7 @@ namespace mu
 				FVector4f colour0 = LoadColor(FScheduledOp::FromOpAndOptions(args.colour0, item, 0));
 				FVector4f colour1 = LoadColor(FScheduledOp::FromOpAndOptions(args.colour1, item, 0));
 
-				ImagePtr pResult = CreateImage(args.size[0], args.size[1], 1, EImageFormat::IF_RGB_UBYTE);
+				ImagePtr pResult = CreateImage(args.size[0], args.size[1], 1, EImageFormat::IF_RGB_UBYTE, EInitializationType::NotInitialized);
                 ImageGradient( pResult.get(), colour0, colour1 );
 
 				StoreImage( item, pResult );
@@ -4045,7 +4044,7 @@ namespace mu
 
                 float c = LoadScalar(FScheduledOp::FromOpAndOptions(args.threshold, item, 0));
 
-                Ptr<Image> Result = CreateImage(pA->GetSizeX(), pA->GetSizeY(), pA->GetLODCount(), EImageFormat::IF_L_UBYTE);
+                Ptr<Image> Result = CreateImage(pA->GetSizeX(), pA->GetSizeY(), pA->GetLODCount(), EImageFormat::IF_L_UBYTE, EInitializationType::NotInitialized);
 				ImageBinarise( Result.get(), pA.get(), c );
 
 				Release(pA);
@@ -4116,7 +4115,7 @@ namespace mu
 
                 ImagePtr pA = CreateImage( SizeX, SizeY, FMath::Max(LODs,1), EImageFormat(args.format), EInitializationType::NotInitialized );
 
-                pA->FillColour(c);
+				ImOp.FillColor(pA.get(), c);
 
 				StoreImage( item, pA );
                 break;
@@ -4180,8 +4179,8 @@ namespace mu
 				ImagePtr pResult;
 				if (!rect.IsEmpty())
 				{
-					pResult = CreateImage( rect.size[0], rect.size[1], 1, pA->GetFormat() );
-					ImageCrop(pResult.get(), m_pSettings->ImageCompressionQuality, pA.get(), rect);
+					pResult = CreateImage( rect.size[0], rect.size[1], 1, pA->GetFormat(), EInitializationType::NotInitialized);
+					ImOp.ImageCrop(pResult.get(), m_pSettings->ImageCompressionQuality, pA.get(), rect);
 				}
 
 				Release(pA);
@@ -4242,12 +4241,12 @@ namespace mu
 					// This is usually enforced at object compilation time.
 					if (pResult->GetFormat() != pB->GetFormat())
 					{
-						MUTABLE_CPUPROFILER_SCOPE(ImagPatcheReformat);
+						MUTABLE_CPUPROFILER_SCOPE(ImagPatchReformat);
 
 						EImageFormat format = GetMostGenericFormat(pResult->GetFormat(), pB->GetFormat());
 
 						const FImageFormatData& finfo = GetImageFormatData(format);
-						if (finfo.m_pixelsPerBlockX == 0)
+						if (finfo.PixelsPerBlockX == 0)
 						{
 							format = GetUncompressedFormat(format);
 						}
@@ -4256,7 +4255,7 @@ namespace mu
 						{
 							Ptr<Image> Formatted = CreateImage(pResult->GetSizeX(), pResult->GetSizeY(), pResult->GetLODCount(), format, EInitializationType::NotInitialized);
 							bool bSuccess = false;
-							ImagePixelFormat(bSuccess, m_pSettings->ImageCompressionQuality, Formatted.get(), pResult.get());
+							ImOp.ImagePixelFormat(bSuccess, m_pSettings->ImageCompressionQuality, Formatted.get(), pResult.get());
 							check(bSuccess);
 							Release(pResult);
 							pResult = Formatted;
@@ -4265,7 +4264,7 @@ namespace mu
 						{
 							Ptr<Image> Formatted = CreateImage(pB->GetSizeX(), pB->GetSizeY(), pB->GetLODCount(), format, EInitializationType::NotInitialized);
 							bool bSuccess = false;
-							ImagePixelFormat(bSuccess, m_pSettings->ImageCompressionQuality, Formatted.get(), pB.get());
+							ImOp.ImagePixelFormat(bSuccess, m_pSettings->ImageCompressionQuality, Formatted.get(), pB.get());
 							check(bSuccess);
 							Release(pB);
 						}
@@ -4274,10 +4273,10 @@ namespace mu
 					// Don't patch if below the image compression block size.
 					const FImageFormatData& finfo = GetImageFormatData(pResult->GetFormat());
 					bApplyPatch =
-						(rect.min[0] % finfo.m_pixelsPerBlockX == 0) &&
-						(rect.min[1] % finfo.m_pixelsPerBlockY == 0) &&
-						(rect.size[0] % finfo.m_pixelsPerBlockX == 0) &&
-						(rect.size[1] % finfo.m_pixelsPerBlockY == 0) &&
+						(rect.min[0] % finfo.PixelsPerBlockX == 0) &&
+						(rect.min[1] % finfo.PixelsPerBlockY == 0) &&
+						(rect.size[0] % finfo.PixelsPerBlockX == 0) &&
+						(rect.size[1] % finfo.PixelsPerBlockY == 0) &&
 						(rect.min[0] + rect.size[0]) <= pResult->GetSizeX() &&
 						(rect.min[1] + rect.size[1]) <= pResult->GetSizeY()
 						;
@@ -4285,7 +4284,7 @@ namespace mu
 
 				if (bApplyPatch)
 				{
-					ImageCompose(pResult.get(), pB.get(), rect);
+					ImOp.ImageCompose(pResult.get(), pB.get(), rect);
 					pResult->m_flags = 0;
 				}
 				else
@@ -4358,7 +4357,7 @@ namespace mu
 					}
 
                     // Flat mesh UV raster
-					Ptr<Image> ResultImage = CreateImage(SizeX, SizeY, 1, EImageFormat::IF_L_UBYTE);
+					Ptr<Image> ResultImage = CreateImage(SizeX, SizeY, 1, EImageFormat::IF_L_UBYTE, EInitializationType::Black);
 					if (pMesh)
 					{
 						ImageRasterMesh(pMesh.get(), ResultImage.get(), args.blockIndex, CropMin, UncroppedSize);
@@ -4484,7 +4483,7 @@ namespace mu
 						MUTABLE_CPUPROFILER_SCOPE(ImageResize_MaskFixForProjection);
 
 						Ptr<Image> Resized = CreateImage(SizeX, SizeY, Mask->GetLODCount(), Mask->GetFormat(), EInitializationType::NotInitialized);
-						ImageResizeLinear(Resized.get(), 0, Mask.get());
+						ImOp.ImageResizeLinear(Resized.get(), 0, Mask.get());
 						Release(Mask);
 						Mask = Resized;
 					}
@@ -4508,7 +4507,7 @@ namespace mu
 					MUTABLE_CPUPROFILER_SCOPE(RunCode_RasterMesh_ReformatSource);
 					Ptr<Image> Formatted = CreateImage(Source->GetSizeX(), Source->GetSizeY(), Source->GetLODCount(), Format, EInitializationType::NotInitialized);
 					bool bSuccess = false;
-					ImagePixelFormat(bSuccess, m_pSettings->ImageCompressionQuality, Formatted.get(), Source.get());
+					ImOp.ImagePixelFormat(bSuccess, m_pSettings->ImageCompressionQuality, Formatted.get(), Source.get());
 					check(bSuccess); 
 					Release(Source);
 					Source = Formatted;
@@ -4543,7 +4542,7 @@ namespace mu
 					{
 						MUTABLE_CPUPROFILER_SCOPE(RunCode_RasterMesh_BilinearMipGen);
 
-						Ptr<Image> NewImage = CreateImage(Source->GetSizeX(), Source->GetSizeY(), 2, Source->GetFormat());
+						Ptr<Image> NewImage = CreateImage(Source->GetSizeX(), Source->GetSizeY(), 2, Source->GetFormat(), EInitializationType::NotInitialized);
 
 						check(NewImage->GetDataSize() >= Source->GetDataSize());
 						FMemory::Memcpy(NewImage->GetData(), Source->GetData(), Source->GetDataSize());
@@ -4556,7 +4555,7 @@ namespace mu
 				}
 
 				// Allocate new image after bilinear mip generation to reduce operation memory peak.
-				Ptr<Image> pNew = CreateImage(SizeX, SizeY, 1, Format);
+				Ptr<Image> New = CreateImage(SizeX, SizeY, 1, Format, EInitializationType::Black);
 
 				if (args.projector && Source && Source->GetSizeX() > 0 && Source->GetSizeY() > 0)
 				{
@@ -4565,7 +4564,7 @@ namespace mu
 					switch (Projector.type)
 					{
 					case PROJECTOR_TYPE::PLANAR:
-						ImageRasterProjectedPlanar(pMesh.get(), pNew.get(),
+						ImageRasterProjectedPlanar(pMesh.get(), New.get(),
 							Source.get(), Mask.get(),
 							args.bIsRGBFadingEnabled, args.bIsAlphaFadingEnabled,
 							SamplingMethod,
@@ -4576,7 +4575,7 @@ namespace mu
 						break;
 
 					case PROJECTOR_TYPE::WRAPPING:
-						ImageRasterProjectedWrapping(pMesh.get(), pNew.get(),
+						ImageRasterProjectedWrapping(pMesh.get(), New.get(),
 							Source.get(), Mask.get(),
 							args.bIsRGBFadingEnabled, args.bIsAlphaFadingEnabled,
 							SamplingMethod,
@@ -4587,7 +4586,7 @@ namespace mu
 						break;
 
 					case PROJECTOR_TYPE::CYLINDRICAL:
-						ImageRasterProjectedCylindrical(pMesh.get(), pNew.get(),
+						ImageRasterProjectedCylindrical(pMesh.get(), New.get(),
 							Source.get(), Mask.get(),
 							args.bIsRGBFadingEnabled, args.bIsAlphaFadingEnabled,
 							FadeStartRad, FadeEndRad,
@@ -4606,7 +4605,7 @@ namespace mu
 				Release(pMesh);
 				Release(Source);
 				Release(Mask);
-				StoreImage(item, pNew);
+				StoreImage(item, New);
 
                 break;
             }
@@ -4633,7 +4632,7 @@ namespace mu
 
                 Ptr<const Image> Mask = LoadImage( FCacheAddress(args.mask,item) );
 
-                Ptr<Image> Result = CreateImage( Mask->GetSizeX(), Mask->GetSizeY(), Mask->GetLODCount(), EImageFormat::IF_L_UBYTE);
+                Ptr<Image> Result = CreateImage( Mask->GetSizeX(), Mask->GetSizeY(), Mask->GetLODCount(), EImageFormat::IF_L_UBYTE, EInitializationType::NotInitialized);
 
                 ImageMakeGrowMap(Result.get(), Mask.get(), args.border );
 				Result->m_flags |= Image::IF_CANNOT_BE_SCALED;
@@ -4684,7 +4683,7 @@ namespace mu
 					MUTABLE_CPUPROFILER_SCOPE(ImageResize_EmergencyHackForDisplacementStep1);
 
 					Ptr<Image> Resized = CreateImage(pMap->GetSizeX(), pMap->GetSizeY(), Source->GetLODCount(), Source->GetFormat(), EInitializationType::NotInitialized);
-					ImageResizeLinear(Resized.get(), 0, Source.get());
+					ImOp.ImageResizeLinear(Resized.get(), 0, Source.get());
 					Release(Source);
 					Source = Resized;
 				}
@@ -4701,7 +4700,7 @@ namespace mu
 					{
 						MUTABLE_CPUPROFILER_SCOPE(ImageResize_EmergencyHackForDisplacementStep2);
 						Ptr<Image> Resized = CreateImage(OriginalSourceScale[0], OriginalSourceScale[1], Result->GetLODCount(), Result->GetFormat(), EInitializationType::NotInitialized);
-						ImageResizeLinear(Resized.get(), 0, Result.get());
+						ImOp.ImageResizeLinear(Resized.get(), 0, Result.get());
 						Release(Result);
 						Result = Resized;
 					}
@@ -4764,10 +4763,10 @@ namespace mu
 					static_cast<uint16>(FMath::Clamp(FMath::FloorToInt(float(BaseSize.Y) * FMath::Abs(Scale.Y)), 2, BaseSize.Y)));
 	
 				Ptr<Image> pSampleImage = CreateImage(SampleImageSize[0], SampleImageSize[1], BaseLODs, BaseFormat, EInitializationType::NotInitialized);
-				ImageResizeLinear(pSampleImage.get(), 0, pBaseImage.get());
+				ImOp.ImageResizeLinear( pSampleImage.get(), 0, pBaseImage.get());
 				Release(pBaseImage);
 
-				Ptr<Image> Result = CreateImage(BaseSize.X, BaseSize.Y, 1, BaseFormat);
+				Ptr<Image> Result = CreateImage(BaseSize.X, BaseSize.Y, 1, BaseFormat, EInitializationType::NotInitialized);
 				ImageTransform(Result.get(), pSampleImage.get(), Offset, Scale, Rotation, static_cast<EAddressMode>(Args.AddressMode));
 
 				Release(pSampleImage);
@@ -6128,7 +6127,7 @@ namespace mu
 				EImageFormat NewFormat = args.format;
 				if (args.formatIfAlpha != EImageFormat::IF_NONE
 					&&
-					GetImageFormatData(OldFormat).m_channels > 3)
+					GetImageFormatData(OldFormat).Channels > 3)
 				{
 					NewFormat = args.formatIfAlpha;
 				}
diff --git a/Engine/Plugins/Experimental/Mutable/Source/MutableRuntime/Private/MuR/CodeRunner.h b/Engine/Plugins/Experimental/Mutable/Source/MutableRuntime/Private/MuR/CodeRunner.h
--- a/Engine/Plugins/Experimental/Mutable/Source/MutableRuntime/Private/MuR/CodeRunner.h
+++ b/Engine/Plugins/Experimental/Mutable/Source/MutableRuntime/Private/MuR/CodeRunner.h
@@ -505,13 +505,13 @@ namespace  mu
 			m_pSystem->WorkingMemoryManager.CurrentInstanceCache->SetExtensionData(To, Resource);
 		}
 
-		inline Ptr<Image> CreateImage(uint32 SizeX, uint32 SizeY, uint32 Lods, EImageFormat Format, EInitializationType Init = EInitializationType::Black)
+		inline Ptr<Image> CreateImage(uint32 SizeX, uint32 SizeY, uint32 Lods, EImageFormat Format, EInitializationType Init)
 		{
 			Ptr<Image> Result = m_pSystem->WorkingMemoryManager.CreateImage(SizeX, SizeY, Lods, Format, Init);
 			return MoveTemp(Result);
 		}
 
-		Ptr<Image> CreateImageLike(const Image* Ref, EInitializationType Init = EInitializationType::Black)
+		Ptr<Image> CreateImageLike(const Image* Ref, EInitializationType Init)
 		{
 			Ptr<Image> Result = m_pSystem->WorkingMemoryManager.CreateImage(Ref->GetSizeX(), Ref->GetSizeY(), Ref->GetLODCount(), Ref->GetFormat(), Init);
 			return MoveTemp(Result);
@@ -636,4 +636,31 @@ namespace  mu
 
     };
 
+
+	/** Helper function to create the memory-tracked image operator. */
+	inline FImageOperator MakeImageOperator(CodeRunner* Runner)
+	{
+		return FImageOperator(
+			// Create
+			[Runner](int32 x, int32 y, int32 m, EImageFormat f, EInitializationType i)
+			{
+				return Runner->CreateImage(x, y, m, f, i);
+			},
+
+			// Release
+			[Runner](Ptr<Image>& i)
+			{
+				Runner->Release(i);
+			},
+
+			// Clone
+			[Runner](const Image* i)
+			{
+				Ptr<Image> New = Runner->CreateImage(i->GetSizeX(), i->GetSizeY(), i->GetLODCount(), i->GetFormat(), EInitializationType::NotInitialized);
+				New->Copy(i);
+				return New;
+			}
+		);
+	}
+
 }
diff --git a/Engine/Plugins/Experimental/Mutable/Source/MutableRuntime/Private/MuR/CodeRunnerTasked.cpp b/Engine/Plugins/Experimental/Mutable/Source/MutableRuntime/Private/MuR/CodeRunnerTasked.cpp
--- a/Engine/Plugins/Experimental/Mutable/Source/MutableRuntime/Private/MuR/CodeRunnerTasked.cpp
+++ b/Engine/Plugins/Experimental/Mutable/Source/MutableRuntime/Private/MuR/CodeRunnerTasked.cpp
@@ -23,10 +23,7 @@
 #include "MuR/MutableMath.h"
 #include "MuR/MutableTrace.h"
 #include "MuR/OpImageBlend.h"
-#include "MuR/OpImageCompose.h"
-#include "MuR/OpImageMipmap.h"
 #include "MuR/OpImageNormalCombine.h"
-#include "MuR/OpImageResize.h"
 #include "MuR/OpImageSaturate.h"
 #include "MuR/OpImageInvert.h"
 #include "MuR/Operations.h"
@@ -272,7 +269,7 @@ namespace mu
 			// TODO: We should get the actual format from the arguments for a more precise calculation.
 			OP::ImagePixelFormatArgs args = Program.GetOpArgs<OP::ImagePixelFormatArgs>(Candidate.At);
 			bool bIsCompressed = GetUncompressedFormat(args.format) != args.format;
-			bool bIsSmallish = GetImageFormatData(args.format).m_bytesPerBlock < 2;
+			bool bIsSmallish = GetImageFormatData(args.format).BytesPerBlock < 2;
 			if (bIsCompressed || bIsSmallish)
 			{
 				OpDelta = -10;
@@ -617,7 +614,7 @@ namespace mu
 				}
 
 				// If we reached here it means we didn't find an op to wait for. Try to wait for a loading op.
-				// \todo: unify laoding ops with normal ones?
+				// \todo: unify loading ops with normal ones?
 				for (FRomLoadOp& o : m_romLoadOps)
 				{
 					if (o.m_romIndex >= 0)
@@ -734,6 +731,8 @@ namespace mu
 			return false;
 		}
 
+		FImageOperator ImOp = MakeImageOperator(Runner);
+
 		// Input data fixes
 		InitialFormat = Base->GetFormat();
 
@@ -742,7 +741,7 @@ namespace mu
 			EImageFormat UncompressedFormat = GetUncompressedFormat(InitialFormat);
 			Ptr<Image> Formatted = Runner->CreateImage( Base->GetSizeX(), Base->GetSizeY(), Base->GetLODCount(), UncompressedFormat, EInitializationType::NotInitialized );
 			bool bSuccess = false;
-			ImagePixelFormat(bSuccess, ImageCompressionQuality, Formatted.get(), Base.get());
+			ImOp.ImagePixelFormat(bSuccess, ImageCompressionQuality, Formatted.get(), Base.get());
 			check(bSuccess); // Decompression cannot fail
 			Runner->Release(Base);
 			Base = Formatted;
@@ -752,7 +751,7 @@ namespace mu
 		{
 			Ptr<Image> Formatted = Runner->CreateImage(Blended->GetSizeX(), Blended->GetSizeY(), Blended->GetLODCount(), Base->GetFormat(), EInitializationType::NotInitialized);
 			bool bSuccess = false;
-			ImagePixelFormat(bSuccess, ImageCompressionQuality, Formatted.get(), Blended.get());
+			ImOp.ImagePixelFormat(bSuccess, ImageCompressionQuality, Formatted.get(), Blended.get());
 			check(bSuccess); // Decompression cannot fail
 			Runner->Release(Blended);
 			Blended = Formatted;
@@ -762,7 +761,7 @@ namespace mu
 		{
 			MUTABLE_CPUPROFILER_SCOPE(ImageResize_EmergencyFix);
 			Ptr<Image> Resized = Runner->CreateImage(Base->GetSizeX(), Base->GetSizeY(), 1, Blended->GetFormat(), EInitializationType::NotInitialized);
-			ImageResizeLinear(Resized.get(), ImageCompressionQuality, Blended.get());
+			ImOp.ImageResizeLinear(Resized.get(), ImageCompressionQuality, Blended.get());
 			Runner->Release(Blended);
 			Blended = Resized;
 
@@ -774,7 +773,7 @@ namespace mu
 			{
 				MUTABLE_CPUPROFILER_SCOPE(ImageResize_EmergencyFix);
 				Ptr<Image> Resized = Runner->CreateImage(Base->GetSizeX(), Base->GetSizeY(), 1, Mask->GetFormat(), EInitializationType::NotInitialized);
-				ImageResizeLinear(Resized.get(), ImageCompressionQuality, Mask.get());
+				ImOp.ImageResizeLinear(Resized.get(), ImageCompressionQuality, Mask.get());
 				Runner->Release(Mask);
 				Mask = Resized;
 
@@ -786,15 +785,13 @@ namespace mu
 				MUTABLE_CPUPROFILER_SCOPE(ImageLayer_EmergencyFix);
 
 				int32 levelCount = Base->GetLODCount();
-				ImagePtr pDest = Runner->CreateImage(Mask->GetSizeX(), Mask->GetSizeY(),
-					levelCount,
-					Mask->GetFormat());
+				Ptr<Image> Dest = Runner->CreateImage(Mask->GetSizeX(), Mask->GetSizeY(), levelCount, Mask->GetFormat(), EInitializationType::NotInitialized);
 
 				FMipmapGenerationSettings settings{};
-				ImageMipmap(Runner->m_pSystem->WorkingMemoryManager, ImageCompressionQuality, pDest.get(), Mask.get(), levelCount, settings);
+				ImOp.ImageMipmap(ImageCompressionQuality, Dest.get(), Mask.get(), levelCount, settings);
 
 				Runner->Release(Mask);
-				Mask = pDest;
+				Mask = Dest;
 			}
 		}
 
@@ -958,7 +955,9 @@ namespace mu
 			{
 				Ptr<Image> Formatted = Runner->CreateImage(Result->GetSizeX(), Result->GetSizeY(), Result->GetLODCount(), InitialFormat, EInitializationType::NotInitialized);
 				bool bSuccess = false;
-				ImagePixelFormat(bSuccess, ImageCompressionQuality, Formatted.get(), Result.get());
+
+				FImageOperator ImOp = MakeImageOperator(Runner);
+				ImOp.ImagePixelFormat(bSuccess, ImageCompressionQuality, Formatted.get(), Result.get());
 				check(bSuccess);
 
 				Runner->Release(Result);
@@ -1035,11 +1034,13 @@ namespace mu
 
 		if (Args.mask && Mask)
 		{
+			FImageOperator ImOp = MakeImageOperator(Runner);
+
 			if (Base->GetSize() != Mask->GetSize())
 			{
 				MUTABLE_CPUPROFILER_SCOPE(ImageResize_EmergencyFixSize);
 				Ptr<Image> Resized = Runner->CreateImage(Base->GetSizeX(), Base->GetSizeY(), 1, Mask->GetFormat(), EInitializationType::NotInitialized);
-				ImageResizeLinear(Resized.get(), ImageCompressionQuality, Mask.get());
+				ImOp.ImageResizeLinear(Resized.get(), ImageCompressionQuality, Mask.get() );
 				Runner->Release(Mask);
 				Mask = Resized;
 			}
@@ -1049,10 +1050,10 @@ namespace mu
 			{
 				MUTABLE_CPUPROFILER_SCOPE(ImageResize_EmergencyFixMips);
 				int32 levelCount = Base->GetLODCount();
-				Ptr<Image> pDest = Runner->CreateImage(Mask->GetSizeX(), Mask->GetSizeY(), levelCount, Mask->GetFormat());
+				Ptr<Image> pDest = Runner->CreateImage(Mask->GetSizeX(), Mask->GetSizeY(), levelCount, Mask->GetFormat(), EInitializationType::NotInitialized);
 
 				FMipmapGenerationSettings settings{};
-				ImageMipmap(Runner->m_pSystem->WorkingMemoryManager, ImageCompressionQuality, pDest.get(), Mask.get(), levelCount, settings);
+				ImOp.ImageMipmap(ImageCompressionQuality, pDest.get(), Mask.get(), levelCount, settings);
 
 				Runner->Release(Mask);
 				Mask = pDest;
@@ -1130,7 +1131,13 @@ namespace mu
 					case EBlendType::BT_OVERLAY: BufferLayerColour<OverlayChannel>(Result.get(), Result.get(), Color); break;
 					case EBlendType::BT_LIGHTEN: BufferLayerColour<LightenChannel>(Result.get(), Result.get(), Color); break;
 					case EBlendType::BT_MULTIPLY: BufferLayerColour<MultiplyChannel>(Result.get(), Result.get(), Color); break;
-					case EBlendType::BT_BLEND: Result->FillColour(Color); break;
+					case EBlendType::BT_BLEND: 
+					{
+						// In this case we know it is already an uncompressed image, and we won't need additional allocations;
+						FImageOperator ImOp = FImageOperator::GetDefault();
+						ImOp.FillColor( Result.get(), Color); 
+						break;
+					}
 					default: check(false);
 					}
 				}
@@ -1248,7 +1255,7 @@ namespace mu
 		TargetFormat = Args.format;
 		if (Args.formatIfAlpha != EImageFormat::IF_NONE
 			&&
-			GetImageFormatData(Base->GetFormat()).m_channels > 3)
+			GetImageFormatData(Base->GetFormat()).Channels > 3)
 		{
 			TargetFormat = Args.formatIfAlpha;
 		}
@@ -1271,7 +1278,8 @@ namespace mu
 		MUTABLE_CPUPROFILER_SCOPE(FImagePixelFormatTask);
 		
 		bool bSuccess = false;
-		ImagePixelFormat(bSuccess, ImageCompressionQuality, Result.get(), Base.get(), -1);
+		FImageOperator ImOp = FImageOperator::GetDefault();
+		ImOp.ImagePixelFormat(bSuccess, ImageCompressionQuality, Result.get(), Base.get(), -1);
 
 		int32 OriginalDataSize = FMath::Max(Result->m_data.Num(), Base->m_data.Num());
 		int32 ExcessDataSize = OriginalDataSize * 4 + 4;
@@ -1283,7 +1291,7 @@ namespace mu
 			// We need to support it anyway for small mips or scaled images.
 			Result->m_data.SetNum(ExcessDataSize);
 			bSuccess = false;
-			ImagePixelFormat(bSuccess,ImageCompressionQuality, Result.get(), Base.get(), -1);
+			ImOp.ImagePixelFormat(bSuccess,ImageCompressionQuality, Result.get(), Base.get(), -1);
 			ExcessDataSize *= 4;
 		}
 
@@ -1323,7 +1331,7 @@ namespace mu
 		OP::ImageMipmapArgs Args;
 		Ptr<const Image> Base;
 		Ptr<Image> Result;
-		FScratchImageMipmap Scratch;
+		FImageOperator::FScratchImageMipmap Scratch;
 	};
 
 
@@ -1379,7 +1387,8 @@ namespace mu
 		Result = Runner->CreateImage(Base->GetSizeX(), Base->GetSizeY(), LevelCount, Base->GetFormat(), EInitializationType::NotInitialized);
 		Result->m_flags = Base->m_flags;
 
-		ImageMipmap_PrepareScratch( Runner->m_pSystem->WorkingMemoryManager, Result.get(), Base.get(), LevelCount, Scratch);
+		FImageOperator ImOp = MakeImageOperator(Runner);
+		ImOp.ImageMipmap_PrepareScratch( Result.get(), Base.get(), LevelCount, Scratch);
 
 		return true;
 	}
@@ -1391,15 +1400,17 @@ namespace mu
 		// This runs in a worker thread
 		MUTABLE_CPUPROFILER_SCOPE(FImageMipmapTask);
 
+		FImageOperator ImOp = FImageOperator::GetDefault();
 		FMipmapGenerationSettings settings{};
-		ImageMipmap(Scratch, ImageCompressionQuality, Result.get(), Base.get(), Result->GetLODCount(), settings);
+		ImOp.ImageMipmap(Scratch, ImageCompressionQuality, Result.get(), Base.get(), Result->GetLODCount(), settings);
 	}
 
 
 	//---------------------------------------------------------------------------------------------
 	void FImageMipmapTask::Complete(CodeRunner* Runner)
 	{
-		ImageMipmap_ReleaseScratch(Runner->m_pSystem->WorkingMemoryManager, Scratch);
+		FImageOperator ImOp = MakeImageOperator(Runner);
+		ImOp.ImageMipmap_ReleaseScratch(Scratch);
 
 		// This runs in the Runner thread
 		Runner->Release(Base);
@@ -1462,6 +1473,10 @@ namespace mu
 		// Create destination data
 		EImageFormat format = (EImageFormat)Args.format;
 
+		FImageOperator ImOp = MakeImageOperator(Runner);
+
+		int32 ResultLODs = Sources[0]->GetLODCount();
+
 		// Be defensive: ensure image sizes match.
 		for (int i = 1; i < MUTABLE_OP_MAX_SWIZZLE_CHANNELS; ++i)
 		{
@@ -1469,13 +1484,14 @@ namespace mu
 			{
 				MUTABLE_CPUPROFILER_SCOPE(ImageResize_ForSwizzle);
 				Ptr<Image> Resized = Runner->CreateImage(Sources[0]->GetSizeX(), Sources[0]->GetSizeY(), 1, Sources[i]->GetFormat(), EInitializationType::NotInitialized);
-				ImageResizeLinear(Resized.get(), 0, Sources[i].get());
+				ImOp.ImageResizeLinear(Resized.get(), 0, Sources[i].get());
 				Runner->Release(Sources[i]);
 				Sources[i] = Resized;
+				ResultLODs = 1;
 			}
 		}
 
-		Result = Runner->CreateImage(Sources[0]->GetSizeX(), Sources[0]->GetSizeY(), Sources[0]->GetLODCount(), Args.format, EInitializationType::Black);
+		Result = Runner->CreateImage(Sources[0]->GetSizeX(), Sources[0]->GetSizeY(), ResultLODs, Args.format, EInitializationType::Black);
 		return true;
 	}
 
@@ -1497,6 +1513,8 @@ namespace mu
 			Runner->Release(Sources[i]);
 		}
 
+		// \TODO: If Result LODs differ from Source[0]'s, rebuild mips?
+
 		// If no shortcut was taken
 		if (Result)
 		{
@@ -1663,9 +1681,8 @@ namespace mu
 
 		// Warning: This will actually allocate temp memory that may exceed the budget.
 		// \TODO: Fix it.
-		// Warning: This may allocate memory if the image is compressed.
-		// \TODO: treat this case in the Prepare.
-		ImageResizeLinear( Result.get(), ImageCompressionQuality, Base.get() );
+		FImageOperator ImOp = FImageOperator::GetDefault();
+		ImOp.ImageResizeLinear( Result.get(), ImageCompressionQuality, Base.get());
 
 		int32 LodCount = Result->GetLODCount();
 		if (LodCount>1)
@@ -1762,7 +1779,9 @@ namespace mu
 		// This runs on a random worker thread
 		MUTABLE_CPUPROFILER_SCOPE(FImageResizeRelTask);
 
-		ImageResizeLinear(Result.get(), ImageCompressionQuality, Base.get());
+		// \TODO: Track allocs
+		FImageOperator ImOp = FImageOperator::GetDefault();
+		ImOp.ImageResizeLinear(Result.get(), ImageCompressionQuality, Base.get());
 
 		int32 LodCount = Result->GetLODCount();
 		if (LodCount > 1)
@@ -1927,6 +1946,8 @@ namespace mu
 		{
 			MUTABLE_CPUPROFILER_SCOPE(ImageComposeWithoutMask);
 
+			FImageOperator ImOp = MakeImageOperator(Runner);
+
 			EImageFormat Format = GetMostGenericFormat(Result->GetFormat(), Block->GetFormat());
 
 			// Resize image if it doesn't fit in the new block size
@@ -1937,7 +1958,7 @@ namespace mu
 				// This now happens more often since the generation of specific mips on request. For this reason
 				// this warning is usually acceptable.
 				Ptr<Image> Resized = Runner->CreateImage(Rect.size[0], Rect.size[1], 1, Block->GetFormat(), EInitializationType::NotInitialized );
-				ImageResizeLinear(Resized.get(), ImageCompressionQuality, Block.get() );
+				ImOp.ImageResizeLinear(Resized.get(), ImageCompressionQuality, Block.get());
 				Runner->Release(Block);
 				Block = Resized;
 			}
@@ -1952,7 +1973,7 @@ namespace mu
 				{
 					Ptr<Image> Formatted = Runner->CreateImage(Result->GetSizeX(), Result->GetSizeY(), Result->GetLODCount(), Format, EInitializationType::NotInitialized);
 					bool bSuccess = false;
-					ImagePixelFormat(bSuccess, ImageCompressionQuality, Formatted.get(), Result.get());
+					ImOp.ImagePixelFormat(bSuccess, ImageCompressionQuality, Formatted.get(), Result.get());
 					check(bSuccess); // Decompression cannot fail
 					Runner->Release(Result);
 					Result = Formatted;
@@ -1961,7 +1982,7 @@ namespace mu
 				{
 					Ptr<Image> Formatted = Runner->CreateImage(Block->GetSizeX(), Block->GetSizeY(), Block->GetLODCount(), Format, EInitializationType::NotInitialized);
 					bool bSuccess = false;
-					ImagePixelFormat(bSuccess, ImageCompressionQuality, Formatted.get(), Block.get());
+					ImOp.ImagePixelFormat(bSuccess, ImageCompressionQuality, Formatted.get(), Block.get());
 					check(bSuccess); // Decompression cannot fail
 					Runner->Release(Block);
 					Block = Formatted;
@@ -1984,7 +2005,9 @@ namespace mu
 			MUTABLE_CPUPROFILER_SCOPE(ImageComposeWithoutMask);
 
 			// Compose without a mask
-			ImageCompose(Result.get(), Block.get(), Rect);
+			// \TODO: track allocs
+			FImageOperator ImOp = FImageOperator::GetDefault();
+			ImOp.ImageCompose(Result.get(), Block.get(), Rect);
 		}
 		else
 		{
diff --git a/Engine/Plugins/Experimental/Mutable/Source/MutableRuntime/Private/MuR/Image.cpp b/Engine/Plugins/Experimental/Mutable/Source/MutableRuntime/Private/MuR/Image.cpp
--- a/Engine/Plugins/Experimental/Mutable/Source/MutableRuntime/Private/MuR/Image.cpp
+++ b/Engine/Plugins/Experimental/Mutable/Source/MutableRuntime/Private/MuR/Image.cpp
@@ -3,12 +3,12 @@
 #include "MuR/Image.h"
 
 #include "HAL/UnrealMemory.h"
+#include "HAL/IConsoleManager.h"
 #include "Math/UnrealMathSSE.h"
 #include "Math/NumericLimits.h"
 #include "MuR/ImagePrivate.h"
 #include "MuR/MutableMath.h"
 #include "MuR/MutableTrace.h"
-#include "MuR/OpImageResize.h"
 
 #include <initializer_list>
 
@@ -81,6 +81,29 @@ namespace mu
     }
 
 
+	//---------------------------------------------------------------------------------------------
+	void FMipmapGenerationSettings::Serialise(OutputArchive& arch) const
+	{
+		uint32 ver = 0;
+		arch << ver;
+
+		arch << m_sharpenFactor;
+		arch << m_filterType;
+		arch << m_ditherMipmapAlpha;
+	}
+
+	void FMipmapGenerationSettings::Unserialise(InputArchive& arch)
+	{
+		uint32 ver = 0;
+		arch >> ver;
+		check(ver == 0);
+
+		arch >> m_sharpenFactor;
+		arch >> m_filterType;
+		arch >> m_ditherMipmapAlpha;
+	}
+
+
     //---------------------------------------------------------------------------------------------
     Image::Image()
     {
@@ -96,16 +119,18 @@ namespace mu
         check(format != EImageFormat::IF_NONE);
 		check(format < EImageFormat::IF_COUNT);
 
-		check(sizeX <= TNumericLimits<uint16>::Max())
-		check(sizeY <= TNumericLimits<uint16>::Max())
-		check(lods <= TNumericLimits<uint8>::Max())
+		check(sizeX <= TNumericLimits<uint16>::Max());
+		check(sizeY <= TNumericLimits<uint16>::Max());
+		check(lods <= TNumericLimits<uint8>::Max());
+
+		// TODO: check that lods is sensible for the size.
 
         m_format = format;
         m_size = FImageSize( (uint16)sizeX, (uint16)sizeY );
-        m_lods = (uint8_t)lods;
+        m_lods = (uint8)lods;
 
         const FImageFormatData& fdata = GetImageFormatData( format );
-        int32 PixelsPerBlock = fdata.m_pixelsPerBlockX*fdata.m_pixelsPerBlockY;
+        int32 PixelsPerBlock = fdata.PixelsPerBlockX*fdata.PixelsPerBlockY;
         if (PixelsPerBlock)
         {
 			if (Init == EInitializationType::Black)
@@ -134,11 +159,11 @@ namespace mu
 		else
 		{
 			const FImageFormatData& fdata = GetImageFormatData(m_format);
-			check(fdata.m_bytesPerBlock <= FImageFormatData::MAX_BYTES_PER_BLOCK);
+			check(fdata.BytesPerBlock <= FImageFormatData::MAX_BYTES_PER_BLOCK);
 
 			constexpr uint8 ZeroedBlock[FImageFormatData::MAX_BYTES_PER_BLOCK] = { 0 };
 
-			const SIZE_T BlockSizeSanitized = FMath::Min<SIZE_T>(FImageFormatData::MAX_BYTES_PER_BLOCK, fdata.m_bytesPerBlock);
+			const SIZE_T BlockSizeSanitized = FMath::Min<SIZE_T>(FImageFormatData::MAX_BYTES_PER_BLOCK, fdata.BytesPerBlock);
 			const bool bIsFormatBlackBlockZeroed = FMemory::Memcmp(fdata.BlackBlock, ZeroedBlock, BlockSizeSanitized) == 0;
 
 			if (bIsFormatBlackBlockZeroed)
@@ -151,10 +176,10 @@ namespace mu
 			{
 				m_data.SetNumUninitialized(DataSize);
 
-				check(fdata.m_bytesPerBlock > 0);
-				for (int32 BlockDataOffset = 0; BlockDataOffset < DataSize; BlockDataOffset += fdata.m_bytesPerBlock)
+				check(fdata.BytesPerBlock > 0);
+				for (int32 BlockDataOffset = 0; BlockDataOffset < DataSize; BlockDataOffset += fdata.BytesPerBlock)
 				{
-					FMemory::Memcpy(m_data.GetData() + BlockDataOffset, fdata.BlackBlock, fdata.m_bytesPerBlock);
+					FMemory::Memcpy(m_data.GetData() + BlockDataOffset, fdata.BlackBlock, fdata.BytesPerBlock);
 				}
 			}
 		}
@@ -224,45 +249,50 @@ namespace mu
 
 
 	//---------------------------------------------------------------------------------------------
-	Ptr<Image> Image::ExtractMip(int32 Mip) const
+	Ptr<Image> FImageOperator::ExtractMip(const Image* This, int32 Mip)
 	{
-		if (Mip == 0 && m_lods == 1)
+		if (Mip == 0 && This->m_lods == 1)
 		{
-			return Clone();
+			return CloneImage(This);
 		}
 
-		FIntVector2 MipSize = CalculateMipSize(Mip);
+		FIntVector2 MipSize = This->CalculateMipSize(Mip);
 
 		int32 Quality = 4;
 
-		if (m_lods > Mip)
+		if (This->m_lods > Mip)
 		{
-			const uint8* SourceData = GetMipData(Mip);
-			int32 DataSize = CalculateDataSize(Mip);
+			const uint8* SourceData = This->GetMipData(Mip);
+			int32 DataSize = This->CalculateDataSize(Mip);
 
 			if (DataSize)
 			{
-				Ptr<Image> pResult = new Image(MipSize[0], MipSize[1], 1, m_format );
-				pResult->m_flags = m_flags;
+				Ptr<Image> pResult = CreateImage(MipSize[0], MipSize[1], 1, This->m_format, EInitializationType::NotInitialized);
+				pResult->m_flags = This->m_flags;
 				uint8* DestData = pResult->GetData();
 				FMemory::Memcpy(DestData, SourceData, DataSize);
 				return pResult;
 			}
 			else
 			{
-				EImageFormat uncompressedFormat = GetUncompressedFormat(m_format);
+				EImageFormat uncompressedFormat = GetUncompressedFormat(This->m_format);
 				// TODO: OnlyLOD=Mip and then ExtractMip(0)?
-				Ptr<Image> pTemp = ImagePixelFormat(Quality, this, uncompressedFormat);
-				pTemp = pTemp->ExtractMip(Mip);
-				Ptr<Image> pResult = ImagePixelFormat(Quality, pTemp.get(), m_format);
-				pResult->m_flags = m_flags;
-				return pResult;
+				Ptr<Image> Temp = ImagePixelFormat(Quality, This, uncompressedFormat);
+				Temp = ExtractMip(Temp.get(), Mip);
+				Ptr<Image> Result = ImagePixelFormat(Quality, Temp.get(), This->m_format);
+				Result->m_flags = This->m_flags;
+				ReleaseImage(Temp);
+				return Result;
 			}
 		}
 
 		// We need to generate the mip
 		// \TODO: optimize, quality
-		return ImageResizeLinear(Quality, this, FImageSize(MipSize[0],MipSize[1]))->ExtractMip(0);
+		Ptr<Image> Resized = CreateImage(MipSize[0], MipSize[1],1,This->GetFormat(),EInitializationType::NotInitialized);
+		ImageResizeLinear(Resized.get(), Quality, This);
+		Ptr<Image> Result = ExtractMip( Resized.get(), 0);
+		ReleaseImage(Resized);
+		return Result;
 	}
 
 
@@ -350,14 +380,14 @@ namespace mu
 		int32 res = 0;
 
 		const FImageFormatData& fdata = GetImageFormatData(Format);
-		if (fdata.m_bytesPerBlock)
+		if (fdata.BytesPerBlock)
 		{
 			for (int32 LodIndex = 0; LodIndex < FMath::Max(1, LodCount); ++LodIndex)
 			{
-				int32 blocksX = FMath::DivideAndRoundUp(SizeX, int32(fdata.m_pixelsPerBlockX));
-				int32 blocksY = FMath::DivideAndRoundUp(SizeY, int32(fdata.m_pixelsPerBlockY));
+				int32 blocksX = FMath::DivideAndRoundUp(SizeX, int32(fdata.PixelsPerBlockX));
+				int32 blocksY = FMath::DivideAndRoundUp(SizeY, int32(fdata.PixelsPerBlockY));
 
-				res += (blocksX * blocksY) * fdata.m_bytesPerBlock;
+				res += (blocksX * blocksY) * fdata.BytesPerBlock;
 
 				SizeX = FMath::DivideAndRoundUp(SizeX, 2);
 				SizeY = FMath::DivideAndRoundUp(SizeY, 2);
@@ -381,18 +411,18 @@ namespace mu
 		int32 res = 0;
 
         const FImageFormatData& fdata = GetImageFormatData( m_format );
-        if (fdata.m_bytesPerBlock)
+        if (fdata.BytesPerBlock)
         {
             FImageSize s = m_size;
 
             for ( int l=0; l< FMath::Max(1,(int)m_lods); ++l )
             {
-                int blocksX = FMath::DivideAndRoundUp( s[0], (uint16)fdata.m_pixelsPerBlockX );
-                int blocksY = FMath::DivideAndRoundUp( s[1], (uint16)fdata.m_pixelsPerBlockY );
+                int blocksX = FMath::DivideAndRoundUp( s[0], (uint16)fdata.PixelsPerBlockX );
+                int blocksY = FMath::DivideAndRoundUp( s[1], (uint16)fdata.PixelsPerBlockY );
 
                 if ( mip==l )
                 {
-                    res = (blocksX*blocksY) * fdata.m_bytesPerBlock;
+                    res = (blocksX*blocksY) * fdata.BytesPerBlock;
                     break;
                 }
 
@@ -411,16 +441,16 @@ namespace mu
 		int32 res = 0;
 
         const FImageFormatData& fdata = GetImageFormatData( m_format );
-        if (fdata.m_bytesPerBlock)
+        if (fdata.BytesPerBlock)
         {
             FImageSize s = m_size;
 
             for ( int l=0; l<FMath::Max(1,(int)m_lods); ++l )
             {
-                int blocksX = FMath::DivideAndRoundUp( s[0], (uint16)fdata.m_pixelsPerBlockX );
-                int blocksY = FMath::DivideAndRoundUp( s[1], (uint16)fdata.m_pixelsPerBlockY );
+                int blocksX = FMath::DivideAndRoundUp( s[0], (uint16)fdata.PixelsPerBlockX );
+                int blocksY = FMath::DivideAndRoundUp( s[1], (uint16)fdata.PixelsPerBlockY );
 
-                res += (blocksX*blocksY) * fdata.m_pixelsPerBlockX * fdata.m_pixelsPerBlockY;
+                res += (blocksX*blocksY) * fdata.PixelsPerBlockX * fdata.PixelsPerBlockY;
 
                 s[0] = FMath::DivideAndRoundUp( s[0], (uint16)2 );
                 s[1] = FMath::DivideAndRoundUp( s[1], (uint16)2 );
@@ -446,18 +476,18 @@ namespace mu
 		int32 res = 0;
 
         const FImageFormatData& fdata = GetImageFormatData( m_format );
-        if (fdata.m_bytesPerBlock)
+        if (fdata.BytesPerBlock)
         {
             FImageSize s = m_size;
 
             for ( int l=0; l<mip+1; ++l )
             {
-                int blocksX = FMath::DivideAndRoundUp( s[0], (uint16)fdata.m_pixelsPerBlockX );
-                int blocksY = FMath::DivideAndRoundUp( s[1], (uint16)fdata.m_pixelsPerBlockY );
+                int blocksX = FMath::DivideAndRoundUp( s[0], (uint16)fdata.PixelsPerBlockX );
+                int blocksY = FMath::DivideAndRoundUp( s[1], (uint16)fdata.PixelsPerBlockY );
 
                 if ( l==mip )
                 {
-                    res = (blocksX*blocksY) * fdata.m_pixelsPerBlockX * fdata.m_pixelsPerBlockY;
+                    res = (blocksX*blocksY) * fdata.PixelsPerBlockX * fdata.PixelsPerBlockY;
                 }
 
                 s[0] = FMath::DivideAndRoundUp( s[0], (uint16)2 );
@@ -521,17 +551,17 @@ namespace mu
 
 		const FImageFormatData& fdata = GetImageFormatData(m_format);
 
-		if (fdata.m_bytesPerBlock)
+		if (fdata.BytesPerBlock)
 		{
 			// Fixed-sized formats
 			FImageSize s = m_size;
 
 			for (int l = 0; l < mip; ++l)
 			{
-				int blocksX = FMath::DivideAndRoundUp(s[0], (uint16)fdata.m_pixelsPerBlockX);
-				int blocksY = FMath::DivideAndRoundUp(s[1], (uint16)fdata.m_pixelsPerBlockY);
+				int blocksX = FMath::DivideAndRoundUp(s[0], (uint16)fdata.PixelsPerBlockX);
+				int blocksY = FMath::DivideAndRoundUp(s[1], (uint16)fdata.PixelsPerBlockY);
 
-				pResult += (blocksX * blocksY) * fdata.m_bytesPerBlock;
+				pResult += (blocksX * blocksY) * fdata.BytesPerBlock;
 
 				s[0] = FMath::DivideAndRoundUp(s[0], (uint16)2);
 				s[1] = FMath::DivideAndRoundUp(s[1], (uint16)2);
@@ -566,17 +596,17 @@ namespace mu
 
 		const FImageFormatData& fdata = GetImageFormatData(m_format);
 
-		if (fdata.m_bytesPerBlock)
+		if (fdata.BytesPerBlock)
 		{
 			// Fixed-sized formats
 			FImageSize s = m_size;
 
 			for (int l = 0; l < mip; ++l)
 			{
-				int blocksX = FMath::DivideAndRoundUp(s[0], (uint16)fdata.m_pixelsPerBlockX);
-				int blocksY = FMath::DivideAndRoundUp(s[1], (uint16)fdata.m_pixelsPerBlockY);
+				int blocksX = FMath::DivideAndRoundUp(s[0], (uint16)fdata.PixelsPerBlockX);
+				int blocksY = FMath::DivideAndRoundUp(s[1], (uint16)fdata.PixelsPerBlockY);
 
-				pResult += (blocksX * blocksY) * fdata.m_bytesPerBlock;
+				pResult += (blocksX * blocksY) * fdata.BytesPerBlock;
 
 				s[0] = FMath::DivideAndRoundUp(s[0], (uint16)2);
 				s[1] = FMath::DivideAndRoundUp(s[1], (uint16)2);
@@ -606,7 +636,7 @@ namespace mu
 		int32 res = 0;
 
         const FImageFormatData& fdata = GetImageFormatData( m_format );
-        if (fdata.m_bytesPerBlock)
+        if (fdata.BytesPerBlock)
         {
 			return CalculateDataSize();
         }
@@ -643,21 +673,21 @@ namespace mu
         const FImageFormatData& fdata = GetImageFormatData( m_format );
 
         int pixelX = FMath::Max( 0, FMath::Min( m_size[0]-1, (int)(coords[0] * m_size[0] )));
-        int blockX = pixelX / fdata.m_pixelsPerBlockX;
-        int blockPixelX = pixelX % fdata.m_pixelsPerBlockX;
+        int blockX = pixelX / fdata.PixelsPerBlockX;
+        int blockPixelX = pixelX % fdata.PixelsPerBlockX;
 
         int pixelY = FMath::Max( 0, FMath::Min( m_size[1]-1, (int)(coords[1] * m_size[1] )));
-        int blockY = pixelY / fdata.m_pixelsPerBlockY;
-        int blockPixelY = pixelY % fdata.m_pixelsPerBlockY;
+        int blockY = pixelY / fdata.PixelsPerBlockY;
+        int blockPixelY = pixelY % fdata.PixelsPerBlockY;
 
-        int blocksPerRow = m_size[0] / fdata.m_pixelsPerBlockX;
+        int blocksPerRow = m_size[0] / fdata.PixelsPerBlockX;
         int blockOffset = blockX + blockY * blocksPerRow;
 
         // Non-generic part
         if ( m_format== EImageFormat::IF_RGB_UBYTE )
         {
-            int byteOffset = blockOffset * fdata.m_bytesPerBlock
-                    + ( blockPixelY * fdata.m_pixelsPerBlockX + blockPixelX ) * 3;
+            int byteOffset = blockOffset * fdata.BytesPerBlock
+                    + ( blockPixelY * fdata.PixelsPerBlockX + blockPixelX ) * 3;
 
             result[0] = m_data[ byteOffset+0 ] / 255.0f;
             result[1] = m_data[ byteOffset+1 ] / 255.0f;
@@ -666,8 +696,8 @@ namespace mu
         }
         else if ( m_format== EImageFormat::IF_RGBA_UBYTE )
         {
-            int byteOffset = blockOffset * fdata.m_bytesPerBlock
-                    + ( blockPixelY * fdata.m_pixelsPerBlockX + blockPixelX ) * 4;
+            int byteOffset = blockOffset * fdata.BytesPerBlock
+                    + ( blockPixelY * fdata.PixelsPerBlockX + blockPixelX ) * 4;
 
             result[0] = m_data[ byteOffset+0 ] / 255.0f;
             result[1] = m_data[ byteOffset+1 ] / 255.0f;
@@ -676,8 +706,8 @@ namespace mu
         }
         else if ( m_format== EImageFormat::IF_BGRA_UBYTE )
         {
-            int byteOffset = blockOffset * fdata.m_bytesPerBlock
-                    + ( blockPixelY * fdata.m_pixelsPerBlockX + blockPixelX ) * 4;
+            int byteOffset = blockOffset * fdata.BytesPerBlock
+                    + ( blockPixelY * fdata.PixelsPerBlockX + blockPixelX ) * 4;
 
             result[0] = m_data[ byteOffset+2 ] / 255.0f;
             result[1] = m_data[ byteOffset+1 ] / 255.0f;
@@ -686,8 +716,8 @@ namespace mu
         }
         else if ( m_format== EImageFormat::IF_L_UBYTE )
         {
-            int byteOffset = blockOffset * fdata.m_bytesPerBlock
-                    + ( blockPixelY * fdata.m_pixelsPerBlockX + blockPixelX ) * 1;
+            int byteOffset = blockOffset * fdata.BytesPerBlock
+                    + ( blockPixelY * fdata.PixelsPerBlockX + blockPixelX ) * 1;
 
             result[0] = m_data[ byteOffset ] / 255.0f;
             result[1] = m_data[ byteOffset ] / 255.0f;
@@ -1026,7 +1056,7 @@ namespace mu
   //      case IF_RGBA_UBYTE:
   //      case IF_BGRA_UBYTE:
   //      {
-  //          size_t bytesPerPixel = GetImageFormatData( m_format ).m_bytesPerBlock;
+  //          size_t bytesPerPixel = GetImageFormatData( m_format ).BytesPerBlock;
   //          size_t rowStride = sx * bytesPerPixel;
 
   //          // Find top
@@ -1104,7 +1134,7 @@ namespace mu
 	//---------------------------------------------------------------------------------------------
 	void Image::ReduceLODsTo(int32 NewLODCount)
 	{
-		bool bIsBlockBased = GetImageFormatData(m_format).m_bytesPerBlock > 0;
+		bool bIsBlockBased = GetImageFormatData(m_format).BytesPerBlock > 0;
 		int32 MaxLODs = GetMipmapCount(m_size[0], m_size[1]);
 		int32 LODSToSkip = MaxLODs - NewLODCount;
 		if (LODSToSkip > 0 && bIsBlockBased)
@@ -1133,7 +1163,7 @@ namespace mu
 	//---------------------------------------------------------------------------------------------
 	void Image::ReduceLODs(int32 LODSToSkip)
 	{
-		bool bIsBlockBased = GetImageFormatData(m_format).m_bytesPerBlock > 0;
+		bool bIsBlockBased = GetImageFormatData(m_format).BytesPerBlock > 0;
 		if (LODSToSkip > 0 && bIsBlockBased)
 		{
 			check(LODSToSkip < m_lods);
@@ -1158,19 +1188,19 @@ namespace mu
 
 
 	//---------------------------------------------------------------------------------------------
-	void Image::FillColour(FVector4f c)
+	void FImageOperator::FillColor(Image* Target, FVector4f Color)
 	{
-		EImageFormat Format = GetFormat();
+		EImageFormat Format = Target->GetFormat();
 		switch (Format)
 		{
 		case EImageFormat::IF_RGB_UBYTE:
 		{
 			// TODO: Optimize: don't write bytes one by one
-			int pixelCount = CalculatePixelCount();
-			uint8* pData = GetData();
-			uint8 r = uint8(FMath::Clamp(255.0f * c[0], 0.0f, 255.0f));
-			uint8 g = uint8(FMath::Clamp(255.0f * c[1], 0.0f, 255.0f));
-			uint8 b = uint8(FMath::Clamp(255.0f * c[2], 0.0f, 255.0f));
+			int pixelCount = Target->CalculatePixelCount();
+			uint8* pData = Target->GetData();
+			uint8 r = uint8(FMath::Clamp(255.0f * Color[0], 0.0f, 255.0f));
+			uint8 g = uint8(FMath::Clamp(255.0f * Color[1], 0.0f, 255.0f));
+			uint8 b = uint8(FMath::Clamp(255.0f * Color[2], 0.0f, 255.0f));
 			for (int p = 0; p < pixelCount; ++p)
 			{
 				pData[0] = r;
@@ -1184,12 +1214,12 @@ namespace mu
 		case EImageFormat::IF_RGBA_UBYTE:
 		{
 			// TODO: Optimize: don't write bytes one by one
-			int pixelCount = CalculatePixelCount();
-			uint8* pData = GetData();
-			uint8 r = uint8(FMath::Clamp(255.0f * c[0], 0.0f, 255.0f));
-			uint8 g = uint8(FMath::Clamp(255.0f * c[1], 0.0f, 255.0f));
-			uint8 b = uint8(FMath::Clamp(255.0f * c[2], 0.0f, 255.0f));
-			uint8 a = uint8(FMath::Clamp(255.0f * c[3], 0.0f, 255.0f));
+			int pixelCount = Target->CalculatePixelCount();
+			uint8* pData = Target->GetData();
+			uint8 r = uint8(FMath::Clamp(255.0f * Color[0], 0.0f, 255.0f));
+			uint8 g = uint8(FMath::Clamp(255.0f * Color[1], 0.0f, 255.0f));
+			uint8 b = uint8(FMath::Clamp(255.0f * Color[2], 0.0f, 255.0f));
+			uint8 a = uint8(FMath::Clamp(255.0f * Color[3], 0.0f, 255.0f));
 			for (int p = 0; p < pixelCount; ++p)
 			{
 				pData[0] = r;
@@ -1204,12 +1234,12 @@ namespace mu
 		case EImageFormat::IF_BGRA_UBYTE:
 		{
 			// TODO: Optimize: don't write bytes one by one
-			int pixelCount = CalculatePixelCount();
-			uint8* pData = GetData();
-			uint8 r = uint8(FMath::Clamp(255.0f * c[0], 0.0f, 255.0f));
-			uint8 g = uint8(FMath::Clamp(255.0f * c[1], 0.0f, 255.0f));
-			uint8 b = uint8(FMath::Clamp(255.0f * c[2], 0.0f, 255.0f));
-			uint8 a = uint8(FMath::Clamp(255.0f * c[3], 0.0f, 255.0f));
+			int pixelCount = Target->CalculatePixelCount();
+			uint8* pData = Target->GetData();
+			uint8 r = uint8(FMath::Clamp(255.0f * Color[0], 0.0f, 255.0f));
+			uint8 g = uint8(FMath::Clamp(255.0f * Color[1], 0.0f, 255.0f));
+			uint8 b = uint8(FMath::Clamp(255.0f * Color[2], 0.0f, 255.0f));
+			uint8 a = uint8(FMath::Clamp(255.0f * Color[3], 0.0f, 255.0f));
 			for (int p = 0; p < pixelCount; ++p)
 			{
 				pData[0] = b;
@@ -1223,9 +1253,9 @@ namespace mu
 
 		case EImageFormat::IF_L_UBYTE:
 		{
-			int pixelCount = CalculatePixelCount();
-			uint8* pData = GetData();
-			uint8 v = FMath::Min<uint8>(255, FMath::Max<uint8>(0, uint8(255.0f * c[0])));
+			int pixelCount = Target->CalculatePixelCount();
+			uint8* pData = Target->GetData();
+			uint8 v = FMath::Min<uint8>(255, FMath::Max<uint8>(0, uint8(255.0f * Color[0])));
 			FMemory::Memset(pData, v, pixelCount);
 			break;
 		}
@@ -1234,28 +1264,29 @@ namespace mu
 		{
 			// Generic case that supports compressed formats.
 			const FImageFormatData& FormatData = GetImageFormatData(Format);
-			Ptr<Image> Block = new Image(FormatData.m_pixelsPerBlockX, FormatData.m_pixelsPerBlockY, 1, EImageFormat::IF_RGBA_UBYTE, EInitializationType::NotInitialized);
+			Ptr<Image> Block = CreateImage(FormatData.PixelsPerBlockX, FormatData.PixelsPerBlockY, 1, EImageFormat::IF_RGBA_UBYTE, EInitializationType::NotInitialized);
 
-			uint32 Pixel = (uint32(FMath::Clamp(255.0f * c[0], 0.0f, 255.0f)) << 0)
-				| (uint32(FMath::Clamp(255.0f * c[1], 0.0f, 255.0f)) << 8)
-				| (uint32(FMath::Clamp(255.0f * c[2], 0.0f, 255.0f)) << 16)
-				| (uint32(FMath::Clamp(255.0f * c[3], 0.0f, 255.0f)) << 24);
+			uint32 Pixel = (uint32(FMath::Clamp(255.0f * Color[0], 0.0f, 255.0f)) << 0)
+				| (uint32(FMath::Clamp(255.0f * Color[1], 0.0f, 255.0f)) << 8)
+				| (uint32(FMath::Clamp(255.0f * Color[2], 0.0f, 255.0f)) << 16)
+				| (uint32(FMath::Clamp(255.0f * Color[3], 0.0f, 255.0f)) << 24);
 			uint32* BlockData = reinterpret_cast<uint32*>(Block->GetData());
-			for (int32 I = 0; I < FormatData.m_pixelsPerBlockX * FormatData.m_pixelsPerBlockY; ++I)
+			for (int32 I = 0; I < FormatData.PixelsPerBlockX * FormatData.PixelsPerBlockY; ++I)
 			{
 				*BlockData = Pixel;
 				BlockData += 1;
 			}
 			Ptr<Image> Converted = ImagePixelFormat(0, Block.get(), Format);
+			ReleaseImage(Block);
 
-			int32 DataSize = CalculateDataSize();
-			check(DataSize % FormatData.m_bytesPerBlock == 0);
-			int32 BlockCount = DataSize / FormatData.m_bytesPerBlock;
-			uint8* TargetData = GetData();
+			int32 DataSize = Target->CalculateDataSize();
+			check(DataSize % FormatData.BytesPerBlock == 0);
+			int32 BlockCount = DataSize / FormatData.BytesPerBlock;
+			uint8* TargetData = Target->GetData();
 			while (BlockCount)
 			{
-				FMemory::Memcpy(TargetData, Converted->GetData(), FormatData.m_bytesPerBlock);
-				TargetData += FormatData.m_bytesPerBlock;
+				FMemory::Memcpy(TargetData, Converted->GetData(), FormatData.BytesPerBlock);
+				TargetData += FormatData.BytesPerBlock;
 				--BlockCount;
 			}
 			break;
diff --git a/Engine/Plugins/Experimental/Mutable/Source/MutableRuntime/Private/MuR/ImagePrivate.h b/Engine/Plugins/Experimental/Mutable/Source/MutableRuntime/Private/MuR/ImagePrivate.h
--- a/Engine/Plugins/Experimental/Mutable/Source/MutableRuntime/Private/MuR/ImagePrivate.h
+++ b/Engine/Plugins/Experimental/Mutable/Source/MutableRuntime/Private/MuR/ImagePrivate.h
@@ -18,35 +18,6 @@ namespace mu
 	MUTABLE_DEFINE_ENUM_SERIALISABLE( EAddressMode );
 
 
-	struct FMipmapGenerationSettings
-	{
-		float m_sharpenFactor = 0.0f;
-		EMipmapFilterType m_filterType = EMipmapFilterType::MFT_SimpleAverage;
-		EAddressMode m_addressMode = EAddressMode::None;
-		bool m_ditherMipmapAlpha = false;
-
-		void Serialise( OutputArchive& arch ) const
-		{
-			uint32 ver = 0;
-			arch << ver;
-
-			arch << m_sharpenFactor;
-			arch << m_filterType;
-			arch << m_ditherMipmapAlpha;
-		}
-
-		void Unserialise( InputArchive& arch )
-		{
-			uint32 ver = 0;
-			arch >> ver;
-			check(ver == 0);
-
-			arch >> m_sharpenFactor;
-			arch >> m_filterType;
-			arch >> m_ditherMipmapAlpha;
-		}
-	};
-
 	//---------------------------------------------------------------------------------------------
 	//!
 	//---------------------------------------------------------------------------------------------
@@ -82,8 +53,8 @@ namespace mu
 	inline EImageFormat GetMostGenericFormat(EImageFormat a, EImageFormat b)
     {
         if (a==b) return a;
-        if ( GetImageFormatData(a).m_channels>GetImageFormatData(b).m_channels ) return a;
-        if ( GetImageFormatData(b).m_channels>GetImageFormatData(a).m_channels ) return b;
+        if ( GetImageFormatData(a).Channels>GetImageFormatData(b).Channels ) return a;
+        if ( GetImageFormatData(b).Channels>GetImageFormatData(a).Channels ) return b;
         if (a== EImageFormat::IF_BC2 || a== EImageFormat::IF_BC3 || a== EImageFormat::IF_ASTC_4x4_RGBA_LDR) return a;
         if (b== EImageFormat::IF_BC2 || b== EImageFormat::IF_BC3 || b== EImageFormat::IF_ASTC_4x4_RGBA_LDR) return b;
 
diff --git a/Engine/Plugins/Experimental/Mutable/Source/MutableRuntime/Private/MuR/Model.cpp b/Engine/Plugins/Experimental/Mutable/Source/MutableRuntime/Private/MuR/Model.cpp
--- a/Engine/Plugins/Experimental/Mutable/Source/MutableRuntime/Private/MuR/Model.cpp
+++ b/Engine/Plugins/Experimental/Mutable/Source/MutableRuntime/Private/MuR/Model.cpp
@@ -29,58 +29,9 @@ namespace mu
     //---------------------------------------------------------------------------------------------
     void FProgram::Check()
     {
-//    #ifdef MUTABLE_DEBUG
-//        // Process all from root
-//        if (m_states.size())
-//        {
-//            // Super debug: enable in case of emergencies.
-////            InvariantCodeVisitor invariantVisitor;
-////            invariantVisitor.Apply( *this, m_states[0].m_root );
-//        }
-
-//        for ( size_t i=0; i<m_code.size(); ++i )
-//        {
-//            if ( m_code[i].type == OP_TYPE::IM_PIXELFORMAT )
-//            {
-//                check( m_code[i].args.ImagePixelFormat.format != IF_NONE );
-//                check( m_code[i].args.ImagePixelFormat.source != i );
-//            }
-//    //
-//    //            if ( m_code[i].type == OP_TYPE::IM_MIPMAP )
-//    //            {
-//    //                check( m_code[i].args.ImageMipmap.blockLevels != 0 );
-//    //            }
-
-//            if ( m_code[i].type == OP_TYPE::ME_APPLYLAYOUT )
-//            {
-//                check( m_code[i].args.MeshApplyLayout.mesh != 0 );
-//                check( m_code[i].args.MeshApplyLayout.mesh != i );
-//            }
-
-//            if ( m_code[i].type == OP_TYPE::ME_MORPH2 )
-//            {
-//                for (int j=0;j<MUTABLE_OP_MAX_MORPH2_TARGETS;++j)
-//                {
-//                    auto target = m_code[i].args.MeshMorph2.targets[j];
-//                    check( m_code[target].type!=OP_TYPE::ME_APPLYLAYOUT );
-//                }
-//            }
-
-//            if ( m_code[i].type == OP_TYPE::ME_CONDITIONAL ||
-//                 m_code[i].type == OP_TYPE::IN_CONDITIONAL ||
-//                 m_code[i].type == OP_TYPE::IM_CONDITIONAL ||
-//                 m_code[i].type == OP_TYPE::LA_CONDITIONAL )
-//            {
-//                for (int j=0;j<MUTABLE_OP_MAX_MORPH2_TARGETS;++j)
-//                {
-//                    check( m_code[i].args.Conditional.condition!=i );
-//                    check( m_code[i].args.Conditional.yes!=i );
-//                    check( m_code[i].args.Conditional.no!=i );
-//                }
-//            }
-
-//        }
-//    #endif
+    #ifdef MUTABLE_DEBUG
+		// Insert debug checks here.
+    #endif
     }
 
 
@@ -844,7 +795,10 @@ namespace mu
         newKey.m_id = newId;
         newKey.m_lastRequestId = m_lastResourceResquestId;
         newKey.m_rootAddress = rootAt;
-        newKey.m_parameterValuesBlob = MoveTemp(parameterValuesBlob);
+
+		// Don't MoveTemp so that we get a tight array on the destination.
+		//newKey.m_parameterValuesBlob = MoveTemp(parameterValuesBlob);
+		newKey.m_parameterValuesBlob = parameterValuesBlob;
 
         // TODO: Move the constant to settings?
         const size_t maxGeneratedResourcesIDCacheSize = 1024;
diff --git a/Engine/Plugins/Experimental/Mutable/Source/MutableRuntime/Private/MuR/ModelPrivate.h b/Engine/Plugins/Experimental/Mutable/Source/MutableRuntime/Private/MuR/ModelPrivate.h
--- a/Engine/Plugins/Experimental/Mutable/Source/MutableRuntime/Private/MuR/ModelPrivate.h
+++ b/Engine/Plugins/Experimental/Mutable/Source/MutableRuntime/Private/MuR/ModelPrivate.h
@@ -12,8 +12,6 @@
 #include "MuR/ImagePrivate.h"
 #include "MuR/MeshPrivate.h"
 #include "MuR/ParametersPrivate.h"
-#include "MuR/OpImageResize.h"
-#include "MuR/OpImageMipmap.h"
 
 #define MUTABLE_MAX_RUNTIME_PARAMETERS_PER_STATE	64
 #define MUTABLE_GROW_BORDER_VALUE					2
@@ -408,9 +406,11 @@ namespace mu
 			// TODO:
 			int32 CompressionQuality = 4;
 
+			FImageOperator ImOp = FImageOperator::GetDefault();
+
 			// TODO: Not efficient if we don't make mips (no need to clone base)
 			// TODO: If the image already has mips, we will be duplicating them...
-			Ptr<Image> pMip = pImage->ExtractMip(0);			
+			Ptr<Image> pMip = ImOp.ExtractMip(pImage.get(),0);
 			for ( int Mip=0; Mip<MipsToStore; ++Mip )
 			{
 				check(pMip->GetFormat() == pImage->GetFormat());
@@ -448,11 +448,11 @@ namespace mu
 					if (Mip > pImage->GetLODCount())
 					{
 						// Generate from the last mip.
-						pMip = pMip->ExtractMip(1);
+						pMip = ImOp.ExtractMip(pMip.get(), 1);
 					}
 					else
 					{
-						pMip = pImage->ExtractMip(Mip+1);
+						pMip = ImOp.ExtractMip(pImage.get(),Mip+1);
 					}
 					check(pMip);
 				}
@@ -669,7 +669,7 @@ namespace mu
 			{
 				MUTABLE_CPUPROFILER_SCOPE(ComposeConstantImage);
 
-				Ptr<Image> Result = CreateImage( CurrentMip->GetSizeX(), CurrentMip->GetSizeY(), FinalLODs, CurrentMip->GetFormat() );
+				Ptr<Image> Result = CreateImage( CurrentMip->GetSizeX(), CurrentMip->GetSizeY(), FinalLODs, CurrentMip->GetFormat(), EInitializationType::NotInitialized );
 				Result->m_flags = CurrentMip->m_flags;
 
 				// Some non-block pixel formats require separate memory size calculation
@@ -826,7 +826,7 @@ namespace mu
 		TArray<RESOURCE_KEY> m_generatedResources;
 
         //! Get a resource key for a given resource with given parameter values.
-        uint32 GetResourceKey( uint32 paramListIndex, OP::ADDRESS rootAt, const Parameters* pParams );
+        uint32 GetResourceKey( uint32 ParamListIndex, OP::ADDRESS RootAt, const Parameters* );
 
         //! @}
 
diff --git a/Engine/Plugins/Experimental/Mutable/Source/MutableRuntime/Private/MuR/OpImageCompose.cpp b/Engine/Plugins/Experimental/Mutable/Source/MutableRuntime/Private/MuR/OpImageCompose.cpp
new file mode 100644
--- /dev/null
+++ b/Engine/Plugins/Experimental/Mutable/Source/MutableRuntime/Private/MuR/OpImageCompose.cpp
@@ -0,0 +1,117 @@
+// Copyright Epic Games, Inc. All Rights Reserved.
+
+#include "MuR/ImagePrivate.h"
+#include "MuR/MutableMath.h"
+
+
+namespace mu
+{
+
+	void FImageOperator::ImageCompose( Image* pBase, const Image* pBlock, const box< UE::Math::TIntVector2<uint16> >& rect )
+	{
+		check(pBase && pBlock);
+		check(pBase != pBlock);
+		check( rect.min[0]>=0 && rect.min[1]>=0 );
+		check( rect.size[0]>=0 && rect.size[1]>=0 );
+		check( rect.min[0]+rect.size[0]<=pBase->GetSizeX() );
+		check( rect.min[1]+rect.size[1]<=pBase->GetSizeY() );
+		check( rect.size[0]==pBlock->GetSizeX() );
+		check( rect.size[1]==pBlock->GetSizeY() );
+        check( pBase->GetFormat()==pBlock->GetFormat() );
+
+        // This assert may be acceptable if the missing lods are due to the image block format size
+        // limiting it, or disadjustment due to image scaling down. It is handled below.
+        //check( pBase->GetLODCount()<=pBlock->GetLODCount() );
+
+		EImageFormat OriginalFormat = pBase->GetFormat();
+		const FImageFormatData& finfo = GetImageFormatData( OriginalFormat );
+
+		// If this is true, we are composing coordinates or sizes that don't fall on pixel format
+		// block borders.
+		bool bSubPixelBlockCompose = 
+			rect.min[0] % finfo.PixelsPerBlockX != 0
+			|| rect.min[1] % finfo.PixelsPerBlockY != 0
+			|| rect.size[0] % finfo.PixelsPerBlockX != 0
+			|| rect.size[1] % finfo.PixelsPerBlockY != 0;
+
+		if (bSubPixelBlockCompose)
+		{
+			// Bad performance case, but likely a very small image anyway.
+			MUTABLE_CPUPROFILER_SCOPE(ImageCompose_SubBlock_Fix);
+			EImageFormat UncompressedFormat = GetUncompressedFormat(OriginalFormat);
+			constexpr int32 Quality = 0;
+			Ptr<Image> TempBase = ImagePixelFormat(Quality, pBase, UncompressedFormat);
+			Ptr<Image> TempBlock = ImagePixelFormat(Quality, pBlock, UncompressedFormat);
+			ImageCompose( TempBase.get(), TempBlock.get(), rect);
+			ReleaseImage(TempBlock);
+			Ptr<Image> NewBase = ImagePixelFormat(Quality, TempBase.get(), OriginalFormat);
+			ReleaseImage(TempBase);
+			// ImageCompose above may have reduced the LODs because of the block lods.
+			pBase->CopyMove( NewBase.get() );
+			ReleaseImage(NewBase);
+			return;
+		}
+		else
+		{
+			// Sizes in blocks of the current mips
+			UE::Math::TIntVector2<uint16> baseMipSize(
+				pBase->GetSizeX() / finfo.PixelsPerBlockX,
+				pBase->GetSizeY() / finfo.PixelsPerBlockY);
+
+			UE::Math::TIntVector2<uint16> blockMipPos(
+				rect.min[0] / finfo.PixelsPerBlockX,
+				rect.min[1] / finfo.PixelsPerBlockY);
+
+			UE::Math::TIntVector2<uint16> blockMipSize(
+				rect.size[0] / finfo.PixelsPerBlockX,
+				rect.size[1] / finfo.PixelsPerBlockY);
+
+			int doneMips = 0;
+			for (; doneMips < pBase->GetLODCount() && doneMips < pBlock->GetLODCount(); ++doneMips)
+			{
+				//UE_LOG(LogMutableCore, Warning, "Block Mip Pos : %d, %d", blockMipPos[0], blockMipPos[1]);
+
+				// Sizes of a row in bytes
+				int baseRowSize = finfo.BytesPerBlock * baseMipSize[0];
+				int blockRowSize = finfo.BytesPerBlock * blockMipSize[0];
+
+				uint8_t* pBaseBuf = pBase->GetMipData(doneMips);
+				const uint8_t* pBlockBuf = pBlock->GetMipData(doneMips);
+
+				int skipBytes = baseRowSize * blockMipPos[1] + blockMipPos[0] * finfo.BytesPerBlock;
+
+				pBaseBuf += skipBytes;
+				for (int y = 0; y < blockMipSize[1]; ++y)
+				{
+					check(pBaseBuf + blockRowSize <= pBase->GetData() + pBase->GetDataSize());
+					check(pBlockBuf + blockRowSize <= pBlock->GetData() + pBlock->GetDataSize());
+					FMemory::Memcpy(pBaseBuf, pBlockBuf, blockRowSize);
+					pBlockBuf += blockRowSize;
+					pBaseBuf += baseRowSize;
+				}
+
+				// We may not proceed to the next mip, if the layout block size is smaller than the
+				// pixel format block size.
+				if (blockMipPos[0] % 2 || blockMipPos[1] % 2 || blockMipSize[0] % 2 || blockMipSize[1] % 2)
+				{
+					++doneMips;
+					break;
+				}
+
+				baseMipSize[0] = FMath::DivideAndRoundUp(baseMipSize[0], uint16(2));
+				baseMipSize[1] = FMath::DivideAndRoundUp(baseMipSize[1], uint16(2));
+				blockMipPos[0] = FMath::DivideAndRoundUp(blockMipPos[0], uint16(2));
+				blockMipPos[1] = FMath::DivideAndRoundUp(blockMipPos[1], uint16(2));
+				blockMipSize[0] = FMath::DivideAndRoundUp(blockMipSize[0], uint16(2));
+				blockMipSize[1] = FMath::DivideAndRoundUp(blockMipSize[1], uint16(2));
+			}
+
+			// Adjust the actually valid mips, maybe wasting some of the reserved memory.
+			// TODO: Review this.
+			pBase->m_lods = (uint8)doneMips;
+			pBase->m_data.SetNum(pBase->CalculateDataSize());
+		}
+    }
+
+}
+
diff --git a/Engine/Plugins/Experimental/Mutable/Source/MutableRuntime/Private/MuR/OpImageCompose.h b/Engine/Plugins/Experimental/Mutable/Source/MutableRuntime/Private/MuR/OpImageCompose.h
deleted file mode 100644
diff --git a/Engine/Plugins/Experimental/Mutable/Source/MutableRuntime/Private/MuR/OpImageCrop.cpp b/Engine/Plugins/Experimental/Mutable/Source/MutableRuntime/Private/MuR/OpImageCrop.cpp
new file mode 100644
--- /dev/null
+++ b/Engine/Plugins/Experimental/Mutable/Source/MutableRuntime/Private/MuR/OpImageCrop.cpp
@@ -0,0 +1,89 @@
+// Copyright Epic Games, Inc. All Rights Reserved.
+
+#include "MuR/ImagePrivate.h"
+#include "MuR/MutableMath.h"
+#include "MuR/Platform.h"
+
+
+namespace mu
+{
+
+	bool FImageOperator::ImageCrop( Image* InCropped, int32 CompressionQuality, const Image* InBase, const box< vec2<int32> >& Rect)
+	{
+		Ptr<const Image> Base = InBase;
+		Ptr<Image> BaseReformat;
+		Ptr<Image> Cropped = InCropped;
+
+		EImageFormat BaseFormat = Base->GetFormat();
+		EImageFormat UncompressedFormat = GetUncompressedFormat(BaseFormat);
+
+		if (BaseFormat != UncompressedFormat)
+		{
+			// Compressed formats need decompression + compression after crop			
+			// \TODO: This may use some additional untracked memory locally in this function.
+			BaseReformat = ImagePixelFormat(CompressionQuality, Base.get(), UncompressedFormat);
+			Base = BaseReformat;
+			Cropped = CreateImage( InCropped->GetSizeX(), InCropped->GetSizeY(), InCropped->GetLODCount(), UncompressedFormat, EInitializationType::NotInitialized);
+        }
+
+		const FImageFormatData& finfo = GetImageFormatData(UncompressedFormat);
+
+		check(Base && Cropped);
+		check(Cropped->GetSizeX() == Rect.size[0]);
+		check(Cropped->GetSizeY() == Rect.size[1]);
+
+		// TODO: better error control. This happens if some layouts are corrupt.
+		bool bCorrect =
+				( Rect.min[0]>=0 && Rect.min[1]>=0 ) &&
+				( Rect.size[0]>=0 && Rect.size[1]>=0 ) &&
+				( Rect.min[0]+Rect.size[0]<=Base->GetSizeX() ) &&
+				( Rect.min[1]+Rect.size[1]<=Base->GetSizeY() );
+		if (!bCorrect)
+		{
+			return false;
+		}
+
+		// Block images are not supported for now
+		check( finfo.PixelsPerBlockX == 1 );
+		check( finfo.PixelsPerBlockY == 1 );
+
+		checkf( Rect.min[0] % finfo.PixelsPerBlockX == 0, TEXT("Rect must snap to blocks.") );
+		checkf( Rect.min[1] % finfo.PixelsPerBlockY == 0, TEXT("Rect must snap to blocks.") );
+		checkf( Rect.size[0] % finfo.PixelsPerBlockX == 0, TEXT("Rect must snap to blocks.") );
+		checkf( Rect.size[1] % finfo.PixelsPerBlockY == 0, TEXT("Rect must snap to blocks.") );
+
+		int baseRowSize = finfo.BytesPerBlock * Base->GetSizeX() / finfo.PixelsPerBlockX;
+		int cropRowSize = finfo.BytesPerBlock * Rect.size[0] / finfo.PixelsPerBlockX;
+
+        const uint8_t* pBaseBuf = Base->GetData();
+        uint8_t* pCropBuf = Cropped->GetData();
+
+		int skipPixels = Base->GetSizeX() * Rect.min[1] + Rect.min[0];
+		pBaseBuf += finfo.BytesPerBlock * skipPixels / finfo.PixelsPerBlockX;
+		for ( int y=0; y<Rect.size[1]; ++y )
+		{
+			FMemory::Memcpy( pCropBuf, pBaseBuf, cropRowSize );
+			pCropBuf += cropRowSize;
+			pBaseBuf += baseRowSize;
+		}
+
+		if (BaseFormat != UncompressedFormat)
+		{
+			ReleaseImage(BaseReformat);
+
+			bool bSuccess = false;
+			int32 DataSize = Cropped->m_data.Num();
+			while (!bSuccess)
+			{
+				InCropped->m_data.SetNumUninitialized(DataSize);
+				ImagePixelFormat(bSuccess, CompressionQuality, InCropped, Cropped.get() );
+				DataSize = (DataSize + 16) * 2;
+			}
+
+			ReleaseImage(Cropped);
+		}
+
+		return true;
+	}
+
+}
diff --git a/Engine/Plugins/Experimental/Mutable/Source/MutableRuntime/Private/MuR/OpImageCrop.h b/Engine/Plugins/Experimental/Mutable/Source/MutableRuntime/Private/MuR/OpImageCrop.h
deleted file mode 100644
diff --git a/Engine/Plugins/Experimental/Mutable/Source/MutableRuntime/Private/MuR/OpImageDifference.h b/Engine/Plugins/Experimental/Mutable/Source/MutableRuntime/Private/MuR/OpImageDifference.h
deleted file mode 100644
diff --git a/Engine/Plugins/Experimental/Mutable/Source/MutableRuntime/Private/MuR/OpImageDisplace.h b/Engine/Plugins/Experimental/Mutable/Source/MutableRuntime/Private/MuR/OpImageDisplace.h
--- a/Engine/Plugins/Experimental/Mutable/Source/MutableRuntime/Private/MuR/OpImageDisplace.h
+++ b/Engine/Plugins/Experimental/Mutable/Source/MutableRuntime/Private/MuR/OpImageDisplace.h
@@ -236,7 +236,7 @@ namespace mu
 			default:
 			{
 				// Generic
-				int pixelSize = GetImageFormatData( pResult->GetFormat() ).m_bytesPerBlock;
+				int pixelSize = GetImageFormatData( pResult->GetFormat() ).BytesPerBlock;
 				//for ( int y=0; y<sizeY; ++y )
 				const auto ProcessRow = [
 					pSourceData, pMap, pResult, sizeX, sizeY, pixelSize
@@ -272,7 +272,7 @@ namespace mu
 
         else if (mapFormat==EImageFormat::IF_L_UBYTE_RLE )
 		{
-            int pixelSize = GetImageFormatData( pResult->GetFormat() ).m_bytesPerBlock;
+            int pixelSize = GetImageFormatData( pResult->GetFormat() ).BytesPerBlock;
 
 			const uint8* pMapData = pMap->GetData();
 			uint8* pResultData = pResult->GetData();
diff --git a/Engine/Plugins/Experimental/Mutable/Source/MutableRuntime/Private/MuR/OpImageGrow.h b/Engine/Plugins/Experimental/Mutable/Source/MutableRuntime/Private/MuR/OpImageGrow.h
--- a/Engine/Plugins/Experimental/Mutable/Source/MutableRuntime/Private/MuR/OpImageGrow.h
+++ b/Engine/Plugins/Experimental/Mutable/Source/MutableRuntime/Private/MuR/OpImageGrow.h
@@ -25,7 +25,7 @@ namespace mu
 
 		ImagePtrConst pSource = pImage->Clone();
 
-		int rowSize = sizeX*GetImageFormatData( pImage->GetFormat() ).m_bytesPerBlock;
+		int rowSize = sizeX*GetImageFormatData( pImage->GetFormat() ).BytesPerBlock;
 
 		for ( int y=0; y<sizeY; ++y )
 		{
@@ -80,7 +80,7 @@ namespace mu
             return;
         }
 
-		int pixelSize = GetImageFormatData( pImage->GetFormat() ).m_bytesPerBlock;
+		int pixelSize = GetImageFormatData( pImage->GetFormat() ).BytesPerBlock;
 		int rowSize = sizeX*pixelSize;
 
 		for ( int y=0; y<sizeY; ++y )
diff --git a/Engine/Plugins/Experimental/Mutable/Source/MutableRuntime/Private/MuR/OpImageInvert.h b/Engine/Plugins/Experimental/Mutable/Source/MutableRuntime/Private/MuR/OpImageInvert.h
--- a/Engine/Plugins/Experimental/Mutable/Source/MutableRuntime/Private/MuR/OpImageInvert.h
+++ b/Engine/Plugins/Experimental/Mutable/Source/MutableRuntime/Private/MuR/OpImageInvert.h
@@ -15,7 +15,7 @@ namespace mu
 	{
 		MUTABLE_CPUPROFILER_SCOPE(ImageInvert);
 
-		ImagePtr pDest = new Image(pA->GetSizeX(), pA->GetSizeY(), pA->GetLODCount(), pA->GetFormat());
+		ImagePtr pDest = new Image(pA->GetSizeX(), pA->GetSizeY(), pA->GetLODCount(), pA->GetFormat(), EInitializationType::NotInitialized);
 
 		uint8* pDestBuf = pDest->GetData();
 		const uint8* pABuf = pA->GetData();
diff --git a/Engine/Plugins/Experimental/Mutable/Source/MutableRuntime/Private/MuR/OpImageLayer.h b/Engine/Plugins/Experimental/Mutable/Source/MutableRuntime/Private/MuR/OpImageLayer.h
--- a/Engine/Plugins/Experimental/Mutable/Source/MutableRuntime/Private/MuR/OpImageLayer.h
+++ b/Engine/Plugins/Experimental/Mutable/Source/MutableRuntime/Private/MuR/OpImageLayer.h
@@ -490,10 +490,6 @@ namespace mu
             int rows = Base->GetSizeY();
             int width = Base->GetSizeX();
 
-			// debug test
-			ImagePtr TempDecodedMask = new Image(Mask->GetSizeX(), Mask->GetSizeY(), Mask->GetLODCount(), EImageFormat::IF_L_UBYTE);
-			UncompressRLE_L(Mask, TempDecodedMask.get());
-
             for (int32 lod=0;lod<LODCount; ++lod)
             {
 				// Skip mip size and row sizes.
@@ -897,8 +893,8 @@ namespace mu
 
 		// No longer required.
 		//check(pBase->GetFormat() == pBlended->GetFormat());
-		check(BaseChannelOffset + CHANNELS_TO_BLEND <= GetImageFormatData(pBase->GetFormat()).m_channels);
-		check(BlendedChannelOffset + CHANNELS_TO_BLEND <= GetImageFormatData(pBlended->GetFormat()).m_channels);
+		check(BaseChannelOffset + CHANNELS_TO_BLEND <= GetImageFormatData(pBase->GetFormat()).Channels);
+		check(BlendedChannelOffset + CHANNELS_TO_BLEND <= GetImageFormatData(pBlended->GetFormat()).Channels);
 
 		check(bOnlyFirstLOD || pBase->GetLODCount() <= pBlended->GetLODCount());
 
@@ -2036,7 +2032,7 @@ namespace mu
         check( pBase->GetLODCount() <= pMask->GetLODCount() );
         check( pBase->GetLODCount() <= pBlended->GetLODCount() );
 
-		int pixelSize = GetImageFormatData( pBase->GetFormat() ).m_bytesPerBlock;
+		int pixelSize = GetImageFormatData( pBase->GetFormat() ).BytesPerBlock;
 
 		int start = ( pBase->GetSizeX()*rect.min[1]+rect.min[0] ) * pixelSize;
 		int stride = ( pBase->GetSizeX() - rect.size[0] ) * pixelSize;
diff --git a/Engine/Plugins/Experimental/Mutable/Source/MutableRuntime/Private/MuR/OpImageMipmap.cpp b/Engine/Plugins/Experimental/Mutable/Source/MutableRuntime/Private/MuR/OpImageMipmap.cpp
--- a/Engine/Plugins/Experimental/Mutable/Source/MutableRuntime/Private/MuR/OpImageMipmap.cpp
+++ b/Engine/Plugins/Experimental/Mutable/Source/MutableRuntime/Private/MuR/OpImageMipmap.cpp
@@ -1,6 +1,5 @@
 // Copyright Epic Games, Inc. All Rights Reserved.
 
-#include "MuR/OpImageMipmap.h"
 #include "MuR/ImagePrivate.h"
 #include "MuR/SystemPrivate.h"
 #include "Async/ParallelFor.h"
@@ -502,7 +501,7 @@ namespace mu
 
 
     //---------------------------------------------------------------------------------------------
-    void ImageMipmap_PrepareScratch(FWorkingMemoryManager& Manager, Image* Dest, const Image* Base, int32 LevelCount, FScratchImageMipmap& Scratch )
+    void FImageOperator::ImageMipmap_PrepareScratch(Image* Dest, const Image* Base, int32 LevelCount, FScratchImageMipmap& Scratch )
     {
         int StartLevel = Base->GetLODCount() - 1;
 
@@ -534,23 +533,23 @@ namespace mu
         {
 			// Uncompress the last mip that we already have
 			FIntVector2 UncompressedSize = Base->CalculateMipSize(StartLevel);
-			Scratch.Uncompressed = Manager.CreateImage( 
+			Scratch.Uncompressed = CreateImage( 
 				(uint16)UncompressedSize[0], (uint16)UncompressedSize[1], 
 				1, 
-				EImageFormat::IF_RGBA_UBYTE);
+				EImageFormat::IF_RGBA_UBYTE, EInitializationType::NotInitialized );
 
 			FIntVector2 UncompressedMipsSize = Base->CalculateMipSize(StartLevel + 1);
 			// Generate the mipmaps from there on
-			Scratch.UncompressedMips = Manager.CreateImage( 
+			Scratch.UncompressedMips = CreateImage( 
 				(uint16)UncompressedMipsSize[0], (uint16)UncompressedMipsSize[1],
 				FMath::Max(1, LevelCount - StartLevel - 1), 
-				EImageFormat::IF_RGBA_UBYTE);
+				EImageFormat::IF_RGBA_UBYTE, EInitializationType::NotInitialized);
 
 			// Compress the mipmapped image
-			Scratch.CompressedMips = Manager.CreateImage(
+			Scratch.CompressedMips = CreateImage(
 				(uint16)UncompressedMipsSize[0], (uint16)UncompressedMipsSize[1],
 				Scratch.UncompressedMips->GetLODCount(),
-				Base->GetFormat());
+				Base->GetFormat(), EInitializationType::NotInitialized);
 
             break;
         }
@@ -560,25 +559,25 @@ namespace mu
         {
             // Uncompress the last mip that we already have
 			FIntVector2 UncompressedSize = Base->CalculateMipSize(StartLevel);
-            Scratch.Uncompressed = Manager.CreateImage(
+            Scratch.Uncompressed = CreateImage(
                 (uint16)UncompressedSize[0], (uint16)UncompressedSize[1],
                 1,
-                EImageFormat::IF_L_UBYTE);
+                EImageFormat::IF_L_UBYTE, EInitializationType::NotInitialized);
 
 
 			FIntVector2 UncompressedMipsSize = Base->CalculateMipSize(StartLevel + 1);
             // Generate the mipmaps from there on
-            Scratch.UncompressedMips = Manager.CreateImage(
+            Scratch.UncompressedMips = CreateImage(
                 (uint16)UncompressedMipsSize[0], (uint16)UncompressedMipsSize[1],
 				FMath::Max(1, LevelCount - StartLevel - 1),
-                EImageFormat::IF_L_UBYTE);
+                EImageFormat::IF_L_UBYTE, EInitializationType::NotInitialized);
 
 
             // Compress the mipmapped image
-            Scratch.CompressedMips = Manager.CreateImage(
+            Scratch.CompressedMips = CreateImage(
                 (uint16)UncompressedMipsSize[0], (uint16)UncompressedMipsSize[1],
                 Scratch.UncompressedMips->GetLODCount(),
-                Base->GetFormat());
+                Base->GetFormat(), EInitializationType::NotInitialized);
 
             // Preallocate ample memory for the compressed data
 			uint32 TotalMemory = Scratch.UncompressedMips->GetDataSize();
@@ -597,15 +596,15 @@ namespace mu
     }
 
 
-	void ImageMipmap_ReleaseScratch(FWorkingMemoryManager& Manager, FScratchImageMipmap& Scratch)
+	void FImageOperator::ImageMipmap_ReleaseScratch(FScratchImageMipmap& Scratch)
 	{
-		Manager.Release(Scratch.Uncompressed);
-		Manager.Release(Scratch.UncompressedMips);
-		Manager.Release(Scratch.CompressedMips);
+		ReleaseImage(Scratch.Uncompressed);
+		ReleaseImage(Scratch.UncompressedMips);
+		ReleaseImage(Scratch.CompressedMips);
 	}
 
 
-	void ImageMipmap(FScratchImageMipmap& Scratch, int32 CompressionQuality, Image* Dest, const Image* Base,
+	void FImageOperator::ImageMipmap(FScratchImageMipmap& Scratch, int32 CompressionQuality, Image* Dest, const Image* Base,
 		int32 LevelCount,
 		const FMipmapGenerationSettings& Settings, bool bGenerateOnlyTail)
 	{
@@ -731,16 +730,16 @@ namespace mu
 	}
 
 
-	void ImageMipmap(FWorkingMemoryManager& Manager, int32 CompressionQuality, Image* Dest, const Image* Base,
+	void FImageOperator::ImageMipmap(int32 CompressionQuality, Image* Dest, const Image* Base,
 		int32 LevelCount,
 		const FMipmapGenerationSettings& Settings, bool bGenerateOnlyTail)
 	{
 		FScratchImageMipmap Scratch;
-		ImageMipmap_PrepareScratch(Manager, Dest, Base, LevelCount, Scratch);
+		ImageMipmap_PrepareScratch(Dest, Base, LevelCount, Scratch);
 
 		ImageMipmap(Scratch, CompressionQuality, Dest, Base, LevelCount, Settings, bGenerateOnlyTail);
 
-		ImageMipmap_ReleaseScratch(Manager, Scratch);
+		ImageMipmap_ReleaseScratch(Scratch);
 	}
 
 
diff --git a/Engine/Plugins/Experimental/Mutable/Source/MutableRuntime/Private/MuR/OpImageMipmap.h b/Engine/Plugins/Experimental/Mutable/Source/MutableRuntime/Private/MuR/OpImageMipmap.h
deleted file mode 100644
diff --git a/Engine/Plugins/Experimental/Mutable/Source/MutableRuntime/Private/MuR/OpImagePixelFormat.cpp b/Engine/Plugins/Experimental/Mutable/Source/MutableRuntime/Private/MuR/OpImagePixelFormat.cpp
--- a/Engine/Plugins/Experimental/Mutable/Source/MutableRuntime/Private/MuR/OpImagePixelFormat.cpp
+++ b/Engine/Plugins/Experimental/Mutable/Source/MutableRuntime/Private/MuR/OpImagePixelFormat.cpp
@@ -18,14 +18,11 @@ namespace mu
 {
 
 
-//-------------------------------------------------------------------------------------------------
-Ptr<Image> ImagePixelFormat( int32 CompressionQuality, const Image* Base, EImageFormat TargetFormat, int32 OnlyLOD )
+Ptr<Image> FImageOperator::ImagePixelFormat( int32 CompressionQuality, const Image* Base, EImageFormat TargetFormat, int32 OnlyLOD )
 {
 	MUTABLE_CPUPROFILER_SCOPE(ImagePixelFormat);
-
 	TRACE_CPUPROFILER_EVENT_SCOPE_TEXT(*FString::Printf(TEXT("From %d to %d"), int32(Base->GetFormat()), int32(TargetFormat)));
 
-
 	FIntVector2 resultSize;
     int resultLODCount = 0;
 
@@ -41,7 +38,7 @@ Ptr<Image> ImagePixelFormat( int32 CompressionQuality, const Image* Base, EImage
         resultLODCount = 1;
     }
 
-    Ptr<Image> Result = new Image( (uint16)resultSize[0], (uint16)resultSize[1], resultLODCount, TargetFormat, EInitializationType::NotInitialized );
+    Ptr<Image> Result = CreateImage( resultSize[0], resultSize[1], resultLODCount, TargetFormat, EInitializationType::NotInitialized );
 	Result->m_flags = Base->m_flags;
 
 	if (Base->GetSizeX()<=0 ||Base->GetSizeY()<=0)
@@ -67,8 +64,7 @@ Ptr<Image> ImagePixelFormat( int32 CompressionQuality, const Image* Base, EImage
 }
 
 
-//-------------------------------------------------------------------------------------------------
-void ImagePixelFormat( bool& bOutSuccess, int32 CompressionQuality, Image* pResult, const Image* Base, int32 OnlyLOD )
+void FImageOperator::ImagePixelFormat( bool& bOutSuccess, int32 CompressionQuality, Image* pResult, const Image* Base, int32 OnlyLOD)
 {
 	MUTABLE_CPUPROFILER_SCOPE(ImagePixelFormatInPlace);
 
@@ -311,16 +307,16 @@ void ImagePixelFormat( bool& bOutSuccess, int32 CompressionQuality, Image* pResu
             case EImageFormat::IF_BGRA_UBYTE:
             {
                 check( OnlyLOD == -1 );
-                ImagePtr pTemp = ImagePixelFormat( CompressionQuality, Base, EImageFormat::IF_L_UBYTE );
+                ImagePtr Temp = ImagePixelFormat( CompressionQuality, Base, EImageFormat::IF_L_UBYTE, -1 );
  
 				// Try to compress
 				if (!pResult->m_data.Num())
 				{
 					// Allocate memory for the compressed data. TODO: Smaller?
-					pResult->m_data.SetNum(pTemp->m_data.Num());
+					pResult->m_data.SetNum(Temp->m_data.Num());
 				}
 
-				CompressRLE_L(bOutSuccess, pTemp.get(), pResult );
+				CompressRLE_L(bOutSuccess, Temp.get(), pResult );
 
 #ifdef MUTABLE_DEBUG_RLE		
 				if (bOutSuccess)
@@ -330,6 +326,7 @@ void ImagePixelFormat( bool& bOutSuccess, int32 CompressionQuality, Image* pResu
                     check( !FMemory::Memcmp(pTest->GetData(), pTemp->GetData(), pTemp->GetDataSize() ) );
                 }
 #endif
+				ReleaseImage(Temp);
 
                 break;
             }
@@ -378,8 +375,8 @@ void ImagePixelFormat( bool& bOutSuccess, int32 CompressionQuality, Image* pResu
             case EImageFormat::IF_BGRA_UBYTE:
             {
                 check( OnlyLOD == -1 );
-                ImagePtr pTemp = ImagePixelFormat( CompressionQuality, Base, EImageFormat::IF_L_UBYTE );
-				CompressRLE_L1(bOutSuccess, pTemp.get(), pResult );
+                ImagePtr Temp = ImagePixelFormat( CompressionQuality, Base, EImageFormat::IF_L_UBYTE, -1);
+				CompressRLE_L1(bOutSuccess, Temp.get(), pResult );
 
 #ifdef bOutSuccess		
 				if (bOutSuccess)
@@ -390,6 +387,7 @@ void ImagePixelFormat( bool& bOutSuccess, int32 CompressionQuality, Image* pResu
 				}
 #endif
 
+				ReleaseImage(Temp);
                 break;
             }
 
@@ -457,14 +455,15 @@ void ImagePixelFormat( bool& bOutSuccess, int32 CompressionQuality, Image* pResu
             case EImageFormat::IF_RGBA_UBYTE_RLE:
             {
                 check( OnlyLOD == -1 );
-                ImagePtr pTemp = ImagePixelFormat( CompressionQuality, Base, EImageFormat::IF_RGBA_UBYTE );
-                const uint8_t* pTempBuf = pTemp->GetData();
+                ImagePtr Temp = ImagePixelFormat( CompressionQuality, Base, EImageFormat::IF_RGBA_UBYTE, -1);
+                const uint8* TempBuf = Temp->GetData();
                 for ( int i=0; i<pixelCount; ++i )
                 {
-                    pDestBuf[i*3+0] = pTempBuf[i*4+0];
-                    pDestBuf[i*3+1] = pTempBuf[i*4+1];
-                    pDestBuf[i*3+2] = pTempBuf[i*4+2];
+                    pDestBuf[i*3+0] = TempBuf[i*4+0];
+                    pDestBuf[i*3+1] = TempBuf[i*4+1];
+                    pDestBuf[i*3+2] = TempBuf[i*4+2];
                 }
+				ReleaseImage(Temp);
                 break;
             }
 
@@ -649,14 +648,14 @@ void ImagePixelFormat( bool& bOutSuccess, int32 CompressionQuality, Image* pResu
         {
             // TODO: Optimise
             pResult->m_format = EImageFormat::IF_RGBA_UBYTE;
-            ImagePixelFormat(bOutSuccess, CompressionQuality, pResult, Base, OnlyLOD );
+            ImagePixelFormat(bOutSuccess, CompressionQuality, pResult, Base, OnlyLOD);
 			check(bOutSuccess);
             pResult->m_format = EImageFormat::IF_BGRA_UBYTE;
 
-            uint8_t* pDestBuf = pResult->GetData();
+            uint8* pDestBuf = pResult->GetData();
             for ( int i=0; i<pixelCount; ++i )
             {
-                uint8_t temp = pDestBuf[i*4+0];
+                uint8 temp = pDestBuf[i*4+0];
                 pDestBuf[i*4+0] = pDestBuf[i*4+2];
                 pDestBuf[i*4+2] = temp;
             }
@@ -890,30 +889,32 @@ void ImagePixelFormat( bool& bOutSuccess, int32 CompressionQuality, Image* pResu
 			case EImageFormat::IF_RGB_UBYTE_RLE:
 			{
 				check(OnlyLOD == -1);
-				ImagePtr pTemp = ImagePixelFormat(CompressionQuality, Base, EImageFormat::IF_RGB_UBYTE);
-				const uint8_t* pTempData = pTemp->GetData();
+				ImagePtr Temp = ImagePixelFormat(CompressionQuality, Base, EImageFormat::IF_RGB_UBYTE);
+				const uint8* TempData = Temp->GetData();
 				for (int i = 0; i < pixelCount; ++i)
 				{
-					pDestBuf[i * 4 + 0] = pTempData[i * 3 + 0];
-					pDestBuf[i * 4 + 1] = pTempData[i * 3 + 1];
-					pDestBuf[i * 4 + 2] = pTempData[i * 3 + 2];
+					pDestBuf[i * 4 + 0] = TempData[i * 3 + 0];
+					pDestBuf[i * 4 + 1] = TempData[i * 3 + 1];
+					pDestBuf[i * 4 + 2] = TempData[i * 3 + 2];
 					pDestBuf[i * 4 + 3] = 255;
 				}
+				ReleaseImage(Temp);
 				break;
 			}
 
 			case EImageFormat::IF_L_UBYTE_RLE:
 			{
 				check(OnlyLOD == -1);
-				ImagePtr pTemp = ImagePixelFormat(CompressionQuality, Base, EImageFormat::IF_L_UBYTE);
-				const uint8_t* pTempData = pTemp->GetData();
+				ImagePtr Temp = ImagePixelFormat(CompressionQuality, Base, EImageFormat::IF_L_UBYTE);
+				const uint8* TempData = Temp->GetData();
 				for (int i = 0; i < pixelCount; ++i)
 				{
-					pDestBuf[i * 4 + 0] = pTempData[i];
-					pDestBuf[i * 4 + 1] = pTempData[i];
-					pDestBuf[i * 4 + 2] = pTempData[i];
+					pDestBuf[i * 4 + 0] = TempData[i];
+					pDestBuf[i * 4 + 1] = TempData[i];
+					pDestBuf[i * 4 + 2] = TempData[i];
 					pDestBuf[i * 4 + 3] = 255;
 				}
+				ReleaseImage(Temp);
 				break;
 			}
 
@@ -936,12 +937,6 @@ void ImagePixelFormat( bool& bOutSuccess, int32 CompressionQuality, Image* pResu
             {
                 check( OnlyLOD == -1 );
                 CompressRLE_RGBA( Base, pResult );
-
-                // TODO: TEST BLEH
-                //                ImagePtr pTest = ImagePixelFormat( CompressionQuality,
-                //                pResult, EImageFormat::IF_RGBA_UBYTE ); check(
-                //                !memcmp(pTest->GetData(), Base->GetData(),
-                //                Base->GetDataSize() ) );
                 break;
             }
 
@@ -950,14 +945,16 @@ void ImagePixelFormat( bool& bOutSuccess, int32 CompressionQuality, Image* pResu
                 check( OnlyLOD == -1 );
 
                 // \todo: optimise
-                ImagePtr pTemp = ImagePixelFormat( CompressionQuality, Base, EImageFormat::IF_RGBA_UBYTE );
-                CompressRLE_RGBA( pTemp.get(), pResult );
+                ImagePtr Temp = ImagePixelFormat( CompressionQuality, Base, EImageFormat::IF_RGBA_UBYTE);
+                CompressRLE_RGBA( Temp.get(), pResult );
 
-                // TODO: TEST BLEH
+                // Test
                 //                ImagePtr pTest = ImagePixelFormat( CompressionQuality,
                 //                pResult, EImageFormat::IF_RGBA_UBYTE ); check(
                 //                !memcmp(pTest->GetData(), Base->GetData(),
                 //                Base->GetDataSize() ) );
+
+				ReleaseImage(Temp);
                 break;
             }
 
@@ -966,11 +963,12 @@ void ImagePixelFormat( bool& bOutSuccess, int32 CompressionQuality, Image* pResu
                 check( OnlyLOD == -1 );
 
                 // \todo: optimise
-                ImagePtr pTemp1 = ImagePixelFormat( CompressionQuality, Base, EImageFormat::IF_RGB_UBYTE );
-                ImagePtr pTemp2 =
-                    ImagePixelFormat( CompressionQuality, pTemp1.get(), EImageFormat::IF_RGBA_UBYTE );
-                CompressRLE_RGBA( pTemp2.get(), pResult );
-                break;
+                ImagePtr Temp1 = ImagePixelFormat( CompressionQuality, Base, EImageFormat::IF_RGB_UBYTE);
+                ImagePtr Temp2 = ImagePixelFormat( CompressionQuality, Temp1.get(), EImageFormat::IF_RGBA_UBYTE);
+				ReleaseImage(Temp1);
+				CompressRLE_RGBA( Temp2.get(), pResult );
+				ReleaseImage(Temp2);
+				break;
             }
 
             default:
@@ -993,7 +991,7 @@ void ImagePixelFormat( bool& bOutSuccess, int32 CompressionQuality, Image* pResu
                 check( OnlyLOD == -1 );
                 CompressRLE_RGB( Base, pResult );
 
-                // TODO: TEST BLEH
+                // Test
                 // ImagePtr pTest = ImagePixelFormat( CompressionQuality, pDest.get(),
                 // EImageFormat::IF_RGB_UBYTE ); check( !memcmp(pTest->GetData(), Base->GetData(),
                 // Base->GetDataSize() ) );
@@ -1005,11 +1003,12 @@ void ImagePixelFormat( bool& bOutSuccess, int32 CompressionQuality, Image* pResu
                 check( OnlyLOD == -1 );
 
                 // \todo: optimise
-                ImagePtr pTemp1 = ImagePixelFormat( CompressionQuality, Base, EImageFormat::IF_RGBA_UBYTE );
-                ImagePtr pTemp2 =
-                    ImagePixelFormat( CompressionQuality, pTemp1.get(), EImageFormat::IF_RGB_UBYTE );
-                CompressRLE_RGB( pTemp2.get(), pResult );
-                break;
+                ImagePtr Temp1 = ImagePixelFormat( CompressionQuality, Base, EImageFormat::IF_RGBA_UBYTE);
+                ImagePtr Temp2 = ImagePixelFormat( CompressionQuality, Temp1.get(), EImageFormat::IF_RGB_UBYTE);
+				ReleaseImage(Temp1);
+				CompressRLE_RGB( Temp2.get(), pResult );
+				ReleaseImage(Temp2);
+				break;
             }
 
             default:
@@ -1046,28 +1045,30 @@ void ImagePixelFormat( bool& bOutSuccess, int32 CompressionQuality, Image* pResu
 
             case EImageFormat::IF_L_UBYTE:
             {
-                ImagePtr pTemp = ImagePixelFormat( CompressionQuality, Base, EImageFormat::IF_RGB_UBYTE );
+                ImagePtr Temp = ImagePixelFormat( CompressionQuality, Base, EImageFormat::IF_RGB_UBYTE);
                 for ( int m=0; m<resultLODCount; ++m )
                 {
                     FIntVector2 mipSize = pResult->CalculateMipSize( m );
                     miro::RGB_to_BC1(
-                        mipSize[0], mipSize[1], pTemp->GetMipData( baseLOD + m ),
+                        mipSize[0], mipSize[1], Temp->GetMipData( baseLOD + m ),
                         pResult->GetMipData( m ), CompressionQuality );
                 }
+				ReleaseImage(Temp);
                 break;
             }
 
             case EImageFormat::IF_BC3:
             {
-                ImagePtr pTemp = ImagePixelFormat( CompressionQuality, Base, EImageFormat::IF_RGBA_UBYTE );
+                ImagePtr Temp = ImagePixelFormat( CompressionQuality, Base, EImageFormat::IF_RGBA_UBYTE);
                 for (int m = 0; m < resultLODCount; ++m)
                 {
                     FIntVector2 mipSize = pResult->CalculateMipSize(m);
                     miro::RGBA_to_BC1(
-                        mipSize[0], mipSize[1], pTemp->GetMipData( baseLOD + m ),
+                        mipSize[0], mipSize[1], Temp->GetMipData( baseLOD + m ),
                         pResult->GetMipData( m ), CompressionQuality );
                 }
-                break;
+				ReleaseImage(Temp);
+				break;
             }
 
             default:
@@ -1168,16 +1169,13 @@ void ImagePixelFormat( bool& bOutSuccess, int32 CompressionQuality, Image* pResu
 			case EImageFormat::IF_RGBA_UBYTE:
 			case EImageFormat::IF_BGRA_UBYTE:
 			{	
-                ImagePtr pTempBase =
-                        ImagePixelFormat( CompressionQuality, Base, EImageFormat::IF_L_UBYTE );
+                ImagePtr TempBase = ImagePixelFormat( CompressionQuality, Base, EImageFormat::IF_L_UBYTE, -1);
                 for (int m = 0; m < resultLODCount; ++m)
                 {
-                    FIntVector2 mipSize = pTempBase->CalculateMipSize( m );
-
-                    miro::L_to_BC4(
-                        mipSize[0], mipSize[1], Base->GetMipData( baseLOD + m ),
-                        pResult->GetMipData( m ), CompressionQuality );
+                    FIntVector2 MipSize = TempBase->CalculateMipSize( m );
+                    miro::L_to_BC4( MipSize[0], MipSize[1], Base->GetMipData( baseLOD + m ), pResult->GetMipData( m ), CompressionQuality );
                 }
+				ReleaseImage(TempBase);
 				break;
 			}
 
diff --git a/Engine/Plugins/Experimental/Mutable/Source/MutableRuntime/Private/MuR/OpImageProject.cpp b/Engine/Plugins/Experimental/Mutable/Source/MutableRuntime/Private/MuR/OpImageProject.cpp
--- a/Engine/Plugins/Experimental/Mutable/Source/MutableRuntime/Private/MuR/OpImageProject.cpp
+++ b/Engine/Plugins/Experimental/Mutable/Source/MutableRuntime/Private/MuR/OpImageProject.cpp
@@ -61,7 +61,7 @@ namespace
             m_pTargetData = pTargetData;
             m_pMaskData = pMaskData;
 
-            check( GetImageFormatData( pSource->GetFormat() ).m_bytesPerBlock
+            check( GetImageFormatData( pSource->GetFormat() ).BytesPerBlock
                             ==
                             PIXEL_SIZE );
         }
@@ -248,7 +248,7 @@ namespace
             TargetDataPtr = pTargetData;
             MaskDataPtr = pMaskData;
 
-            check(GetImageFormatData(Source->GetFormat() ).m_bytesPerBlock == PIXEL_SIZE);
+            check(GetImageFormatData(Source->GetFormat() ).BytesPerBlock == PIXEL_SIZE);
         }
 
 		inline void ProcessPixel(uint8* BufferPosPtr, float Varying[4]) const
@@ -603,7 +603,7 @@ namespace
             m_pTargetData = pTargetData;
             m_pMaskData = pMaskData;
 
-            check( GetImageFormatData( pSource->GetFormat() ).m_bytesPerBlock
+            check( GetImageFormatData( pSource->GetFormat() ).BytesPerBlock
                             ==
                             PIXEL_SIZE );
         }
@@ -749,7 +749,7 @@ void ImageRasterProjected_Optimised( const Mesh* pMesh, Image* pImage,
     }
 
 	EImageFormat format = pImage->GetFormat();
-    int pixelSize = GetImageFormatData( format ).m_bytesPerBlock;
+    int pixelSize = GetImageFormatData( format ).BytesPerBlock;
 
     int32 sizeX = pImage->GetSizeX();
     int32 sizeY = pImage->GetSizeY();
@@ -870,7 +870,7 @@ void ImageRasterProjected_OptimisedWrapping( const Mesh* pMesh, Image* pImage, P
     }
 
 	EImageFormat format = pImage->GetFormat();
-    int pixelSize = GetImageFormatData( format ).m_bytesPerBlock;
+    int pixelSize = GetImageFormatData( format ).BytesPerBlock;
 
     int sizeX = pImage->GetSizeX();
     int sizeY = pImage->GetSizeY();
diff --git a/Engine/Plugins/Experimental/Mutable/Source/MutableRuntime/Private/MuR/OpImageResize.cpp b/Engine/Plugins/Experimental/Mutable/Source/MutableRuntime/Private/MuR/OpImageResize.cpp
--- a/Engine/Plugins/Experimental/Mutable/Source/MutableRuntime/Private/MuR/OpImageResize.cpp
+++ b/Engine/Plugins/Experimental/Mutable/Source/MutableRuntime/Private/MuR/OpImageResize.cpp
@@ -1,170 +1,793 @@
 // Copyright Epic Games, Inc. All Rights Reserved.
 
-#include "MuR/OpImageResize.h"
+#include "MuR/ImagePrivate.h"
+#include "MuR/MutableTrace.h"
+#include "Async/ParallelFor.h"
 
 namespace mu
 {
+	/** */
+	template< int32 NC >
+	void ImageMagnifyX(Image* Dest, const Image* Base)
+	{
+		int32 BaseSizeX = Base->GetSizeX();
+		int32 DestSizeX = Dest->GetSizeX();
+		int32 SizeY = Base->GetSizeY();
+
+		uint32 dx_16 = (uint32(BaseSizeX) << 16) / DestSizeX;
 
-	Ptr<Image> ImageResizeLinear(int32 imageCompressionQuality, const Image* pBasePtr, FImageSize destSize)
+		// Linear filtering
+		//for ( int32 y=0; y<SizeY; ++y )
+		const auto& ProcessLine = [
+			Dest, Base, dx_16, BaseSizeX, DestSizeX
+		] (uint32 y)
+		{
+			uint32 px_16 = 0;
+			const uint8* BaseBuf = Base->GetData() + y * BaseSizeX * NC;
+			uint8* DestBuf = Dest->GetData() + y * DestSizeX * NC;
+
+			for (int32 x = 0; x < DestSizeX; ++x)
+			{
+				uint32 px = px_16 >> 16;
+				uint32 epx_16 = px_16 + dx_16;
+
+				if ((px_16 & 0xffff0000) == ((epx_16 - 1) & 0xffff0000))
+				{
+					// One fraction
+					for (int32 c = 0; c < NC; ++c)
+					{
+						DestBuf[c] = BaseBuf[px * NC + c];
+					}
+				}
+				else
+				{
+					// Two fractions
+					uint32 frac1 = (px_16 & 0xffff);
+					uint32 frac0 = 0x10000 - frac1;
+
+					for (int32 c = 0; c < NC; ++c)
+					{
+						DestBuf[c] = uint8((BaseBuf[px * NC + c] * frac0 + BaseBuf[(px + 1) * NC + c] * frac1) >> 16);
+					}
+
+					++px;
+				}
+
+				px_16 = epx_16;
+				DestBuf += NC;
+			}
+		};
+
+		// TODO: ensure 64K-aligned batches?
+		ParallelFor(SizeY, ProcessLine);
+	}
+
+
+	/** */
+	void ImageMagnifyX(Image* Dest, const Image* Base)
 	{
-		MUTABLE_CPUPROFILER_SCOPE(ImageResizeLinear);
+		MUTABLE_CPUPROFILER_SCOPE(ImageMagnifyX)
+
+		check(Dest->GetSizeY() == Base->GetSizeY());
+		check(Dest->GetSizeX() > Base->GetSizeX());
+
+		switch (Base->GetFormat())
+		{
+
+		case EImageFormat::IF_L_UBYTE:
+		{
+			ImageMagnifyX<1>(Dest, Base);
+			break;
+		}
+
+		case EImageFormat::IF_RGB_UBYTE:
+		{
+			ImageMagnifyX<3>(Dest, Base);
+			break;
+		}
 
-		if (pBasePtr->GetSize() == destSize)
+		case EImageFormat::IF_BGRA_UBYTE:
+		case EImageFormat::IF_RGBA_UBYTE:
 		{
-			return pBasePtr->Clone();
+			ImageMagnifyX<4>(Dest, Base);
+			break;
 		}
 
-		check(!(pBasePtr->m_flags & Image::IF_CANNOT_BE_SCALED));
+		default:
+			// Case not implemented
+			check(false);
+		}
+	}
 
-		ImagePtrConst pBase = pBasePtr;
 
-		// Shouldn't happen! But if it does...
-		EImageFormat sourceFormat = pBase->GetFormat();
-		EImageFormat uncompressedFormat = GetUncompressedFormat(sourceFormat);
-		if (sourceFormat != uncompressedFormat)
+	/** General image minimisation. */
+	template< int32 NC >
+	void ImageMinifyX(Image* Dest, const Image* Base)
+	{
+		int32 BaseSizeX = Base->GetSizeX();
+		int32 DestSizeX = Dest->GetSizeX();
+		int32 SizeY = Base->GetSizeY();
+
+		uint32 dx_16 = (uint32(BaseSizeX) << 16) / DestSizeX;
+
+		// Linear filtering
+		//for ( int32 y=0; y<SizeY; ++y )
+		const auto& ProcessLine = [
+			Dest, Base, dx_16, BaseSizeX, DestSizeX
+		] (uint32 y)
+		{
+			const uint8* BaseBuf = Base->GetData() + y * BaseSizeX * NC;
+			uint8* DestBuf = Dest->GetData() + y * DestSizeX * NC;
+
+			uint32 px_16 = 0;
+			for (int32 x = 0; x < DestSizeX; ++x)
+			{
+				uint32 r_16[NC];
+				for (int32 c = 0; c < NC; ++c)
+				{
+					r_16[c] = 0;
+				}
+
+				uint32 epx_16 = px_16 + dx_16;
+				uint32 px = px_16 >> 16;
+				uint32 epx = epx_16 >> 16;
+
+				// First fraction
+				uint32 frac0 = px_16 & 0xffff;
+				if (frac0)
+				{
+					for (int32 c = 0; c < NC; ++c)
+					{
+						r_16[c] += (0x10000 - frac0) * BaseBuf[px * NC + c];
+					}
+
+					++px;
+				}
+
+				// Whole pixels
+				while (px < epx)
+				{
+					for (int32 c = 0; c < NC; ++c)
+					{
+						r_16[c] += uint32(BaseBuf[px * NC + c]) << 16;
+					}
+
+					++px;
+				}
+
+				// Second fraction
+				uint32 frac1 = epx_16 & 0xffff;
+				if (frac1)
+				{
+					for (int32 c = 0; c < NC; ++c)
+					{
+						r_16[c] += frac1 * BaseBuf[px * NC + c];
+					}
+				}
+
+				for (int32 c = 0; c < NC; ++c)
+				{
+					DestBuf[c] = (uint8)(r_16[c] / dx_16);
+				}
+
+				px_16 = epx_16;
+				DestBuf += NC;
+			}
+		};
+
+		ParallelFor(SizeY, ProcessLine);
+	}
+
+
+	//---------------------------------------------------------------------------------------------
+	//! Optimised for whole factors
+	//---------------------------------------------------------------------------------------------
+	template< int32 NC, int32 FACTOR >
+	void ImageMinifyX_Exact(Image* Dest, const Image* Base)
+	{
+		int32 BaseSizeX = Base->GetSizeX();
+		int32 DestSizeX = Dest->GetSizeX();
+		int32 SizeY = Base->GetSizeY();
+
+		const uint8* BaseBuf = Base->GetData();
+		uint8* DestBuf = Dest->GetData();
+
+		// Linear filtering
+		const auto& ProcessLine = [
+			Dest, Base, BaseSizeX, DestSizeX
+		] (uint32 y)
+		{
+			const uint8* BaseBuf = Base->GetData() + y * BaseSizeX * NC;
+			uint8* DestBuf = Dest->GetData() + y * DestSizeX * NC;
+
+			uint32 r[NC];
+			for (int32 x = 0; x < DestSizeX; ++x)
+			{
+				for (int32 c = 0; c < NC; ++c)
+				{
+					r[c] = 0;
+					for (int32 a = 0; a < FACTOR; ++a)
+					{
+						r[c] += BaseBuf[a * NC + c];
+					}
+				}
+
+				for (int32 c = 0; c < NC; ++c)
+				{
+					DestBuf[c] = (uint8)(r[c] / FACTOR);
+				}
+
+				DestBuf += NC;
+				BaseBuf += NC * FACTOR;
+			}
+		};
+
+		ParallelFor(SizeY, ProcessLine);
+	}
+
+	inline uint32 AverageChannel(uint32 a, uint32 b)
+	{
+		uint32 r = (a + b) >> 1;
+		return r;
+	}
+
+	template<>
+	void ImageMinifyX_Exact<4, 2>(Image* Dest, const Image* Base)
+	{
+		int32 BaseSizeX = Base->GetSizeX();
+		int32 DestSizeX = Dest->GetSizeX();
+		int32 SizeY = Base->GetSizeY();
+
+		const uint8* BaseBuf = Base->GetData();
+		uint8* DestBuf = Dest->GetData();
+
+		int32 TotalBasePixels = BaseSizeX * SizeY;
+		constexpr int32 BasePixelsPerBatch = 4096 * 2;
+		int32 NumBatches = FMath::DivideAndRoundUp(TotalBasePixels, BasePixelsPerBatch);
+
+		// Linear filtering
+		const auto& ProcessBatchUnaligned =
+			[DestBuf, BaseBuf, BaseSizeX, DestSizeX, BasePixelsPerBatch, TotalBasePixels]
+		(int32 BatchIndex)
+		{
+			const uint8* pBatchBaseBuf = BaseBuf + BatchIndex * BasePixelsPerBatch * 4;
+			uint8* pBatchDestBuf = DestBuf + BatchIndex * BasePixelsPerBatch * 4 / 2;
+
+			int32 NumBasePixels = FMath::Min(BasePixelsPerBatch, TotalBasePixels - BatchIndex * BasePixelsPerBatch);
+
+			uint16 r[4];
+			for (int32 x = 0; x < NumBasePixels; x += 2)
+			{
+				r[0] = pBatchBaseBuf[0] + pBatchBaseBuf[0 + 4];
+				r[1] = pBatchBaseBuf[1] + pBatchBaseBuf[1 + 4];
+				r[2] = pBatchBaseBuf[2] + pBatchBaseBuf[2 + 4];
+				r[3] = pBatchBaseBuf[3] + pBatchBaseBuf[3 + 4];
+
+				pBatchDestBuf[0] = (uint8)(r[0] >> 1);
+				pBatchDestBuf[1] = (uint8)(r[1] >> 1);
+				pBatchDestBuf[2] = (uint8)(r[2] >> 1);
+				pBatchDestBuf[3] = (uint8)(r[3] >> 1);
+
+				pBatchBaseBuf += 4 * 2;
+				pBatchDestBuf += 4;
+			}
+		};
+
+		const auto& ProcessBatchAligned =
+			[DestBuf, BaseBuf, BaseSizeX, DestSizeX, BasePixelsPerBatch, TotalBasePixels]
+		(int32 BatchIndex)
+		{
+			const uint32* pBatchBaseBuf = reinterpret_cast<const uint32*>(BaseBuf) + BatchIndex * BasePixelsPerBatch;
+			uint32* pBatchDestBuf = reinterpret_cast<uint32*>(DestBuf) + BatchIndex * (BasePixelsPerBatch >> 1);
+
+			int32 NumBasePixels = FMath::Min(BasePixelsPerBatch, TotalBasePixels - BatchIndex * BasePixelsPerBatch);
+
+			for (int32 p = 0; p < NumBasePixels; ++p)
+			{
+				uint32 FullSource0 = pBatchBaseBuf[p * 2 + 0];
+				uint32 FullSource1 = pBatchBaseBuf[p * 2 + 1];
+
+				uint32 FullResult = 0;
+
+				FullResult |= AverageChannel((FullSource0 >> 0) & 0xff, (FullSource1 >> 0) & 0xff) << 0;
+				FullResult |= AverageChannel((FullSource0 >> 8) & 0xff, (FullSource1 >> 8) & 0xff) << 8;
+				FullResult |= AverageChannel((FullSource0 >> 16) & 0xff, (FullSource1 >> 16) & 0xff) << 16;
+				FullResult |= AverageChannel((FullSource0 >> 24) & 0xff, (FullSource1 >> 24) & 0xff) << 24;
+
+				pBatchDestBuf[p] = FullResult;
+			}
+		};
+
+
+		ParallelFor(NumBatches, ProcessBatchUnaligned);
+	}
+
+
+	//---------------------------------------------------------------------------------------------
+	//! Image minify X version hub.
+	//---------------------------------------------------------------------------------------------
+	void ImageMinifyX(Image* Dest, const Image* Base)
+	{
+		MUTABLE_CPUPROFILER_SCOPE(ImageMinifyX)
+
+		check(Dest->GetSizeY() == Base->GetSizeY());
+		check(Dest->GetSizeX() < Base->GetSizeX());
+
+		switch (Base->GetFormat())
+		{
+
+		case EImageFormat::IF_L_UBYTE:
 		{
-			pBase = ImagePixelFormat(imageCompressionQuality, pBasePtr, uncompressedFormat);
+			if (2 * Dest->GetSizeX() == Base->GetSizeX())
+			{
+				// Optimised case
+				ImageMinifyX_Exact<1, 2>(Dest, Base);
+			}
+			else
+			{
+				// Generic case
+				ImageMinifyX<1>(Dest, Base);
+			}
+			break;
 		}
 
-		FImageSize baseSize = FImageSize(pBase->GetSizeX(), pBase->GetSizeY());
+		case EImageFormat::IF_RGB_UBYTE:
+		{
+			if (2 * Dest->GetSizeX() == Base->GetSizeX())
+			{
+				// Optimised case
+				ImageMinifyX_Exact<3, 2>(Dest, Base);
+			}
+			else
+			{
+				// Generic case
+				ImageMinifyX<3>(Dest, Base);
+			}
+			break;
+		}
 
-		ImagePtr pDest = new Image(destSize[0], destSize[1], 1, pBase->GetFormat());
-		if (!destSize[0] || !destSize[1] || !baseSize[0] || !baseSize[1])
+		case EImageFormat::IF_BGRA_UBYTE:
+		case EImageFormat::IF_RGBA_UBYTE:
 		{
-			return pDest;
+			if (2 * Dest->GetSizeX() == Base->GetSizeX())
+			{
+				// Optimised case
+				ImageMinifyX_Exact<4, 2>(Dest, Base);
+			}
+			else if (4 * Dest->GetSizeX() == Base->GetSizeX())
+			{
+				// Optimised case
+				ImageMinifyX_Exact<4, 4>(Dest, Base);
+			}
+			else
+			{
+				// Generic case
+				ImageMinifyX<4>(Dest, Base);
+			}
+			break;
 		}
 
-		// First resize X
-		ImagePtr pTemp;
-		if (destSize[0] > baseSize[0])
+		default:
+			// Case not implemented
+			check(false);
+		}
+
+	}
+
+
+	//---------------------------------------------------------------------------------------------
+	//!
+	//---------------------------------------------------------------------------------------------
+	template< int32 NC >
+	void ImageMagnifyY(Image* Dest, const Image* Base)
+	{
+		if (!Base || !Dest ||
+			!Base->GetSizeX() || !Base->GetSizeY() || !Dest->GetSizeX() || !Dest->GetSizeY())
 		{
-			pTemp = new Image(destSize[0], baseSize[1], 1, pBase->GetFormat());
-			ImageMagnifyX(pTemp.get(), pBase.get());
+			return;
 		}
-		else if (destSize[0] < baseSize[0])
+
+		int32 BaseSizeY = Base->GetSizeY();
+		int32 DestSizeY = Dest->GetSizeY();
+		int32 SizeX = Base->GetSizeX();
+
+		SIZE_T rowSize = SizeX * NC;
+
+		// Common case, optimised.
+		if (DestSizeY == BaseSizeY * 2)
 		{
-			pTemp = new Image(destSize[0], baseSize[1], 1, pBase->GetFormat());
-			ImageMinifyX(pTemp.get(), pBase.get());
+			//for (int32 y = 0; y < BaseSizeY; ++y)
+			const auto& ProcessLine = [
+				Dest, Base, rowSize
+			] (uint32 y)
+			{
+				uint8* DestBuf = Dest->GetData() + 2 * y * rowSize;
+				const uint8* BaseBuf = Base->GetData() + y * rowSize;
+
+				FMemory::Memcpy(DestBuf, BaseBuf, rowSize);
+				DestBuf += rowSize;
+
+				FMemory::Memcpy(DestBuf, BaseBuf, rowSize);
+			};
+
+				ParallelFor(BaseSizeY, ProcessLine);
 		}
 		else
 		{
-			pTemp = pBase->Clone();
+			uint32 dy_16 = (uint32(BaseSizeY) << 16) / DestSizeY;
+
+			// Linear filtering
+			// \todo: optimise: swap loops, etc.
+			//for ( int32 x=0; x<SizeX; ++x )
+			const auto& ProcessColumn = [
+				Dest, Base, SizeX, DestSizeY, dy_16
+			] (uint32 x)
+			{
+				uint32 py_16 = 0;
+				uint8* DestBuf = Dest->GetData() + x * NC;
+				const uint8* BaseBuf = Base->GetData();
+
+				for (int32 y = 0; y < DestSizeY; ++y)
+				{
+					uint32 py = py_16 >> 16;
+					uint32 epy_16 = py_16 + dy_16;
+
+					if ((py_16 & 0xffff0000) == ((epy_16 - 1) & 0xffff0000))
+					{
+						// One fraction
+						for (int32 c = 0; c < NC; ++c)
+						{
+							DestBuf[c] = BaseBuf[(py * SizeX + x) * NC + c];
+						}
+					}
+					else
+					{
+						// Two fractions
+						uint32 frac1 = (py_16 & 0xffff);
+						uint32 frac0 = 0x10000 - frac1;
+
+						for (int32 c = 0; c < NC; ++c)
+						{
+							DestBuf[c] = (uint8)((BaseBuf[(py * SizeX + x) * NC + c] * frac0 +
+								BaseBuf[((py + 1) * SizeX + x) * NC + c] * frac1
+								) >> 16);
+						}
+
+						++py;
+					}
+
+					py_16 = epy_16;
+					DestBuf += SizeX * NC;
+				}
+			};
+
+				ParallelFor(SizeX, ProcessColumn);
 		}
+	}
 
-		// Now resize Y
-		if (destSize[1] > baseSize[1])
+
+	void ImageMagnifyY(Image* Dest, const Image* Base)
+	{
+		check(Dest->GetSizeY() > Base->GetSizeY());
+		check(Dest->GetSizeX() == Base->GetSizeX());
+
+		MUTABLE_CPUPROFILER_SCOPE(ImageMagnifyY)
+
+		switch (Base->GetFormat())
+		{
+
+		case EImageFormat::IF_L_UBYTE:
 		{
-			ImageMagnifyY(pDest.get(), pTemp.get());
+			ImageMagnifyY<1>(Dest, Base);
+			break;
 		}
-		else if (destSize[1] < baseSize[1])
+
+		case EImageFormat::IF_RGB_UBYTE:
 		{
-			ImageMinifyY(pDest.get(), pTemp.get());
+			ImageMagnifyY<3>(Dest, Base);
+			break;
 		}
-		else
+
+		case EImageFormat::IF_RGBA_UBYTE:
+		case EImageFormat::IF_BGRA_UBYTE:
 		{
-			pDest = pTemp;
+			ImageMagnifyY<4>(Dest, Base);
+			break;
 		}
 
+		default:
+			// Case not implemented
+			check(false);
+		}
+	}
 
-		// Reset format if it was changed to scale
-		if (sourceFormat != uncompressedFormat)
+
+	//---------------------------------------------------------------------------------------------
+	//!
+	//---------------------------------------------------------------------------------------------
+	template< int32 NC >
+	void ImageMinifyY(Image* Dest, const Image* Base)
+	{
+		int32 BaseSizeY = Base->GetSizeY();
+		int32 DestSizeY = Dest->GetSizeY();
+		int32 SizeX = Base->GetSizeX();
+
+		uint32 dy_16 = (uint32(BaseSizeY) << 16) / DestSizeY;
+
+		const uint8* BaseBuf = Base->GetData();
+
+		// Linear filtering
+		//for (int32 x = 0; x < SizeX; ++x)
+		const auto& ProcessColumn = [
+			Dest, BaseBuf, SizeX, DestSizeY, dy_16
+		] (uint32 x)
 		{
-			pDest = ImagePixelFormat(imageCompressionQuality, pDest.get(), sourceFormat);
+			uint8* DestBuf = Dest->GetData() + x * NC;
+			uint32 py_16 = 0;
+			for (int32 y = 0; y < DestSizeY; ++y)
+			{
+				uint32 r_16[NC];
+				for (int32 c = 0; c < NC; ++c)
+				{
+					r_16[c] = 0;
+				}
+
+				uint32 epy_16 = py_16 + dy_16;
+				uint32 py = py_16 >> 16;
+				uint32 epy = epy_16 >> 16;
+
+				// First fraction
+				uint32 frac0 = py_16 & 0xffff;
+				if (frac0)
+				{
+					for (int32 c = 0; c < NC; ++c)
+					{
+						r_16[c] += (0x10000 - frac0) * BaseBuf[(py * SizeX + x) * NC + c];
+					}
+
+					++py;
+				}
+
+				// Whole pixels
+				while (py < epy)
+				{
+					for (int32 c = 0; c < NC; ++c)
+					{
+						r_16[c] += uint32(BaseBuf[(py * SizeX + x) * NC + c]) << 16;
+					}
+
+					++py;
+				}
+
+				// Second fraction
+				uint32 frac1 = epy_16 & 0xffff;
+				if (frac1)
+				{
+					for (int32 c = 0; c < NC; ++c)
+					{
+						r_16[c] += frac1 * BaseBuf[(py * SizeX + x) * NC + c];
+					}
+				}
+
+				for (int32 c = 0; c < NC; ++c)
+				{
+					DestBuf[c] = (uint8)(r_16[c] / dy_16);
+				}
+
+				py_16 = epy_16;
+				DestBuf += SizeX * NC;
+			}
+		};
+
+		ParallelFor(SizeX, ProcessColumn);
+
+	}
+
+
+	//---------------------------------------------------------------------------------------------
+	//!
+	//---------------------------------------------------------------------------------------------
+	template< int32 NC, int32 FACTOR >
+	void ImageMinifyY_Exact(Image* Dest, const Image* Base)
+	{
+		int32 DestSizeY = Dest->GetSizeY();
+		int32 SizeX = Base->GetSizeX();
+
+		// Linear filtering
+		//for ( int32 y=0; y<DestSizeY; ++y )
+		const auto ProcessRow = [
+			Dest, Base, SizeX, DestSizeY
+		] (uint32 y)
+		{
+			uint8* DestBuf = Dest->GetData() + y * NC * SizeX;
+			const uint8* BaseBuf = Base->GetData() + y * FACTOR * SizeX * NC;
+
+			for (int32 x = 0; x < SizeX; ++x)
+			{
+				uint32 r[NC];
+				for (int32 c = 0; c < NC; ++c)
+				{
+					r[c] = 0;
+				}
+
+				// Whole pixels
+				for (int32 f = 0; f < FACTOR; ++f)
+				{
+					for (int32 c = 0; c < NC; ++c)
+					{
+						r[c] += BaseBuf[SizeX * NC * f + x * NC + c];
+					}
+				}
+
+				for (int32 c = 0; c < NC; ++c)
+				{
+					DestBuf[c] = (uint8)(r[c] / FACTOR);
+				}
+
+				DestBuf += NC;
+			}
+			BaseBuf += FACTOR * SizeX * NC;
+		};
+
+		ParallelFor(DestSizeY, ProcessRow);
+	}
+
+
+	//---------------------------------------------------------------------------------------------
+	//! Image minify Y version hub.
+	//---------------------------------------------------------------------------------------------
+	void ImageMinifyY(Image* Dest, const Image* Base)
+	{
+		check(Dest->GetSizeY() < Base->GetSizeY());
+		check(Dest->GetSizeX() == Base->GetSizeX());
+
+		MUTABLE_CPUPROFILER_SCOPE(ImageMinifyY)
+
+		switch (Base->GetFormat())
+		{
+
+		case EImageFormat::IF_L_UBYTE:
+		{
+			if (2 * Dest->GetSizeY() == Base->GetSizeY())
+			{
+				// Optimised case
+				ImageMinifyY_Exact<1, 2>(Dest, Base);
+			}
+			else
+			{
+				// Generic case
+				ImageMinifyY<1>(Dest, Base);
+			}
+			break;
 		}
 
-		// Update the relevancy data of the image.
-		if (pBase->m_flags & Image::EImageFlags::IF_HAS_RELEVANCY_MAP)
+		case EImageFormat::IF_RGB_UBYTE:
 		{
-			pDest->m_flags |= Image::EImageFlags::IF_HAS_RELEVANCY_MAP;
+			if (2 * Dest->GetSizeY() == Base->GetSizeY())
+			{
+				// Optimised case
+				ImageMinifyY_Exact<3, 2>(Dest, Base);
+			}
+			else
+			{
+				// Generic case
+				ImageMinifyY<3>(Dest, Base);
+			}
+			break;
+		}
 
-			float FactorY = float(destSize[1]) / float(baseSize[1]);
+		case EImageFormat::IF_RGBA_UBYTE:
+		case EImageFormat::IF_BGRA_UBYTE:
+		{
+			if (2 * Dest->GetSizeY() == Base->GetSizeY())
+			{
+				// Optimised case
+				ImageMinifyY_Exact<4, 2>(Dest, Base);
+			}
+			else
+			{
+				// Generic case
+				ImageMinifyY<4>(Dest, Base);
+			}
+			break;
+		}
 
-			pDest->RelevancyMinY = uint16(FMath::FloorToFloat(pBase->RelevancyMinY * FactorY));
-			pDest->RelevancyMaxY = uint16(FMath::Min((int32)FMath::CeilToFloat(pBase->RelevancyMinY * FactorY), pDest->GetSizeY() - 1));
+		default:
+			// Case not implemented
+			check(false);
+			//mu::Halt();
 		}
 
-		return pDest;
 	}
 
 
-	void ImageResizeLinear(Image* pDest, int32 imageCompressionQuality, const Image* pBasePtr)
+	void FImageOperator::ImageResizeLinear(Image* Dest, int32 ImageCompressionQuality, const Image* InBase)
 	{
 		MUTABLE_CPUPROFILER_SCOPE(ImageResizeLinear);
 
-		check(!(pBasePtr->m_flags & Image::IF_CANNOT_BE_SCALED));
+		check(!(InBase->m_flags & Image::IF_CANNOT_BE_SCALED));
 
-		ImagePtrConst pBase = pBasePtr;
+		ImagePtrConst Base = InBase;
+		ImagePtr BaseReformat;
 
 		// Shouldn't happen! But if it does...
-		EImageFormat sourceFormat = pBase->GetFormat();
-		EImageFormat uncompressedFormat = GetUncompressedFormat(sourceFormat);
-		if (sourceFormat != uncompressedFormat)
+		EImageFormat SourceFormat = Base->GetFormat();
+		EImageFormat UncompressedFormat = GetUncompressedFormat(SourceFormat);
+		if (SourceFormat != UncompressedFormat)
 		{
-			pBase = ImagePixelFormat(imageCompressionQuality, pBasePtr, uncompressedFormat);
+			BaseReformat = ImagePixelFormat(ImageCompressionQuality, InBase, UncompressedFormat);
+			Base = BaseReformat;
 		}
 
-		FImageSize baseSize = FImageSize(pBase->GetSizeX(), pBase->GetSizeY());
-		FImageSize destSize = FImageSize(pDest->GetSizeX(), pDest->GetSizeY());
-		if (!destSize[0] || !destSize[1] || !baseSize[0] || !baseSize[1])
+		FImageSize BaseSize = FImageSize(Base->GetSizeX(), Base->GetSizeY());
+		FImageSize DestSize = FImageSize(Dest->GetSizeX(), Dest->GetSizeY());
+		if (!DestSize[0] || !DestSize[1] || !BaseSize[0] || !BaseSize[1])
 		{
 			return;
 		}
 
 		// First resize X
-		ImagePtr pTemp;
-		if (destSize[0] > baseSize[0])
+		ImagePtr Temp;
+		if (DestSize[0] > BaseSize[0])
 		{
-			pTemp = new Image(destSize[0], baseSize[1], 1, pBase->GetFormat());
-			ImageMagnifyX(pTemp.get(), pBase.get());
+			Temp = CreateImage(DestSize[0], BaseSize[1], 1, Base->GetFormat(), EInitializationType::NotInitialized);
+			ImageMagnifyX(Temp.get(), Base.get());
 		}
-		else if (destSize[0] < baseSize[0])
+		else if (DestSize[0] < BaseSize[0])
 		{
-			pTemp = new Image(destSize[0], baseSize[1], 1, pBase->GetFormat());
-			ImageMinifyX(pTemp.get(), pBase.get());
+			Temp = CreateImage(DestSize[0], BaseSize[1], 1, Base->GetFormat(), EInitializationType::NotInitialized);
+			ImageMinifyX(Temp.get(), Base.get());
 		}
 		else
 		{
-			pTemp = pBase->Clone();
+			Temp = CloneImage(Base.get());
 		}
 
 		// Now resize Y
-		ImagePtr pTemp2;
-		if (destSize[1] > baseSize[1])
+		ImagePtr Temp2;
+		if (DestSize[1] > BaseSize[1])
 		{
-			pTemp2 = new Image(destSize[0], destSize[1], 1, pBase->GetFormat());
-			ImageMagnifyY(pTemp2.get(), pTemp.get());
+			Temp2 = CreateImage(DestSize[0], DestSize[1], 1, Base->GetFormat(), EInitializationType::NotInitialized);
+			ImageMagnifyY(Temp2.get(), Temp.get());
+			ReleaseImage(Temp);
 		}
-		else if (destSize[1] < baseSize[1])
+		else if (DestSize[1] < BaseSize[1])
 		{
-			pTemp2 = new Image(destSize[0], destSize[1], 1, pBase->GetFormat());
-			ImageMinifyY(pTemp2.get(), pTemp.get());
+			Temp2 = CreateImage(DestSize[0], DestSize[1], 1, Base->GetFormat(), EInitializationType::NotInitialized);
+			ImageMinifyY(Temp2.get(), Temp.get());
+			ReleaseImage(Temp);
 		}
 		else
 		{
-			pTemp2 = pTemp;
+			Temp2 = Temp;
 		}
-		pTemp = nullptr;
+		Temp = nullptr;
 
 
 		// Reset format if it was changed to scale
-		if (sourceFormat != uncompressedFormat)
+		if (SourceFormat != UncompressedFormat)
 		{
-			pTemp2 = ImagePixelFormat(imageCompressionQuality, pTemp2.get(), sourceFormat);
+			Ptr<Image> OldTemp = Temp2;
+			Temp2 = ImagePixelFormat(ImageCompressionQuality, Temp2.get(), SourceFormat);
+			ReleaseImage(OldTemp);
 		}
 
-		pDest->CopyMove(pTemp2.get());
-		pTemp2 = nullptr;
+		Dest->CopyMove(Temp2.get());
+		ReleaseImage(Temp2);
+		Temp2 = nullptr;
 
 		// Update the relevancy data of the image.
-		if (pBase->m_flags & Image::EImageFlags::IF_HAS_RELEVANCY_MAP)
+		if (Base->m_flags & Image::EImageFlags::IF_HAS_RELEVANCY_MAP)
 		{
-			pDest->m_flags |= Image::EImageFlags::IF_HAS_RELEVANCY_MAP;
+			Dest->m_flags |= Image::EImageFlags::IF_HAS_RELEVANCY_MAP;
 
-			float FactorY = float(destSize[1]) / float(baseSize[1]);
+			float FactorY = float(DestSize[1]) / float(BaseSize[1]);
 
-			pDest->RelevancyMinY = uint16(FMath::FloorToFloat(pBase->RelevancyMinY * FactorY));
-			pDest->RelevancyMaxY = uint16(FMath::Min((int32)FMath::CeilToFloat(pBase->RelevancyMinY * FactorY), pDest->GetSizeY() - 1));
+			Dest->RelevancyMinY = uint16(FMath::FloorToFloat(Base->RelevancyMinY * FactorY));
+			Dest->RelevancyMaxY = uint16(FMath::Min((int32)FMath::CeilToFloat(Base->RelevancyMinY * FactorY), Dest->GetSizeY() - 1));
 		}
+
+		ReleaseImage(BaseReformat);
 	}
 
 }
\ No newline at end of file
diff --git a/Engine/Plugins/Experimental/Mutable/Source/MutableRuntime/Private/MuR/OpImageResize.h b/Engine/Plugins/Experimental/Mutable/Source/MutableRuntime/Private/MuR/OpImageResize.h
deleted file mode 100644
diff --git a/Engine/Plugins/Experimental/Mutable/Source/MutableRuntime/Private/MuR/OpImageSwizzle.cpp b/Engine/Plugins/Experimental/Mutable/Source/MutableRuntime/Private/MuR/OpImageSwizzle.cpp
--- a/Engine/Plugins/Experimental/Mutable/Source/MutableRuntime/Private/MuR/OpImageSwizzle.cpp
+++ b/Engine/Plugins/Experimental/Mutable/Source/MutableRuntime/Private/MuR/OpImageSwizzle.cpp
@@ -22,18 +22,16 @@ namespace mu
 
         // Pixelcount should already match, but due to bugs it may not be the case. Try to detect it,
         // but avoid crashing below:
-        uint16 numChannels = GetImageFormatData(Format).m_channels;
+        uint16 numChannels = GetImageFormatData(Format).Channels;
         for (uint16 c=0;c<numChannels; ++c)
         {
             if (Sources[c])
             {
-                size_t sourcePixelCount = Sources[c]->CalculatePixelCount();
-                if (PixelCount>sourcePixelCount)
+                int32 SourcePixelCount = Sources[c]->CalculatePixelCount();
+                if (PixelCount> SourcePixelCount)
                 {
-                    check(false);
-
                     // Something went wrong
-					PixelCount = sourcePixelCount;
+					PixelCount = SourcePixelCount;
                 }
             }
         }
@@ -232,7 +230,7 @@ namespace mu
 	}
 
 
-	Ptr<Image> ImageSwizzle( EImageFormat Format, const Ptr<const Image> Sources[], const uint8 Channels[] )
+	Ptr<Image> FImageOperator::ImageSwizzle( EImageFormat Format, const Ptr<const Image> Sources[], const uint8 Channels[] )
 	{
 		MUTABLE_CPUPROFILER_SCOPE(ImageSwizzle);
 
@@ -241,11 +239,9 @@ namespace mu
 			return nullptr;
 		}
 
-		ImagePtr Dest = new Image(Sources[0]->GetSizeX(), Sources[0]->GetSizeY(),
-			Sources[0]->GetLODCount(),
-			Format);
+		Ptr<Image> Dest = CreateImage(Sources[0]->GetSizeX(), Sources[0]->GetSizeY(), Sources[0]->GetLODCount(), Format, EInitializationType::Black);
 
-		ImageSwizzle(Dest.get(), Sources, Channels);
+		mu::ImageSwizzle(Dest.get(), Sources, Channels);
 
 		return Dest;
 	}
diff --git a/Engine/Plugins/Experimental/Mutable/Source/MutableRuntime/Private/MuR/Operations.h b/Engine/Plugins/Experimental/Mutable/Source/MutableRuntime/Private/MuR/Operations.h
--- a/Engine/Plugins/Experimental/Mutable/Source/MutableRuntime/Private/MuR/Operations.h
+++ b/Engine/Plugins/Experimental/Mutable/Source/MutableRuntime/Private/MuR/Operations.h
@@ -669,12 +669,6 @@ namespace mu
             uint32 blockIndex;
         };
 
-        struct ImageDifferenceArgs
-        {
-            ADDRESS a;
-            ADDRESS b;
-        };
-
         struct ImageInterpolateArgs
         {
             ADDRESS factor;
@@ -1130,7 +1124,6 @@ namespace mu
             ImageResizeVarArgs ImageResizeVar;
             ImageResizeRelArgs ImageResizeRel;
             ImageBlankLayoutArgs ImageBlankLayout;
-            ImageDifferenceArgs ImageDifference;
             ImageInterpolateArgs ImageInterpolate;
             ImageInterpolate3Args ImageInterpolate3;
             ImageSaturateArgs ImageSaturate;
diff --git a/Engine/Plugins/Experimental/Mutable/Source/MutableRuntime/Private/MuR/System.cpp b/Engine/Plugins/Experimental/Mutable/Source/MutableRuntime/Private/MuR/System.cpp
--- a/Engine/Plugins/Experimental/Mutable/Source/MutableRuntime/Private/MuR/System.cpp
+++ b/Engine/Plugins/Experimental/Mutable/Source/MutableRuntime/Private/MuR/System.cpp
@@ -311,7 +311,7 @@ namespace mu
 				// We always need to return something valid.
 				if (!pResult)
 				{
-					pResult = new mu::Image(16, 16, 1, EImageFormat::IF_RGBA_UBYTE);
+					pResult = new mu::Image(16, 16, 1, EImageFormat::IF_RGBA_UBYTE, EInitializationType::Black);
 				}
 
 				break;
diff --git a/Engine/Plugins/Experimental/Mutable/Source/MutableRuntime/Private/MuR/SystemPrivate.h b/Engine/Plugins/Experimental/Mutable/Source/MutableRuntime/Private/MuR/SystemPrivate.h
--- a/Engine/Plugins/Experimental/Mutable/Source/MutableRuntime/Private/MuR/SystemPrivate.h
+++ b/Engine/Plugins/Experimental/Mutable/Source/MutableRuntime/Private/MuR/SystemPrivate.h
@@ -1160,7 +1160,7 @@ namespace mu
         void MarkRomUsed( int32 RomIndex, const TSharedPtr<const Model>& );
 
 		/** */
-		UE_NODISCARD Ptr<Image> CreateImage(uint32 SizeX, uint32 SizeY, uint32 Lods, EImageFormat Format, EInitializationType Init = EInitializationType::Black)
+		UE_NODISCARD Ptr<Image> CreateImage(uint32 SizeX, uint32 SizeY, uint32 Lods, EImageFormat Format, EInitializationType Init)
 		{
 			CheckRunnerThread();
 
diff --git a/Engine/Plugins/Experimental/Mutable/Source/MutableRuntime/Public/MuR/Image.h b/Engine/Plugins/Experimental/Mutable/Source/MutableRuntime/Public/MuR/Image.h
--- a/Engine/Plugins/Experimental/Mutable/Source/MutableRuntime/Public/MuR/Image.h
+++ b/Engine/Plugins/Experimental/Mutable/Source/MutableRuntime/Public/MuR/Image.h
@@ -175,75 +175,6 @@ namespace mu
 
 	static_assert(uint32(EMinFilterMethod::MaxValue) <= (1 << 3), "EMinFilterMethod enum cannot hold more than 8 values");
 
-	struct FImageFormatData
-	{
-		static constexpr SIZE_T MAX_BYTES_PER_BLOCK = 16;
-
-		FImageFormatData
-			(
-				unsigned pixelsPerBlockX = 0,
-				unsigned pixelsPerBlockY = 0,
-				unsigned bytesPerBlock = 0,
-				unsigned channels = 0
-			)
-		{
-			m_pixelsPerBlockX = (uint8)pixelsPerBlockX;
-			m_pixelsPerBlockY = (uint8)pixelsPerBlockY;
-			m_bytesPerBlock = (uint16)bytesPerBlock;
-			m_channels = (uint16)channels;	
-		}
-
-		FImageFormatData
-			(
-				unsigned pixelsPerBlockX,
-				unsigned pixelsPerBlockY,
-				unsigned bytesPerBlock,
-				unsigned channels,
-				std::initializer_list<uint8> BlackBlockInit
-			)
-			: FImageFormatData(pixelsPerBlockX, pixelsPerBlockY, bytesPerBlock, channels)
-		{
-			check(MAX_BYTES_PER_BLOCK >= BlackBlockInit.size());
-
-			const SIZE_T SanitizedBlockSize = FMath::Min<SIZE_T>(MAX_BYTES_PER_BLOCK, BlackBlockInit.size());
-			FMemory::Memcpy(BlackBlock, BlackBlockInit.begin(), SanitizedBlockSize);
-		}
-
-		//! For block based formats, size of the block size. For uncompressed formats it will
-		//! always be 1,1. For non-block-based compressed formats, it will be 0,0.
-        uint8 m_pixelsPerBlockX, m_pixelsPerBlockY;
-
-		//! Number of bytes used by every pixel block, if uncompressed or block-compressed format.
-		//! For non-block-compressed formats, it returns 0.
-        uint16 m_bytesPerBlock;
-
-		//! Channels in every pixel of the image.
-        uint16 m_channels;
-
-		//! Representation of a black block of the image.
-		uint8 BlackBlock[MAX_BYTES_PER_BLOCK] = { 0 };
-	};
-		
-	MUTABLERUNTIME_API const FImageFormatData& GetImageFormatData(EImageFormat format );
-
-	/** Convert an image to another pixel format.
-	* Allocates the destination image.
-	* \warning Not all format conversions are implemented.
-	* \param onlyLOD If different than -1, only the specified lod level will be converted in the returned image.
-	* \return nullptr if the conversion failed, usually because not enough memory was allocated in the result. This is only checked for RLE compression.
-	*/
-	MUTABLERUNTIME_API Ptr<Image> ImagePixelFormat(int32 Quality, const Image* Base, EImageFormat TargetFormat, int32 OnlyLOD = -1);
-	
-	/** Convert an image to another pixel format.
-	* \warning Not all format conversions are implemented.
-	* \param onlyLOD If different than -1, only the specified lod level will be converted in the returned image.
-	* \return false if the conversion failed, usually because not enough memory was allocated in the result. This is only checked for RLE compression.
-	*/
-	MUTABLERUNTIME_API void ImagePixelFormat(bool& bOutSuccess, int32 Quality, Image* Result, const Image* Base, int32 OnlyLOD = -1);
-
-	MUTABLERUNTIME_API Ptr<Image> ImageSwizzle(EImageFormat Format, const Ptr<const Image> Sources[], const uint8 Channels[]);
-	MUTABLERUNTIME_API void ImageSwizzle(Image* Result, const Ptr<const Image> Sources[], const uint8 Channels[]);
-
     //! \brief 2D image resource with mipmaps.
 	//! \ingroup runtime
     class MUTABLERUNTIME_API Image : public Resource
@@ -264,14 +195,11 @@ namespace mu
         //!         base level. It must be a number between 1 and the maximum possible levels, which
         //!         depends on the image size.
         //! \param format Pixel format.
-        Image( uint32 sizeX, uint32 sizeY, uint32 lods, EImageFormat format, EInitializationType Init=EInitializationType::Black);
+        Image( uint32 sizeX, uint32 sizeY, uint32 lods, EImageFormat format, EInitializationType Init );
 		
 		/** */
 		static Ptr<Image> CreateAsReference( uint32 ID );
 
-		/** */
-		Ptr<Image> ExtractMip(int32 Mip) const;
-
 		//! Serialisation
 		static void Serialise( const Image* p, OutputArchive& arch );
 		static Ptr<Image> StaticUnserialise( InputArchive& arch );
@@ -391,7 +319,27 @@ namespace mu
 		}
 
 		//! Copy another image.
-		void CopyMove( Image* Other )
+		void Copy(const Image* Other)
+		{
+			LLM_SCOPE_BYNAME(TEXT("MutableRuntime"));
+			MUTABLE_CPUPROFILER_SCOPE(Copy);
+
+			if (Other == this)
+			{
+				return;
+			}
+
+			m_size = Other->m_size;
+			m_format = Other->m_format;
+			m_lods = Other->m_lods;
+			m_flags = Other->m_flags;
+			RelevancyMinY = Other->RelevancyMinY;
+			RelevancyMaxY = Other->RelevancyMaxY;
+			m_data =Other->m_data;
+			ReferenceID = Other->ReferenceID;
+		}
+
+		void CopyMove(Image* Other)
 		{
 			LLM_SCOPE_BYNAME(TEXT("MutableRuntime"));
 			MUTABLE_CPUPROFILER_SCOPE(CopyMove);
@@ -409,15 +357,15 @@ namespace mu
 			RelevancyMaxY = Other->RelevancyMaxY;
 			m_data = MoveTemp(Other->m_data);
 			ReferenceID = Other->ReferenceID;
-			
-			Other->m_size = FImageSize(0,0);
+
+			Other->m_size = FImageSize(0, 0);
 			Other->m_format = EImageFormat::IF_NONE;
 			Other->m_lods = 0;
 			Other->m_flags = 0;
 			Other->RelevancyMinY = 0;
 			Other->RelevancyMaxY = 0;
 			Other->m_data.SetNum(0);
-			Other->ReferenceID=0;
+			Other->ReferenceID = 0;
 		}
 
 
@@ -494,9 +442,163 @@ namespace mu
 		//! Get the rect inside the image bounding the non-black content of the image.
 		void GetNonBlackRect(FImageRect& rect) const;
 		void GetNonBlackRect_Reference(FImageRect& rect) const;
+	};
+
+	/** */
+	struct FImageFormatData
+	{
+		static constexpr SIZE_T MAX_BYTES_PER_BLOCK = 16;
+
+		FImageFormatData
+		(
+			uint8 InPixelsPerBlockX = 0,
+			uint8 InPixelsPerBlockY = 0,
+			uint16 InBytesPerBlock = 0,
+			uint16 InChannels = 0
+		)
+		{
+			PixelsPerBlockX = InPixelsPerBlockX;
+			PixelsPerBlockY = InPixelsPerBlockY;
+			BytesPerBlock = InBytesPerBlock;
+			Channels = InChannels;
+		}
+
+		FImageFormatData
+		(
+			uint8 InPixelsPerBlockX,
+			uint8 InPixelsPerBlockY,
+			uint16 InBytesPerBlock,
+			uint16 InChannels,
+			std::initializer_list<uint8> BlackBlockInit
+		)
+			: FImageFormatData(InPixelsPerBlockX, InPixelsPerBlockY, InBytesPerBlock, InChannels)
+		{
+			check(MAX_BYTES_PER_BLOCK >= BlackBlockInit.size());
+
+			const SIZE_T SanitizedBlockSize = FMath::Min<SIZE_T>(MAX_BYTES_PER_BLOCK, BlackBlockInit.size());
+			FMemory::Memcpy(BlackBlock, BlackBlockInit.begin(), SanitizedBlockSize);
+		}
+
+		/** For block based formats, size of the block size.For uncompressed formats it will always be 1,1. For non-block-based compressed formats, it will be 0,0. */
+		uint8 PixelsPerBlockX, PixelsPerBlockY;
+
+		/** Number of bytes used by every pixel block, if uncompressed or block-compressed format.
+		 * For non-block-compressed formats, it returns 0.
+		 */
+		uint16 BytesPerBlock;
+
+		/** Channels in every pixel of the image. */
+		uint16 Channels;
+
+		/** Representation of a black block of the image. */
+		uint8 BlackBlock[MAX_BYTES_PER_BLOCK] = { 0 };
+	};
+
+	MUTABLERUNTIME_API const FImageFormatData& GetImageFormatData(EImageFormat format);
+
+	MUTABLERUNTIME_API struct FMipmapGenerationSettings
+	{
+		float m_sharpenFactor = 0.0f;
+		EMipmapFilterType m_filterType = EMipmapFilterType::MFT_SimpleAverage;
+		EAddressMode m_addressMode = EAddressMode::None;
+		bool m_ditherMipmapAlpha = false;
+
+		MUTABLERUNTIME_API void Serialise(OutputArchive& arch) const;
+		MUTABLERUNTIME_API void Unserialise(InputArchive& arch);
+	};
+
+	/** This struct contains image operations that may need to allocate and free temporary images to work.
+	* It's purpose is to override the creation, release and clone functions depending on the context.
+	*/
+	MUTABLERUNTIME_API struct FImageOperator
+	{
+		/** Common callback used for functions that can create temporary images. */
+		typedef TFunction<Ptr<Image>(int32, int32, int32, EImageFormat, EInitializationType)> FImageCreateFunc;
+		typedef TFunction<void(Ptr<Image>&)> FImageReleaseFunc;
+		typedef TFunction<Ptr<Image>(const Image*)> FImageCloneFunc;
+
+		FImageCreateFunc CreateImage;
+		FImageReleaseFunc ReleaseImage;
+		FImageCloneFunc CloneImage;
+
+		FImageOperator(FImageCreateFunc Create, FImageReleaseFunc Release, FImageCloneFunc Clone) : CreateImage(Create), ReleaseImage(Release), CloneImage(Clone) {};
+
+		/** Create an default version for untracked resources. */
+		static inline FImageOperator GetDefault()
+		{
+			return FImageOperator(
+				[](int32 x, int32 y, int32 m, EImageFormat f, EInitializationType i) { return new Image(x, y, m, f, EInitializationType::NotInitialized); },
+				[](Ptr<Image>& i) {i = nullptr; },
+				[](const Image* i) { return i->Clone(); }
+			);
+		}
+
+		/** Convert an image to another pixel format.
+		* Allocates the destination image.
+		* \warning Not all format conversions are implemented.
+		* \param onlyLOD If different than -1, only the specified lod level will be converted in the returned image.
+		* \return nullptr if the conversion failed, usually because not enough memory was allocated in the result. This is only checked for RLE compression.
+		*/
+		MUTABLERUNTIME_API Ptr<Image> ImagePixelFormat(int32 Quality, const Image* Base, EImageFormat TargetFormat, int32 OnlyLOD = -1);
+
+		/** Convert an image to another pixel format.
+		* \warning Not all format conversions are implemented.
+		* \param onlyLOD If different than -1, only the specified lod level will be converted in the returned image.
+		* \return false if the conversion failed, usually because not enough memory was allocated in the result. This is only checked for RLE compression.
+		*/
+		MUTABLERUNTIME_API void ImagePixelFormat(bool& bOutSuccess, int32 Quality, Image* Result, const Image* Base, int32 OnlyLOD = -1);
+
+		MUTABLERUNTIME_API Ptr<Image> ImageSwizzle(EImageFormat Format, const Ptr<const Image> Sources[], const uint8 Channels[]);
+
+		/** */
+		MUTABLERUNTIME_API Ptr<Image> ExtractMip(const Image* From, int32 Mip);
+
+		/** Bilinear filter image resize. */
+		void ImageResizeLinear(Image* Dest, int32 ImageCompressionQuality, const Image* Base);
 
 		/** Fill the image with a plain colour. */
-		void FillColour(FVector4f c);
+		void FillColor(Image* Image, FVector4f Color);
+
+		/** Support struct to keep preallocated data required for some mipmap operations. */
+		struct FScratchImageMipmap
+		{
+			Ptr<Image> Uncompressed;
+			Ptr<Image> UncompressedMips;
+			Ptr<Image> CompressedMips;
+		};
+
+		/** Generate the mipmaps for images.
+		* if bGenerateOnlyTail is true, generates the mips missing from Base to LevelCount and sets
+		* them in Dest (the full chain is spit in two images). Otherwise generate the mips missing
+		* from Base up to LevelCount and append them in Dest to the already generated Base's mips.
+		*/
+		void ImageMipmap(int32 CompressionQuality, Image* Dest, const Image* Base,
+			int32 LevelCount,
+			const FMipmapGenerationSettings&, bool bGenerateOnlyTail = false);
+
+		/** Mipmap separating the worst case treatment in 3 steps to manage allocations of temp data. */
+		void ImageMipmap_PrepareScratch(Image* Dest, const Image* Base, int32 LevelCount, FScratchImageMipmap&);
+		void ImageMipmap(FImageOperator::FScratchImageMipmap&, int32 CompressionQuality, Image* Dest, const Image* Base,
+			int32 LevelCount,
+			const FMipmapGenerationSettings&, bool bGenerateOnlyTail = false);
+		void ImageMipmap_ReleaseScratch(FScratchImageMipmap&);
+
+		/** */
+		MUTABLERUNTIME_API bool ImageCrop(Image* Cropped, int32 CompressionQuality, const Image* Base, const box< vec2<int32> >& Rect);
+
+		/** */
+		void ImageCompose(Image* Base, const Image* Block, const box< UE::Math::TIntVector2<uint16> >& Rect);
+
 	};
+
+	/** Update all the mipmaps in the image from the data in the base one.
+	* Only the mipmaps already existing in the image are updated.
+	*/
+	MUTABLERUNTIME_API void ImageMipmapInPlace(int32 CompressionQuality, Image* Base, const FMipmapGenerationSettings&);
+
+	/** */
+	MUTABLERUNTIME_API void ImageSwizzle(Image* Result, const Ptr<const Image> Sources[], const uint8 Channels[]);
+
+
 }
 
diff --git a/Engine/Plugins/Experimental/Mutable/Source/MutableTools/Private/MuT/ASTOpImagePixelFormat.cpp b/Engine/Plugins/Experimental/Mutable/Source/MutableTools/Private/MuT/ASTOpImagePixelFormat.cpp
--- a/Engine/Plugins/Experimental/Mutable/Source/MutableTools/Private/MuT/ASTOpImagePixelFormat.cpp
+++ b/Engine/Plugins/Experimental/Mutable/Source/MutableTools/Private/MuT/ASTOpImagePixelFormat.cpp
@@ -108,7 +108,7 @@ namespace mu
 
 		EImageFormat format = Format;
 		bool bIsCompressedFormat = IsCompressedFormat(Format);
-		//bool isBlockFormat = GetImageFormatData( format ).m_pixelsPerBlockX!=0;
+		//bool isBlockFormat = GetImageFormatData( format ).PixelsPerBlockX!=0;
 
 		// The instruction can be sunk
 		OP_TYPE sourceType = sourceAt->GetOpType();
@@ -182,7 +182,7 @@ namespace mu
 			EImageFormat layoutFormat = (EImageFormat)nop->op.args.ImageBlankLayout.format;
 			if (FormatIfAlpha != EImageFormat::IF_NONE
 				&&
-				GetImageFormatData(layoutFormat).m_channels > 3)
+				GetImageFormatData(layoutFormat).Channels > 3)
 			{
 				format = FormatIfAlpha;
 			}
@@ -200,7 +200,7 @@ namespace mu
 			EImageFormat layoutFormat = (EImageFormat)nop->op.args.ImagePlainColour.format;
 			if (FormatIfAlpha != EImageFormat::IF_NONE
 				&&
-				GetImageFormatData(layoutFormat).m_channels > 3)
+				GetImageFormatData(layoutFormat).Channels > 3)
 			{
 				format = FormatIfAlpha;
 			}
@@ -249,7 +249,7 @@ namespace mu
 
 		if (FormatIfAlpha != EImageFormat::IF_NONE
 			&&
-			GetImageFormatData(res.m_format).m_channels > 3)
+			GetImageFormatData(res.m_format).Channels > 3)
 		{
 			res.m_format = FormatIfAlpha;
 		}
@@ -339,7 +339,7 @@ namespace mu
 
 		EImageFormat format = currentFormatOp->Format;
 		bool bIsCompressedFormat = IsCompressedFormat(format);
-		bool isBlockFormat = GetImageFormatData(format).m_pixelsPerBlockX != 0;
+		bool isBlockFormat = GetImageFormatData(format).PixelsPerBlockX != 0;
 
 		// Already visited?
 		const Ptr<ASTOp>* Cached = OldToNew.Find({ at,currentFormatOp });
@@ -381,8 +381,8 @@ namespace mu
 			{
 				// We can only optimise if the layout grid blocks size in pixels is
 				// a multiple of the image format block size.
-				int imageFormatBlockSizeX = GetImageFormatData(format).m_pixelsPerBlockX;
-				int imageFormatBlockSizeY = GetImageFormatData(format).m_pixelsPerBlockY;
+				int imageFormatBlockSizeX = GetImageFormatData(format).PixelsPerBlockX;
+				int imageFormatBlockSizeY = GetImageFormatData(format).PixelsPerBlockY;
 				bool acceptable = imageFormatBlockSizeX == 1 && imageFormatBlockSizeY == 1;
 
 				if (!acceptable)
diff --git a/Engine/Plugins/Experimental/Mutable/Source/MutableTools/Private/MuT/ASTOpImageSwizzle.cpp b/Engine/Plugins/Experimental/Mutable/Source/MutableTools/Private/MuT/ASTOpImageSwizzle.cpp
--- a/Engine/Plugins/Experimental/Mutable/Source/MutableTools/Private/MuT/ASTOpImageSwizzle.cpp
+++ b/Engine/Plugins/Experimental/Mutable/Source/MutableTools/Private/MuT/ASTOpImageSwizzle.cpp
@@ -233,7 +233,7 @@ namespace mu
 
 		// If we are not changing channel order, just remove the swizzle and adjust the format.
 		bool bSameChannelOrder = true;
-		int32 NumChannelsInFormat = GetImageFormatData(Format).m_channels;
+		int32 NumChannelsInFormat = GetImageFormatData(Format).Channels;
 		for (int32 c = 0; c < NumChannelsInFormat; ++c)
 		{
 			if (Sources[c] && SourceChannels[c] != c)
diff --git a/Engine/Plugins/Experimental/Mutable/Source/MutableTools/Private/MuT/CodeGenerator.cpp b/Engine/Plugins/Experimental/Mutable/Source/MutableTools/Private/MuT/CodeGenerator.cpp
--- a/Engine/Plugins/Experimental/Mutable/Source/MutableTools/Private/MuT/CodeGenerator.cpp
+++ b/Engine/Plugins/Experimental/Mutable/Source/MutableTools/Private/MuT/CodeGenerator.cpp
@@ -1389,7 +1389,7 @@ namespace mu
 								bop->op.args.ImageBlankLayout.blockSize[0] = uint16(blockSizeX);
 								bop->op.args.ImageBlankLayout.blockSize[1] = uint16(blockSizeY);
 								// We support block compression directly here, but not non-block compression
-								if (GetImageFormatData(desc.m_format).m_pixelsPerBlockX == 0)
+								if (GetImageFormatData(desc.m_format).PixelsPerBlockX == 0)
 								{
 									// It's something like RLE
 									bop->op.args.ImageBlankLayout.format = GetUncompressedFormat(desc.m_format);
@@ -1605,8 +1605,8 @@ namespace mu
 								// we will devide the layout block by the format block
 								const FImageFormatData& finfo = GetImageFormatData(desc.m_format);
 
-								int mipsX = (int)ceilf(logf((float)blockSizeX / finfo.m_pixelsPerBlockX) / logf(2.0f));
-								int mipsY = (int)ceilf(logf((float)blockSizeY / finfo.m_pixelsPerBlockY) / logf(2.0f));
+								int mipsX = (int)ceilf(logf((float)blockSizeX / finfo.PixelsPerBlockX) / logf(2.0f));
+								int mipsY = (int)ceilf(logf((float)blockSizeY / finfo.PixelsPerBlockY) / logf(2.0f));
 								mop->BlockLevels = (uint8_t)FMath::Max(mipsX, mipsY);
 
 								mop->AddressMode = mipmapNode->GetPrivate()->m_settings.m_addressMode;
@@ -1632,8 +1632,8 @@ namespace mu
 								// we will devide the layout block by the format block
 								const FImageFormatData& finfo = GetImageFormatData(desc.m_format);
 
-								int mipsX = (int)ceilf(logf((float)blockSizeX / finfo.m_pixelsPerBlockX) / logf(2.0f));
-								int mipsY = (int)ceilf(logf((float)blockSizeY / finfo.m_pixelsPerBlockY) / logf(2.0f));
+								int mipsX = (int)ceilf(logf((float)blockSizeX / finfo.PixelsPerBlockX) / logf(2.0f));
+								int mipsY = (int)ceilf(logf((float)blockSizeY / finfo.PixelsPerBlockY) / logf(2.0f));
 								mop->BlockLevels = (uint8_t)FMath::Max(mipsX, mipsY);
 
 							// Not important for the end of the mip tail?
diff --git a/Engine/Plugins/Experimental/Mutable/Source/MutableTools/Private/MuT/CodeGenerator_Image.cpp b/Engine/Plugins/Experimental/Mutable/Source/MutableTools/Private/MuT/CodeGenerator_Image.cpp
--- a/Engine/Plugins/Experimental/Mutable/Source/MutableTools/Private/MuT/CodeGenerator_Image.cpp
+++ b/Engine/Plugins/Experimental/Mutable/Source/MutableTools/Private/MuT/CodeGenerator_Image.cpp
@@ -7,7 +7,6 @@
 #include "MuR/Mesh.h"
 #include "MuR/MutableMath.h"
 #include "MuR/MutableTrace.h"
-#include "MuR/OpImageCrop.h"
 #include "MuR/OpImageProject.h"
 #include "MuR/Operations.h"
 #include "MuR/Parameters.h"
@@ -240,8 +239,16 @@ namespace mu
         //
         if ( pImage->GetSizeX()!=cropRect.size[0] || pImage->GetSizeY()!=cropRect.size[1] )
         {
-            Ptr<Image> pCropped = new Image(cropRect.size[0], cropRect.size[1], 1, pImage->GetFormat());
-			ImageCrop(pCropped.get(), m_compilerOptions->ImageCompressionQuality, pImage.get(), cropRect );
+			FImageOperator ImOp
+			(
+				[](int32 x, int32 y, int32 m, EImageFormat f, EInitializationType i) { return new Image(x, y, m, f, i); },
+				[](Ptr<Image>& i) {i = nullptr; },
+				[](const Image* i) { return i->Clone(); }
+			);
+
+
+            Ptr<Image> pCropped = new Image(cropRect.size[0], cropRect.size[1], 1, pImage->GetFormat(), EInitializationType::NotInitialized);
+			ImOp.ImageCrop(pCropped.get(), m_compilerOptions->ImageCompressionQuality, pImage.get(), cropRect);
             op->SetValue( pCropped, m_compilerOptions->OptimisationOptions.bUseDiskCache );
         }
         else
@@ -1783,7 +1790,7 @@ namespace mu
             // Make a checkered debug image
             const FImageSize size = MUTABLE_MISSING_IMAGE_DESC.m_size;
 
-            ImagePtr pImage = new Image( size[0], size[1], 1, format );
+            ImagePtr pImage = new Image( size[0], size[1], 1, format, EInitializationType::NotInitialized);
 
             switch (format)
             {
@@ -1892,7 +1899,7 @@ namespace mu
     Ptr<ASTOp> CodeGenerator::GeneratePlainImageCode( const vec3<float>& colour, const FImageGenerationOptions& Options )
     {
         const int size = 4;
-        ImagePtr pImage = new Image( size, size, 1, EImageFormat::IF_RGB_UBYTE );
+        ImagePtr pImage = new Image( size, size, 1, EImageFormat::IF_RGB_UBYTE, EInitializationType::NotInitialized);
 
         uint8_t* pData = pImage->GetData();
         for ( int p=0; p<size*size; ++p )
diff --git a/Engine/Plugins/Experimental/Mutable/Source/MutableTools/Private/MuT/CodeOptimiser_Semantic.cpp b/Engine/Plugins/Experimental/Mutable/Source/MutableTools/Private/MuT/CodeOptimiser_Semantic.cpp
--- a/Engine/Plugins/Experimental/Mutable/Source/MutableTools/Private/MuT/CodeOptimiser_Semantic.cpp
+++ b/Engine/Plugins/Experimental/Mutable/Source/MutableTools/Private/MuT/CodeOptimiser_Semantic.cpp
@@ -35,36 +35,6 @@
 namespace mu
 {
 
-
-    //---------------------------------------------------------------------------------------------
-    //! Get an uncompressed version of the image.
-    //! If the image is not compressed, it si returned unmodified.
-    //---------------------------------------------------------------------------------------------
-    ImagePtrConst GetUncompressed( ImagePtrConst pSource )
-    {
-        ImagePtrConst pResult = pSource;
-
-        if ( pSource->GetFormat()== EImageFormat::IF_L_UBYTE_RLE )
-        {
-            ImagePtr pNew = new Image( pSource->GetSizeX(), pSource->GetSizeY(),
-                                 pSource->GetLODCount(),
-				EImageFormat::IF_L_UBYTE );
-            UncompressRLE_L( pSource.get(), pNew.get() );
-            pResult = pNew;
-        }
-        else if ( pSource->GetFormat()== EImageFormat::IF_L_UBIT_RLE )
-        {
-            ImagePtr pNew = new Image( pSource->GetSizeX(), pSource->GetSizeY(),
-                                 pSource->GetLODCount(),
-				EImageFormat::IF_L_UBYTE );
-            UncompressRLE_L1( pSource.get(), pNew.get() );
-            pResult = pNew;
-        }
-
-        return pResult;
-    }
-
-
     //---------------------------------------------------------------------------------------------
     //---------------------------------------------------------------------------------------------
     //---------------------------------------------------------------------------------------------
diff --git a/Engine/Plugins/Experimental/Mutable/Source/MutableTools/Private/MuT/DataPacker.cpp b/Engine/Plugins/Experimental/Mutable/Source/MutableTools/Private/MuT/DataPacker.cpp
--- a/Engine/Plugins/Experimental/Mutable/Source/MutableTools/Private/MuT/DataPacker.cpp
+++ b/Engine/Plugins/Experimental/Mutable/Source/MutableTools/Private/MuT/DataPacker.cpp
@@ -739,6 +739,8 @@ namespace mu
 				ASTOpConstantResource* typed = dynamic_cast<ASTOpConstantResource*>(n.get());
                 Ptr<const Image> pOld = static_cast<const Image*>(typed->GetValue().get());
 
+				FImageOperator ImOp = FImageOperator::GetDefault();
+
 				// See if there is a better format for this image
 				FVector4f PlainColor;
 				if ( pOld->IsPlainColour(PlainColor) )
@@ -766,7 +768,7 @@ namespace mu
 				}
 				else if ( accFormat.m_supportedFormats[typed][(size_t)EImageFormat::IF_L_UBIT_RLE] )
                 {
-                    ImagePtr pNew = ImagePixelFormat( imageCompressionQuality, pOld.get(), EImageFormat::IF_L_UBIT_RLE );
+                    ImagePtr pNew = ImOp.ImagePixelFormat( imageCompressionQuality, pOld.get(), EImageFormat::IF_L_UBIT_RLE );
 
                     // Only replace if the compression was worth!
                     size_t oldSize = pOld->GetDataSize();
@@ -778,7 +780,7 @@ namespace mu
                 }
                 else if ( accFormat.m_supportedFormats[typed][(size_t)EImageFormat::IF_L_UBYTE_RLE] )
                 {
-                    ImagePtr pNew = ImagePixelFormat( imageCompressionQuality, pOld.get(), EImageFormat::IF_L_UBYTE_RLE );
+                    ImagePtr pNew = ImOp.ImagePixelFormat( imageCompressionQuality, pOld.get(), EImageFormat::IF_L_UBYTE_RLE );
 
                     // Only replace if the compression was worth!
                     size_t oldSize = pOld->GetDataSize();
diff --git a/Engine/Plugins/Experimental/Mutable/Source/MutableTools/Private/MuT/ImageDescGenerator.cpp b/Engine/Plugins/Experimental/Mutable/Source/MutableTools/Private/MuT/ImageDescGenerator.cpp
--- a/Engine/Plugins/Experimental/Mutable/Source/MutableTools/Private/MuT/ImageDescGenerator.cpp
+++ b/Engine/Plugins/Experimental/Mutable/Source/MutableTools/Private/MuT/ImageDescGenerator.cpp
@@ -290,7 +290,7 @@ namespace mu
         // Necessary since dxt1 reports 4 channels just in case
         if (m_desc.m_format!=node.m_format)
         {
-            if (GetImageFormatData(m_desc.m_format).m_channels>3
+            if (GetImageFormatData(m_desc.m_format).Channels >3
                     &&
                     node.m_formatIfAlpha!= EImageFormat::IF_NONE
                     )
diff --git a/Engine/Plugins/Experimental/Mutable/Source/MutableTools/Private/MuT/NodeImageSwizzle.cpp b/Engine/Plugins/Experimental/Mutable/Source/MutableTools/Private/MuT/NodeImageSwizzle.cpp
--- a/Engine/Plugins/Experimental/Mutable/Source/MutableTools/Private/MuT/NodeImageSwizzle.cpp
+++ b/Engine/Plugins/Experimental/Mutable/Source/MutableTools/Private/MuT/NodeImageSwizzle.cpp
@@ -71,7 +71,7 @@ namespace mu
 	{
 		m_pD->m_format = format;
 
-		int channelCount = GetImageFormatData( format ).m_channels;
+		int channelCount = GetImageFormatData( format ).Channels;
 		m_pD->m_sources.SetNum( channelCount );
 		m_pD->m_sourceChannels.SetNum( channelCount );
 	}
