[mutable] Fixes in DLL exports for Linux editor compilation.
#rnx

diff --git a/Engine/Plugins/Experimental/Mutable/Source/MutableRuntime/Public/MuR/Image.h b/Engine/Plugins/Experimental/Mutable/Source/MutableRuntime/Public/MuR/Image.h
--- a/Engine/Plugins/Experimental/Mutable/Source/MutableRuntime/Public/MuR/Image.h
+++ b/Engine/Plugins/Experimental/Mutable/Source/MutableRuntime/Public/MuR/Image.h
@@ -496,21 +496,21 @@ namespace mu
 
 	MUTABLERUNTIME_API const FImageFormatData& GetImageFormatData(EImageFormat format);
 
-	MUTABLERUNTIME_API struct FMipmapGenerationSettings
+	struct MUTABLERUNTIME_API FMipmapGenerationSettings
 	{
 		float m_sharpenFactor = 0.0f;
 		EMipmapFilterType m_filterType = EMipmapFilterType::MFT_SimpleAverage;
 		EAddressMode m_addressMode = EAddressMode::None;
 		bool m_ditherMipmapAlpha = false;
 
-		MUTABLERUNTIME_API void Serialise(OutputArchive& arch) const;
-		MUTABLERUNTIME_API void Unserialise(InputArchive& arch);
+		void Serialise(OutputArchive& arch) const;
+		void Unserialise(InputArchive& arch);
 	};
 
 	/** This struct contains image operations that may need to allocate and free temporary images to work.
 	* It's purpose is to override the creation, release and clone functions depending on the context.
 	*/
-	MUTABLERUNTIME_API struct FImageOperator
+	struct FImageOperator
 	{
 		/** Common callback used for functions that can create temporary images. */
 		typedef TFunction<Ptr<Image>(int32, int32, int32, EImageFormat, EInitializationType)> FImageCreateFunc;
