[mutable] Reduce error/warning spam

#rnx

diff --git a/Engine/Plugins/Experimental/Mutable/Source/CustomizableObjectEditor/Private/MuCOE/CompilationMessageCache.cpp b/Engine/Plugins/Experimental/Mutable/Source/CustomizableObjectEditor/Private/MuCOE/CompilationMessageCache.cpp
--- a/Engine/Plugins/Experimental/Mutable/Source/CustomizableObjectEditor/Private/MuCOE/CompilationMessageCache.cpp
+++ b/Engine/Plugins/Experimental/Mutable/Source/CustomizableObjectEditor/Private/MuCOE/CompilationMessageCache.cpp
@@ -2,9 +2,12 @@
 
 #include "MuCOE/CompilationMessageCache.h"
 
+FCompilationMessageCache::FCompilationMessageCache()
+{
+	ClearMessageCounters();
+}
 
-
-bool FCompilationMessageCache::AddMessage(const FText& InMessage, const TArray<const UCustomizableObjectNode*>& InArrayNode, EMessageSeverity::Type MessageSeverity /* = EMessageSeverity::Warning */)
+bool FCompilationMessageCache::AddMessage(const FText& InMessage, const TArray<const UCustomizableObjectNode*>& InArrayNode, EMessageSeverity::Type MessageSeverity /* = EMessageSeverity::Warning */, const ELoggerSpamBin SpamBin /*= ELoggerSpamBin::ShowAll*/)
 {
 	// Skip message if an identical one has already been reported
 	const FLoggedMessage Cached = { InMessage, InArrayNode, MessageSeverity };
@@ -13,6 +16,18 @@ bool FCompilationMessageCache::AddMessage(const FText& InMessage, const TArray<c
 		return false;
 	}
 
+	// Skip message if it has a spam bin other than all and is introduced too many times
+	if (SpamBin != ELoggerSpamBin::ShowAll)
+	{
+		uint32 currentNumOfBinMessages = ++SpamBinCounts[static_cast<uint8>(SpamBin)];
+		if (currentNumOfBinMessages > MaxSpamMessages)
+		{
+			IgnoredCount++;
+			IgnoredMessages.Add(Cached);
+			return false;
+		}
+	}
+
 	LoggedMessages.Add(Cached);
 
 	switch (MessageSeverity)
@@ -74,6 +89,7 @@ void FCompilationMessageCache::GetMessages(TArray<FText>& OutWarningMessages, TA
 void FCompilationMessageCache::ClearMessagesArray()
 {
 	LoggedMessages.Empty();
+	IgnoredMessages.Empty();
 }
 
 
@@ -82,6 +98,9 @@ void FCompilationMessageCache::ClearMessageCounters()
 	PerformanceWarningCount = 0;
 	WarningCount = 0;
 	ErrorCount = 0;
+	IgnoredCount = 0;
+	SpamBinCounts.Empty();
+	SpamBinCounts.SetNumZeroed(static_cast<uint8>(ELoggerSpamBin::ShowAll), false);
 }
 
 uint32 FCompilationMessageCache::GetWarningCount(bool bIncludePerformanceWarnings) const
@@ -98,3 +117,8 @@ uint32 FCompilationMessageCache::GetErrorCount() const
 {
 	return ErrorCount;
 }
+
+uint32 FCompilationMessageCache::GetIgnoredCount() const
+{
+	return IgnoredCount;
+}
diff --git a/Engine/Plugins/Experimental/Mutable/Source/CustomizableObjectEditor/Private/MuCOE/CustomizableObjectCompileRunnable.cpp b/Engine/Plugins/Experimental/Mutable/Source/CustomizableObjectEditor/Private/MuCOE/CustomizableObjectCompileRunnable.cpp
--- a/Engine/Plugins/Experimental/Mutable/Source/CustomizableObjectEditor/Private/MuCOE/CustomizableObjectCompileRunnable.cpp
+++ b/Engine/Plugins/Experimental/Mutable/Source/CustomizableObjectEditor/Private/MuCOE/CustomizableObjectCompileRunnable.cpp
@@ -103,35 +103,32 @@ uint32 FCustomizableObjectCompileRunnable::Run()
 	for (int i = 0; i < pLog->GetMessageCount(); ++i)
 	{
 		const FString& Message = pLog->GetMessageText(i);
-		mu::ErrorLogMessageType MessageType = pLog->GetMessageType(i);
-		mu::ErrorLogMessageAttachedDataView MessageAttachedData = pLog->GetMessageAttachedData(i);
+		const mu::ErrorLogMessageType MessageType = pLog->GetMessageType(i);
+		const mu::ErrorLogMessageAttachedDataView MessageAttachedData = pLog->GetMessageAttachedData(i);
 
-		if (MessageType == mu::ELMT_WARNING)
+		if (MessageType == mu::ELMT_WARNING || MessageType == mu::ELMT_ERROR)
 		{
-			if (MessageAttachedData.m_unassignedUVs && MessageAttachedData.m_unassignedUVsSize > 0) 
-			{			
-				TSharedPtr<FErrorAttachedData> ErrorAttachedData = MakeShared<FErrorAttachedData>();
-				ErrorAttachedData->UnassignedUVs.Reset();
-				ErrorAttachedData->UnassignedUVs.Append(MessageAttachedData.m_unassignedUVs, MessageAttachedData.m_unassignedUVsSize);
-				ArrayWarning.Add(FError(FText::AsCultureInvariant(Message), ErrorAttachedData, pLog->GetMessageContext(i)));
-			}
-			else
-			{
-				ArrayWarning.Add(FError(FText::AsCultureInvariant(Message), pLog->GetMessageContext(i)));
+			const EMessageSeverity::Type Severity = MessageType == mu::ELMT_WARNING ? EMessageSeverity::Warning : EMessageSeverity::Error;
+			const ELoggerSpamBin SpamBin = [&] {
+				switch (pLog->GetMessageSpamBin(i)) {
+				case mu::ErrorLogMessageSpamBin::ELMSB_UNKNOWN_TAG:
+					return ELoggerSpamBin::TagsNotFound;
+				case mu::ErrorLogMessageSpamBin::ELMSB_ALL:
+				default:
+					return ELoggerSpamBin::ShowAll;
 			}
-		}
-		else if (MessageType == mu::ELMT_ERROR)
-		{
+			}();
+
 			if (MessageAttachedData.m_unassignedUVs && MessageAttachedData.m_unassignedUVsSize > 0) 
 			{			
 				TSharedPtr<FErrorAttachedData> ErrorAttachedData = MakeShared<FErrorAttachedData>();
 				ErrorAttachedData->UnassignedUVs.Reset();
 				ErrorAttachedData->UnassignedUVs.Append(MessageAttachedData.m_unassignedUVs, MessageAttachedData.m_unassignedUVsSize);
-				ArrayError.Add(FError(FText::AsCultureInvariant(Message), ErrorAttachedData, pLog->GetMessageContext(i)));
+				ArrayErrors.Add(FError(Severity, FText::AsCultureInvariant(Message), ErrorAttachedData, pLog->GetMessageContext(i), SpamBin));
 			}
 			else
 			{
-				ArrayError.Add(FError(FText::AsCultureInvariant(Message), pLog->GetMessageContext(i)));
+				ArrayErrors.Add(FError(Severity, FText::AsCultureInvariant(Message), pLog->GetMessageContext(i), SpamBin));
 			}
 		}
 	}
@@ -152,15 +149,9 @@ bool FCustomizableObjectCompileRunnable::IsCompleted() const
 }
 
 
-const TArray<FCustomizableObjectCompileRunnable::FError>& FCustomizableObjectCompileRunnable::GetArrayError() const
-{
-	return ArrayError;
-}
-
-
-const TArray<FCustomizableObjectCompileRunnable::FError>& FCustomizableObjectCompileRunnable::GetArrayWarning() const
+const TArray<FCustomizableObjectCompileRunnable::FError>& FCustomizableObjectCompileRunnable::GetArrayErrors() const
 {
-	return ArrayWarning;
+	return ArrayErrors;
 }
 
 
diff --git a/Engine/Plugins/Experimental/Mutable/Source/CustomizableObjectEditor/Private/MuCOE/CustomizableObjectCompileRunnable.h b/Engine/Plugins/Experimental/Mutable/Source/CustomizableObjectEditor/Private/MuCOE/CustomizableObjectCompileRunnable.h
--- a/Engine/Plugins/Experimental/Mutable/Source/CustomizableObjectEditor/Private/MuCOE/CustomizableObjectCompileRunnable.h
+++ b/Engine/Plugins/Experimental/Mutable/Source/CustomizableObjectEditor/Private/MuCOE/CustomizableObjectCompileRunnable.h
@@ -6,6 +6,7 @@
 #include "MuCO/CustomizableObject.h"
 #include "MuR/Ptr.h"
 #include "MuT/Node.h"
+#include "MuCOE/CustomizableObjectEditorLogger.h"
 
 class ITargetPlatform;
 
@@ -21,20 +22,21 @@ public:
 
 	struct FError
 	{
+		EMessageSeverity::Type Severity;
+		ELoggerSpamBin SpamBin = ELoggerSpamBin::ShowAll;
 		FText Message;
 		TSharedPtr<FErrorAttachedData> AttachedData;
 		const void* Context;
 
-		FError(const FText& InMessage, const void* InContext) : Message(InMessage), Context(InContext) {}
-		FError(const FText& InMessage, const TSharedPtr<FErrorAttachedData>& InAttachedData, const void* InContext) 
-			: Message(InMessage), AttachedData(InAttachedData), Context(InContext) {}
+		FError(const EMessageSeverity::Type InSeverity, const FText& InMessage, const void* InContext, const ELoggerSpamBin InSpamBin = ELoggerSpamBin::ShowAll) : Severity(InSeverity), SpamBin(InSpamBin), Message(InMessage), Context(InContext) {}
+		FError(const EMessageSeverity::Type InSeverity, const FText& InMessage, const TSharedPtr<FErrorAttachedData>& InAttachedData, const void* InContext, const ELoggerSpamBin InSpamBin = ELoggerSpamBin::ShowAll)
+			: Severity(InSeverity), SpamBin(InSpamBin), Message(InMessage), AttachedData(InAttachedData), Context(InContext) {}
 	};
 
 private:
 
 	mu::Ptr<mu::Node> MutableRoot;
-	TArray<FError> ArrayWarning;
-	TArray<FError> ArrayError;
+	TArray<FError> ArrayErrors;
 
 public:
 
@@ -49,10 +51,7 @@ public:
 	bool IsCompleted() const;
 
 	//
-	const TArray<FError>& GetArrayError() const;
-	
-	//
-	const TArray<FError>& GetArrayWarning() const;
+	const TArray<FError>& GetArrayErrors() const;
 
 public:
 
diff --git a/Engine/Plugins/Experimental/Mutable/Source/CustomizableObjectEditor/Private/MuCOE/CustomizableObjectCompiler.cpp b/Engine/Plugins/Experimental/Mutable/Source/CustomizableObjectEditor/Private/MuCOE/CustomizableObjectCompiler.cpp
--- a/Engine/Plugins/Experimental/Mutable/Source/CustomizableObjectEditor/Private/MuCOE/CustomizableObjectCompiler.cpp
+++ b/Engine/Plugins/Experimental/Mutable/Source/CustomizableObjectEditor/Private/MuCOE/CustomizableObjectCompiler.cpp
@@ -1548,6 +1548,7 @@ void FCustomizableObjectCompiler::NotifyCompilationErrors() const
 {
 	const uint32 NumWarnings = CompilationLogsContainer.GetWarningCount(false);
 	const uint32 NumErrors = CompilationLogsContainer.GetErrorCount();
+	const uint32 NumIgnoreds = CompilationLogsContainer.GetIgnoredCount();
 	const bool NoWarningsOrErrors = !(NumWarnings || NumErrors);
 
 	if (Options.bSilentCompilation && NoWarningsOrErrors)
@@ -1574,8 +1575,11 @@ void FCustomizableObjectCompiler::NotifyCompilationErrors() const
 	const FText Prefix = FText::FromString(CurrentObject ? CurrentObject->GetName() : "Customizable Object");
 
 	const FText Message = NoWarningsOrErrors ?
-		FText::Format(LOCTEXT("CompilationFinishedSuccessfully", "{0} finished compiling successfully"), Prefix) :
-		FText::Format(LOCTEXT("CompilationFinished", "{0} finished compiling successfully with {1} {1}|plural(one=warning,other=warnings) and {2} {2}|plural(one=error,other=errors)"), Prefix, NumWarnings, NumErrors);
+		FText::Format(LOCTEXT("CompilationFinishedSuccessfully", "{0} finished compiling successfully."), Prefix) :
+		NumIgnoreds > 0 ?
+		FText::Format(LOCTEXT("CompilationFinished", "{0} finished compiling successfully with {1} {1}|plural(one=warning,other=warnings), {2} {2}|plural(one=error,other=errors) and {3} more similar warnings."), Prefix, NumWarnings, NumErrors, NumIgnoreds)
+		:
+		FText::Format(LOCTEXT("CompilationFinished", "{0} finished compiling successfully with {1} {1}|plural(one=warning,other=warnings) and {2} {2}|plural(one=error,other=errors)."), Prefix, NumWarnings, NumErrors);
 	
 	FCustomizableObjectEditorLogger::CreateLog(Message)
 	.Category(ELoggerCategory::Compilation)
@@ -1585,27 +1589,28 @@ void FCustomizableObjectCompiler::NotifyCompilationErrors() const
 }
 
 
-void FCustomizableObjectCompiler::CompilerLog(const FText& Message, const TArray<const UCustomizableObjectNode*>& ArrayNode, EMessageSeverity::Type MessageSeverity, bool bAddBaseObjectInfo)
+void FCustomizableObjectCompiler::CompilerLog(const FText& Message, const TArray<const UCustomizableObjectNode*>& ArrayNode, const EMessageSeverity::Type MessageSeverity, const bool bAddBaseObjectInfo, const ELoggerSpamBin SpamBin)
 {
-	if (CompilationLogsContainer.AddMessage(Message, ArrayNode, MessageSeverity)) // Cache the message for later reference
+	if (CompilationLogsContainer.AddMessage(Message, ArrayNode, MessageSeverity, SpamBin)) // Cache the message for later reference
 	{
 		FCustomizableObjectEditorLogger::CreateLog(Message)
 			.Severity(MessageSeverity)
 			.Nodes(ArrayNode)
 			.BaseObject(bAddBaseObjectInfo)
+			.SpamBin(SpamBin)
 			.Log();
 	}
 }
 
 
-void FCustomizableObjectCompiler::CompilerLog(const FText& Message, const UCustomizableObjectNode* Node, EMessageSeverity::Type MessageSeverity, bool AddBaseObjectInfo)
+void FCustomizableObjectCompiler::CompilerLog(const FText& Message, const UCustomizableObjectNode* Node, const EMessageSeverity::Type MessageSeverity, const bool bAddBaseObjectInfo, const ELoggerSpamBin SpamBin)
 {
 	TArray<const UCustomizableObjectNode*> ArrayNode;
 	if (Node)
 	{
 		ArrayNode.Add(Node);
 	}
-	CompilerLog(Message, ArrayNode, MessageSeverity, AddBaseObjectInfo);
+	CompilerLog(Message, ArrayNode, MessageSeverity, bAddBaseObjectInfo, SpamBin);
 }
 
 
@@ -1613,46 +1618,27 @@ void FCustomizableObjectCompiler::UpdateCompilerLogData()
 {
 	FMessageLogModule& MessageLogModule = FModuleManager::LoadModuleChecked<FMessageLogModule>("MessageLog");
 	MessageLogModule.RegisterLogListing(FName("Mutable"), LOCTEXT("MutableLog", "Mutable"));
-	const TArray<FCustomizableObjectCompileRunnable::FError>& ArrayCompileWarning = CompileTask->GetArrayWarning();
-	const TArray<FCustomizableObjectCompileRunnable::FError>& ArrayCompileError = CompileTask->GetArrayError();
+	const TArray<FCustomizableObjectCompileRunnable::FError>& ArrayCompileErrors = CompileTask->GetArrayErrors();
 
 	FText ObjectName = CurrentObject ? FText::FromString(CurrentObject->GetName()) : LOCTEXT("Unknown Object", "Unknown Object");
 
 	int32 i;
-	for (i = 0; i < ArrayCompileError.Num(); ++i)
+	for (i = 0; i < ArrayCompileErrors.Num(); ++i)
 	{
-		const UCustomizableObjectNode** pNode = GeneratedNodes.Find(ArrayCompileError[i].Context);
+		const UCustomizableObjectNode** pNode = GeneratedNodes.Find(ArrayCompileErrors[i].Context);
 
-		if (ArrayCompileError[i].AttachedData && pNode)
+		if (ArrayCompileErrors[i].AttachedData && pNode)
 		{
 			UCustomizableObjectNode::FAttachedErrorDataView ErrorDataView;
-			ErrorDataView.UnassignedUVs = { ArrayCompileError[i].AttachedData->UnassignedUVs.GetData(),
-											ArrayCompileError[i].AttachedData->UnassignedUVs.Num() };
+			ErrorDataView.UnassignedUVs = { ArrayCompileErrors[i].AttachedData->UnassignedUVs.GetData(),
+											ArrayCompileErrors[i].AttachedData->UnassignedUVs.Num() };
 
 			const UCustomizableObjectNode* Node = Cast<UCustomizableObjectNode>(*pNode);
 			const_cast<UCustomizableObjectNode*>(Node)->AddAttachedErrorData(ErrorDataView);
 		}
 
-		FText FullMsg = FText::Format(LOCTEXT("MutableMessage", "{0} : {1}"), ObjectName, ArrayCompileError[i].Message);
-		CompilerLog(FullMsg, pNode ? *pNode : nullptr, EMessageSeverity::Error, true);
-	}
-
-	for (i = 0; i < ArrayCompileWarning.Num(); ++i)
-	{
-		const UCustomizableObjectNode** pNode = GeneratedNodes.Find(ArrayCompileWarning[i].Context);
-
-		if (ArrayCompileWarning[i].AttachedData && pNode)
-		{
-			UCustomizableObjectNode::FAttachedErrorDataView ErrorDataView;
-			ErrorDataView.UnassignedUVs = { ArrayCompileWarning[i].AttachedData->UnassignedUVs.GetData(),
-											ArrayCompileWarning[i].AttachedData->UnassignedUVs.Num() };
-
-			const UCustomizableObjectNode* Node = Cast<UCustomizableObjectNode>(*pNode);
-			const_cast<UCustomizableObjectNode*>(Node)->AddAttachedErrorData(ErrorDataView);
-		}
-
-		FText FullMsg = FText::Format(LOCTEXT("MutableMessage", "{0} : {1}"), ObjectName, ArrayCompileWarning[i].Message);
-		CompilerLog(FullMsg, pNode?*pNode:nullptr, EMessageSeverity::Warning, true);
+		FText FullMsg = FText::Format(LOCTEXT("MutableMessage", "{0} : {1}"), ObjectName, ArrayCompileErrors[i].Message);
+		CompilerLog(FullMsg, pNode ? *pNode : nullptr, ArrayCompileErrors[i].Severity, true, ArrayCompileErrors[i].SpamBin);
 	}
 }
 
diff --git a/Engine/Plugins/Experimental/Mutable/Source/CustomizableObjectEditor/Private/MuCOE/CustomizableObjectCompiler.h b/Engine/Plugins/Experimental/Mutable/Source/CustomizableObjectEditor/Private/MuCOE/CustomizableObjectCompiler.h
--- a/Engine/Plugins/Experimental/Mutable/Source/CustomizableObjectEditor/Private/MuCOE/CustomizableObjectCompiler.h
+++ b/Engine/Plugins/Experimental/Mutable/Source/CustomizableObjectEditor/Private/MuCOE/CustomizableObjectCompiler.h
@@ -5,6 +5,7 @@
 #include "AssetRegistry/AssetData.h"
 #include "MuCOE/CompilationMessageCache.h"
 #include "MuCOE/GenerateMutableSource/GenerateMutableSource.h"
+#include "MuCOE/CustomizableObjectEditorLogger.h"
 #include "UObject/GCObject.h"
 
 class FCustomizableObjectCompileRunnable;
@@ -44,8 +45,8 @@ public:
 	/** Generate the Mutable Graph from the Unreal Graph. */
 	mu::NodePtr Export(UCustomizableObject* Object, const FCompilationOptions& Options);
 
-	void CompilerLog(const FText& Message, const TArray<const UCustomizableObjectNode*>& ArrayNode, EMessageSeverity::Type MessageSeverity = EMessageSeverity::Warning, bool bAddBaseObjectInfo = true);
-	void CompilerLog(const FText& Message, const UCustomizableObjectNode* Node = nullptr, EMessageSeverity::Type MessageSeverity = EMessageSeverity::Warning, bool bAddBaseObjectInfo = true);
+	void CompilerLog(const FText& Message, const TArray<const UCustomizableObjectNode*>& ArrayNode, const EMessageSeverity::Type MessageSeverity = EMessageSeverity::Warning, const bool bAddBaseObjectInfo = true, const ELoggerSpamBin SpamBin = ELoggerSpamBin::ShowAll);
+	void CompilerLog(const FText& Message, const UCustomizableObjectNode* Node = nullptr, const EMessageSeverity::Type MessageSeverity = EMessageSeverity::Warning, const bool bAddBaseObjectInfo = true, const ELoggerSpamBin SpamBin = ELoggerSpamBin::ShowAll);
 	void NotifyCompilationErrors() const;
 
 	void FinishCompilation();
diff --git a/Engine/Plugins/Experimental/Mutable/Source/CustomizableObjectEditor/Private/MuCOE/CustomizableObjectEditorLogger.cpp b/Engine/Plugins/Experimental/Mutable/Source/CustomizableObjectEditor/Private/MuCOE/CustomizableObjectEditorLogger.cpp
--- a/Engine/Plugins/Experimental/Mutable/Source/CustomizableObjectEditor/Private/MuCOE/CustomizableObjectEditorLogger.cpp
+++ b/Engine/Plugins/Experimental/Mutable/Source/CustomizableObjectEditor/Private/MuCOE/CustomizableObjectEditorLogger.cpp
@@ -203,6 +203,13 @@ FLogParameters& FLogParameters::FixNotification(const bool FixNotification)
 }
 
 
+FLogParameters& FLogParameters::SpamBin(const ELoggerSpamBin SpamBin)
+{
+	ParamSpamBin = SpamBin;
+	return *this;
+}
+
+
 void FLogParameters::Log()
 {
 	Logger.Log(*this);
diff --git a/Engine/Plugins/Experimental/Mutable/Source/CustomizableObjectEditor/Public/MuCOE/CompilationMessageCache.h b/Engine/Plugins/Experimental/Mutable/Source/CustomizableObjectEditor/Public/MuCOE/CompilationMessageCache.h
--- a/Engine/Plugins/Experimental/Mutable/Source/CustomizableObjectEditor/Public/MuCOE/CompilationMessageCache.h
+++ b/Engine/Plugins/Experimental/Mutable/Source/CustomizableObjectEditor/Public/MuCOE/CompilationMessageCache.h
@@ -2,6 +2,7 @@
 #pragma once
 
 #include "Logging/TokenizedMessage.h"
+#include "MuCOE/CustomizableObjectEditorLogger.h"
 
 // Forward declarations
 class UCustomizableObjectNode;
@@ -10,6 +11,7 @@ class UCustomizableObjectNode;
 class CUSTOMIZABLEOBJECTEDITOR_API FCompilationMessageCache
 {
 public:
+	FCompilationMessageCache();
 
 	/** Method designed to serve as a way to add a new message to the collection of messages the compiler has produced as
 	 * part of the compilation process.
@@ -18,7 +20,7 @@ public:
 	 * @param MessageSeverity Severity of the message.
 	 * @return True if the message is new and had not been added before. If it returns false it means it has not been logged again.
 	 */
-	 bool AddMessage(const FText& InMessage, const TArray<const UCustomizableObjectNode*>& InArrayNode, EMessageSeverity::Type MessageSeverity = EMessageSeverity::Warning);
+	 bool AddMessage(const FText& InMessage, const TArray<const UCustomizableObjectNode*>& InArrayNode, const EMessageSeverity::Type MessageSeverity = EMessageSeverity::Warning, const ELoggerSpamBin SpamBin = ELoggerSpamBin::ShowAll);
 	
 	
 	/** Provides the caller with separated arrays with all the warning and error messages.
@@ -48,6 +50,11 @@ public:
 	 * @return Amount of error messages cached.
 	 */
 	uint32 GetErrorCount() const;
+
+	/** Provides the caller with the amount of messages ignored because they were too similar even if not identical.
+	 * @return Amount of error messages ignored.
+	 */
+	uint32 GetIgnoredCount() const;
 	
 private:
 	
@@ -63,13 +70,16 @@ private:
 			return Message.EqualTo(o.Message) && Context == o.Context && Severity == o.Severity;
 		}
 	};
+	const uint32 MaxSpamMessages = 10;
 
 	/** Array with all the logged messages stored. */
 	TArray<FLoggedMessage> LoggedMessages;
+	TArray<FLoggedMessage> IgnoredMessages;
 
 	// Counters for the type of messages we have on LoggedMessages
-	uint32 PerformanceWarningCount = 0;
-	uint32 WarningCount = 0;
-	uint32 ErrorCount = 0;
-	
+	uint32 PerformanceWarningCount;
+	uint32 WarningCount;
+	uint32 ErrorCount;
+	uint32 IgnoredCount;
+	TArray<uint8> SpamBinCounts;
 };
diff --git a/Engine/Plugins/Experimental/Mutable/Source/CustomizableObjectEditor/Public/MuCOE/CustomizableObjectEditorLogger.h b/Engine/Plugins/Experimental/Mutable/Source/CustomizableObjectEditor/Public/MuCOE/CustomizableObjectEditorLogger.h
--- a/Engine/Plugins/Experimental/Mutable/Source/CustomizableObjectEditor/Public/MuCOE/CustomizableObjectEditorLogger.h
+++ b/Engine/Plugins/Experimental/Mutable/Source/CustomizableObjectEditor/Public/MuCOE/CustomizableObjectEditorLogger.h
@@ -19,6 +19,14 @@ enum class ELoggerCategory
 };
 
 
+/** Notification spam prevention bins. Each bin will only show a limited amount of messages of its type, after which it will simply state "and X more BINNAME warnings/errors" */
+enum class ELoggerSpamBin : uint8
+{
+	TagsNotFound = 0,
+	ShowAll	// Notifications added to this bin will always be shown.
+};
+
+
 /** Log parameters. Holds all parameters required to create a new CustomizableObjectEditorLogger log. */
 class FLogParameters
 {
@@ -44,6 +52,8 @@ public:
 	FLogParameters& Notification(bool Notification = true);
 
 	FLogParameters& FixNotification(bool FixNotification = true);
+
+	FLogParameters& SpamBin(ELoggerSpamBin SpamBin);
 	
 	/** Actually display the log. */
 	void Log();
@@ -62,6 +72,8 @@ private:
 	TAttribute<FText> ParamSubText;
 	
 	ELoggerCategory ParamCategory = ELoggerCategory::General;
+
+	ELoggerSpamBin ParamSpamBin = ELoggerSpamBin::ShowAll;
 	
 	EMessageSeverity::Type ParamSeverity = EMessageSeverity::Info;
 	
diff --git a/Engine/Plugins/Experimental/Mutable/Source/MutableTools/Private/MuT/CodeGenerator_Mesh.cpp b/Engine/Plugins/Experimental/Mutable/Source/MutableTools/Private/MuT/CodeGenerator_Mesh.cpp
--- a/Engine/Plugins/Experimental/Mutable/Source/MutableTools/Private/MuT/CodeGenerator_Mesh.cpp
+++ b/Engine/Plugins/Experimental/Mutable/Source/MutableTools/Private/MuT/CodeGenerator_Mesh.cpp
@@ -1204,8 +1204,9 @@ class Node;
 				Aux = tag.c_str();
                 m_pErrorLog->GetPrivate()->Add( 
 					FString::Printf(TEXT("Unknown tag found in mesh variation [%s]."), StringCast<TCHAR>(Aux).Get()),
-					ELMT_WARNING, 
-					node.m_errorContext
+					ELMT_WARNING,
+					node.m_errorContext,
+					ELMSB_UNKNOWN_TAG
 				);
                 continue;
             }
diff --git a/Engine/Plugins/Experimental/Mutable/Source/MutableTools/Private/MuT/ErrorLog.cpp b/Engine/Plugins/Experimental/Mutable/Source/MutableTools/Private/MuT/ErrorLog.cpp
--- a/Engine/Plugins/Experimental/Mutable/Source/MutableTools/Private/MuT/ErrorLog.cpp
+++ b/Engine/Plugins/Experimental/Mutable/Source/MutableTools/Private/MuT/ErrorLog.cpp
@@ -104,6 +104,20 @@ namespace mu
 		return result;
 	}
 
+
+	//---------------------------------------------------------------------------------------------
+	ErrorLogMessageSpamBin ErrorLog::GetMessageSpamBin(int index) const
+	{
+		ErrorLogMessageSpamBin result = ELMSB_ALL;
+
+		if (index >= 0 && index < (int)m_pD->m_messages.Num())
+		{
+			result = m_pD->m_messages[index].m_spam;
+		}
+
+		return result;
+	}
+
 	//---------------------------------------------------------------------------------------------
     ErrorLogMessageAttachedDataView ErrorLog::GetMessageAttachedData( int index ) const
 	{
@@ -124,32 +138,39 @@ namespace mu
 	}
 
 	//---------------------------------------------------------------------------------------------
-	void ErrorLog::Private::Add(const FString& InMessage, ErrorLogMessageType InType, const void* InContext)
+	void ErrorLog::Private::Add(const FString& InMessage,
+								ErrorLogMessageType InType,
+								const void* InContext,
+								ErrorLogMessageSpamBin InSpamBin)
 	{
 		m_messages.Add(FMessage());
 		m_messages.Last().m_type = InType;
+		m_messages.Last().m_spam = InSpamBin;
 		m_messages.Last().m_text = InMessage;
 		m_messages.Last().m_context = InContext;
 	}
 
 	//---------------------------------------------------------------------------------------------
 	void ErrorLog::Private::Add(const FString& InMessage,
-                                const ErrorLogMessageAttachedDataView& dataView,
-                                ErrorLogMessageType type, const void* context )
+                                const ErrorLogMessageAttachedDataView& InDataView,
+                                ErrorLogMessageType InType, const void* InContext,
+								ErrorLogMessageSpamBin InSpamBin)
 	{
 		m_messages.Add(FMessage() );
-		m_messages.Last().m_type = type;
+		m_messages.Last().m_type = InType;
+		m_messages.Last().m_spam = InSpamBin;
 		m_messages.Last().m_text = InMessage;
-		m_messages.Last().m_context = context;
+		m_messages.Last().m_context = InContext;
 		m_messages.Last().m_data = MakeShared<FErrorData>();
 
-        if ( dataView.m_unassignedUVs && dataView.m_unassignedUVsSize > 0 )
+        if ( InDataView.m_unassignedUVs && InDataView.m_unassignedUVsSize > 0 )
         {
 			// \TODO: Review
-			m_messages.Last().m_data->m_unassignedUVs.Append(dataView.m_unassignedUVs, dataView.m_unassignedUVsSize);
+			m_messages.Last().m_data->m_unassignedUVs.Append(InDataView.m_unassignedUVs, InDataView.m_unassignedUVsSize);
         }
 	}
 	
+	
 	//---------------------------------------------------------------------------------------------
 	void ErrorLog::Log() const
 	{
diff --git a/Engine/Plugins/Experimental/Mutable/Source/MutableTools/Private/MuT/ErrorLogPrivate.h b/Engine/Plugins/Experimental/Mutable/Source/MutableTools/Private/MuT/ErrorLogPrivate.h
--- a/Engine/Plugins/Experimental/Mutable/Source/MutableTools/Private/MuT/ErrorLogPrivate.h
+++ b/Engine/Plugins/Experimental/Mutable/Source/MutableTools/Private/MuT/ErrorLogPrivate.h
@@ -30,6 +30,7 @@ namespace mu
 		struct FMessage
 		{
 			ErrorLogMessageType m_type = ELMT_NONE;
+			ErrorLogMessageSpamBin m_spam = ELMSB_ALL;
 			FString m_text;
             TSharedPtr<FErrorData> m_data;
 			const void* m_context = nullptr;
@@ -41,10 +42,10 @@ namespace mu
 		//-----------------------------------------------------------------------------------------
 
 		//!
-		void Add(const FString& Message, ErrorLogMessageType Type, const void* Context);
+		void Add(const FString& Message, ErrorLogMessageType Type, const void* Context, ErrorLogMessageSpamBin SpamBin = ELMSB_ALL);
 
         //!
-        void Add(const FString& InMessage, const ErrorLogMessageAttachedDataView& data, ErrorLogMessageType type, const void* context );
+        void Add(const FString& Message, const ErrorLogMessageAttachedDataView& Data, ErrorLogMessageType Type, const void* Context, ErrorLogMessageSpamBin SpamBin = ELMSB_ALL);
 	};
 
 
diff --git a/Engine/Plugins/Experimental/Mutable/Source/MutableTools/Public/MuT/ErrorLog.h b/Engine/Plugins/Experimental/Mutable/Source/MutableTools/Public/MuT/ErrorLog.h
--- a/Engine/Plugins/Experimental/Mutable/Source/MutableTools/Public/MuT/ErrorLog.h
+++ b/Engine/Plugins/Experimental/Mutable/Source/MutableTools/Public/MuT/ErrorLog.h
@@ -26,6 +26,14 @@ namespace mu
 		ELMT_INFO
 	} ErrorLogMessageType;
 
+	//! Categories of message stored in the log for the purpose of limiiting duplication of non-identical messages
+	//! \ingroup tools
+	typedef enum
+	{
+		ELMSB_ALL = 0,
+		ELMSB_UNKNOWN_TAG
+	} ErrorLogMessageSpamBin;
+
 	struct ErrorLogMessageAttachedDataView 
 	{
 		const float* m_unassignedUVs = nullptr;
@@ -63,6 +71,9 @@ namespace mu
 		//!
 		ErrorLogMessageType GetMessageType( int index ) const;
 
+		//!
+		ErrorLogMessageSpamBin GetMessageSpamBin(int index) const;
+
 		//! Get the attached data of a message.
 		//! \param index index of message data from 0 to GetMessageCount(ELMT_NONE)-1
 		ErrorLogMessageAttachedDataView GetMessageAttachedData( int index ) const;
