[Mutable] Generate the CO Identifier using the path and name of the package to avoid issues with CO duplicated outside the editor overriding the same compiled data.

#rnx

diff --git a/Engine/Plugins/Experimental/Mutable/Source/CustomizableObject/Classes/MuCO/CustomizableObject.h b/Engine/Plugins/Experimental/Mutable/Source/CustomizableObject/Classes/MuCO/CustomizableObject.h
--- a/Engine/Plugins/Experimental/Mutable/Source/CustomizableObject/Classes/MuCO/CustomizableObject.h
+++ b/Engine/Plugins/Experimental/Mutable/Source/CustomizableObject/Classes/MuCO/CustomizableObject.h
@@ -1439,7 +1439,7 @@ private:
 	// This is a manual version number for the binary blobs in this asset.
 	// Increasing it invalidates all the previously compiled models.
 	// Warning: If while merging code both versions have changed, take the highest+1.
-	static const int32 CurrentSupportedVersion = 392;
+	static const int32 CurrentSupportedVersion = 393;
 
 public:
 
@@ -1464,7 +1464,7 @@ public:
 	FString GetCompiledDataFolderPath(bool bIsEditorData) const;
 
 	// Compose file name 
-	FString GetCompiledDataFileName(bool bIsModel, const ITargetPlatform* InTargetPlatform = nullptr) const;
+	FString GetCompiledDataFileName(bool bIsModel, const ITargetPlatform* InTargetPlatform = nullptr, bool bIsDiskStreamer = false);
 
 	/** Used to set the flag IsRoot */
 	void GetAssetRegistryTags(TArray<FAssetRegistryTag>& OutTags) const override;
@@ -1474,9 +1474,6 @@ public:
 
 	bool IsEditorOnly() const override;
 
-	// Initialize if not set.
-	void InitializeIdentifier();
-
 	// Begin UObject interface.
 	void PreSave(FObjectPreSaveContext ObjectSaveContext) override;
 	void PostRename(UObject* OldOuter, const FName OldName) override;
@@ -1624,6 +1621,10 @@ public:
 	UPROPERTY()
 	bool bIsChildObject = false;
 
+	/** Unique Identifier - used to locate Model and Streamable data on disk. Should not be modified. */
+	UPROPERTY(Transient)
+	FGuid Identifier;
+
 	FPostCompileDelegate PostCompileDelegate;
 
 	void PostCompile();
@@ -1662,10 +1663,6 @@ public:
 	UPROPERTY(Transient)
 	TObjectPtr<UMutableMaskOutCache> MaskOutCache_HardRef;
 
-	/** Unique Identifier - used to locate Model and Streamable data on disk. Should not be modified. */
-	UPROPERTY()
-	FGuid Identifier;
-
 	/** Unique identifier. Regenerated each time the object is compiled. */
 	UPROPERTY()
 	FGuid CompilationGuid;
diff --git a/Engine/Plugins/Experimental/Mutable/Source/CustomizableObject/Private/MuCO/CustomizableObject.cpp b/Engine/Plugins/Experimental/Mutable/Source/CustomizableObject/Private/MuCO/CustomizableObject.cpp
--- a/Engine/Plugins/Experimental/Mutable/Source/CustomizableObject/Private/MuCO/CustomizableObject.cpp
+++ b/Engine/Plugins/Experimental/Mutable/Source/CustomizableObject/Private/MuCO/CustomizableObject.cpp
@@ -110,11 +110,6 @@ void UCustomizableObject::PreSave(FObjectPreSaveContext ObjectSaveContext)
 		}
 	}
 
-	if (!Identifier.IsValid())
-	{
-		Identifier = FGuid::NewGuid();
-	}
-
 #if WITH_EDITORONLY_DATA
 	if (ObjectSaveContext.IsCooking() && !bIsChildObject)
 	{
@@ -243,9 +238,6 @@ void UCustomizableObject::PostDuplicate(EDuplicateMode::Type DuplicateMode)
 
 	if (DuplicateMode == EDuplicateMode::Normal)
 	{
-		// The only place where we can change the identifier
-		Identifier = FGuid::NewGuid();
-
 		// Create a new Private Data or it will use the same as the original Customizable Object
 		PrivateData = TSharedPtr<FCustomizableObjectPrivateData>(new FCustomizableObjectPrivateData());
 	}
@@ -286,7 +278,7 @@ void UCustomizableObject::BeginCacheForCookedPlatformData(const ITargetPlatform*
 	else
 	{
 		ClearCompiledData();
-		PrivateData->SetModel(nullptr); // Discard compilation
+		PrivateData->SetModel(nullptr, FGuid()); // Discard compilation
 		if (TargetPlatform)
 		{
 			PrivateData->CachedPlatformNames.Add(TargetPlatform->PlatformName());
@@ -459,7 +451,7 @@ void UCustomizableObject::SaveCompiledData(FArchive& MemoryWriter, bool bSkipEdi
 
 void UCustomizableObject::LoadCompiledData(FArchive& MemoryReader, bool bSkipEditorOnlyData)
 {
-	PrivateData->SetModel(nullptr);
+	PrivateData->SetModel(nullptr, FGuid());
 	ClearCompiledData();
 
 	MutableCompiledDataStreamHeader Header;
@@ -564,7 +556,7 @@ void UCustomizableObject::LoadCompiledData(FArchive& MemoryReader, bool bSkipEdi
 			mu::InputArchive arch(&stream);
 			TSharedPtr<mu::Model, ESPMode::ThreadSafe> Model = mu::Model::StaticUnserialise( arch );
 
-			PrivateData->SetModel(Model);
+			PrivateData->SetModel(Model, Identifier);
 		}
 	}
 
@@ -575,15 +567,7 @@ void UCustomizableObject::LoadCompiledDataFromDisk(bool bIsEditorData, const ITa
 {
 	FString PlatformName = InTargetPlatform ? InTargetPlatform->PlatformName() : FPlatformProperties::PlatformName();
 
-	if (bIsEditorData)
-	{
-		// Customizable Object outdated
-		if (!Identifier.IsValid())
-		{
-			return;
-		}
-	}
-	else // Loading cooked data
+	if (!bIsEditorData) // Loading cooked data
 	{
 		// If we don't use OnCookStart there's nothing to be loaded. Same case for child objects.
 		if (!bUsesOnCookStart || bIsChildObject) return;
@@ -701,7 +685,7 @@ void UCustomizableObject::CompileForTargetPlatform(const ITargetPlatform* Target
 		(TargetPlatform && Relevancy == ECustomizableObjectRelevancy::ClientOnly && !TargetPlatform->IsServerOnly());
 
 	// Discard any older compilation
-	PrivateData->SetModel(nullptr);
+	PrivateData->SetModel(nullptr, FGuid());
 
 	if (bIsRootObject && bIsRelevantForThisTarget)
 	{
@@ -751,7 +735,7 @@ bool UCustomizableObject::ConditionalAutoCompile()
 	}
 
 	// Discard any older compilation
-	PrivateData->SetModel(nullptr);
+	PrivateData->SetModel(nullptr, FGuid());
 
 	// Sync/Async compilation
 	if (System->IsAutoCompilationSync())
@@ -807,16 +791,6 @@ FReply UCustomizableObject::AddNewParameterProfile(FString Name, UCustomizableOb
 }
 
 
-void UCustomizableObject::InitializeIdentifier()
-{
-	if (!Identifier.IsValid())
-	{
-		Identifier = FGuid::NewGuid();
-		MarkPackageDirty();
-	}
-}
-
-
 FString UCustomizableObject::GetCompiledDataFolderPath(bool bIsEditorData) const
 {	
 	const FString FolderName = bIsEditorData ? TEXT("MutableStreamedDataEditor/") : TEXT("MutableStreamedData/");
@@ -824,10 +798,19 @@ FString UCustomizableObject::GetCompiledDataFolderPath(bool bIsEditorData) const
 }
 
 
-FString UCustomizableObject::GetCompiledDataFileName(bool bIsModel, const ITargetPlatform* InTargetPlatform) const
+FString UCustomizableObject::GetCompiledDataFileName(bool bIsModel, const ITargetPlatform* InTargetPlatform, bool bIsDiskStreamer)
 {
+	// Generate the Identifier using the path and name of the asset
+	if (!bIsDiskStreamer)
+	{
+		uint32 FullPathHash = GetTypeHash(GetFullName());
+		uint32 OutermostHash = GetTypeHash(GetNameSafe(GetOutermost()));
+		uint32 OuterHash = GetTypeHash(GetName());
+		Identifier = FGuid(0, FullPathHash, OutermostHash, OuterHash);
+	}
+
 	const FString PlatformName = InTargetPlatform ? InTargetPlatform->PlatformName() : FPlatformProperties::PlatformName();
-	const FString FileIdentifier = Identifier.IsValid() ? Identifier.ToString() : VersionId.ToString();
+	const FString FileIdentifier = bIsDiskStreamer ? GetPrivate()->Identifier.ToString() : Identifier.ToString();
 	const FString Extension = bIsModel ? TEXT("_M.mut") : TEXT("_S.mut");
 	return PlatformName + FileIdentifier + Extension;
 }
@@ -906,7 +889,7 @@ void UCustomizableObject::LoadEmbeddedData(FArchive& Ar)
 		UnrealMutableInputStream stream(Ar);
 		mu::InputArchive arch(&stream);
 		TSharedPtr<mu::Model, ESPMode::ThreadSafe> Model = mu::Model::StaticUnserialise( arch );
-		PrivateData->SetModel( Model );
+		PrivateData->SetModel( Model, FGuid());
 
 		// Create parameter properties
 		UpdateParameterPropertiesFromModel();
@@ -1210,7 +1193,7 @@ TSharedPtr<mu::Model, ESPMode::ThreadSafe> UCustomizableObject::GetModel() const
 #if WITH_EDITOR
 void UCustomizableObject::SetModel(TSharedPtr<mu::Model, ESPMode::ThreadSafe> Model)
 {
-	PrivateData->SetModel(Model);
+	PrivateData->SetModel(Model, Identifier);
 	
 	UpdateCompiledDataFromModel();
 }
@@ -1780,9 +1763,13 @@ FGuid UCustomizableObject::GetCompilationGuid() const
 //-------------------------------------------------------------------------------------------------
 //-------------------------------------------------------------------------------------------------
 
-void FCustomizableObjectPrivateData::SetModel(const TSharedPtr<mu::Model, ESPMode::ThreadSafe>& Model)
+void FCustomizableObjectPrivateData::SetModel(const TSharedPtr<mu::Model, ESPMode::ThreadSafe>& Model, const FGuid Id)
 {
 	MutableModel = Model;
+
+#if WITH_EDITOR
+	Identifier = Id;
+#endif
 }
 
 
diff --git a/Engine/Plugins/Experimental/Mutable/Source/CustomizableObject/Private/MuCO/CustomizableObjectPrivate.h b/Engine/Plugins/Experimental/Mutable/Source/CustomizableObject/Private/MuCO/CustomizableObjectPrivate.h
--- a/Engine/Plugins/Experimental/Mutable/Source/CustomizableObject/Private/MuCO/CustomizableObjectPrivate.h
+++ b/Engine/Plugins/Experimental/Mutable/Source/CustomizableObject/Private/MuCO/CustomizableObjectPrivate.h
@@ -16,7 +16,7 @@ private:
 
 public:
 
-	void SetModel(const TSharedPtr<mu::Model, ESPMode::ThreadSafe>& Model);
+	void SetModel(const TSharedPtr<mu::Model, ESPMode::ThreadSafe>& Model, const FGuid Identifier);
 	const TSharedPtr<mu::Model, ESPMode::ThreadSafe>& GetModel();
 	TSharedPtr<const mu::Model, ESPMode::ThreadSafe> GetModel() const;
 
@@ -24,8 +24,9 @@ public:
 	bool bLocked = false;
 
 #if WITH_EDITOR
-	bool bModelCompiledForCook = false;
+	FGuid Identifier;
 
+	bool bModelCompiledForCook = false;
 	TArray<FString> CachedPlatformNames;
 #endif
 
diff --git a/Engine/Plugins/Experimental/Mutable/Source/CustomizableObject/Private/MuCO/UnrealMutableModelDiskStreamer.cpp b/Engine/Plugins/Experimental/Mutable/Source/CustomizableObject/Private/MuCO/UnrealMutableModelDiskStreamer.cpp
--- a/Engine/Plugins/Experimental/Mutable/Source/CustomizableObject/Private/MuCO/UnrealMutableModelDiskStreamer.cpp
+++ b/Engine/Plugins/Experimental/Mutable/Source/CustomizableObject/Private/MuCO/UnrealMutableModelDiskStreamer.cpp
@@ -143,7 +143,7 @@ bool FUnrealMutableModelBulkStreamer::PrepareStreamingForObject(UCustomizableObj
 
 #if WITH_EDITOR
 		FString FolderPath = CustomizableObject->GetCompiledDataFolderPath(true);
-		FString FullFileName = FolderPath + CustomizableObject->GetCompiledDataFileName(false);
+		FString FullFileName = FolderPath + CustomizableObject->GetCompiledDataFileName(false, nullptr, true);
 
 		IAsyncReadFileHandle* ReadFileHandle = FPlatformFileManager::Get().GetPlatformFile().OpenAsyncRead(*FullFileName);
 		if (!ReadFileHandle)
diff --git a/Engine/Plugins/Experimental/Mutable/Source/CustomizableObjectEditor/Private/MuCOE/CustomizableObjectCompiler.cpp b/Engine/Plugins/Experimental/Mutable/Source/CustomizableObjectEditor/Private/MuCOE/CustomizableObjectCompiler.cpp
--- a/Engine/Plugins/Experimental/Mutable/Source/CustomizableObjectEditor/Private/MuCOE/CustomizableObjectCompiler.cpp
+++ b/Engine/Plugins/Experimental/Mutable/Source/CustomizableObjectEditor/Private/MuCOE/CustomizableObjectCompiler.cpp
@@ -1266,12 +1266,6 @@ void FCustomizableObjectCompiler::CompileInternal(UCustomizableObject* Object, c
 			}
 		}
 
-		// Initializes a unique identifier for this object
-		if (!Options.bIsCooking)
-		{
-			Object->InitializeIdentifier();
-		}
-
 		if (CompileTask.IsValid()) // Don't start compilation if there's a compilation running
 		{
 			// TODO : warning?
