[mutable] Move resource key cache to the System instead of keeping in in every Model.
Add a CVar to control the size of the generated resource cache.
Modify the generated resource cache IDs to include the root address and save time and boilerplate on resource generation.
Cleanup streamers by splitting reader and writer.
#rnx

diff --git a/Engine/Plugins/Experimental/Mutable/Source/CustomizableObject/Private/MuCO/CustomizableObjectMipDataProvider.cpp b/Engine/Plugins/Experimental/Mutable/Source/CustomizableObject/Private/MuCO/CustomizableObjectMipDataProvider.cpp
--- a/Engine/Plugins/Experimental/Mutable/Source/CustomizableObject/Private/MuCO/CustomizableObjectMipDataProvider.cpp
+++ b/Engine/Plugins/Experimental/Mutable/Source/CustomizableObject/Private/MuCO/CustomizableObjectMipDataProvider.cpp
@@ -151,7 +151,7 @@ namespace impl
 
 				// This ID may be different than the ID obtained the first time the image was generated, because the mutable
 				// runtime cannot remember all the resources it has built, and only remembers a fixed amount.
-				mu::RESOURCE_ID MipImageID = Instance->GetImageId(ImageRef.LOD, ImageRef.Component, SurfaceIndex, ImageRef.Image);
+				mu::FResourceID MipImageID = Instance->GetImageId(ImageRef.LOD, ImageRef.Component, SurfaceIndex, ImageRef.Image);
 
 				UCustomizableObjectSystem* COSystem = UCustomizableObjectSystem::GetInstance();
 				FCustomizableObjectSystemPrivate* SystemPrivate = COSystem ? COSystem->GetPrivate() : nullptr;
diff --git a/Engine/Plugins/Experimental/Mutable/Source/CustomizableObject/Private/MuCO/CustomizableObjectSystem.cpp b/Engine/Plugins/Experimental/Mutable/Source/CustomizableObject/Private/MuCO/CustomizableObjectSystem.cpp
--- a/Engine/Plugins/Experimental/Mutable/Source/CustomizableObject/Private/MuCO/CustomizableObjectSystem.cpp
+++ b/Engine/Plugins/Experimental/Mutable/Source/CustomizableObject/Private/MuCO/CustomizableObjectSystem.cpp
@@ -101,6 +101,12 @@ TAutoConsoleVariable<bool> CVarEnableImageCache(
 	TEXT("Enables or disables the Instance Image Cache. Avoids regenerating images when no parameters or mips have changed."),
 	ECVF_Scalability);
 
+static TAutoConsoleVariable<int32> CVarGeneratedResourcesCacheSize(
+	TEXT("mutable.GeneratedResourcesCacheSize"),
+	512,
+	TEXT("Limit the number of resources (images and meshes) that will be tracked for reusal. Each tracked resource uses a small amout of memory for its key."),
+	ECVF_Scalability);
+
 
 bool FMutablePendingInstanceWork::ArePendingUpdatesEmpty() const
 {
@@ -345,6 +351,7 @@ void UCustomizableObjectSystem::InitSystem()
 #endif
 
 	Private->LastWorkingMemoryBytes = CVarWorkingMemory.GetValueOnGameThread() * 1024;
+	Private->LastGeneratedResourceCacheSize = CVarGeneratedResourcesCacheSize.GetValueOnGameThread();
 
 	mu::Ptr<mu::Settings> pSettings = new mu::Settings;
 	check(pSettings);
@@ -354,7 +361,7 @@ void UCustomizableObjectSystem::InitSystem()
 	Private->MutableSystem = new mu::System(pSettings, Private->ExtensionDataStreamer);
 	check(Private->MutableSystem);
 
-	Private->Streamer = MakeShared<FUnrealMutableModelBulkStreamer>();
+	Private->Streamer = MakeShared<FUnrealMutableModelBulkReader>();
 	check(Private->Streamer != nullptr);
 	Private->MutableSystem->SetStreamingInterface(Private->Streamer);
 
@@ -1084,7 +1091,7 @@ bool UCustomizableObjectSystem::CheckIfDiskOrMipUpdateOperationsPending(const UC
 
 	// Ensure that we don't try to handle any further streaming operations for this object
 	check(GetPrivate());
-	if (const FUnrealMutableModelBulkStreamer* Streamer = GetPrivate()->Streamer.Get())
+	if (const FUnrealMutableModelBulkReader* Streamer = GetPrivate()->Streamer.Get())
 	{
 		if (Streamer->AreTherePendingStreamingOperationsForObject(&Object))
 		{
@@ -1186,6 +1193,14 @@ void FCustomizableObjectSystemPrivate::UpdateMemoryLimit()
 		check(MutableSystem);
 		MutableSystem->SetWorkingMemoryBytes(MemoryBytes);
 	}
+
+	const uint32 GeneratedResourceCacheSize = CVarGeneratedResourcesCacheSize.GetValueOnGameThread();
+	if (GeneratedResourceCacheSize != LastGeneratedResourceCacheSize)
+	{
+		LastGeneratedResourceCacheSize = GeneratedResourceCacheSize;
+		check(MutableSystem);
+		MutableSystem->SetGeneratedCacheSize(GeneratedResourceCacheSize);
+	}
 }
 
 
@@ -1456,7 +1471,7 @@ namespace impl
 							OperationData->InstanceUpdateData.Vectors.Push({});
 							FInstanceUpdateData::FVector& Vector = OperationData->InstanceUpdateData.Vectors.Last();
 							Vector.Name = Instance->GetVectorName(MutableLODIndex, ComponentIndex, InstanceSurfaceIndex, VectorIndex);
-							Instance->GetVector(MutableLODIndex, ComponentIndex, InstanceSurfaceIndex, VectorIndex, &Vector.Vector.R, &Vector.Vector.G, &Vector.Vector.B, &Vector.Vector.A);
+							Vector.Vector = Instance->GetVector(MutableLODIndex, ComponentIndex, InstanceSurfaceIndex, VectorIndex);
 						}
 
 						// Scalars
@@ -1508,7 +1523,7 @@ namespace impl
 		check(System != nullptr);
 
 		// Generate all the required resources, that are not cached
-		TArray<mu::RESOURCE_ID> ImagesInThisInstance;
+		TArray<mu::FResourceID> ImagesInThisInstance;
 		for (FInstanceUpdateData::FImage& Image : OperationData->InstanceUpdateData.Images)
 		{
 			MUTABLE_CPUPROFILER_SCOPE(GetImage);
diff --git a/Engine/Plugins/Experimental/Mutable/Source/CustomizableObject/Private/MuCO/CustomizableObjectSystemPrivate.h b/Engine/Plugins/Experimental/Mutable/Source/CustomizableObject/Private/MuCO/CustomizableObjectSystemPrivate.h
--- a/Engine/Plugins/Experimental/Mutable/Source/CustomizableObject/Private/MuCO/CustomizableObjectSystemPrivate.h
+++ b/Engine/Plugins/Experimental/Mutable/Source/CustomizableObject/Private/MuCO/CustomizableObjectSystemPrivate.h
@@ -269,10 +269,10 @@ public:
 
 struct FMutableImageCacheKey
 {
-	mu::RESOURCE_ID Resource = 0;
+	mu::FResourceID Resource = 0;
 	int32 SkippedMips = 0;
 
-	FMutableImageCacheKey(mu::RESOURCE_ID InResource, int32 InSkippedMips)
+	FMutableImageCacheKey(mu::FResourceID InResource, int32 InSkippedMips)
 		: Resource(InResource), SkippedMips(InSkippedMips) {}
 
 	inline bool operator==(const FMutableImageCacheKey& Other) const
@@ -292,7 +292,7 @@ inline uint32 GetTypeHash(const FMutableImageCacheKey& Key)
 struct FMutableResourceCache
 {
 	TWeakObjectPtr<const UCustomizableObject> Object;
-	TMap<mu::RESOURCE_ID, TWeakObjectPtr<USkeletalMesh> > Meshes;
+	TMap<mu::FResourceID, TWeakObjectPtr<USkeletalMesh> > Meshes;
 	TMap<FMutableImageCacheKey, TWeakObjectPtr<UTexture2D> > Images;
 
 	void Clear()
@@ -368,7 +368,7 @@ struct FInstanceGeneratedData
 		/** True if it can be reused */
 		bool bGenerated = false;
 
-		mu::RESOURCE_ID MeshID;
+		mu::FResourceID MeshID;
 
 		/** Range in the Surfaces array */
 		uint16 FirstSurface = 0;
@@ -403,7 +403,7 @@ struct FInstanceUpdateData
 	struct FImage
 	{
 		FString Name;
-		mu::RESOURCE_ID ImageID;
+		mu::FResourceID ImageID;
 		
 		// LOD of the ImageId. If the texture is shared between LOD, first LOD where this image can be found. 
 		int32 BaseLOD;
@@ -455,7 +455,7 @@ struct FInstanceUpdateData
 		// Reuse component from a previously generated SkeletalMesh
 		bool bReuseMesh = false;
 
-		mu::RESOURCE_ID MeshID;
+		mu::FResourceID MeshID;
 		mu::MeshPtrConst Mesh;
 
 		/** Range in the Surfaces array */
@@ -612,9 +612,10 @@ public:
 
 	/** Store the last streaming memory size in bytes, to change it when it is safe. */
 	uint64 LastWorkingMemoryBytes = 0;
+	uint32 LastGeneratedResourceCacheSize = 0;
 
 	// This object is responsible for streaming data to the MutableSystem.
-	TSharedPtr<class FUnrealMutableModelBulkStreamer> Streamer;
+	TSharedPtr<class FUnrealMutableModelBulkReader> Streamer;
 
 	// 
 	TSharedPtr<class FUnrealExtensionDataStreamer> ExtensionDataStreamer;
@@ -630,7 +631,7 @@ public:
 	// List of textures currently cached and valid for the current object that we are operating on.
 	// This array gets generated when the object cached resources are protected in SetResourceCacheProtected
 	// from the game thread, and it is read from the Mutable thread only while updating the instance.
-	TArray<mu::RESOURCE_ID> ProtectedObjectCachedImages;
+	TArray<mu::FResourceID> ProtectedObjectCachedImages;
 
 	// The pending instance updates, discards or releases
 	FMutablePendingInstanceWork MutablePendingInstanceWork;
diff --git a/Engine/Plugins/Experimental/Mutable/Source/CustomizableObject/Private/MuCO/UnrealMutableImageProvider.cpp b/Engine/Plugins/Experimental/Mutable/Source/CustomizableObject/Private/MuCO/UnrealMutableImageProvider.cpp
--- a/Engine/Plugins/Experimental/Mutable/Source/CustomizableObject/Private/MuCO/UnrealMutableImageProvider.cpp
+++ b/Engine/Plugins/Experimental/Mutable/Source/CustomizableObject/Private/MuCO/UnrealMutableImageProvider.cpp
@@ -117,9 +117,9 @@ mu::EImageFormat GetMutablePixelFormat(EPixelFormat InTextureFormat)
 
 //-------------------------------------------------------------------------------------------------
 #ifdef MUTABLE_USE_NEW_TASKGRAPH
-	TTuple<UE::Tasks::FTask, TFunction<void()>> FUnrealMutableImageProvider::GetImageAsync(mu::EXTERNAL_IMAGE_ID id, uint8 MipmapsToSkip, TFunction<void(mu::Ptr<mu::Image>)>& ResultCallback)
+	TTuple<UE::Tasks::FTask, TFunction<void()>> FUnrealMutableImageProvider::GetImageAsync(mu::FExternalImageID id, uint8 MipmapsToSkip, TFunction<void(mu::Ptr<mu::Image>)>& ResultCallback)
 #else
-	TTuple<FGraphEventRef, TFunction<void()>> FUnrealMutableImageProvider::GetImageAsync(mu::EXTERNAL_IMAGE_ID id, uint8 MipmapsToSkip, TFunction<void(mu::Ptr<mu::Image>)>& ResultCallback)
+	TTuple<FGraphEventRef, TFunction<void()>> FUnrealMutableImageProvider::GetImageAsync(mu::FExternalImageID id, uint8 MipmapsToSkip, TFunction<void(mu::Ptr<mu::Image>)>& ResultCallback)
 #endif
 {
 	// Thread: worker
@@ -393,7 +393,7 @@ mu::EImageFormat GetMutablePixelFormat(EPixelFormat InTextureFormat)
 
 
 // This should mantain parity with the descriptor of the images generated by GetImageAsync 
-mu::FImageDesc FUnrealMutableImageProvider::GetImageDesc(mu::EXTERNAL_IMAGE_ID id, uint8 MipmapsToSkip) 
+mu::FImageDesc FUnrealMutableImageProvider::GetImageDesc(mu::FExternalImageID id, uint8 MipmapsToSkip) 
 {
 	MUTABLE_CPUPROFILER_SCOPE(FUnrealMutableImageProvider::GetImageDesc);
 
@@ -463,7 +463,7 @@ mu::FImageDesc FUnrealMutableImageProvider::GetImageDesc(mu::EXTERNAL_IMAGE_ID i
 }
 
 
-void FUnrealMutableImageProvider::CacheImage(const mu::EXTERNAL_IMAGE_ID& Id, bool bUser)
+void FUnrealMutableImageProvider::CacheImage(const mu::FExternalImageID& Id, bool bUser)
 {
 	check( IsInGameThread() );
 
@@ -551,7 +551,7 @@ void FUnrealMutableImageProvider::CacheImage(const mu::EXTERNAL_IMAGE_ID& Id, bo
 }
 
 
-void FUnrealMutableImageProvider::UnCacheImage(const mu::EXTERNAL_IMAGE_ID& Id, bool bUser)
+void FUnrealMutableImageProvider::UnCacheImage(const mu::FExternalImageID& Id, bool bUser)
 {
 	check(IsInGameThread());
 
@@ -687,7 +687,7 @@ mu::FImageDesc FUnrealMutableImageProvider::CreateDummyDesc()
 
 void FUnrealMutableImageProvider::AddReferencedObjects(FReferenceCollector& Collector)
 {
-	for (TPair<mu::EXTERNAL_IMAGE_ID, FUnrealMutableImageInfo>& Image : GlobalExternalImages)
+	for (TPair<mu::FExternalImageID, FUnrealMutableImageInfo>& Image : GlobalExternalImages)
 	{
 		if (Image.Value.TextureToLoad)
 		{
diff --git a/Engine/Plugins/Experimental/Mutable/Source/CustomizableObject/Private/MuCO/UnrealMutableImageProvider.h b/Engine/Plugins/Experimental/Mutable/Source/CustomizableObject/Private/MuCO/UnrealMutableImageProvider.h
--- a/Engine/Plugins/Experimental/Mutable/Source/CustomizableObject/Private/MuCO/UnrealMutableImageProvider.h
+++ b/Engine/Plugins/Experimental/Mutable/Source/CustomizableObject/Private/MuCO/UnrealMutableImageProvider.h
@@ -23,23 +23,23 @@ public:
 	// mu::ImageParameterGenerator interface
 	// Thread: worker
 #ifdef MUTABLE_USE_NEW_TASKGRAPH
-	virtual TTuple<UE::Tasks::FTask, TFunction<void()>> GetImageAsync(mu::EXTERNAL_IMAGE_ID Id, uint8 MipmapsToSkip, TFunction<void(mu::Ptr<mu::Image>)>& ResultCallback) override;
+	virtual TTuple<UE::Tasks::FTask, TFunction<void()>> GetImageAsync(mu::FExternalImageID Id, uint8 MipmapsToSkip, TFunction<void(mu::Ptr<mu::Image>)>& ResultCallback) override;
 #else
-	virtual TTuple<FGraphEventRef, TFunction<void()>> GetImageAsync(mu::EXTERNAL_IMAGE_ID Id, uint8 MipmapsToSkip, TFunction<void(mu::Ptr<mu::Image>)>& ResultCallback) override;
+	virtual TTuple<FGraphEventRef, TFunction<void()>> GetImageAsync(mu::FExternalImageID Id, uint8 MipmapsToSkip, TFunction<void(mu::Ptr<mu::Image>)>& ResultCallback) override;
 #endif
 
-	virtual mu::FImageDesc GetImageDesc(mu::EXTERNAL_IMAGE_ID Id, uint8 MipmapsToSkip) override;
+	virtual mu::FImageDesc GetImageDesc(mu::FExternalImageID Id, uint8 MipmapsToSkip) override;
 
 	
 	// Own interface	
 	// Thread: Game
 	/** Add a reference to the image. If it was not cached it caches it.
 	 * @param bUser if true, adds a reference to the user reference counter. If false, adds a reference to the system reference counter. */
-	void CacheImage(const mu::EXTERNAL_IMAGE_ID& Id, bool bUser);
+	void CacheImage(const mu::FExternalImageID& Id, bool bUser);
 
 	/** Removes a reference to the image. If all references are removed, it uncaches the image.
 	 * @param bUser if true, removes a reference from the user reference counter. If false, removes a reference from the system reference counter. */
-	void UnCacheImage(const mu::EXTERNAL_IMAGE_ID& Id, bool bUser);
+	void UnCacheImage(const mu::FExternalImageID& Id, bool bUser);
 
 	/** Removes a reference to all images. All images which no longer have references will be uncached.
 	 * @param bUser if true, removes a references from the user reference counter. If false, removes a reference from the system reference counter. */
diff --git a/Engine/Plugins/Experimental/Mutable/Source/CustomizableObject/Private/MuCO/UnrealMutableModelDiskStreamer.cpp b/Engine/Plugins/Experimental/Mutable/Source/CustomizableObject/Private/MuCO/UnrealMutableModelDiskStreamer.cpp
--- a/Engine/Plugins/Experimental/Mutable/Source/CustomizableObject/Private/MuCO/UnrealMutableModelDiskStreamer.cpp
+++ b/Engine/Plugins/Experimental/Mutable/Source/CustomizableObject/Private/MuCO/UnrealMutableModelDiskStreamer.cpp
@@ -44,56 +44,13 @@ void UnrealMutableInputStream::Read(void* pData, uint64 size)
 
 
 //-------------------------------------------------------------------------------------------------
-FUnrealMutableModelBulkStreamer::FUnrealMutableModelBulkStreamer(FArchive* InMainDataArchive, FArchive* InStreamedDataArchive)
-{
-#if WITH_EDITOR
-	MainDataArchive = InMainDataArchive;
-	StreamedDataArchive = InStreamedDataArchive;
-	CurrentWriteFile = nullptr;
-#endif
-}
-
-FUnrealMutableModelBulkStreamer::~FUnrealMutableModelBulkStreamer()
+FUnrealMutableModelBulkReader::~FUnrealMutableModelBulkReader()
 {
 	EndStreaming();
 }
 
 
-void FUnrealMutableModelBulkStreamer::OpenWriteFile(const char* strModelName, uint64 key0)
-{
-#if WITH_EDITOR
-	if (key0 == 0) // Model
-	{
-		check(MainDataArchive);
-		CurrentWriteFile = MainDataArchive;
-	}
-	else
-	{
-		check(StreamedDataArchive);
-		CurrentWriteFile = StreamedDataArchive;
-	}
-#endif
-}
-
-
-void FUnrealMutableModelBulkStreamer::Write(const void* pBuffer, uint64 size)
-{
-#if WITH_EDITOR
-	check(CurrentWriteFile);
-	CurrentWriteFile->Serialize(const_cast<void*>(pBuffer), size);
-#endif
-}
-
-
-void FUnrealMutableModelBulkStreamer::CloseWriteFile()
-{
-#if WITH_EDITOR
-	CurrentWriteFile = nullptr;
-#endif
-}
-
-
-bool FUnrealMutableModelBulkStreamer::PrepareStreamingForObject(UCustomizableObject* CustomizableObject)
+bool FUnrealMutableModelBulkReader::PrepareStreamingForObject(UCustomizableObject* CustomizableObject)
 {
 	// This happens in the game thread
 	check(IsInGameThread());
@@ -131,9 +88,9 @@ bool FUnrealMutableModelBulkStreamer::PrepareStreamingForObject(UCustomizableObj
 
 	// Is the object already prepared for streaming?
 	bool bAlreadyStreaming = Objects.FindByPredicate(
-		[CustomizableObject](const FObjectData& d) 
+		[CustomizableObject](const FObjectData& d)
 		{ return d.Model.Pin().Get() == CustomizableObject->GetModel().Get(); })
-		!= 
+		!=
 		nullptr;
 
 	if (!bAlreadyStreaming)
@@ -192,7 +149,7 @@ bool FUnrealMutableModelBulkStreamer::PrepareStreamingForObject(UCustomizableObj
 
 
 #if WITH_EDITOR
-void FUnrealMutableModelBulkStreamer::CancelStreamingForObject(const UCustomizableObject* CustomizableObject)
+void FUnrealMutableModelBulkReader::CancelStreamingForObject(const UCustomizableObject* CustomizableObject)
 {
 	// This happens in the game thread
 	check(IsInGameThread());
@@ -205,7 +162,7 @@ void FUnrealMutableModelBulkStreamer::CancelStreamingForObject(const UCustomizab
 	// See if we can free previuously allocated resources
 	for (int32 ObjectIndex = 0; ObjectIndex < Objects.Num(); ++ObjectIndex)
 	{
-		if (Objects[ObjectIndex].Model.Pin()==CustomizableObject->GetModel())
+		if (Objects[ObjectIndex].Model.Pin() == CustomizableObject->GetModel())
 		{
 			for (TPair<OPERATION_ID, IAsyncReadRequest*>& it : Objects[ObjectIndex].CurrentReadRequests)
 			{
@@ -225,7 +182,7 @@ void FUnrealMutableModelBulkStreamer::CancelStreamingForObject(const UCustomizab
 }
 
 
-bool FUnrealMutableModelBulkStreamer::AreTherePendingStreamingOperationsForObject(const UCustomizableObject* CustomizableObject) const
+bool FUnrealMutableModelBulkReader::AreTherePendingStreamingOperationsForObject(const UCustomizableObject* CustomizableObject) const
 {
 	// This happens in the game thread
 	check(IsInGameThread());
@@ -240,7 +197,7 @@ bool FUnrealMutableModelBulkStreamer::AreTherePendingStreamingOperationsForObjec
 	{
 		if (Objects[ObjectIndex].Model.Pin() == CustomizableObject->GetModel())
 		{
-			if(!Objects[ObjectIndex].CurrentReadRequests.IsEmpty())
+			if (!Objects[ObjectIndex].CurrentReadRequests.IsEmpty())
 			{
 				return true;
 			}
@@ -253,7 +210,7 @@ bool FUnrealMutableModelBulkStreamer::AreTherePendingStreamingOperationsForObjec
 #endif // WITH_EDITOR
 
 
-void FUnrealMutableModelBulkStreamer::EndStreaming()
+void FUnrealMutableModelBulkReader::EndStreaming()
 {
 	for (FObjectData& o : Objects)
 	{
@@ -279,7 +236,7 @@ FAutoConsoleVariableRef CVarStreamPriority(
 	TEXT(""));
 
 
-mu::ModelStreamer::OPERATION_ID FUnrealMutableModelBulkStreamer::BeginReadBlock(const mu::Model* Model, uint64 Key, void* pBuffer, uint64 size)
+mu::ModelReader::OPERATION_ID FUnrealMutableModelBulkReader::BeginReadBlock(const mu::Model* Model, uint64 Key, void* pBuffer, uint64 size)
 {
 	MUTABLE_CPUPROFILER_SCOPE(FUnrealMutableModelBulkStreamer::OpenReadFile);
 
@@ -297,7 +254,7 @@ mu::ModelStreamer::OPERATION_ID FUnrealMutableModelBulkStreamer::BeginReadBlock(
 		return -1;
 	}
 
-	mu::ModelStreamer::OPERATION_ID Result = 0;
+	mu::ModelReader::OPERATION_ID Result = 0;
 
 	check(!ObjectData->ReadFileHandles.IsEmpty());
 
@@ -329,7 +286,7 @@ mu::ModelStreamer::OPERATION_ID FUnrealMutableModelBulkStreamer::BeginReadBlock(
 }
 
 
-bool FUnrealMutableModelBulkStreamer::IsReadCompleted(mu::ModelStreamer::OPERATION_ID OperationId)
+bool FUnrealMutableModelBulkReader::IsReadCompleted(mu::ModelReader::OPERATION_ID OperationId)
 {
 	MUTABLE_CPUPROFILER_SCOPE(FUnrealMutableModelBulkStreamer::IsReadCompleted);
 
@@ -347,7 +304,7 @@ bool FUnrealMutableModelBulkStreamer::IsReadCompleted(mu::ModelStreamer::OPERATI
 }
 
 
-void FUnrealMutableModelBulkStreamer::EndRead(mu::ModelStreamer::OPERATION_ID OperationId)
+void FUnrealMutableModelBulkReader::EndRead(mu::ModelReader::OPERATION_ID OperationId)
 {
 	MUTABLE_CPUPROFILER_SCOPE(FUnrealMutableModelBulkStreamer::EndRead);
 
@@ -373,10 +330,52 @@ void FUnrealMutableModelBulkStreamer::EndRead(mu::ModelStreamer::OPERATION_ID Op
 			break;
 		}
 	}
-	
+
 	if (!bFound)
 	{
 		UE_LOG(LogMutable, Error, TEXT("Operation not found in EndRead."));
 		check(false);
 	}
 }
+
+
+#if WITH_EDITOR
+
+//-------------------------------------------------------------------------------------------------
+FUnrealMutableModelBulkWriter::FUnrealMutableModelBulkWriter(FArchive* InMainDataArchive, FArchive* InStreamedDataArchive)
+{
+	MainDataArchive = InMainDataArchive;
+	StreamedDataArchive = InStreamedDataArchive;
+	CurrentWriteFile = nullptr;
+}
+
+
+void FUnrealMutableModelBulkWriter::OpenWriteFile(uint64 key0)
+{
+	if (key0 == 0) // Model
+	{
+		check(MainDataArchive);
+		CurrentWriteFile = MainDataArchive;
+	}
+	else
+	{
+		check(StreamedDataArchive);
+		CurrentWriteFile = StreamedDataArchive;
+	}
+}
+
+
+void FUnrealMutableModelBulkWriter::Write(const void* pBuffer, uint64 size)
+{
+	check(CurrentWriteFile);
+	CurrentWriteFile->Serialize(const_cast<void*>(pBuffer), size);
+}
+
+
+void FUnrealMutableModelBulkWriter::CloseWriteFile()
+{
+	CurrentWriteFile = nullptr;
+}
+
+#endif
+
diff --git a/Engine/Plugins/Experimental/Mutable/Source/CustomizableObject/Public/MuCO/UnrealMutableModelDiskStreamer.h b/Engine/Plugins/Experimental/Mutable/Source/CustomizableObject/Public/MuCO/UnrealMutableModelDiskStreamer.h
--- a/Engine/Plugins/Experimental/Mutable/Source/CustomizableObject/Public/MuCO/UnrealMutableModelDiskStreamer.h
+++ b/Engine/Plugins/Experimental/Mutable/Source/CustomizableObject/Public/MuCO/UnrealMutableModelDiskStreamer.h
@@ -13,22 +13,6 @@ namespace mu { class Model; }
 struct FMutableStreamableBlock;
 
 
-// Implementation fo the basic mutable streams.
-#if WITH_EDITOR
-class UnrealMutableOutputStream : public mu::OutputStream
-{
-public:
-
-	UnrealMutableOutputStream(FArchive& ar);
-
-	// mu::OutputStream interface
-	void Write(const void* pData, uint64 size) override;
-
-private:
-	FArchive& m_ar;
-};
-#endif
-
 
 class UnrealMutableInputStream : public mu::InputStream
 {
@@ -45,13 +29,12 @@ private:
 
 
 // Implementation of a mutable streamer using bulk storage.
-class CUSTOMIZABLEOBJECT_API FUnrealMutableModelBulkStreamer : public mu::ModelStreamer
+class CUSTOMIZABLEOBJECT_API FUnrealMutableModelBulkReader : public mu::ModelReader
 {
 public:
 	// 
-	FUnrealMutableModelBulkStreamer(FArchive* InMainDataArchive = nullptr, FArchive* InStreamedDataArchive = nullptr);
-	~FUnrealMutableModelBulkStreamer();
-	
+	~FUnrealMutableModelBulkReader();
+
 	// Own interface
 
 	/** Make sure that the provided object can stream data. */
@@ -72,29 +55,14 @@ public:
 
 	/** Release all the pending resources. This disables treamings for all objects. */
 	void EndStreaming();
-	
-	// mu::ModelStreamer interface
-	void OpenWriteFile(const char* strModelName, uint64 key0) override;
-	void Write(const void* pBuffer, uint64 size) override;
-	void CloseWriteFile() override;
-	
+
+	// mu::ModelReader interface
 	OPERATION_ID BeginReadBlock(const mu::Model*, uint64 key0, void* pBuffer, uint64 size) override;
 	bool IsReadCompleted(OPERATION_ID) override;
 	void EndRead(OPERATION_ID) override;
 
 protected:
 
-#if WITH_EDITOR
-	// Write
-	// Non-owned pointer to an archive where we'll store the main model data (non-streamable)
-	FArchive* MainDataArchive = nullptr;
-
-	// Non-owned pointer to an archive where we'll store the resouces (streamable)
-	FArchive* StreamedDataArchive = nullptr;
-
-	FArchive* CurrentWriteFile = nullptr;
-#endif
-
 	/** Streaming data for one object. */
 	struct FObjectData
 	{
@@ -109,3 +77,46 @@ protected:
 	/** This is used to generate unique ids for read requests. */
 	OPERATION_ID LastOperationID = 0;
 };
+
+
+#if WITH_EDITOR
+
+class UnrealMutableOutputStream : public mu::OutputStream
+{
+public:
+
+	UnrealMutableOutputStream(FArchive& ar);
+
+	// mu::OutputStream interface
+	void Write(const void* pData, uint64 size) override;
+
+private:
+	FArchive& m_ar;
+};
+
+
+// Implementation of a mutable streamer using bulk storage.
+class CUSTOMIZABLEOBJECT_API FUnrealMutableModelBulkWriter : public mu::ModelWriter
+{
+public:
+	// 
+	FUnrealMutableModelBulkWriter(FArchive* InMainDataArchive = nullptr, FArchive* InStreamedDataArchive = nullptr);
+
+	// mu::ModelWriter interface
+	void OpenWriteFile(uint64 key0) override;
+	void Write(const void* pBuffer, uint64 size) override;
+	void CloseWriteFile() override;
+
+protected:
+
+	// Non-owned pointer to an archive where we'll store the main model data (non-streamable)
+	FArchive* MainDataArchive = nullptr;
+
+	// Non-owned pointer to an archive where we'll store the resouces (streamable)
+	FArchive* StreamedDataArchive = nullptr;
+
+	FArchive* CurrentWriteFile = nullptr;
+
+};
+
+#endif
diff --git a/Engine/Plugins/Experimental/Mutable/Source/CustomizableObjectEditor/Private/MuCOE/CustomizableObjectCompileRunnable.cpp b/Engine/Plugins/Experimental/Mutable/Source/CustomizableObjectEditor/Private/MuCOE/CustomizableObjectCompileRunnable.cpp
--- a/Engine/Plugins/Experimental/Mutable/Source/CustomizableObjectEditor/Private/MuCOE/CustomizableObjectCompileRunnable.cpp
+++ b/Engine/Plugins/Experimental/Mutable/Source/CustomizableObjectEditor/Private/MuCOE/CustomizableObjectCompileRunnable.cpp
@@ -190,7 +190,7 @@ uint32 FCustomizableObjectSaveDDRunnable::Run()
 		ModelMemoryWriter << bModelSerialized;
 		if (bModelSerialized)
 		{
-			FUnrealMutableModelBulkStreamer Streamer(&ModelMemoryWriter, &StreamableMemoryWriter);
+			FUnrealMutableModelBulkWriter Streamer(&ModelMemoryWriter, &StreamableMemoryWriter);
 			mu::Model::Serialise(Model.Get(), Streamer);
 		}
 	}
@@ -236,7 +236,7 @@ uint32 FCustomizableObjectSaveDDRunnable::Run()
 			// Serialize mu::Model and streamable resources
 			*ModelMemoryWriter << bModelSerialized;
 
-			FUnrealMutableModelBulkStreamer Streamer(ModelMemoryWriter.Get(), StreamableMemoryWriter.Get());
+			FUnrealMutableModelBulkWriter Streamer(ModelMemoryWriter.Get(), StreamableMemoryWriter.Get());
 			mu::Model::Serialise(Model.Get(), Streamer);
 
 			// Save to disk
diff --git a/Engine/Plugins/Experimental/Mutable/Source/CustomizableObjectEditor/Private/MuCOE/SMutableCodeViewer.cpp b/Engine/Plugins/Experimental/Mutable/Source/CustomizableObjectEditor/Private/MuCOE/SMutableCodeViewer.cpp
--- a/Engine/Plugins/Experimental/Mutable/Source/CustomizableObjectEditor/Private/MuCOE/SMutableCodeViewer.cpp
+++ b/Engine/Plugins/Experimental/Mutable/Source/CustomizableObjectEditor/Private/MuCOE/SMutableCodeViewer.cpp
@@ -2316,9 +2316,9 @@ namespace
 		
 	public:
 #ifdef MUTABLE_USE_NEW_TASKGRAPH
-		TTuple<UE::Tasks::FTask, TFunction<void()>> GetImageAsync(mu::EXTERNAL_IMAGE_ID id, uint8 MipmapsToSkip, TFunction<void(mu::Ptr<mu::Image>)>& ResultCallback) override
+		TTuple<UE::Tasks::FTask, TFunction<void()>> GetImageAsync(mu::FExternalImageID id, uint8 MipmapsToSkip, TFunction<void(mu::Ptr<mu::Image>)>& ResultCallback) override
 #else
-		TTuple<FGraphEventRef, TFunction<void()>> GetImageAsync(mu::EXTERNAL_IMAGE_ID id, uint8 MipmapsToSkip, TFunction<void(mu::Ptr<mu::Image>)>& ResultCallback) override
+		TTuple<FGraphEventRef, TFunction<void()>> GetImageAsync(mu::FExternalImageID id, uint8 MipmapsToSkip, TFunction<void(mu::Ptr<mu::Image>)>& ResultCallback) override
 #endif
 		{
 			MUTABLE_CPUPROFILER_SCOPE(TestImageProvider_GetImage);
@@ -2369,7 +2369,7 @@ namespace
 			return MakeTuple(CompletionEvent, []() -> void {});
 		}
 
-		mu::FImageDesc GetImageDesc(mu::EXTERNAL_IMAGE_ID Id, uint8 MipmapsToSkip) override
+		mu::FImageDesc GetImageDesc(mu::FExternalImageID Id, uint8 MipmapsToSkip) override
 		{
 			return IMAGE_DESC;
 		}
diff --git a/Engine/Plugins/Experimental/Mutable/Source/MutableRuntime/Private/MuR/CodeRunner.cpp b/Engine/Plugins/Experimental/Mutable/Source/MutableRuntime/Private/MuR/CodeRunner.cpp
--- a/Engine/Plugins/Experimental/Mutable/Source/MutableRuntime/Private/MuR/CodeRunner.cpp
+++ b/Engine/Plugins/Experimental/Mutable/Source/MutableRuntime/Private/MuR/CodeRunner.cpp
@@ -138,9 +138,9 @@ namespace mu
 
     //---------------------------------------------------------------------------------------------
 #ifdef MUTABLE_USE_NEW_TASKGRAPH
-	TTuple<UE::Tasks::FTask, TFunction<void()>> CodeRunner::LoadExternalImageAsync(EXTERNAL_IMAGE_ID Id, uint8 MipmapsToSkip, TFunction<void(Ptr<Image>)>& ResultCallback)
+	TTuple<UE::Tasks::FTask, TFunction<void()>> CodeRunner::LoadExternalImageAsync(FExternalImageID Id, uint8 MipmapsToSkip, TFunction<void(Ptr<Image>)>& ResultCallback)
 #else
-	TTuple<FGraphEventRef, TFunction<void()>> CodeRunner::LoadExternalImageAsync(EXTERNAL_IMAGE_ID Id, uint8 MipmapsToSkip, TFunction<void(Ptr<Image>)>& ResultCallback)
+	TTuple<FGraphEventRef, TFunction<void()>> CodeRunner::LoadExternalImageAsync(FExternalImageID Id, uint8 MipmapsToSkip, TFunction<void(Ptr<Image>)>& ResultCallback)
 #endif
     {
 		MUTABLE_CPUPROFILER_SCOPE(LoadExternalImageAsync);
@@ -153,7 +153,7 @@ namespace mu
 
 			// Don't cache for now. Need to figure out how to invalidate them.
 			// \TODO: Like constants? attached to a cache level?
-			//m_pSystem->m_externalImages.push_back( pair<EXTERNAL_IMAGE_ID,ImagePtr>(id,pResult) );
+			//m_pSystem->m_externalImages.push_back( pair<FExternalImageID,ImagePtr>(id,pResult) );
 		}
 		else
 		{
@@ -175,7 +175,7 @@ namespace mu
 
 	
     //---------------------------------------------------------------------------------------------
-	FImageDesc CodeRunner::GetExternalImageDesc(EXTERNAL_IMAGE_ID Id, uint8 MipmapsToSkip)
+	FImageDesc CodeRunner::GetExternalImageDesc(FExternalImageID Id, uint8 MipmapsToSkip)
 	{
 		MUTABLE_CPUPROFILER_SCOPE(GetExternalImageDesc);
 
@@ -781,6 +781,11 @@ namespace mu
     {
 		MUTABLE_CPUPROFILER_SCOPE(RunCode_InstanceAddResource);
 
+		if (!pModel || !m_pSystem)
+		{
+			return;
+		}
+
 		OP_TYPE type = pModel->GetPrivate()->m_program.GetOpType(item.At);
         switch (type)
         {
@@ -813,14 +818,16 @@ namespace mu
 
                 if ( args.value )
                 {
-                    RESOURCE_ID meshId = pModel->GetPrivate()->GetResourceKey(
-                                args.relevantParametersListIndex,
-                                args.value,
-                                pParams);
-                    OP::ADDRESS nameAd = args.name;
-                    check(  nameAd < (uint32)pModel->GetPrivate()->m_program.m_constantStrings.Num() );
-                    const char* strName = pModel->GetPrivate()->m_program.m_constantStrings[ nameAd ].c_str();
-                    pResult->GetPrivate()->AddMesh( 0, 0, meshId, strName );
+					FWorkingMemoryManager::FModelCacheEntry* ModelCache = m_pSystem->WorkingMemoryManager.FindModelCache(pModel);
+					check(ModelCache);
+					if (ModelCache)
+					{
+						FResourceID MeshId = ModelCache->GetResourceKey(args.relevantParametersListIndex, args.value, pParams, m_pSystem->WorkingMemoryManager.MaxGeneratedResourceCacheSize);
+						OP::ADDRESS NameAd = args.name;
+						check(NameAd < (uint32)pModel->GetPrivate()->m_program.m_constantStrings.Num());
+						const char* Name = pModel->GetPrivate()->m_program.m_constantStrings[NameAd].c_str();
+						pResult->GetPrivate()->AddMesh(0, 0, MeshId, Name);
+					}
                 }
                 StoreInstance( item, pResult );
                 break;
@@ -857,14 +864,16 @@ namespace mu
 
                 if ( args.value )
                 {
-                    RESOURCE_ID imageId = pModel->GetPrivate()->GetResourceKey(
-                                args.relevantParametersListIndex,
-                                args.value,
-                                pParams);
-                    OP::ADDRESS nameAd = args.name;
-                    check(  nameAd < (uint32)pModel->GetPrivate()->m_program.m_constantStrings.Num() );
-                    const char* strName = pModel->GetPrivate()->m_program.m_constantStrings[ nameAd ].c_str();
-                    pResult->GetPrivate()->AddImage( 0, 0, 0, imageId, strName );
+					FWorkingMemoryManager::FModelCacheEntry* ModelCache = m_pSystem->WorkingMemoryManager.FindModelCache(pModel);
+					check(ModelCache);
+					if (ModelCache)
+					{
+						FResourceID ImageId = ModelCache->GetResourceKey(args.relevantParametersListIndex, args.value, pParams, m_pSystem->WorkingMemoryManager.MaxGeneratedResourceCacheSize );
+						OP::ADDRESS NameAd = args.name;
+						check(NameAd < (uint32)pModel->GetPrivate()->m_program.m_constantStrings.Num());
+						const char* Name = pModel->GetPrivate()->m_program.m_constantStrings[NameAd].c_str();
+						pResult->GetPrivate()->AddImage(0, 0, 0, ImageId, Name);
+					}
                 }
                 StoreInstance( item, pResult );
                 break;
@@ -5947,7 +5956,7 @@ namespace mu
 		{
 			check(item.Stage == 0);
 			OP::ParameterArgs args = program.GetOpArgs<OP::ParameterArgs>(item.At);
-			EXTERNAL_IMAGE_ID id = pParams->GetImageValue(args.variable);
+			FExternalImageID id = pParams->GetImageValue(args.variable);
 			uint8 MipsToSkip = item.ExecutionOptions;
 			m_heapImageDesc[item.CustomState] = GetExternalImageDesc(id, MipsToSkip);
 			StoreValidDesc(item);
diff --git a/Engine/Plugins/Experimental/Mutable/Source/MutableRuntime/Private/MuR/CodeRunner.h b/Engine/Plugins/Experimental/Mutable/Source/MutableRuntime/Private/MuR/CodeRunner.h
--- a/Engine/Plugins/Experimental/Mutable/Source/MutableRuntime/Private/MuR/CodeRunner.h
+++ b/Engine/Plugins/Experimental/Mutable/Source/MutableRuntime/Private/MuR/CodeRunner.h
@@ -108,11 +108,11 @@ namespace  mu
 		//! Load an external image asynchronously, retuns an event to wait for complition and a cleanup function 
 		//! that must be called once the event has completed.
 #ifdef MUTABLE_USE_NEW_TASKGRAPH
-		TTuple<UE::Tasks::FTask, TFunction<void()>> LoadExternalImageAsync(EXTERNAL_IMAGE_ID Id, uint8 MipmapsToSkip, TFunction<void(Ptr<Image>)>& ResultCallback);
+		TTuple<UE::Tasks::FTask, TFunction<void()>> LoadExternalImageAsync(FExternalImageID Id, uint8 MipmapsToSkip, TFunction<void(Ptr<Image>)>& ResultCallback);
 #else
-		TTuple<FGraphEventRef, TFunction<void()>> LoadExternalImageAsync(EXTERNAL_IMAGE_ID Id, uint8 MipmapsToSkip, TFunction<void(Ptr<Image>)>& ResultCallback);
+		TTuple<FGraphEventRef, TFunction<void()>> LoadExternalImageAsync(FExternalImageID Id, uint8 MipmapsToSkip, TFunction<void(Ptr<Image>)>& ResultCallback);
 #endif
- 	    mu::FImageDesc GetExternalImageDesc(EXTERNAL_IMAGE_ID Id, uint8 MipmapsToSkip);
+ 	    mu::FImageDesc GetExternalImageDesc(FExternalImageID Id, uint8 MipmapsToSkip);
 
 		/** Settings that may affect the execution of some operations, like image conversion quality. */
 		Ptr<const Settings> m_pSettings;
@@ -572,7 +572,7 @@ namespace  mu
 		{
 			int32 m_romIndex = 0;
 			DATATYPE ConstantType = DT_NONE;
-			ModelStreamer::OPERATION_ID m_streamID;
+			ModelReader::OPERATION_ID m_streamID;
 			TArray<uint8> m_streamBuffer;
 		};
 		TArray<FRomLoadOp> m_romLoadOps;
diff --git a/Engine/Plugins/Experimental/Mutable/Source/MutableRuntime/Private/MuR/CodeRunnerTasked.cpp b/Engine/Plugins/Experimental/Mutable/Source/MutableRuntime/Private/MuR/CodeRunnerTasked.cpp
--- a/Engine/Plugins/Experimental/Mutable/Source/MutableRuntime/Private/MuR/CodeRunnerTasked.cpp
+++ b/Engine/Plugins/Experimental/Mutable/Source/MutableRuntime/Private/MuR/CodeRunnerTasked.cpp
@@ -2040,6 +2040,11 @@ namespace mu
 	//---------------------------------------------------------------------------------------------
 	bool CodeRunner::FLoadMeshRomTask::Prepare(CodeRunner* Runner, bool& bOutFailed )
 	{
+		if (!Runner || !Runner->m_pSystem)
+		{
+			return false;
+		}
+
 		// This runs in the mutable Runner thread
 		MUTABLE_CPUPROFILER_SCOPE(FLoadMeshRomTask_Prepare);
 		bOutFailed = false;
@@ -2049,8 +2054,8 @@ namespace mu
 		check(RomIndex < program.m_roms.Num());
 		bool bRomIsLoaded = program.IsRomLoaded(RomIndex);
 
-		FWorkingMemoryManager::FModelCacheEntry& ModelCache = Runner->m_pSystem->WorkingMemoryManager.GetModelCache(Runner->m_pModel);
-		bool bRomHasPendingOps = ModelCache.PendingOpsPerRom[RomIndex]!=0;
+		FWorkingMemoryManager::FModelCacheEntry* ModelCache = Runner->m_pSystem->WorkingMemoryManager.FindModelCache(Runner->m_pModel.Get());
+		bool bRomHasPendingOps = ModelCache->PendingOpsPerRom[RomIndex]!=0;
 		if (!bRomIsLoaded && !bRomHasPendingOps)
 		{
 			check(Runner->m_pSystem->StreamInterface);
@@ -2098,7 +2103,7 @@ namespace mu
 				return false;
 			}
 		}
-		++ModelCache.PendingOpsPerRom[RomIndex];
+		++ModelCache->PendingOpsPerRom[RomIndex];
 
 		//UE_LOG(LogMutableCore, Log, TEXT("FLoadMeshRomTask::Prepare romindex %d pending %d."), RomAt.m_rom, Runner->m_romPendingOps[RomAt.m_rom]);
 
@@ -2110,6 +2115,11 @@ namespace mu
 	//---------------------------------------------------------------------------------------------
 	void CodeRunner::FLoadMeshRomTask::Complete(CodeRunner* Runner)
 	{
+		if (!Runner || !Runner->m_pSystem)
+		{
+			return;
+		}
+
 		MUTABLE_CPUPROFILER_SCOPE(FLoadMeshRomTask_Complete);
 
 		// This runs in the Runner thread
@@ -2117,10 +2127,10 @@ namespace mu
 		// Process the constant op normally, now that the rom is loaded.
 		Runner->RunCode(Op, Runner->m_pParams, Runner->m_pModel.Get(), Runner->m_lodMask);
 
-		FWorkingMemoryManager::FModelCacheEntry& ModelCache = Runner->m_pSystem->WorkingMemoryManager.GetModelCache(Runner->m_pModel);
+		FWorkingMemoryManager::FModelCacheEntry* ModelCache = Runner->m_pSystem->WorkingMemoryManager.FindModelCache(Runner->m_pModel.Get());
 
 		Runner->m_pSystem->WorkingMemoryManager.MarkRomUsed(RomIndex, Runner->m_pModel);
-		--ModelCache.PendingOpsPerRom[RomIndex];
+		--ModelCache->PendingOpsPerRom[RomIndex];
 	}
 
 
@@ -2219,13 +2229,18 @@ namespace mu
 	//---------------------------------------------------------------------------------------------
 	bool CodeRunner::FLoadImageRomsTask::Prepare(CodeRunner* Runner, bool& bOutFailed )
 	{
+		if (!Runner || !Runner->m_pSystem)
+		{
+			return false;
+		}
+
 		// This runs in the mutable Runner thread
 		MUTABLE_CPUPROFILER_SCOPE(FLoadImageRomsTask_Prepare);
 		bOutFailed = false;
 
 		FProgram& program = Runner->m_pModel->GetPrivate()->m_program;
 
-		FWorkingMemoryManager::FModelCacheEntry& ModelCache = Runner->m_pSystem->WorkingMemoryManager.GetModelCache(Runner->m_pModel);
+		FWorkingMemoryManager::FModelCacheEntry* ModelCache = Runner->m_pSystem->WorkingMemoryManager.FindModelCache(Runner->m_pModel.Get());
 
 		for (int32 i = 0; i<LODIndexCount; ++i )
 		{
@@ -2242,11 +2257,11 @@ namespace mu
 			check(RomIndex < program.m_roms.Num());
 
 			bool bRomIsLoaded = program.IsRomLoaded(RomIndex);
-			bool bRomHasAlreadyBeenRequested = ModelCache.PendingOpsPerRom[RomIndex] != 0;
-			++ModelCache.PendingOpsPerRom[RomIndex];
+			bool bRomHasAlreadyBeenRequested = ModelCache->PendingOpsPerRom[RomIndex] != 0;
+			++ModelCache->PendingOpsPerRom[RomIndex];
 
 			if (DebugRom && (DebugRomAll || RomIndex == DebugRomIndex))
-				UE_LOG(LogMutableCore, Log, TEXT("Preparing rom %d, now peding ops is %d."), RomIndex, ModelCache.PendingOpsPerRom[RomIndex]);
+				UE_LOG(LogMutableCore, Log, TEXT("Preparing rom %d, now peding ops is %d."), RomIndex, ModelCache->PendingOpsPerRom[RomIndex]);
 
 			if (bRomIsLoaded || bRomHasAlreadyBeenRequested)
 			{
@@ -2307,11 +2322,16 @@ namespace mu
 	//---------------------------------------------------------------------------------------------
 	void CodeRunner::FLoadImageRomsTask::Complete(CodeRunner* Runner)
 	{
+		if (!Runner || !Runner->m_pSystem)
+		{
+			return;
+		}
+
 		MUTABLE_CPUPROFILER_SCOPE(FLoadImageRomsTask_Complete);
 
 		// This runs in the Runner thread
 
-		FWorkingMemoryManager::FModelCacheEntry& ModelCache = Runner->m_pSystem->WorkingMemoryManager.GetModelCache(Runner->m_pModel);
+		FWorkingMemoryManager::FModelCacheEntry* ModelCache = Runner->m_pSystem->WorkingMemoryManager.FindModelCache(Runner->m_pModel.Get());
 
 		// Process the constant op normally, now that the rom is loaded.
 		Runner->RunCode(Op, Runner->m_pParams, Runner->m_pModel.Get(), Runner->m_lodMask);
@@ -2332,11 +2352,11 @@ namespace mu
 			check(RomIndex < program.m_roms.Num());
 
 			Runner->m_pSystem->WorkingMemoryManager.MarkRomUsed(RomIndex, Runner->m_pModel);
-			--ModelCache.PendingOpsPerRom[RomIndex];
+			--ModelCache->PendingOpsPerRom[RomIndex];
 
 			if (DebugRom && (DebugRomAll || RomIndex == DebugRomIndex))
 			{
-				UE_LOG(LogMutableCore, Log, TEXT("FLoadImageRomsTask::Complete rom %d, now peding ops is %d."), RomIndex, ModelCache.PendingOpsPerRom[RomIndex]);
+				UE_LOG(LogMutableCore, Log, TEXT("FLoadImageRomsTask::Complete rom %d, now peding ops is %d."), RomIndex, ModelCache->PendingOpsPerRom[RomIndex]);
 			}
 		}
 	}
@@ -2375,7 +2395,7 @@ namespace mu
 	class FImageExternalLoadTask : public CodeRunner::FIssuedTask
 	{
 	public:
-		FImageExternalLoadTask(const FScheduledOp& InItem, uint8 InMipmapsToSkip, EXTERNAL_IMAGE_ID InId);
+		FImageExternalLoadTask(const FScheduledOp& InItem, uint8 InMipmapsToSkip, FExternalImageID InId);
 
 		// FIssuedTask interface
 		virtual bool Prepare(CodeRunner*, bool& bOutFailed) override;
@@ -2383,7 +2403,7 @@ namespace mu
 		
 	private:
 		uint8 MipmapsToSkip;
-		EXTERNAL_IMAGE_ID Id;
+		FExternalImageID Id;
 
 		Ptr<Image> Result;
 		
@@ -2392,7 +2412,7 @@ namespace mu
 
 
 	//---------------------------------------------------------------------------------------------
-	FImageExternalLoadTask::FImageExternalLoadTask(const FScheduledOp& InOp,  uint8 InMipmapsToSkip, EXTERNAL_IMAGE_ID InId)
+	FImageExternalLoadTask::FImageExternalLoadTask(const FScheduledOp& InOp,  uint8 InMipmapsToSkip, FExternalImageID InId)
 		: FIssuedTask(InOp)
 	{
 		MipmapsToSkip = InMipmapsToSkip;
@@ -2528,7 +2548,7 @@ namespace mu
 			OP::ParameterArgs args = program.GetOpArgs<OP::ParameterArgs>(item.At);
 			Ptr<RangeIndex> Index = BuildCurrentOpRangeIndex(item, m_pParams, m_pModel.Get(), args.variable);
 
-			const EXTERNAL_IMAGE_ID Id = m_pParams->GetImageValue(args.variable, Index);
+			const FExternalImageID Id = m_pParams->GetImageValue(args.variable, Index);
 
 			check(ImageLOD < TNumericLimits<uint8>::Max() && ImageLOD >= 0);
 			check(ImageLOD + static_cast<int32>(item.ExecutionOptions) < TNumericLimits<uint8>::Max());
diff --git a/Engine/Plugins/Experimental/Mutable/Source/MutableRuntime/Private/MuR/Instance.cpp b/Engine/Plugins/Experimental/Mutable/Source/MutableRuntime/Private/MuR/Instance.cpp
--- a/Engine/Plugins/Experimental/Mutable/Source/MutableRuntime/Private/MuR/Instance.cpp
+++ b/Engine/Plugins/Experimental/Mutable/Source/MutableRuntime/Private/MuR/Instance.cpp
@@ -65,41 +65,27 @@ namespace mu
 
 
     //---------------------------------------------------------------------------------------------
-    int Instance::GetLODCount() const
+    int32 Instance::GetLODCount() const
     {
-        return m_pD->GetLODCount();
+		return m_pD->m_lods.Num();
     }
 
 
 	//---------------------------------------------------------------------------------------------
-	int Instance::Private::GetLODCount() const
+	int32 Instance::GetComponentCount( int32 lod ) const
 	{
-		return m_lods.Num();
-	}
-
-
-	//---------------------------------------------------------------------------------------------
-	int Instance::GetComponentCount( int lod ) const
-	{
-		return m_pD->GetComponentCount( lod );
-	}
-
-
-	//---------------------------------------------------------------------------------------------
-	int Instance::Private::GetComponentCount( int lod ) const
-	{
-		check( lod>=0 && lod<m_lods.Num() );
-        if ( lod>=0 && lod<m_lods.Num() )
-        {
-            return m_lods[lod].m_components.Num();
-        }
+		check(lod >= 0 && lod < m_pD->m_lods.Num());
+		if (lod >= 0 && lod < m_pD->m_lods.Num())
+		{
+			return m_pD->m_lods[lod].m_components.Num();
+		}
 
-        return 0;
+		return 0;
 	}
 
 
     //---------------------------------------------------------------------------------------------
-    const char* Instance::GetComponentName( int lod, int comp ) const
+    const char* Instance::GetComponentName( int32 lod, int32 comp ) const
     {
         if ( lod>=0 && lod<m_pD->m_lods.Num() &&
              comp>=0 && comp<m_pD->m_lods[lod].m_components.Num() )
@@ -116,7 +102,7 @@ namespace mu
 
 	
 	//---------------------------------------------------------------------------------------------
-	uint16 Instance::GetComponentId( int lod, int comp ) const
+	uint16 Instance::GetComponentId( int32 lod, int32 comp ) const
 	{
 		if ( lod>=0 && lod<m_pD->m_lods.Num() &&
 			 comp>=0 && comp<m_pD->m_lods[lod].m_components.Num() )
@@ -133,31 +119,24 @@ namespace mu
 
 
     //---------------------------------------------------------------------------------------------
-    int Instance::GetSurfaceCount( int lod, int comp ) const
-    {
-        return m_pD->GetSurfaceCount( lod, comp );
-    }
-
-
-    //---------------------------------------------------------------------------------------------
-    int Instance::Private::GetSurfaceCount( int lod, int comp ) const
+    int32 Instance::GetSurfaceCount( int32 lod, int32 comp ) const
     {
-        if ( lod>=0 && lod<m_lods.Num() &&
-             comp>=0 && comp<m_lods[lod].m_components.Num() )
-        {
-            return m_lods[lod].m_components[comp].m_surfaces.Num();
-        }
+		if (lod >= 0 && lod < m_pD->m_lods.Num() &&
+			comp >= 0 && comp < m_pD->m_lods[lod].m_components.Num())
+		{
+			return m_pD->m_lods[lod].m_components[comp].m_surfaces.Num();
+		}
 		else
 		{
 			check(false);
 		}
 
-        return 0;
-    }
+		return 0;
+	}
 
 
 	//---------------------------------------------------------------------------------------------
-    const char* Instance::GetSurfaceName( int lod, int comp, int surf ) const
+    const char* Instance::GetSurfaceName( int32 lod, int32 comp, int32 surf ) const
     {
         if ( lod>=0 && lod<m_pD->m_lods.Num() &&
              comp>=0 && comp<m_pD->m_lods[lod].m_components.Num() &&
@@ -175,7 +154,7 @@ namespace mu
 
 
     //---------------------------------------------------------------------------------------------
-    uint32_t Instance::GetSurfaceId( int lod, int comp, int surf ) const
+    uint32_t Instance::GetSurfaceId( int32 lod, int32 comp, int32 surf ) const
     {
         if ( lod>=0 && lod<m_pD->m_lods.Num() &&
              comp>=0 && comp<m_pD->m_lods[lod].m_components.Num() &&
@@ -259,7 +238,7 @@ namespace mu
 	
 
     //---------------------------------------------------------------------------------------------
-    uint32_t Instance::GetSurfaceCustomId( int lod, int comp, int surf ) const
+    uint32_t Instance::GetSurfaceCustomId( int32 lod, int32 comp, int32 surf ) const
     {
         if ( lod>=0 && lod<m_pD->m_lods.Num() &&
              comp>=0 && comp<m_pD->m_lods[lod].m_components.Num() &&
@@ -277,133 +256,87 @@ namespace mu
 
 
 	//---------------------------------------------------------------------------------------------
-    int Instance::GetMeshCount( int lod, int comp ) const
-	{
-        return m_pD->GetMeshCount( lod, comp );
-	}
-
-
-	//---------------------------------------------------------------------------------------------
-    int Instance::Private::GetMeshCount( int lod, int comp ) const
+    int Instance::GetMeshCount( int32 lod, int32 comp ) const
 	{
-		check( lod>=0 && lod<m_lods.Num() );
-		check( comp>=0 && comp<m_lods[lod].m_components.Num() );
+		check(lod >= 0 && lod < m_pD->m_lods.Num());
+		check(comp >= 0 && comp < m_pD->m_lods[lod].m_components.Num());
 
-        return m_lods[lod].m_components[comp].m_meshes.Num();
+		return m_pD->m_lods[lod].m_components[comp].m_meshes.Num();
 	}
 
 
 	//---------------------------------------------------------------------------------------------
-    int Instance::GetImageCount( int lod, int comp, int surf ) const
+    int32 Instance::GetImageCount( int32 lod, int32 comp, int32 surf ) const
 	{
-        return m_pD->GetImageCount( lod, comp, surf );
-	}
-
-
-	//---------------------------------------------------------------------------------------------
-    int Instance::Private::GetImageCount( int lod, int comp, int surf ) const
-	{
-		check( lod>=0 && lod<m_lods.Num() );
-		check( comp>=0 && comp<m_lods[lod].m_components.Num() );
-        check( surf>=0 && surf<m_lods[lod].m_components[comp].m_surfaces.Num() );
-
-        return m_lods[lod].m_components[comp].m_surfaces[surf].m_images.Num();
-	}
-
+		check(lod >= 0 && lod < m_pD->m_lods.Num());
+		check(comp >= 0 && comp < m_pD->m_lods[lod].m_components.Num());
+		check(surf >= 0 && surf < m_pD->m_lods[lod].m_components[comp].m_surfaces.Num());
 
-	//---------------------------------------------------------------------------------------------
-    int Instance::GetVectorCount( int lod, int comp, int surf ) const
-	{
-        return m_pD->GetVectorCount( lod, comp, surf );
+		return m_pD->m_lods[lod].m_components[comp].m_surfaces[surf].m_images.Num();
 	}
 
 
 	//---------------------------------------------------------------------------------------------
-    int Instance::Private::GetVectorCount( int lod, int comp, int surf ) const
+    int32 Instance::GetVectorCount( int32 lod, int32 comp, int32 surf ) const
 	{
-		check( lod>=0 && lod<m_lods.Num() );
-		check( comp>=0 && comp<m_lods[lod].m_components.Num() );
-        check( surf>=0 && surf<m_lods[lod].m_components[comp].m_surfaces.Num() );
-
-        return m_lods[lod].m_components[comp].m_surfaces[surf].m_vectors.Num();
-	}
-
+		check(lod >= 0 && lod < m_pD->m_lods.Num());
+		check(comp >= 0 && comp < m_pD->m_lods[lod].m_components.Num());
+		check(surf >= 0 && surf < m_pD->m_lods[lod].m_components[comp].m_surfaces.Num());
 
-	//---------------------------------------------------------------------------------------------
-    int Instance::GetScalarCount( int lod, int comp, int surf ) const
-	{
-        return m_pD->GetScalarCount( lod, comp, surf );
+		return m_pD->m_lods[lod].m_components[comp].m_surfaces[surf].m_vectors.Num();
 	}
 
 
 	//---------------------------------------------------------------------------------------------
-    int Instance::Private::GetScalarCount( int lod, int comp, int surf ) const
+    int32 Instance::GetScalarCount( int32 lod, int32 comp, int32 surf ) const
 	{
-		check( lod>=0 && lod<m_lods.Num() );
-		check( comp>=0 && comp<m_lods[lod].m_components.Num() );
-        check( surf>=0 && surf<m_lods[lod].m_components[comp].m_surfaces.Num() );
+		check(lod >= 0 && lod < m_pD->m_lods.Num());
+		check(comp >= 0 && comp < m_pD->m_lods[lod].m_components.Num());
+		check(surf >= 0 && surf < m_pD->m_lods[lod].m_components[comp].m_surfaces.Num());
 
-        return m_lods[lod].m_components[comp].m_surfaces[surf].m_scalars.Num();
+		return m_pD->m_lods[lod].m_components[comp].m_surfaces[surf].m_scalars.Num();
 	}
 
 
     //---------------------------------------------------------------------------------------------
-    int Instance::GetStringCount( int lod, int comp, int surf ) const
+    int32 Instance::GetStringCount( int32 lod, int32 comp, int32 surf ) const
     {
-        return m_pD->GetStringCount( lod, comp, surf );
-    }
-
+		check(lod >= 0 && lod < m_pD->m_lods.Num());
+		check(comp >= 0 && comp < m_pD->m_lods[lod].m_components.Num());
+		check(surf >= 0 && surf < m_pD->m_lods[lod].m_components[comp].m_surfaces.Num());
 
-    //---------------------------------------------------------------------------------------------
-    int Instance::Private::GetStringCount( int lod, int comp, int surf ) const
-    {
-        check( lod >= 0 && lod < m_lods.Num() );
-        check( comp >= 0 && comp < m_lods[lod].m_components.Num() );
-        check( surf >= 0 && surf < m_lods[lod].m_components[comp].m_surfaces.Num() );
+		return m_pD->m_lods[lod].m_components[comp].m_surfaces[surf].m_strings.Num();
+	}
 
-        return m_lods[lod].m_components[comp].m_surfaces[surf].m_strings.Num();
-    }
 
 
     //---------------------------------------------------------------------------------------------
-    RESOURCE_ID Instance::GetMeshId( int lod, int comp, int mesh ) const
+	FResourceID Instance::GetMeshId( int32 lod, int32 comp, int32 mesh ) const
     {
         check( lod>=0 && lod<m_pD->m_lods.Num() );
         check( comp>=0 && comp<m_pD->m_lods[lod].m_components.Num() );
         check( mesh>=0 && mesh<m_pD->m_lods[lod].m_components[comp].m_meshes.Num() );
 
-        RESOURCE_ID result = m_pD->m_lods[lod].m_components[comp].m_meshes[mesh].m_meshId;
+		FResourceID result = m_pD->m_lods[lod].m_components[comp].m_meshes[mesh].m_meshId;
         return result;
     }
 
 
 	//---------------------------------------------------------------------------------------------
-    const char* Instance::GetMeshName( int lod, int comp, int mesh ) const
-	{
-		check( lod>=0 && lod<m_pD->m_lods.Num() );
-		check( comp>=0 && comp<m_pD->m_lods[lod].m_components.Num() );
-        check( mesh>=0 && mesh<m_pD->m_lods[lod].m_components[comp].m_meshes.Num() );
-
-        const char* strResult = m_pD->m_lods[lod].m_components[comp].m_meshes[mesh].m_name.c_str();
-		return strResult;
-	}
-
-
-	//---------------------------------------------------------------------------------------------
-    RESOURCE_ID Instance::GetImageId( int lod, int comp, int surf, int img ) const
+	FResourceID Instance::GetImageId( int32 lod, int32 comp, int32 surf, int32 img ) const
 	{
 		check( lod>=0 && lod<m_pD->m_lods.Num() );
 		check( comp>=0 && comp<m_pD->m_lods[lod].m_components.Num() );
         check( surf>=0 && surf<m_pD->m_lods[lod].m_components[comp].m_surfaces.Num() );
         check( img>=0 && img<m_pD->m_lods[lod].m_components[comp].m_surfaces[surf].m_images.Num() );
 
-        RESOURCE_ID result = m_pD->m_lods[lod].m_components[comp].m_surfaces[surf].m_images[img].m_imageId;
+		FResourceID result = m_pD->m_lods[lod].m_components[comp].m_surfaces[surf].m_images[img].m_imageId;
         return result;
 	}
 
 
 	//---------------------------------------------------------------------------------------------
-    const char* Instance::GetImageName( int lod, int comp, int surf, int img ) const
+    const char* Instance::GetImageName( int32 lod, int32 comp, int32 surf, int32 img ) const
 	{
 		check( lod>=0 && lod<m_pD->m_lods.Num() );
 		check( comp>=0 && comp<m_pD->m_lods[lod].m_components.Num() );
@@ -416,25 +349,19 @@ namespace mu
 
 
 	//---------------------------------------------------------------------------------------------
-    void Instance::GetVector( int lod, int comp, int surf, int vec,
-							  float* pX, float* pY, float* pZ, float* pW ) const
+	FVector4f Instance::GetVector( int32 lod, int32 comp, int32 surf, int32 vec ) const
 	{
 		check( lod>=0 && lod<m_pD->m_lods.Num() );
 		check( comp>=0 && comp<m_pD->m_lods[lod].m_components.Num() );
         check( surf>=0 && surf<m_pD->m_lods[lod].m_components[comp].m_surfaces.Num() );
         check( vec>=0 && vec<m_pD->m_lods[lod].m_components[comp].m_surfaces[surf].m_vectors.Num() );
 
-        FVector4f r = m_pD->m_lods[lod].m_components[comp].m_surfaces[surf].m_vectors[vec].m_vec;
-
-		if ( pX ) { *pX = r[0]; }
-		if ( pY ) { *pY = r[1]; }
-		if ( pZ ) { *pZ = r[2]; }
-		if ( pW ) { *pW = r[3]; }
+        return m_pD->m_lods[lod].m_components[comp].m_surfaces[surf].m_vectors[vec].m_vec;
 	}
 
 
 	//---------------------------------------------------------------------------------------------
-    const char* Instance::GetVectorName( int lod, int comp, int surf, int vec ) const
+    const char* Instance::GetVectorName( int32 lod, int32 comp, int32 surf, int32 vec ) const
 	{
 		check( lod>=0 && lod<m_pD->m_lods.Num() );
 		check( comp>=0 && comp<m_pD->m_lods[lod].m_components.Num() );
@@ -447,7 +374,7 @@ namespace mu
 
 
 	//---------------------------------------------------------------------------------------------
-    float Instance::GetScalar( int lod, int comp, int surf, int sca ) const
+    float Instance::GetScalar( int32 lod, int32 comp, int32 surf, int32 sca ) const
 	{
 		check( lod>=0 && lod<m_pD->m_lods.Num() );
 		check( comp>=0 && comp<m_pD->m_lods[lod].m_components.Num() );
@@ -460,7 +387,7 @@ namespace mu
 
 
 	//---------------------------------------------------------------------------------------------
-    const char* Instance::GetScalarName( int lod, int comp, int surf, int sca ) const
+    const char* Instance::GetScalarName( int32 lod, int32 comp, int32 surf, int32 sca ) const
 	{
 		check( lod>=0 && lod<m_pD->m_lods.Num() );
 		check( comp>=0 && comp<m_pD->m_lods[lod].m_components.Num() );
@@ -473,7 +400,7 @@ namespace mu
 
 
     //---------------------------------------------------------------------------------------------
-    const char* Instance::GetString( int lod, int comp, int surf, int str ) const
+    const char* Instance::GetString( int32 lod, int32 comp, int32 surf, int32 str ) const
     {
         check( lod >= 0 && lod < m_pD->m_lods.Num() );
         check( comp >= 0 && comp < m_pD->m_lods[lod].m_components.Num() );
@@ -500,7 +427,7 @@ namespace mu
 
 
     //---------------------------------------------------------------------------------------------
-    const char* Instance::GetStringName( int lod, int comp, int surf, int str ) const
+    const char* Instance::GetStringName( int32 lod, int32 comp, int32 surf, int32 str ) const
     {
         check( lod >= 0 && lod < m_pD->m_lods.Num() );
         check( comp >= 0 && comp < m_pD->m_lods[lod].m_components.Num() );
@@ -540,7 +467,7 @@ namespace mu
 
 
     //---------------------------------------------------------------------------------------------
-	int Instance::Private::AddLOD()
+	int32 Instance::Private::AddLOD()
 	{
 		int result = m_lods.Num();
 		m_lods.Add( INSTANCE_LOD() );
@@ -549,10 +476,10 @@ namespace mu
 
 
     //---------------------------------------------------------------------------------------------
-    int Instance::Private::AddComponent( int lod )
+    int32 Instance::Private::AddComponent( int lod )
     {
         // Automatically create the necessary lods and components
-        while ( lod>=GetLODCount() )
+        while (lod >= m_lods.Num())
         {
             AddLOD();
         }
@@ -564,14 +491,14 @@ namespace mu
 
 
     //---------------------------------------------------------------------------------------------
-    int Instance::Private::AddSurface( int lod, int comp )
+    int32 Instance::Private::AddSurface( int lod, int comp )
     {
         // Automatically create the necessary lods and components
-        while ( lod>=GetLODCount() )
+        while (lod >= m_lods.Num())
         {
             AddLOD();
         }
-        while ( comp>=GetComponentCount(lod) )
+        while (comp >= m_lods[lod].m_components.Num())
         {
             AddComponent(lod);
         }
@@ -583,14 +510,14 @@ namespace mu
 
 
     //---------------------------------------------------------------------------------------------
-    void Instance::Private::SetComponentName( int lod, int comp, const char* strName )
+    void Instance::Private::SetComponentName( int32 lod, int32 comp, const char* strName )
     {
         // Automatically create the necessary lods and components
-        while ( lod>=GetLODCount() )
+        while (lod >= m_lods.Num())
         {
             AddLOD();
         }
-        while ( comp>=GetComponentCount(lod) )
+        while (comp >= m_lods[lod].m_components.Num())
         {
             AddComponent( lod );
         }
@@ -608,18 +535,18 @@ namespace mu
 
 
     //---------------------------------------------------------------------------------------------
-    void Instance::Private::SetSurfaceName( int lod, int comp, int surf, const char* strName )
+    void Instance::Private::SetSurfaceName( int32 lod, int32 comp, int32 surf, const char* strName )
     {
         // Automatically create the necessary lods and components
-        while ( lod>=GetLODCount() )
+        while ( lod>=m_lods.Num() )
         {
             AddLOD();
         }
-        while ( comp>=GetComponentCount(lod) )
+        while (comp >= m_lods[lod].m_components.Num())
         {
             AddComponent( lod );
         }
-        while ( surf>=GetSurfaceCount(lod, comp) )
+        while ( surf>=m_lods[lod].m_components[comp].m_surfaces.Num() )
         {
             AddSurface( lod, comp );
         }
@@ -637,16 +564,14 @@ namespace mu
 
 
 	//---------------------------------------------------------------------------------------------
-    int Instance::Private::AddMesh( int lod, int comp,
-                                    RESOURCE_ID meshId,
-                                    const char* strName )
+    int32 Instance::Private::AddMesh( int32 lod, int32 comp, FResourceID meshId, const char* strName )
 	{
 		// Automatically create the necessary lods and components
-		while ( lod>=GetLODCount() )
+		while (lod >= m_lods.Num())
 		{
 			AddLOD();
 		}
-		while ( comp>=GetComponentCount(lod) )
+		while (comp >= m_lods[lod].m_components.Num())
 		{
 			AddComponent( lod );
 		}
@@ -660,23 +585,21 @@ namespace mu
 
 
 	//---------------------------------------------------------------------------------------------
-    int Instance::Private::AddImage( int lod, int comp, int surf,
-                                     RESOURCE_ID imageId,
-                                     const char* strName )
+    int32 Instance::Private::AddImage( int32 lod, int32 comp, int32 surf, FResourceID imageId, const char* strName )
 	{
 		// Automatically create the necessary lods and components
-		while ( lod>=GetLODCount() )
+		while (lod >= m_lods.Num())
 		{
 			AddLOD();
 		}
-		while ( comp>=GetComponentCount(lod) )
+		while (comp >= m_lods[lod].m_components.Num())
 		{
-			AddComponent( lod );
+			AddComponent(lod);
+		}
+		while (surf >= m_lods[lod].m_components[comp].m_surfaces.Num())
+		{
+			AddSurface(lod, comp);
 		}
-        while ( surf>=GetSurfaceCount(lod, comp) )
-        {
-            AddSurface( lod, comp );
-        }
 
         INSTANCE_SURFACE& surface = m_lods[lod].m_components[comp].m_surfaces[surf];
         int result = surface.m_images.Num();
@@ -687,21 +610,21 @@ namespace mu
 
 
 	//---------------------------------------------------------------------------------------------
-    int Instance::Private::AddVector( int lod, int comp, int surf, const FVector4f& vec, const char* strName )
+    int32 Instance::Private::AddVector( int32 lod, int32 comp, int32 surf, const FVector4f& vec, const char* strName )
 	{
 		// Automatically create the necessary lods and components
-		while ( lod>=GetLODCount() )
+		while (lod >= m_lods.Num())
 		{
 			AddLOD();
 		}
-		while ( comp>=GetComponentCount(lod) )
+		while (comp >= m_lods[lod].m_components.Num())
 		{
-			AddComponent( lod );
+			AddComponent(lod);
+		}
+		while (surf >= m_lods[lod].m_components[comp].m_surfaces.Num())
+		{
+			AddSurface(lod, comp);
 		}
-        while ( surf>=GetSurfaceCount(lod, comp) )
-        {
-            AddSurface( lod, comp );
-        }
 
         INSTANCE_SURFACE& surface = m_lods[lod].m_components[comp].m_surfaces[surf];
         int result = surface.m_vectors.Num();
@@ -712,21 +635,21 @@ namespace mu
 
 
     //---------------------------------------------------------------------------------------------
-    int Instance::Private::AddScalar( int lod, int comp, int surf, float sca, const char* strName )
+    int32 Instance::Private::AddScalar( int32 lod, int32 comp, int32 surf, float sca, const char* strName )
     {
         // Automatically create the necessary lods and components
-        while ( lod >= GetLODCount() )
-        {
-            AddLOD();
-        }
-        while ( comp >= GetComponentCount( lod ) )
-        {
-            AddComponent( lod );
-        }
-        while ( surf >= GetSurfaceCount( lod, comp ) )
-        {
-            AddSurface( lod, comp );
-        }
+		while (lod >= m_lods.Num())
+		{
+			AddLOD();
+		}
+		while (comp >= m_lods[lod].m_components.Num())
+		{
+			AddComponent(lod);
+		}
+		while (surf >= m_lods[lod].m_components[comp].m_surfaces.Num())
+		{
+			AddSurface(lod, comp);
+		}
 
         INSTANCE_SURFACE& surface = m_lods[lod].m_components[comp].m_surfaces[surf];
         int result = surface.m_scalars.Num();
@@ -737,22 +660,21 @@ namespace mu
 
 
     //---------------------------------------------------------------------------------------------
-    int Instance::Private::AddString(
-        int lod, int comp, int surf, const char* str, const char* strName )
+    int32 Instance::Private::AddString( int32 lod, int32 comp, int32 surf, const char* str, const char* strName )
     {
         // Automatically create the necessary lods and components
-        while ( lod >= GetLODCount() )
-        {
-            AddLOD();
-        }
-        while ( comp >= GetComponentCount( lod ) )
-        {
-            AddComponent( lod );
-        }
-        while ( surf >= GetSurfaceCount( lod, comp ) )
-        {
-            AddSurface( lod, comp );
-        }
+		while (lod >= m_lods.Num())
+		{
+			AddLOD();
+		}
+		while (comp >= m_lods[lod].m_components.Num())
+		{
+			AddComponent(lod);
+		}
+		while (surf >= m_lods[lod].m_components[comp].m_surfaces.Num())
+		{
+			AddSurface(lod, comp);
+		}
 
         INSTANCE_SURFACE& surface = m_lods[lod].m_components[comp].m_surfaces[surf];
         int result = surface.m_strings.Num();
diff --git a/Engine/Plugins/Experimental/Mutable/Source/MutableRuntime/Private/MuR/InstancePrivate.h b/Engine/Plugins/Experimental/Mutable/Source/MutableRuntime/Private/MuR/InstancePrivate.h
--- a/Engine/Plugins/Experimental/Mutable/Source/MutableRuntime/Private/MuR/InstancePrivate.h
+++ b/Engine/Plugins/Experimental/Mutable/Source/MutableRuntime/Private/MuR/InstancePrivate.h
@@ -11,16 +11,28 @@
 namespace mu
 {
 
+	/** Helper functions to make and read FResourceIDs */
+	inline FResourceID MakeResourceID(uint32 RootAddress, uint32 ParameterBlobIndex)
+	{
+		return (uint64(RootAddress) << 32) | uint64(ParameterBlobIndex);
+	}
+
+	inline uint32 GetResourceIDRoot(FResourceID Id)
+	{
+		return uint32(Id >> 32);
+	}
+
+	/** */
     struct INSTANCE_SURFACE
 	{
 		string m_name;
-        uint32_t InternalId=0;
-        uint32_t ExternalId =0;
-        uint32_t SharedId =0;
+        uint32 InternalId=0;
+        uint32 ExternalId =0;
+        uint32 SharedId =0;
 
 		struct IMAGE
 		{
-            IMAGE( RESOURCE_ID p, const char* strName )
+            IMAGE(FResourceID p, const char* strName )
 			{
                 m_imageId = p;
 				if (strName)
@@ -29,7 +41,7 @@ namespace mu
 				}
 			}
 
-            RESOURCE_ID m_imageId;
+			FResourceID m_imageId;
 			string m_name;
 		};
 
@@ -96,7 +108,7 @@ namespace mu
 
 		struct MESH
 		{
-            MESH(RESOURCE_ID p, const char* strName)
+            MESH(FResourceID p, const char* strName)
 			{
                 m_meshId = p;
 				if (strName)
@@ -105,7 +117,7 @@ namespace mu
 				}
 			}
 
-            RESOURCE_ID m_meshId;
+			FResourceID m_meshId;
 			string m_name;
 		};
 		TArray<MESH, TInlineAllocator<2>> m_meshes;
@@ -144,24 +156,16 @@ namespace mu
             m_id = 0;
         }
 
-		int GetLODCount() const;
-        int GetComponentCount( int lod ) const;
-        int GetSurfaceCount( int lod, int comp ) const;
-        int GetMeshCount( int lod, int comp ) const;
-        int GetImageCount( int lod, int comp, int surf ) const;
-        int GetVectorCount( int lod, int comp, int surf ) const;
-        int GetScalarCount( int lod, int comp, int surf ) const;
-        int GetStringCount( int lod, int comp, int surf ) const;
         int AddLOD();
         int AddComponent( int lod );
-        void SetComponentName( int lod, int comp, const char* strName );
-        int AddMesh(int lod, int comp, RESOURCE_ID, const char* strName);
+        void SetComponentName( int32 lod, int32 comp, const char* strName );
+        int AddMesh(int32 lod, int32 comp, FResourceID, const char* strName);
 		int AddSurface( int lod, int comp );
-        void SetSurfaceName( int lod, int comp, int surf, const char* strName );
-        int AddImage( int lod, int comp, int surf, RESOURCE_ID, const char* strName );
-        int AddVector( int lod, int comp, int surf, const FVector4f&, const char* strName );
-        int AddScalar( int lod, int comp, int surf, float, const char* strName );
-        int AddString( int lod, int comp, int surf, const char* strValue, const char* strName );
+        void SetSurfaceName( int32 lod, int32 comp, int32 surf, const char* strName );
+        int AddImage( int32 lod, int32 comp, int32 surf, FResourceID, const char* strName );
+        int32 AddVector( int32 lod, int32 comp, int32 surf, const FVector4f&, const char* strName );
+        int32 AddScalar( int32 lod, int32 comp, int32 surf, float, const char* strName );
+        int32 AddString( int32 lod, int32 comp, int32 surf, const char* strValue, const char* strName );
 		
 		// Data and Name must be non-null
 		void AddExtensionData(ExtensionDataPtrConst Data, const char* Name);
diff --git a/Engine/Plugins/Experimental/Mutable/Source/MutableRuntime/Private/MuR/Model.cpp b/Engine/Plugins/Experimental/Mutable/Source/MutableRuntime/Private/MuR/Model.cpp
--- a/Engine/Plugins/Experimental/Mutable/Source/MutableRuntime/Private/MuR/Model.cpp
+++ b/Engine/Plugins/Experimental/Mutable/Source/MutableRuntime/Private/MuR/Model.cpp
@@ -108,7 +108,7 @@ namespace mu
         //-----------------------------------------------------------------------------------------
         // Life cycle
         //-----------------------------------------------------------------------------------------
-        OutputModelStream( ModelStreamer* pStreamer )
+        OutputModelStream(ModelWriter* pStreamer )
             : m_pStreamer( pStreamer )
         {
         }
@@ -124,12 +124,12 @@ namespace mu
 
     private:
 
-        ModelStreamer* m_pStreamer;
+        ModelWriter* m_pStreamer;
     };
 
 
     //---------------------------------------------------------------------------------------------
-    void Model::Serialise( Model* p, ModelStreamer& streamer )
+    void Model::Serialise( Model* p, ModelWriter& streamer )
     {
 		LLM_SCOPE_BYNAME(TEXT("MutableRuntime"));
 
@@ -160,7 +160,7 @@ namespace mu
 			Image::Serialise(ResData.Value.get(), MemoryArch);
 			check(RomData.Size == MemStream.GetBufferSize());
 
-			streamer.OpenWriteFile(p->GetLocation(), RomData.Id);
+			streamer.OpenWriteFile(RomData.Id);
 			streamer.Write(MemStream.GetBuffer(), MemStream.GetBufferSize());
 			streamer.CloseWriteFile();
 
@@ -190,7 +190,7 @@ namespace mu
 			Mesh::Serialise(ResData.Value.get(), MemoryArch);
 			check(RomData.Size == MemStream.GetBufferSize());
 
-			streamer.OpenWriteFile(p->GetLocation(), RomData.Id);
+			streamer.OpenWriteFile(RomData.Id);
 			streamer.Write(MemStream.GetBuffer(), MemStream.GetBufferSize());
 			streamer.CloseWriteFile();
 
@@ -200,7 +200,7 @@ namespace mu
 
 		// Store the main data of the model
 		{
-			streamer.OpenWriteFile(p->GetLocation(), 0);
+			streamer.OpenWriteFile(0);
 			OutputModelStream stream(&streamer);
 			OutputArchive arch(&stream);
 
@@ -255,23 +255,6 @@ namespace mu
         return pResult;
     }
 
-
-    //---------------------------------------------------------------------------------------------
-    const char* Model::GetLocation( ) const
-    {
-        return m_pD->m_location.c_str();
-    }
-
-
-    //---------------------------------------------------------------------------------------------
-    void Model::SetLocation( const char* strLocation )
-    {
-        if (strLocation)
-        {
-            m_pD->m_location = strLocation;
-        }
-    }
-
     //---------------------------------------------------------------------------------------------
     Model::Private* Model::GetPrivate() const
     {
@@ -279,18 +262,6 @@ namespace mu
     }
 
 
-    //---------------------------------------------------------------------------------------------
-    void Model::ClearCaches()
-    {
-		LLM_SCOPE_BYNAME(TEXT("MutableRuntime"));
-
-        if (m_pD)
-        {
-            m_pD->m_generatedResources.Empty();
-        }
-    }
-
-
     //---------------------------------------------------------------------------------------------
     bool Model::GetBoolDefaultValue(int32 Index) const
     {
@@ -309,7 +280,7 @@ namespace mu
 
 
     //---------------------------------------------------------------------------------------------
-	int32 Model::GetIntDefaultValue(int Index) const
+	int32 Model::GetIntDefaultValue(int32 Index) const
 	{
 		check(m_pD->m_program.m_parameters.IsValidIndex(Index));
 		check(m_pD->m_program.m_parameters[Index].m_type == PARAMETER_TYPE::T_INT);
@@ -385,7 +356,7 @@ namespace mu
 
 
 	//---------------------------------------------------------------------------------------------
-    EXTERNAL_IMAGE_ID Model::GetImageDefaultValue(int32 Index) const
+    FExternalImageID Model::GetImageDefaultValue(int32 Index) const
     {
 	    check(m_pD->m_program.m_parameters.IsValidIndex(Index));
 		check(m_pD->m_program.m_parameters[Index].m_type == PARAMETER_TYPE::T_IMAGE);
@@ -578,243 +549,6 @@ namespace mu
     }
 
 
-    //---------------------------------------------------------------------------------------------
-    static size_t AddMultiValueKeys(TArray<uint8_t>& parameterValuesBlob, size_t pos,
-                                  const TMap< TArray<int>, PARAMETER_VALUE >& multi )
-    {
-        parameterValuesBlob.SetNum( pos+4 );
-        uint32 s = uint32( multi.Num() );
-        FMemory::Memcpy( &parameterValuesBlob[pos], &s, sizeof(uint32) );
-        pos+=4;
-
-        for(const auto& v: multi)
-        {
-            uint32 ds = uint32( v.Key.Num() );
-
-            parameterValuesBlob.SetNum( pos + 4 + ds*4 );
-
-			FMemory::Memcpy( &parameterValuesBlob[pos], &s, sizeof(uint32) );
-            pos+=4;
-
-			FMemory::Memcpy( &parameterValuesBlob[pos],v.Key.GetData(), ds*sizeof(int32) );
-            pos += ds*sizeof(int32);
-        }
-
-        return pos;
-    }
-
-
-    //---------------------------------------------------------------------------------------------
-    uint32 Model::Private::GetResourceKey( uint32 paramListIndex, OP::ADDRESS rootAt, const Parameters* pParams )
-    {
-		MUTABLE_CPUPROFILER_SCOPE(GetResourceKey)
-
-        // Find the list of relevant parameters
-        const TArray<uint16>* params = nullptr;
-        if (paramListIndex<(uint32)m_program.m_parameterLists.Num())
-        {
-            params = &m_program.m_parameterLists[paramListIndex];
-        }
-        check(params);
-        if (!params)
-        {
-            return 0xffff;
-        }
-
-        // Generate the relevant parameters blob
-		TArray<uint8_t> parameterValuesBlob;
-		parameterValuesBlob.Reserve(1024);
-        for (int param: *params)
-        {
-            int32 pos = parameterValuesBlob.Num();
-			int32 dataSize = 0;
-
-            switch(m_program.m_parameters[param].m_type)
-            {
-            case PARAMETER_TYPE::T_BOOL:
-                dataSize = 1;
-                parameterValuesBlob.Add( pParams->GetPrivate()->m_values[param].Get<ParamBoolType>() ? 1 : 0 );
-                pos += dataSize;
-
-                // Multi-values
-                if ( param < int(pParams->GetPrivate()->m_multiValues.Num()) )
-                {
-                    const auto& multi = pParams->GetPrivate()->m_multiValues[param];
-                    pos = AddMultiValueKeys( parameterValuesBlob, pos, multi );
-                    for(const auto& v: multi)
-                    {
-                        parameterValuesBlob.Add( v.Value.Get<ParamBoolType>() ? 1 : 0 );
-                        pos += dataSize;
-                    }
-                }
-                break;
-
-            case PARAMETER_TYPE::T_INT:
-                dataSize = sizeof(ParamIntType);
-                parameterValuesBlob.SetNum( pos+dataSize );
-                FMemory::Memcpy( &parameterValuesBlob[pos], &pParams->GetPrivate()->m_values[param].Get<ParamIntType>(), dataSize );
-                pos += dataSize;
-
-                // Multi-values
-                if ( param < int(pParams->GetPrivate()->m_multiValues.Num()) )
-                {
-                    const auto& multi = pParams->GetPrivate()->m_multiValues[param];
-                    pos = AddMultiValueKeys( parameterValuesBlob, pos, multi );
-                    parameterValuesBlob.SetNum( pos+multi.Num()*dataSize );
-                    for(const auto& v: multi)
-                    {
-						FMemory::Memcpy( &parameterValuesBlob[pos], &v.Value.Get<ParamIntType>(), dataSize );
-                        pos += dataSize;
-                    }
-                }
-                break;
-
-            //! Floating point value in the range of 0.0 to 1.0
-            case PARAMETER_TYPE::T_FLOAT:
-                dataSize = sizeof(ParamFloatType);
-                parameterValuesBlob.SetNum( pos+dataSize );
-				FMemory::Memcpy( &parameterValuesBlob[pos], &pParams->GetPrivate()->m_values[param].Get<ParamFloatType>(), dataSize );
-                pos += dataSize;
-
-                // Multi-values
-                if ( param < pParams->GetPrivate()->m_multiValues.Num() )
-                {
-                    const auto& multi = pParams->GetPrivate()->m_multiValues[param];
-                    pos = AddMultiValueKeys( parameterValuesBlob, pos, multi );
-                    parameterValuesBlob.SetNum( pos+multi.Num()*dataSize );
-                    for(const auto& v: multi)
-                    {
-						FMemory::Memcpy( &parameterValuesBlob[pos], &v.Value.Get<ParamFloatType>(),dataSize );
-                        pos += dataSize;
-                    }
-                }
-                break;
-
-            //! Floating point RGBA colour, with each channel ranging from 0.0 to 1.0
-            case PARAMETER_TYPE::T_COLOUR:
-                dataSize = sizeof(ParamColorType);
-                parameterValuesBlob.SetNum( pos+dataSize );
-				FMemory::Memcpy( &parameterValuesBlob[pos], &pParams->GetPrivate()->m_values[param].Get<ParamColorType>(), dataSize );
-                pos += dataSize;
-
-                // Multi-values
-                if ( param < int(pParams->GetPrivate()->m_multiValues.Num()) )
-                {
-                    const auto& multi = pParams->GetPrivate()->m_multiValues[param];
-                    pos = AddMultiValueKeys( parameterValuesBlob, pos, multi );
-                    parameterValuesBlob.SetNum( pos+multi.Num()*dataSize );
-                    for(const auto& v: multi)
-                    {
-						FMemory::Memcpy( &parameterValuesBlob[pos], &v.Value.Get<ParamColorType>(),dataSize );
-                        pos += dataSize;
-                    }
-                }
-                break;
-
-            //! 3D Projector type, defining a position, scale and orientation. Basically used for
-            //! projected decals.
-            case PARAMETER_TYPE::T_PROJECTOR:
-                dataSize = sizeof(ParamProjectorType);
-
-                // \todo: padding will be random?
-                parameterValuesBlob.SetNum( pos+dataSize );
-				FMemory::Memcpy( &parameterValuesBlob[pos], &pParams->GetPrivate()->m_values[param].Get<ParamProjectorType>(), dataSize );
-				pos += dataSize;
-
-                // Multi-values
-                if ( param < int(pParams->GetPrivate()->m_multiValues.Num()) )
-                {
-                    const auto& multi = pParams->GetPrivate()->m_multiValues[param];
-                    pos = AddMultiValueKeys( parameterValuesBlob, pos, multi );
-                    parameterValuesBlob.SetNum( pos+multi.Num()*dataSize );
-                    for(const auto& v: multi)
-                    {
-						FMemory::Memcpy( &parameterValuesBlob[pos], &v.Value.Get<ParamProjectorType>(),dataSize );
-                        pos += dataSize;
-                    }
-                }
-                break;
-
-            case PARAMETER_TYPE::T_IMAGE:
-                dataSize = sizeof(ParamImageType);
-                parameterValuesBlob.SetNum( pos+dataSize );
-				FMemory::Memcpy( &parameterValuesBlob[pos],
-                        &pParams->GetPrivate()->m_values[param].Get<ParamImageType>(),
-                        dataSize );
-				pos += dataSize;
-
-                // Multi-values
-                if ( param < int(pParams->GetPrivate()->m_multiValues.Num()) )
-                {
-                    const auto& multi = pParams->GetPrivate()->m_multiValues[param];
-                    pos = AddMultiValueKeys( parameterValuesBlob, pos, multi );
-                    parameterValuesBlob.SetNum( pos+multi.Num()*dataSize );
-                    for(const auto& v: multi)
-                    {
-						FMemory::Memcpy( &parameterValuesBlob[pos], &v.Value.Get<ParamImageType>(),dataSize );
-                        pos += dataSize;
-                    }
-                }
-                break;
-
-            default:
-                // unsupported parameter type
-                check(false);
-            }
-        }
-
-        // Increase the request id
-        ++m_lastResourceResquestId;
-
-        // See if we already have this id
-        size_t oldestCachePosition = 0;
-        for (size_t i=0; i<m_generatedResources.Num(); ++i)
-        {
-            auto& key = m_generatedResources[i];
-            if (key.m_rootAddress==rootAt
-				&&
-				key.m_parameterValuesBlob==parameterValuesBlob)
-            {
-                key.m_lastRequestId = m_lastResourceResquestId;
-                return key.m_id;
-            }
-            else
-            {
-                if ( m_generatedResources[oldestCachePosition].m_lastRequestId
-                     >
-                     key.m_lastRequestId )
-                {
-                    oldestCachePosition = i;
-                }
-            }
-        }
-
-        // Generate a new id
-        uint32 newId = ++m_lastResourceKeyId;
-        RESOURCE_KEY newKey;
-        newKey.m_id = newId;
-        newKey.m_lastRequestId = m_lastResourceResquestId;
-        newKey.m_rootAddress = rootAt;
-
-		// Don't MoveTemp so that we get a tight array on the destination.
-		//newKey.m_parameterValuesBlob = MoveTemp(parameterValuesBlob);
-		newKey.m_parameterValuesBlob = parameterValuesBlob;
-
-        // TODO: Move the constant to settings?
-        const size_t maxGeneratedResourcesIDCacheSize = 1024;
-        if (m_generatedResources.Num()>=maxGeneratedResourcesIDCacheSize)
-        {
-            m_generatedResources[oldestCachePosition] = MoveTemp(newKey);
-        }
-        else
-        {
-            m_generatedResources.Add(MoveTemp(newKey));
-        }
-
-        return newId;
-    }
-
-
     //---------------------------------------------------------------------------------------------
     //---------------------------------------------------------------------------------------------
     //---------------------------------------------------------------------------------------------
diff --git a/Engine/Plugins/Experimental/Mutable/Source/MutableRuntime/Private/MuR/ModelPrivate.h b/Engine/Plugins/Experimental/Mutable/Source/MutableRuntime/Private/MuR/ModelPrivate.h
--- a/Engine/Plugins/Experimental/Mutable/Source/MutableRuntime/Private/MuR/ModelPrivate.h
+++ b/Engine/Plugins/Experimental/Mutable/Source/MutableRuntime/Private/MuR/ModelPrivate.h
@@ -773,9 +773,6 @@ namespace mu
         //!
         FProgram m_program;
 
-        //! Non-persistent, streamer-specific location information
-        string m_location;
-
         //!
         void Serialise( OutputArchive& arch ) const
         {
@@ -795,41 +792,6 @@ namespace mu
             arch >> m_program;
         }
 
-        //! Management of generated resources
-        //! @{
-
-        //! This is used to uniquely identify a generated resource like meshes or images.
-        struct RESOURCE_KEY
-        {
-            //! The id assigned to the generated resource.
-            uint32 m_id;
-
-            //! The last request operation for this resource
-            uint32 m_lastRequestId;
-
-            //! The address that generated this resource
-            OP::ADDRESS m_rootAddress;
-
-            //! An opaque blob with the values of the relevant parameters
-			TArray<uint8> m_parameterValuesBlob;
-        };
-
-        //! The last id generated for a resource
-        uint32 m_lastResourceKeyId = 0;
-
-        //! The last id generated for a resource request. This is used to check the
-        //! relevancy of the resources when flushing the cache
-        uint32 m_lastResourceResquestId = 0;
-
-        //! Cached ids for returned assets
-        //! This is non-persistent runtime data
-		TArray<RESOURCE_KEY> m_generatedResources;
-
-        //! Get a resource key for a given resource with given parameter values.
-        uint32 GetResourceKey( uint32 ParamListIndex, OP::ADDRESS RootAt, const Parameters* );
-
-        //! @}
-
     };
 
 
diff --git a/Engine/Plugins/Experimental/Mutable/Source/MutableRuntime/Private/MuR/Parameters.cpp b/Engine/Plugins/Experimental/Mutable/Source/MutableRuntime/Private/MuR/Parameters.cpp
--- a/Engine/Plugins/Experimental/Mutable/Source/MutableRuntime/Private/MuR/Parameters.cpp
+++ b/Engine/Plugins/Experimental/Mutable/Source/MutableRuntime/Private/MuR/Parameters.cpp
@@ -758,7 +758,7 @@ namespace mu
 
 
     //---------------------------------------------------------------------------------------------
-    EXTERNAL_IMAGE_ID Parameters::GetImageValue( int index, const Ptr<const RangeIndex>& pos ) const
+    FExternalImageID Parameters::GetImageValue( int index, const Ptr<const RangeIndex>& pos ) const
     {
         check( index >= 0 && index < (int)m_pD->m_values.Num() );
         check( GetType( index ) == PARAMETER_TYPE::T_IMAGE );
@@ -799,7 +799,7 @@ namespace mu
 
 
     //---------------------------------------------------------------------------------------------
-    void Parameters::SetImageValue( int index, EXTERNAL_IMAGE_ID id, const Ptr<const RangeIndex>& pos )
+    void Parameters::SetImageValue( int index, FExternalImageID id, const Ptr<const RangeIndex>& pos )
     {
         check( index >= 0 && index < (int)m_pD->m_values.Num() );
         check( GetType( index ) == PARAMETER_TYPE::T_IMAGE );
diff --git a/Engine/Plugins/Experimental/Mutable/Source/MutableRuntime/Private/MuR/ParametersPrivate.h b/Engine/Plugins/Experimental/Mutable/Source/MutableRuntime/Private/MuR/ParametersPrivate.h
--- a/Engine/Plugins/Experimental/Mutable/Source/MutableRuntime/Private/MuR/ParametersPrivate.h
+++ b/Engine/Plugins/Experimental/Mutable/Source/MutableRuntime/Private/MuR/ParametersPrivate.h
@@ -20,9 +20,7 @@ namespace mu
     MUTABLE_DEFINE_ENUM_SERIALISABLE(PROJECTOR_TYPE)
 
 
-    //---------------------------------------------------------------------------------------------
-    //! Information to project an image on a mesh.
-    //---------------------------------------------------------------------------------------------
+    /** Description of a projector to project an image on a mesh. */
     struct FProjector
     {
         PROJECTOR_TYPE type = PROJECTOR_TYPE::PLANAR;
@@ -33,12 +31,12 @@ namespace mu
         float projectionAngle = 0.0f;
 
         //!
-        inline void GetDirectionSideUp(FVector3f& d, FVector3f& s, FVector3f& u ) const
+        inline void GetDirectionSideUp(FVector3f& OutDirection, FVector3f& OutSide, FVector3f& OutUp) const
         {
-            d = direction;
-            u = up;
-            s = FVector3f::CrossProduct( up, direction );
-            s.Normalize();
+			OutDirection = direction;
+            OutUp = up;
+            OutSide = FVector3f::CrossProduct( up, direction );
+			OutSide.Normalize();
         }
 
 
@@ -46,18 +44,10 @@ namespace mu
         void Serialise( OutputArchive& arch ) const
         {
             arch << type;
-            arch << position[0];
-            arch << position[1];
-            arch << position[2];
-            arch << direction[0];
-            arch << direction[1];
-            arch << direction[2];
-            arch << up[0];
-            arch << up[1];
-            arch << up[2];
-            arch << scale[0];
-            arch << scale[1];
-            arch << scale[2];
+            arch << position;
+            arch << direction;
+            arch << up;
+            arch << scale;
             arch << projectionAngle;
         }
 
@@ -65,18 +55,10 @@ namespace mu
         void Unserialise( InputArchive& arch )
         {
             arch >> type;
-            arch >> position[0];
-            arch >> position[1];
-            arch >> position[2];
-            arch >> direction[0];
-            arch >> direction[1];
-            arch >> direction[2];
-            arch >> up[0];
-            arch >> up[1];
-            arch >> up[2];
-            arch >> scale[0];
-            arch >> scale[1];
-            arch >> scale[2];
+            arch >> position;
+            arch >> direction;
+            arch >> up;
+            arch >> scale;
             arch >> projectionAngle;
         }
 
@@ -273,7 +255,7 @@ namespace mu
 	using ParamFloatType = float;
 	using ParamColorType = FVector3f;
 	using ParamProjectorType = FProjector;
-	using ParamImageType = EXTERNAL_IMAGE_ID;
+	using ParamImageType = FExternalImageID;
 	using ParamStringType = string;
 
 	
diff --git a/Engine/Plugins/Experimental/Mutable/Source/MutableRuntime/Private/MuR/System.cpp b/Engine/Plugins/Experimental/Mutable/Source/MutableRuntime/Private/MuR/System.cpp
--- a/Engine/Plugins/Experimental/Mutable/Source/MutableRuntime/Private/MuR/System.cpp
+++ b/Engine/Plugins/Experimental/Mutable/Source/MutableRuntime/Private/MuR/System.cpp
@@ -122,7 +122,7 @@ namespace mu
 
 
     //---------------------------------------------------------------------------------------------
-    void System::SetStreamingInterface(const TSharedPtr<ModelStreamer>& InInterface )
+    void System::SetStreamingInterface(const TSharedPtr<ModelReader>& InInterface )
     {
 		LLM_SCOPE_BYNAME(TEXT("MutableRuntime"));
        
@@ -130,9 +130,9 @@ namespace mu
     }
 
 
-    //---------------------------------------------------------------------------------------------
-    void System::SetWorkingMemoryBytes( uint64 InBytes )
-    {
+	//---------------------------------------------------------------------------------------------
+	void System::SetWorkingMemoryBytes(uint64 InBytes)
+	{
 		LLM_SCOPE_BYNAME(TEXT("MutableRuntime"));
 		MUTABLE_CPUPROFILER_SCOPE(SetWorkingMemoryBytes);
 
@@ -140,6 +140,16 @@ namespace mu
 		m_pD->WorkingMemoryManager.EnsureBudgetBelow(0);
 	}
 
+
+    //---------------------------------------------------------------------------------------------
+    void System::SetGeneratedCacheSize( uint32 InCount )
+    {
+		LLM_SCOPE_BYNAME(TEXT("MutableRuntime"));
+		MUTABLE_CPUPROFILER_SCOPE(SetGeneratedCacheSize);
+
+		m_pD->WorkingMemoryManager.MaxGeneratedResourceCacheSize = InCount;
+	}
+
 	
 	//---------------------------------------------------------------------------------------------
 	void System::ClearWorkingMemory()
@@ -262,9 +272,13 @@ namespace mu
 
 		OP::ADDRESS rootAt = pLiveInstance->Model->GetPrivate()->m_program.m_states[InStateIndex].m_root;
 
+		// Prepare instance cache
 		m_pD->PrepareCache(pLiveInstance->Model.Get(), InStateIndex);
 		pLiveInstance->OldParameters = InParams->Clone();
 
+		// Ensure the model cache has been created
+		m_pD->WorkingMemoryManager.FindOrAddModelCache(pLiveInstance->Model);
+
 		m_pD->RunCode(pLiveInstance->Model, InParams.get(), rootAt, InLodMask);
 
 		Ptr<const Instance> Result = pLiveInstance->Cache->GetInstance(FCacheAddress(rootAt, 0, 0));
@@ -287,7 +301,7 @@ namespace mu
 
 
 	//---------------------------------------------------------------------------------------------
-	Ptr<const Image> System::GetImage(Instance::ID instanceID, RESOURCE_ID imageId, int32 MipsToSkip, int32 InImageLOD)
+	Ptr<const Image> System::GetImage(Instance::ID instanceID, FResourceID ImageId, int32 MipsToSkip, int32 InImageLOD)
 	{
 		LLM_SCOPE_BYNAME(TEXT("MutableRuntime"));
 		MUTABLE_CPUPROFILER_SCOPE(SystemGetImage);
@@ -299,23 +313,13 @@ namespace mu
 		check(pLiveInstance);
 		m_pD->WorkingMemoryManager.CurrentInstanceCache = pLiveInstance->Cache;
 
-		// Find the resource id in the model's resource cache
-		for (const Model::Private::RESOURCE_KEY& res : pLiveInstance->Model->GetPrivate()->m_generatedResources)
-		{
-			if (res.m_id == imageId)
-			{
-				pResult = m_pD->BuildImage(pLiveInstance->Model,
-					pLiveInstance->OldParameters.get(),
-					res.m_rootAddress, MipsToSkip, InImageLOD);
-
-				// We always need to return something valid.
-				if (!pResult)
-				{
-					pResult = new mu::Image(16, 16, 1, EImageFormat::IF_RGBA_UBYTE, EInitializationType::Black);
-				}
+		OP::ADDRESS RootAddress = GetResourceIDRoot(ImageId);
+		pResult = m_pD->BuildImage(pLiveInstance->Model, pLiveInstance->OldParameters.get(), RootAddress, MipsToSkip, InImageLOD);
 
-				break;
-			}
+		// We always need to return something valid.
+		if (!pResult)
+		{
+			pResult = new mu::Image(16, 16, 1, EImageFormat::IF_RGBA_UBYTE, EInitializationType::Black);
 		}
 
 		m_pD->WorkingMemoryManager.CurrentInstanceCache = nullptr;
@@ -334,7 +338,7 @@ namespace mu
 
 
 	//---------------------------------------------------------------------------------------------
-	void System::GetImageDesc(Instance::ID instanceID, RESOURCE_ID imageId, FImageDesc& OutDesc)
+	void System::GetImageDesc(Instance::ID instanceID, FResourceID ImageId, FImageDesc& OutDesc)
 	{
 		LLM_SCOPE_BYNAME(TEXT("MutableRuntime"));
 		MUTABLE_CPUPROFILER_SCOPE(SystemGetImageDesc);
@@ -346,38 +350,30 @@ namespace mu
 		check(pLiveInstance);
 		m_pD->WorkingMemoryManager.CurrentInstanceCache = pLiveInstance->Cache;
 
-		// Find the resource id in the model's resource cache
-		for (const Model::Private::RESOURCE_KEY& res : pLiveInstance->Model->GetPrivate()->m_generatedResources)
-		{
-			if (res.m_id == imageId)
-			{
-				const mu::Model* Model = pLiveInstance->Model.Get();
-				const mu::FProgram& program = Model->GetPrivate()->m_program;
+		OP::ADDRESS RootAddress = GetResourceIDRoot(ImageId);
 
-				// TODO: It should be possible to reuse this data if cleared in the correct places only, together with m_heapImageDesc.
-				int32 VarValue = CVarClearImageDescCache.GetValueOnAnyThread();
-				if (VarValue != 0)
-				{
-					m_pD->WorkingMemoryManager.CurrentInstanceCache->ClearDescCache();
-				}
+		const mu::Model* Model = pLiveInstance->Model.Get();
+		const mu::FProgram& program = Model->GetPrivate()->m_program;
 
-				OP::ADDRESS at = res.m_rootAddress;
-				mu::OP_TYPE opType = program.GetOpType(at);
-				if (GetOpDataType(opType) == DT_IMAGE)
-				{
-					// GetImageDesc may call normal execution paths where meshes are computed.
-					m_pD->WorkingMemoryManager.BeginRunnerThread();
-					
-					int8 executionOptions = 0;
-					CodeRunner Runner(m_pD->Settings, m_pD, EExecutionStrategy::MinimizeMemory, pLiveInstance->Model, pLiveInstance->OldParameters.get(), at, System::AllLODs, executionOptions, 0, FScheduledOp::EType::ImageDesc);
-					Runner.Run();
-					Runner.GetImageDescResult(OutDesc);
+		// TODO: It should be possible to reuse this data if cleared in the correct places only, together with m_heapImageDesc.
+		int32 VarValue = CVarClearImageDescCache.GetValueOnAnyThread();
+		if (VarValue != 0)
+		{
+			m_pD->WorkingMemoryManager.CurrentInstanceCache->ClearDescCache();
+		}
 
-					m_pD->WorkingMemoryManager.EndRunnerThread();
-				}
+		mu::OP_TYPE opType = program.GetOpType(RootAddress);
+		if (GetOpDataType(opType) == DT_IMAGE)
+		{
+			// GetImageDesc may call normal execution paths where meshes are computed.
+			m_pD->WorkingMemoryManager.BeginRunnerThread();
+					
+			int8 executionOptions = 0;
+			CodeRunner Runner(m_pD->Settings, m_pD, EExecutionStrategy::MinimizeMemory, pLiveInstance->Model, pLiveInstance->OldParameters.get(), RootAddress, System::AllLODs, executionOptions, 0, FScheduledOp::EType::ImageDesc);
+			Runner.Run();
+			Runner.GetImageDescResult(OutDesc);
 
-				break;
-			}
+			m_pD->WorkingMemoryManager.EndRunnerThread();
 		}
 
 		m_pD->WorkingMemoryManager.CurrentInstanceCache = nullptr;
@@ -385,8 +381,7 @@ namespace mu
 
 
     //---------------------------------------------------------------------------------------------
-    MeshPtrConst System::GetMesh( Instance::ID instanceID,
-                                   RESOURCE_ID meshId )
+    MeshPtrConst System::GetMesh( Instance::ID instanceID, FResourceID MeshId )
     {
 		LLM_SCOPE_BYNAME(TEXT("MutableRuntime"));
 		MUTABLE_CPUPROFILER_SCOPE(SystemGetMesh);
@@ -398,23 +393,13 @@ namespace mu
 		check(pLiveInstance);
 		m_pD->WorkingMemoryManager.CurrentInstanceCache = pLiveInstance->Cache;
 
-		// Find the resource id in the model's resource cache
-		for (const Model::Private::RESOURCE_KEY& res : pLiveInstance->Model->GetPrivate()->m_generatedResources)
-		{
-			if (res.m_id == meshId)
-			{
-				pResult = m_pD->BuildMesh(pLiveInstance->Model,
-					pLiveInstance->OldParameters.get(),
-					res.m_rootAddress);
-
-				// If the mesh is null it means empty, but we still need to return a valid one
-				if (!pResult)
-				{
-					pResult = new Mesh();
-				}
+		OP::ADDRESS RootAddress = GetResourceIDRoot(MeshId);
+		pResult = m_pD->BuildMesh(pLiveInstance->Model, pLiveInstance->OldParameters.get(), RootAddress);
 
-				break;
-			}
+		// If the mesh is null it means empty, but we still need to return a valid one
+		if (!pResult)
+		{
+			pResult = new Mesh();
 		}
 
 		m_pD->WorkingMemoryManager.CurrentInstanceCache = nullptr;
@@ -552,8 +537,8 @@ namespace mu
             case OP_TYPE::IM_PARAMETER:
             {
 				OP::ParameterArgs args = program.GetOpArgs<OP::ParameterArgs>(at);
-				OP::ADDRESS param = args.variable;
-                Flags[param] = true;
+				OP::ADDRESS ParamIndex = args.variable;
+                Flags[ParamIndex] = true;
                 break;
             }
 
@@ -681,11 +666,11 @@ namespace mu
 
 
 	//---------------------------------------------------------------------------------------------
-	bool System::Private::BuildBool(const TSharedPtr<const Model>& pModel, const Parameters* pParams, OP::ADDRESS at)
+	bool System::Private::BuildBool(const TSharedPtr<const Model>& pModel, const Parameters* Params, OP::ADDRESS at)
 	{
 		WorkingMemoryManager.BeginRunnerThread();
 
-		RunCode(pModel, pParams, at);
+		RunCode(pModel, Params, at);
 
 		bool bResult = false;
 		if (!bUnrecoverableError)
@@ -700,11 +685,11 @@ namespace mu
 
 
 	//---------------------------------------------------------------------------------------------
-	float System::Private::BuildScalar(const TSharedPtr<const Model>& pModel, const Parameters* pParams, OP::ADDRESS at)
+	float System::Private::BuildScalar(const TSharedPtr<const Model>& pModel, const Parameters* Params, OP::ADDRESS at)
 	{
 		WorkingMemoryManager.BeginRunnerThread();
 
-		RunCode(pModel, pParams, at);
+		RunCode(pModel, Params, at);
 
 		float Result = 0.0f;		
 		if (!bUnrecoverableError)
@@ -719,11 +704,11 @@ namespace mu
 
 
 	//---------------------------------------------------------------------------------------------
-	int32 System::Private::BuildInt(const TSharedPtr<const Model>& pModel, const Parameters* pParams, OP::ADDRESS at)
+	int32 System::Private::BuildInt(const TSharedPtr<const Model>& pModel, const Parameters* Params, OP::ADDRESS at)
 	{
 		WorkingMemoryManager.BeginRunnerThread();
 
-		RunCode(pModel, pParams, at);
+		RunCode(pModel, Params, at);
 
 		int32 Result = 0;
 		if (!bUnrecoverableError)
@@ -738,7 +723,7 @@ namespace mu
 
 
 	//---------------------------------------------------------------------------------------------
-	FVector4f System::Private::BuildColour(const TSharedPtr<const Model>& pModel, const Parameters* pParams, OP::ADDRESS at)
+	FVector4f System::Private::BuildColour(const TSharedPtr<const Model>& pModel, const Parameters* Params, OP::ADDRESS at)
 	{
 		WorkingMemoryManager.BeginRunnerThread();
 
@@ -747,7 +732,7 @@ namespace mu
 		mu::OP_TYPE opType = pModel->GetPrivate()->m_program.GetOpType(at);
 		if (GetOpDataType(opType) == DT_COLOUR)
 		{
-			RunCode(pModel, pParams, at);
+			RunCode(pModel, Params, at);
 			if (!bUnrecoverableError)
 			{
 				Result = WorkingMemoryManager.CurrentInstanceCache->GetColour(FCacheAddress(at, 0, 0));
@@ -761,11 +746,11 @@ namespace mu
 
 	
 	//---------------------------------------------------------------------------------------------
-	FProjector System::Private::BuildProjector(const TSharedPtr<const Model>& pModel, const Parameters* pParams, OP::ADDRESS at)
+	FProjector System::Private::BuildProjector(const TSharedPtr<const Model>& pModel, const Parameters* Params, OP::ADDRESS at)
 	{
 		WorkingMemoryManager.BeginRunnerThread();
 
-    	RunCode(pModel, pParams, at);
+    	RunCode(pModel, Params, at);
 
 		FProjector Result;
 		if (!bUnrecoverableError)
@@ -781,7 +766,7 @@ namespace mu
 	
 	//---------------------------------------------------------------------------------------------
 	Ptr<const Image> System::Private::BuildImage(const TSharedPtr<const Model>& pModel,
-		const Parameters* pParams, OP::ADDRESS at, int32 MipsToSkip, int32 InImageLOD)
+		const Parameters* Params, OP::ADDRESS at, int32 MipsToSkip, int32 InImageLOD)
 	{
 		WorkingMemoryManager.BeginRunnerThread();
 
@@ -790,7 +775,7 @@ namespace mu
 		mu::OP_TYPE opType = pModel->GetPrivate()->m_program.GetOpType(at);
 		if (GetOpDataType(opType) == DT_IMAGE)
 		{
-			RunCode(pModel, pParams, at, System::AllLODs, uint8(MipsToSkip), InImageLOD);
+			RunCode(pModel, Params, at, System::AllLODs, uint8(MipsToSkip), InImageLOD);
 			if (!bUnrecoverableError)
 			{
 				Result = WorkingMemoryManager.LoadImage(FCacheAddress(at, 0, MipsToSkip), true);
@@ -804,7 +789,7 @@ namespace mu
 
 
 	//---------------------------------------------------------------------------------------------
-	MeshPtrConst System::Private::BuildMesh(const TSharedPtr<const Model>& pModel, const Parameters* pParams, OP::ADDRESS at)
+	MeshPtrConst System::Private::BuildMesh(const TSharedPtr<const Model>& pModel, const Parameters* Params, OP::ADDRESS at)
 	{
 		WorkingMemoryManager.BeginRunnerThread();
 
@@ -813,7 +798,7 @@ namespace mu
 		mu::OP_TYPE opType = pModel->GetPrivate()->m_program.GetOpType(at);
 		if (GetOpDataType(opType) == DT_MESH)
 		{
-			RunCode(pModel, pParams, at);
+			RunCode(pModel, Params, at);
 			if (!bUnrecoverableError)
 			{
 				Result = WorkingMemoryManager.LoadMesh(FCacheAddress(at, 0, 0), true);
@@ -827,7 +812,7 @@ namespace mu
 
 
 	//---------------------------------------------------------------------------------------------
-	Ptr<const Layout> System::Private::BuildLayout(const TSharedPtr<const Model>& pModel, const Parameters* pParams, OP::ADDRESS at)
+	Ptr<const Layout> System::Private::BuildLayout(const TSharedPtr<const Model>& pModel, const Parameters* Params, OP::ADDRESS at)
 	{
 		WorkingMemoryManager.BeginRunnerThread();
 
@@ -838,7 +823,7 @@ namespace mu
 			mu::OP_TYPE opType = pModel->GetPrivate()->m_program.GetOpType(at);
 			if (GetOpDataType(opType) == DT_LAYOUT)
 			{
-				RunCode(pModel, pParams, at);
+				RunCode(pModel, Params, at);
 				if (!bUnrecoverableError)
 				{
 					Result = WorkingMemoryManager.CurrentInstanceCache->GetLayout(FCacheAddress(at, 0, 0));
@@ -853,7 +838,7 @@ namespace mu
 
 
 	//---------------------------------------------------------------------------------------------
-	Ptr<const String> System::Private::BuildString(const TSharedPtr<const Model>& pModel, const Parameters* pParams, OP::ADDRESS at)
+	Ptr<const String> System::Private::BuildString(const TSharedPtr<const Model>& pModel, const Parameters* Params, OP::ADDRESS at)
 	{
 		WorkingMemoryManager.BeginRunnerThread();
 
@@ -864,7 +849,7 @@ namespace mu
 			mu::OP_TYPE opType = pModel->GetPrivate()->m_program.GetOpType(at);
 			if (GetOpDataType(opType) == DT_STRING)
 			{
-				RunCode(pModel, pParams, at);
+				RunCode(pModel, Params, at);
 				if (!bUnrecoverableError)
 				{
 					Result = WorkingMemoryManager.CurrentInstanceCache->GetString(FCacheAddress(at, 0, 0));
@@ -960,12 +945,12 @@ namespace mu
 		}
 
 		// Data stored per-model, but related to instance construction
-		for (const Model* Model : Models)
+		for (const FModelCacheEntry& CacheEntry:CachePerModel)
 		{
 			// Keys for resources generated so far
-			for (const Model::Private::RESOURCE_KEY& Key : Model->GetPrivate()->m_generatedResources)
+			for (const FModelCacheEntry::FGeneratedResourceData& Data : CacheEntry.GeneratedResources)
 			{
-				InternalBytes += Key.m_parameterValuesBlob.GetAllocatedSize();
+				InternalBytes += Data.ParameterValuesBlob.GetAllocatedSize();
 			}
 		}
 
@@ -999,27 +984,51 @@ namespace mu
 	}
 
 
+
+	//---------------------------------------------------------------------------------------------
+	FWorkingMemoryManager::FModelCacheEntry* FWorkingMemoryManager::FindModelCache(const Model* InModel)
+	{
+		check(InModel);
+
+		for (FModelCacheEntry& c : CachePerModel)
+		{
+			TSharedPtr<const Model> pCandidate = c.Model.Pin();
+			if (pCandidate)
+			{
+				if (pCandidate.Get() == InModel)
+				{
+					return &c;
+				}
+			}
+		}
+
+		return nullptr;
+	}
+
     //---------------------------------------------------------------------------------------------
-	FWorkingMemoryManager::FModelCacheEntry& FWorkingMemoryManager::GetModelCache(const TSharedPtr<const Model>& InModel )
+	FWorkingMemoryManager::FModelCacheEntry& FWorkingMemoryManager::FindOrAddModelCache(const TSharedPtr<const Model>& InModel)
     {
         check(InModel);
 
-        for(FModelCacheEntry& c:CachePerModel)
-        {
-			TSharedPtr<const Model> pCandidate = c.Model.Pin();
-            if (pCandidate)
-            {
-                if (pCandidate==InModel)
-                {
-                    return c;
-                }
-            }
-            else
-            {
-                // Free stray data. TODO: remove vector entry.
-                c.RomWeights.Empty();
-            }
-        }
+		// First clean stray data for models that may have been unloaded.
+		for (int32 CacheIndex=0; CacheIndex<CachePerModel.Num(); )
+		{
+			if (!CachePerModel[CacheIndex].Model.Pin())
+			{
+				CachePerModel.RemoveAtSwap(CacheIndex);
+			}
+			else
+			{
+				++CacheIndex;
+			}
+		}
+
+
+		FModelCacheEntry* ExistingCache = FindModelCache(InModel.Get());
+		if (ExistingCache)
+		{
+			return *ExistingCache;
+		}        
 
         // Not found. Add new
 		FModelCacheEntry n;
@@ -1456,7 +1465,7 @@ namespace mu
     {
         check(pModel);
 
-        FProgram& program = pModel->GetPrivate()->m_program;
+        FProgram& Program = pModel->GetPrivate()->m_program;
 
         // If budget is zero, we don't unload anything here, and we assume it is managed somewhere else.
         if (!BudgetBytes)
@@ -1468,16 +1477,260 @@ namespace mu
 
         // Update current cache
         {
-			FModelCacheEntry& ModelCache = GetModelCache(pModel);
+			FModelCacheEntry* ModelCache = FindModelCache(pModel.Get());
 
-            while (ModelCache.RomWeights.Num()<program.m_roms.Num())
+            if (ModelCache->RomWeights.Num()<Program.m_roms.Num())
             {
-				ModelCache.RomWeights.Add({ 0,0 });
+				ModelCache->RomWeights.SetNumZeroed(Program.m_roms.Num());
             }
 
-            ModelCache.RomWeights[romIndex].Key++;
-            ModelCache.RomWeights[romIndex].Value = RomTick;
+            ModelCache->RomWeights[romIndex].Key++;
+            ModelCache->RomWeights[romIndex].Value = RomTick;
         }
     }
 
+
+	//---------------------------------------------------------------------------------------------
+	static int32 AddMultiValueKeys(TArray<uint8_t>& Blob, int32 pos, const TMap< TArray<int>, PARAMETER_VALUE >& multi)
+	{
+		Blob.SetNum(pos + 4);
+		uint32 s = uint32(multi.Num());
+		FMemory::Memcpy(&Blob[pos], &s, sizeof(uint32));
+		pos += 4;
+
+		for (const auto& v : multi)
+		{
+			uint32 ds = uint32(v.Key.Num());
+
+			Blob.SetNum(pos + 4 + ds * 4);
+
+			FMemory::Memcpy(&Blob[pos], &s, sizeof(uint32));
+			pos += 4;
+
+			FMemory::Memcpy(&Blob[pos], v.Key.GetData(), ds * sizeof(int32));
+			pos += ds * sizeof(int32);
+		}
+
+		return pos;
+	}
+
+
+	//---------------------------------------------------------------------------------------------
+	FResourceID FWorkingMemoryManager::FModelCacheEntry::GetResourceKey(uint32 paramListIndex, OP::ADDRESS RootAt, const Parameters* Params, int32 InMaxResourceKeys)
+	{
+		MUTABLE_CPUPROFILER_SCOPE(GetResourceKey);
+
+		constexpr uint32 ErrorId = 0xffff;
+
+		const mu::Model* ThisModel = Model.Pin().Get();
+		check(ThisModel);
+		if (!ThisModel)
+		{
+			return ErrorId;
+		}
+
+		const FProgram& Program = ThisModel->GetPrivate()->m_program;
+
+		// Find the list of relevant parameters
+		const TArray<uint16>* RelevantParams = nullptr;
+		if (paramListIndex < (uint32)Program.m_parameterLists.Num())
+		{
+			RelevantParams = &Program.m_parameterLists[paramListIndex];
+		}
+		check(RelevantParams);
+		if (!RelevantParams)
+		{
+			return ErrorId;
+		}
+
+		// Generate the relevant parameters blob
+		TArray<uint8> Blob;
+		Blob.Reserve(2048);
+		for (int32 ParamIndex : *RelevantParams)
+		{
+			int32 pos = Blob.Num();
+			int32 dataSize = 0;
+
+			switch (Program.m_parameters[ParamIndex].m_type)
+			{
+			case PARAMETER_TYPE::T_BOOL:
+				dataSize = 1;
+				Blob.Add(Params->GetPrivate()->m_values[ParamIndex].Get<ParamBoolType>() ? 1 : 0);
+				pos += dataSize;
+
+				// Multi-values
+				if (ParamIndex < int(Params->GetPrivate()->m_multiValues.Num()))
+				{
+					const auto& multi = Params->GetPrivate()->m_multiValues[ParamIndex];
+					pos = AddMultiValueKeys(Blob, pos, multi);
+					for (const auto& v : multi)
+					{
+						Blob.Add(v.Value.Get<ParamBoolType>() ? 1 : 0);
+						pos += dataSize;
+					}
+				}
+				break;
+
+			case PARAMETER_TYPE::T_INT:
+				dataSize = sizeof(ParamIntType);
+				Blob.SetNum(pos + dataSize);
+				FMemory::Memcpy(&Blob[pos], &Params->GetPrivate()->m_values[ParamIndex].Get<ParamIntType>(), dataSize);
+				pos += dataSize;
+
+				// Multi-values
+				if (ParamIndex < int(Params->GetPrivate()->m_multiValues.Num()))
+				{
+					const auto& multi = Params->GetPrivate()->m_multiValues[ParamIndex];
+					pos = AddMultiValueKeys(Blob, pos, multi);
+					Blob.SetNum(pos + multi.Num() * dataSize);
+					for (const auto& v : multi)
+					{
+						FMemory::Memcpy(&Blob[pos], &v.Value.Get<ParamIntType>(), dataSize);
+						pos += dataSize;
+					}
+				}
+				break;
+
+				//! Floating point value in the range of 0.0 to 1.0
+			case PARAMETER_TYPE::T_FLOAT:
+				dataSize = sizeof(ParamFloatType);
+				Blob.SetNum(pos + dataSize);
+				FMemory::Memcpy(&Blob[pos], &Params->GetPrivate()->m_values[ParamIndex].Get<ParamFloatType>(), dataSize);
+				pos += dataSize;
+
+				// Multi-values
+				if (ParamIndex < Params->GetPrivate()->m_multiValues.Num())
+				{
+					const auto& multi = Params->GetPrivate()->m_multiValues[ParamIndex];
+					pos = AddMultiValueKeys(Blob, pos, multi);
+					Blob.SetNum(pos + multi.Num() * dataSize);
+					for (const auto& v : multi)
+					{
+						FMemory::Memcpy(&Blob[pos], &v.Value.Get<ParamFloatType>(), dataSize);
+						pos += dataSize;
+					}
+				}
+				break;
+
+				//! Floating point RGBA colour, with each channel ranging from 0.0 to 1.0
+			case PARAMETER_TYPE::T_COLOUR:
+				dataSize = sizeof(ParamColorType);
+				Blob.SetNum(pos + dataSize);
+				FMemory::Memcpy(&Blob[pos], &Params->GetPrivate()->m_values[ParamIndex].Get<ParamColorType>(), dataSize);
+				pos += dataSize;
+
+				// Multi-values
+				if (ParamIndex < int(Params->GetPrivate()->m_multiValues.Num()))
+				{
+					const auto& multi = Params->GetPrivate()->m_multiValues[ParamIndex];
+					pos = AddMultiValueKeys(Blob, pos, multi);
+					Blob.SetNum(pos + multi.Num() * dataSize);
+					for (const auto& v : multi)
+					{
+						FMemory::Memcpy(&Blob[pos], &v.Value.Get<ParamColorType>(), dataSize);
+						pos += dataSize;
+					}
+				}
+				break;
+
+				//! 3D Projector type, defining a position, scale and orientation. Basically used for
+				//! projected decals.
+			case PARAMETER_TYPE::T_PROJECTOR:
+				dataSize = sizeof(ParamProjectorType);
+
+				// \todo: padding will be random?
+				Blob.SetNum(pos + dataSize);
+				FMemory::Memcpy(&Blob[pos], &Params->GetPrivate()->m_values[ParamIndex].Get<ParamProjectorType>(), dataSize);
+				pos += dataSize;
+
+				// Multi-values
+				if (ParamIndex < int(Params->GetPrivate()->m_multiValues.Num()))
+				{
+					const auto& multi = Params->GetPrivate()->m_multiValues[ParamIndex];
+					pos = AddMultiValueKeys(Blob, pos, multi);
+					Blob.SetNum(pos + multi.Num() * dataSize);
+					for (const auto& v : multi)
+					{
+						FMemory::Memcpy(&Blob[pos], &v.Value.Get<ParamProjectorType>(), dataSize);
+						pos += dataSize;
+					}
+				}
+				break;
+
+			case PARAMETER_TYPE::T_IMAGE:
+				dataSize = sizeof(ParamImageType);
+				Blob.SetNum(pos + dataSize);
+				FMemory::Memcpy(&Blob[pos],
+					&Params->GetPrivate()->m_values[ParamIndex].Get<ParamImageType>(),
+					dataSize);
+				pos += dataSize;
+
+				// Multi-values
+				if (ParamIndex < int(Params->GetPrivate()->m_multiValues.Num()))
+				{
+					const auto& multi = Params->GetPrivate()->m_multiValues[ParamIndex];
+					pos = AddMultiValueKeys(Blob, pos, multi);
+					Blob.SetNum(pos + multi.Num() * dataSize);
+					for (const auto& v : multi)
+					{
+						FMemory::Memcpy(&Blob[pos], &v.Value.Get<ParamImageType>(), dataSize);
+						pos += dataSize;
+					}
+				}
+				break;
+
+			default:
+				// unsupported parameter type
+				check(false);
+			}
+		}
+
+		// Increase the request id
+		++LastResourceResquestId;
+
+		// See if we already have this id
+		int32 OldestCachePosition = 0;
+		for (int32 CacheIndex = 0; CacheIndex < GeneratedResources.Num(); ++CacheIndex)
+		{
+			FGeneratedResourceData& Data = GeneratedResources[CacheIndex];
+			if (GetResourceIDRoot(Data.Id) == RootAt
+				&&
+				Data.ParameterValuesBlob == Blob)
+			{
+				Data.LastRequestId = LastResourceResquestId;
+				return Data.Id;
+			}
+			else
+			{
+				if (GeneratedResources[OldestCachePosition].LastRequestId > Data.LastRequestId)
+				{
+					OldestCachePosition = CacheIndex;
+				}
+			}
+		}
+
+		// Generate a new id
+		uint32 NewBlobId = ++LastResourceKeyId;
+		FGeneratedResourceData NewKey;
+		NewKey.Id = MakeResourceID( RootAt, NewBlobId);
+		NewKey.LastRequestId = LastResourceResquestId;
+
+		// Don't MoveTemp so that we get a tight array on the destination.
+		//newKey.m_parameterValuesBlob = MoveTemp(Blob);
+		NewKey.ParameterValuesBlob = Blob;
+
+		// TODO: Move the constant to settings?
+		const size_t maxGeneratedResourcesIDCacheSize = InMaxResourceKeys;
+		if (GeneratedResources.Num() >= maxGeneratedResourcesIDCacheSize)
+		{
+			GeneratedResources[OldestCachePosition] = MoveTemp(NewKey);
+		}
+		else
+		{
+			GeneratedResources.Add(MoveTemp(NewKey));
+		}
+
+		return NewKey.Id;
+	}
+
 }
diff --git a/Engine/Plugins/Experimental/Mutable/Source/MutableRuntime/Private/MuR/SystemPrivate.h b/Engine/Plugins/Experimental/Mutable/Source/MutableRuntime/Private/MuR/SystemPrivate.h
--- a/Engine/Plugins/Experimental/Mutable/Source/MutableRuntime/Private/MuR/SystemPrivate.h
+++ b/Engine/Plugins/Experimental/Mutable/Source/MutableRuntime/Private/MuR/SystemPrivate.h
@@ -1104,6 +1104,38 @@ namespace mu
 
 			/** Count of pending operations for every rom index. */
 			TArray<uint16> PendingOpsPerRom;
+
+			//! Management of generated resources
+			//! @{
+
+			//! This is used to uniquely identify a generated resource like meshes or images.
+			struct FGeneratedResourceData
+			{
+				//! The id assigned to the generated resource.
+				FResourceID Id;
+
+				//! The last request operation for this resource
+				uint32 LastRequestId;
+
+				//! An opaque blob with the values of the relevant parameters
+				TArray<uint8> ParameterValuesBlob;
+			};
+
+			//! The last id generated for a resource
+			uint32 LastResourceKeyId = 0;
+
+			//! The last id generated for a resource request. This is used to check the
+			//! relevancy of the resources when flushing the cache
+			uint32 LastResourceResquestId = 0;
+
+			//! Cached ids for returned assets
+			//! This is non-persistent runtime data
+			TArray<FGeneratedResourceData> GeneratedResources;
+
+			//! Get a resource key for a given resource with given parameter values.
+			FResourceID GetResourceKey(uint32 ParamListIndex, OP::ADDRESS RootAt, const Parameters*, int32 InMaxResourceKeys);
+
+			//! @}
         };
 
 		/** Maximum working memory that mutable should be using. */
@@ -1112,6 +1144,9 @@ namespace mu
 		/** Maximum excess memory reached suring the current operation. */
 		int64 BudgetExcessBytes = 0;
 
+		/** Maximum number of resource keys that will be stored for resource reusal. */
+		int32 MaxGeneratedResourceCacheSize = 1024;
+
 		/** Signed ints to account for imprecisions especially with temp.*/
 		int32 TrackedBudgetBytes_Rom = 0;
 		// int32 TrackedBudgetBytes_Temp = 0;
@@ -1142,7 +1177,8 @@ namespace mu
 		TMap<Ptr<const Resource>, int32> CacheResources;
 
 		/** Given a mutable model, find or create its rom cache. */
-		FModelCacheEntry& GetModelCache(const TSharedPtr<const Model>&);
+		FModelCacheEntry* FindModelCache(const Model*);
+		FModelCacheEntry& FindOrAddModelCache(const TSharedPtr<const Model>&);
 
         /** Make sure the working memory is below the internal budget, even counting with the passed additional memory. 
 		* An optional function can be passed to "block" the unload of certain roms of data.
@@ -1874,7 +1910,7 @@ namespace mu
         Ptr<const Settings> Settings;
 
         //! Data streaming interface, if any.
-		TSharedPtr<ModelStreamer> StreamInterface;
+		TSharedPtr<ModelReader> StreamInterface;
 
 		TSharedPtr<ImageParameterGenerator> ImageParameterGenerator;
 
diff --git a/Engine/Plugins/Experimental/Mutable/Source/MutableRuntime/Public/MuR/Instance.h b/Engine/Plugins/Experimental/Mutable/Source/MutableRuntime/Public/MuR/Instance.h
--- a/Engine/Plugins/Experimental/Mutable/Source/MutableRuntime/Public/MuR/Instance.h
+++ b/Engine/Plugins/Experimental/Mutable/Source/MutableRuntime/Public/MuR/Instance.h
@@ -51,44 +51,44 @@ namespace mu
         Instance::ID GetId() const;
 
 		//! Get the number of levels-of-detail of this instance.
-		int GetLODCount() const;
+		int32 GetLODCount() const;
 
         //! Get the number of components in a level-of-detail
         //! \param lod Index of the level of detail, from 0 to GetLODCount()-1
-        int GetComponentCount( int lod ) const;
+        int32 GetComponentCount( int32 lod ) const;
 
         //! Get the name of a component
         //! \param lod Index of the level of detail, from 0 to GetLODCount()-1
         //! \param comp Index of the component, from 0 to GetComponentCount(lod)-1
-        const char* GetComponentName( int lod, int comp ) const;
+        const char* GetComponentName( int32 lod, int32 comp ) const;
 		
 		//! Get the Id of a component
 		//! \param lod Index of the level of detail, from 0 to GetLODCount()-1
 		//! \param comp Index of the component, from 0 to GetComponentCount(lod)-1
-		uint16 GetComponentId( int lod, int comp ) const;
+		uint16 GetComponentId( int32 lod, int32 comp ) const;
 
         //! Get the number of surfaces in a component
         //! \param lod Index of the level of detail, from 0 to GetLODCount()-1
         //! \param comp Index of the component, from 0 to GetComponentCount(lod)-1
-        int GetSurfaceCount( int lod, int comp ) const;
+        int32 GetSurfaceCount( int32 lod, int32 comp ) const;
 
         //! Get the name of a surface
         //! \param lod Index of the level of detail, from 0 to GetLODCount()-1
         //! \param comp Index of the component, from 0 to GetComponentCount(lod)-1
         //! \param surf Index of the surface, from 0 to GetSurfaceCount(lod,comp)-1
-        const char* GetSurfaceName( int lod, int comp, int surf ) const;
+        const char* GetSurfaceName( int32 lod, int32 comp, int32 surf ) const;
 
         //! Get an id that can be used to match the surface data with the mesh surface data.
         //! \param lod Index of the level of detail, from 0 to GetLODCount()-1
         //! \param comp Index of the component, from 0 to GetComponentCount(lod)-1
         //! \param surf Index of the surface, from 0 to GetSurfaceCount(lod,comp)-1
-        uint32_t GetSurfaceId( int lod, int comp, int surf ) const;
+        uint32 GetSurfaceId( int32 lod, int32 comp, int32 surf ) const;
 
         //! Find a surface index from the internal id (as returned by GetSurfaceId).
         //! \param lod Index of the level of detail, from 0 to GetLODCount()-1
         //! \param comp Index of the component, from 0 to GetComponentCount(lod)-1
         //! \param id ID of the surface to look for.
-        int FindSurfaceById( int lod, int comp, uint32_t id ) const;
+        int32 FindSurfaceById( int32 lod, int32 comp, uint32 id ) const;
 
 
 		//! Find the base surface index and Lod index when reusing surfaces between LODs. Return the surface index
@@ -114,12 +114,12 @@ namespace mu
 		//! \param lod Index of the level of detail, from 0 to GetLODCount()-1
 		//! \param comp Index of the component, from 0 to GetComponentCount(lod)-1
         //! \param surf Index of the surface, from 0 to GetSurfaceCount(lod,comp)-1
-        uint32_t GetSurfaceCustomId( int lod, int comp, int surf ) const;
+        uint32 GetSurfaceCustomId( int32 lod, int32 comp, int32 surf ) const;
 
         //! Get the number of meshes in a surface
 		//! \param lod Index of the level of detail, from 0 to GetLODCount()-1
 		//! \param comp Index of the component, from 0 to GetComponentCount(lod)-1
-        int GetMeshCount( int lod, int comp ) const;
+        int32 GetMeshCount( int32 lod, int32 comp ) const;
 
         //! Get a mesh resource id from a component
         //! \param lod Index of the level of detail, from 0 to GetLODCount()-1
@@ -128,19 +128,13 @@ namespace mu
         //! \return an ummodifiable pointer to the requested mesh. The returned object is guaranteed
         //! to be alive only while this instance is alive. The ownership of the returned object
         //! remains in the instance, so it should not be deleted.
-        RESOURCE_ID GetMeshId( int lod, int comp, int mesh ) const;
-
-		//! Get the name of a mesh in a component
-		//! \param lod Index of the level of detail, from 0 to GetLODCount()-1
-		//! \param comp Index of the component, from 0 to GetComponentCount(lod)-1
-        //! \param mesh Index of the mesh, from 0 to GetMeshCount(lod,comp,surf)-1
-        const char* GetMeshName( int lod, int comp, int mesh ) const;
+		FResourceID GetMeshId( int32 lod, int32 comp, int32 mesh ) const;
 
 		//! Get the number of images in a component
 		//! \param lod Index of the level of detail, from 0 to GetLODCount()
 		//! \param comp Index of the component, from 0 to GetComponentCount(lod)
         //! \param surf Index of the surface, from 0 to GetSurfaceCount(lod,comp)-1
-        int GetImageCount( int lod, int comp, int surf ) const;
+        int32 GetImageCount( int32 lod, int32 comp, int32 surf ) const;
 
         //! Get an image resource id from a component
 		//! \param lod Index of the level of detail, from 0 to GetLODCount()-1
@@ -150,20 +144,20 @@ namespace mu
         //! \return an ummodifiable pointer to the requested image. The returned object is
         //! guaranteed to be alive only while this instance is alive. The ownership of the returned
         //! object remains in the instance, so it should not be deleted.
-        RESOURCE_ID GetImageId( int lod, int comp, int surf, int img ) const;
+		FResourceID GetImageId( int32 lod, int32 comp, int32 surf, int32 img ) const;
 
 		//! Get the name of an image in a component
 		//! \param lod Index of the level of detail, from 0 to GetLODCount()-1
 		//! \param comp Index of the component, from 0 to GetComponentCount(lod)-1
         //! \param surf Index of the surface, from 0 to GetSurfaceCount(lod,comp)-1
         //! \param img Index of the image, from 0 to GetImageCount(lod,comp,surf)-1
-        const char* GetImageName( int lod, int comp, int surf, int img ) const;
+        const char* GetImageName( int32 lod, int32 comp, int32 surf, int32 img ) const;
 
 		//! Get the number of vectors in a component
 		//! \param lod Index of the level of detail, from 0 to GetLODCount()
 		//! \param comp Index of the component, from 0 to GetComponentCount(lod)
         //! \param surf Index of the surface, from 0 to GetSurfaceCount(lod,comp)-1
-        int GetVectorCount( int lod, int comp, int surf ) const;
+        int32 GetVectorCount( int32 lod, int32 comp, int32 surf ) const;
 
 		//! Get a vector from a component
 		//! \param lod Index of the level of detail, from 0 to GetLODCount()-1
@@ -172,55 +166,54 @@ namespace mu
         //! \param vec Index of the vector, from 0 to GetVectorCount(lod,comp,surf)-1
 		//! \param pX, pY, pZ, pW are optional to floats to store the components of the vector
 		//!		   value. They can be null.
-        void GetVector( int lod, int comp, int surf, int vec,
-						float* pX, float* pY, float* pZ, float* pW ) const;
+        FVector4f GetVector( int32 lod, int32 comp, int32 surf, int32 vec ) const;
 
 		//! Get the name of a vector in a component
 		//! \param lod Index of the level of detail, from 0 to GetLODCount()-1
 		//! \param comp Index of the component, from 0 to GetComponentCount(lod)-1
         //! \param surf Index of the surface, from 0 to GetSurfaceCount(lod,comp)-1
         //! \param vec Index of the vector, from 0 to GetVectorCount(lod,comp)-1
-        const char* GetVectorName( int lod, int comp, int surf, int vec ) const;
+        const char* GetVectorName( int32 lod, int32 comp, int32 surf, int32 vec ) const;
 
         //! Get the number of scalar values in a component
         //! \param lod Index of the level of detail, from 0 to GetLODCount()
         //! \param comp Index of the component, from 0 to GetComponentCount(lod)
         //! \param surf Index of the surface, from 0 to GetSurfaceCount(lod,comp)-1
-        int GetScalarCount( int lod, int comp, int surf ) const;
+        int32 GetScalarCount( int32 lod, int32 comp, int32 surf ) const;
 
         //! Get a scalar value from a component
         //! \param lod Index of the level of detail, from 0 to GetLODCount()-1
         //! \param comp Index of the component, from 0 to GetComponentCount(lod)-1
         //! \param surf Index of the surface, from 0 to GetSurfaceCount(lod,comp)-1
         //! \param sca Index of the scalar, from 0 to GetScalarCount(lod,comp)-1
-        float GetScalar( int lod, int comp, int surf, int sca ) const;
+        float GetScalar( int32 lod, int32 comp, int32 surf, int32 sca ) const;
 
         //! Get the name of a scalar from a component
         //! \param lod Index of the level of detail, from 0 to GetLODCount()-1
         //! \param comp Index of the component, from 0 to GetComponentCount(lod)-1
         //! \param surf Index of the surface, from 0 to GetSurfaceCount(lod,comp)-1
         //! \param sca Index of the scalar, from 0 to GetScalarCount(lod,comp)-1
-        const char* GetScalarName( int lod, int comp, int surf, int sca ) const;
+        const char* GetScalarName( int32 lod, int32 comp, int32 surf, int32 sca ) const;
 
         //! Get the number of string values in a component
         //! \param lod Index of the level of detail, from 0 to GetLODCount()
         //! \param comp Index of the component, from 0 to GetComponentCount(lod)
         //! \param surf Index of the surface, from 0 to GetSurfaceCount(lod,comp)-1
-        int GetStringCount( int lod, int comp, int surf ) const;
+        int32 GetStringCount( int32 lod, int32 comp, int32 surf ) const;
 
         //! Get a string value from a component
         //! \param lod Index of the level of detail, from 0 to GetLODCount()-1
         //! \param comp Index of the component, from 0 to GetComponentCount(lod)-1
         //! \param surf Index of the surface, from 0 to GetSurfaceCount(lod,comp)-1
         //! \param sca Index of the string, from 0 to GetStringCount(lod,comp)-1
-        const char* GetString( int lod, int comp, int surf, int str ) const;
+        const char* GetString( int32 lod, int32 comp, int32 surf, int32 str ) const;
 
         //! Get the name of a string from a component
         //! \param lod Index of the level of detail, from 0 to GetLODCount()-1
         //! \param comp Index of the component, from 0 to GetComponentCount(lod)-1
         //! \param surf Index of the surface, from 0 to GetSurfaceCount(lod,comp)-1
         //! \param sca Index of the string, from 0 to GetStringCount(lod,comp)-1
-        const char* GetStringName( int lod, int comp, int surf, int str ) const;
+        const char* GetStringName( int32 lod, int32 comp, int32 surf, int32 str ) const;
 
 		//! Get the number of ExtensionData values in a component
 		int32 GetExtensionDataCount() const;
diff --git a/Engine/Plugins/Experimental/Mutable/Source/MutableRuntime/Public/MuR/Model.h b/Engine/Plugins/Experimental/Mutable/Source/MutableRuntime/Public/MuR/Model.h
--- a/Engine/Plugins/Experimental/Mutable/Source/MutableRuntime/Public/MuR/Model.h
+++ b/Engine/Plugins/Experimental/Mutable/Source/MutableRuntime/Public/MuR/Model.h
@@ -20,7 +20,7 @@
 namespace mu
 {
 	class InputArchive;
-	class ModelStreamer;
+	class ModelWriter;
 	class OutputArchive;
 
 	typedef Ptr<Parameters> ParametersPtr;
@@ -60,18 +60,12 @@ namespace mu
 		//! Special serialise operation that serialises the data in separate "files". An object
         //! with the ModelStreamer interface is responsible of storing this data and providing
 		//! the "file" concept.
-        static void Serialise( Model* p, ModelStreamer& arch );
+        static void Serialise( Model* p, ModelWriter& arch );
 
 		//! Return true if the model has external data in other files. This kind of models will
 		//! require data streaming when used.
 		bool HasExternalData() const;
 
-        //! Get the streamer-specific location data.
-        //! This is a non-persistent string with meaning specific to every ModelStreamer
-        //! implementation.
-        const char* GetLocation( ) const;
-        void SetLocation( const char* strLocation );
-
 		//-----------------------------------------------------------------------------------------
 		// Own interface
 		//-----------------------------------------------------------------------------------------
@@ -102,11 +96,6 @@ namespace mu
 
         //! Free memory used by streaming data that may be loaded again when needed.
         void UnloadExternalData();
-
-        //! Free memory used in internal runtime caches. This is useful for long-running processes
-        //! that keep models loaded. This could be called when a game finishes, or a change of
-        //! level.
-        void ClearCaches();
     	
 		//! Return the default value of a boolean parameter.
 		//! \pre The parameter specified by index is a T_BOOL.
@@ -146,7 +135,7 @@ namespace mu
         //! \pre The parameter specified by index is a T_IMAGE.
         //! \param Index Index of the parameter from 0 to GetCount()-1
 		//! \return The externalId specified when setting the image value (\see SetImageValue)
-        EXTERNAL_IMAGE_ID GetImageDefaultValue(int32 Index) const;
+        FExternalImageID GetImageDefaultValue(int32 Index) const;
 
     	int32 GetRomCount() const;
 
diff --git a/Engine/Plugins/Experimental/Mutable/Source/MutableRuntime/Public/MuR/Parameters.h b/Engine/Plugins/Experimental/Mutable/Source/MutableRuntime/Public/MuR/Parameters.h
--- a/Engine/Plugins/Experimental/Mutable/Source/MutableRuntime/Public/MuR/Parameters.h
+++ b/Engine/Plugins/Experimental/Mutable/Source/MutableRuntime/Public/MuR/Parameters.h
@@ -351,13 +351,13 @@ namespace mu
         //! \param index Index of the parameter from 0 to GetCount()-1
 		//! \param pos Only for multidimensional parameters: relevant position to set in the ranges
 	    //! \return The externalId specified when setting the image value (\see SetImageValue)
-        EXTERNAL_IMAGE_ID GetImageValue( int index, const Ptr<const RangeIndex>& pos = nullptr) const;
+        FExternalImageID GetImageValue( int index, const Ptr<const RangeIndex>& pos = nullptr) const;
 
         //! If the parameter is of the image type, set its value.
         //! \param index Index of the parameter from 0 to GetCount()-1
         //! \param externalId Application-specific id used to identify this image during replication.
 		//! \param pos Only for multidimensional parameters: relevant position to set in the ranges
-		void SetImageValue( int index, EXTERNAL_IMAGE_ID externalId, const Ptr<const RangeIndex>& pos = nullptr);
+		void SetImageValue( int index, FExternalImageID externalId, const Ptr<const RangeIndex>& pos = nullptr);
 
         //! Return the value of a float parameter.
         //! \pre The parameter specified by index is a T_FLOAT.
diff --git a/Engine/Plugins/Experimental/Mutable/Source/MutableRuntime/Public/MuR/Serialisation.h b/Engine/Plugins/Experimental/Mutable/Source/MutableRuntime/Public/MuR/Serialisation.h
--- a/Engine/Plugins/Experimental/Mutable/Source/MutableRuntime/Public/MuR/Serialisation.h
+++ b/Engine/Plugins/Experimental/Mutable/Source/MutableRuntime/Public/MuR/Serialisation.h
@@ -72,13 +72,8 @@ namespace mu
 	};
 
 
-    //! \brief Interface of any class responsible of providing model data.
-    //!
-    //! When building a Model that was generated with split data, an object implementing this
-    //! interface must be provided to the System. The System will call this interface to fetch the
-    //! data it needs from the model.
-    //! \ingroup runtime
-    class MUTABLERUNTIME_API ModelStreamer : public Base
+	/** */
+	class MUTABLERUNTIME_API ModelReader : public Base
     {
     public:
 
@@ -87,7 +82,7 @@ namespace mu
         //-----------------------------------------------------------------------------------------
 
         //! Ensure virtual destruction.
-        virtual ~ModelStreamer() {}
+        virtual ~ModelReader() {}
 
         //-----------------------------------------------------------------------------------------
         // Reading interface
@@ -121,35 +116,50 @@ namespace mu
         //! any more to identify the same operation and becomes free.
         virtual void EndRead( OPERATION_ID ) = 0;
 
-        //-----------------------------------------------------------------------------------------
-        // Writing interface
-        //-----------------------------------------------------------------------------------------
+    };
 
-        //! \brief Open a file for writing data.
-        //!
-        //! Only one file can be open a time for writing. This method is only required when
-        //! writing model data from tools.
-        //! \param strModelName Name of the model where the file to open belongs to. This string's
-        //!         meaning depends on the implementation of the ModelStreamer subclasses. It could
-        //!         be a file path, or a resource identifier, etc. It will be the same for all
-        //!         files in a model.
-        //! \param key key identifying the model data fragment that is requested.
-        //!         This key interpretation depends on the implementation of the ModelStreamer,
-        virtual void OpenWriteFile( const char* strModelName, uint64 key0 ) = 0;
 
-        //! \brief Write a piece of data to the currently open file.
-        //!
-        //! There must be a file open with OpenWriteFile before calling this.
-        //! \param pBuffer pointer to the data that will be written to the file.
-        //! \param size size of the data to write on the file, in bytes.
-        virtual void Write( const void* pBuffer, uint64 size ) = 0;
+	/** */
+	class MUTABLERUNTIME_API ModelWriter : public Base
+	{
+	public:
 
-        //! \brief Close the file open for writing in a previous call to OpenWriteFile in this
-        //! object.
-        virtual void CloseWriteFile() = 0;
+		//-----------------------------------------------------------------------------------------
+		// Life cycle
+		//-----------------------------------------------------------------------------------------
+
+		//! Ensure virtual destruction.
+		virtual ~ModelWriter() {}
+
+		//-----------------------------------------------------------------------------------------
+		// Writing interface
+		//-----------------------------------------------------------------------------------------
+
+		//! \brief Open a file for writing data.
+		//!
+		//! Only one file can be open a time for writing. This method is only required when
+		//! writing model data from tools.
+		//! \param strModelName Name of the model where the file to open belongs to. This string's
+		//!         meaning depends on the implementation of the ModelStreamer subclasses. It could
+		//!         be a file path, or a resource identifier, etc. It will be the same for all
+		//!         files in a model.
+		//! \param key key identifying the model data fragment that is requested.
+		//!         This key interpretation depends on the implementation of the ModelStreamer,
+		virtual void OpenWriteFile(uint64 key0) = 0;
+
+		//! \brief Write a piece of data to the currently open file.
+		//!
+		//! There must be a file open with OpenWriteFile before calling this.
+		//! \param pBuffer pointer to the data that will be written to the file.
+		//! \param size size of the data to write on the file, in bytes.
+		virtual void Write(const void* pBuffer, uint64 size) = 0;
+
+		//! \brief Close the file open for writing in a previous call to OpenWriteFile in this
+		//! object.
+		virtual void CloseWriteFile() = 0;
 
 
-    };
+	};
 
 
 
diff --git a/Engine/Plugins/Experimental/Mutable/Source/MutableRuntime/Public/MuR/System.h b/Engine/Plugins/Experimental/Mutable/Source/MutableRuntime/Public/MuR/System.h
--- a/Engine/Plugins/Experimental/Mutable/Source/MutableRuntime/Public/MuR/System.h
+++ b/Engine/Plugins/Experimental/Mutable/Source/MutableRuntime/Public/MuR/System.h
@@ -50,7 +50,7 @@ namespace mu
 {
 	// Forward references
 	class Model;
-	class ModelStreamer;
+	class ModelReader;
 	class Parameters;
     class Mesh;
 	class ExtensionDataStreamer;
@@ -91,12 +91,12 @@ namespace mu
 
         //! Returns the completion event and a cleanup function that must be called once event is completed.
 #ifdef MUTABLE_USE_NEW_TASKGRAPH
-		virtual TTuple<UE::Tasks::FTask, TFunction<void()>> GetImageAsync(EXTERNAL_IMAGE_ID Id, uint8 MipmapsToSkip, TFunction<void(Ptr<Image>)>& ResultCallback) = 0;
+		virtual TTuple<UE::Tasks::FTask, TFunction<void()>> GetImageAsync(FExternalImageID Id, uint8 MipmapsToSkip, TFunction<void(Ptr<Image>)>& ResultCallback) = 0;
 #else
-		virtual TTuple<FGraphEventRef, TFunction<void()>> GetImageAsync(EXTERNAL_IMAGE_ID Id, uint8 MipmapsToSkip, TFunction<void(Ptr<Image>)>& ResultCallback) = 0;
+		virtual TTuple<FGraphEventRef, TFunction<void()>> GetImageAsync(FExternalImageID Id, uint8 MipmapsToSkip, TFunction<void(Ptr<Image>)>& ResultCallback) = 0;
 #endif
 
-        virtual mu::FImageDesc GetImageDesc(EXTERNAL_IMAGE_ID Id, uint8 MipmapsToSkip) = 0;
+        virtual mu::FImageDesc GetImageDesc(FExternalImageID Id, uint8 MipmapsToSkip) = 0;
     };
 
 
@@ -118,7 +118,7 @@ namespace mu
         System( const Ptr<Settings>& Settings = nullptr, const TSharedPtr<ExtensionDataStreamer>& DataStreamer = nullptr );
 
         //! Set a new provider for model data. 
-        void SetStreamingInterface(const TSharedPtr<ModelStreamer>& );
+        void SetStreamingInterface(const TSharedPtr<ModelReader>& );
 
         /** Set the working memory limit, overrding any set in the settings when the system was created.
          * Refer to Settings::SetWorkingMemoryBudget for more information.
@@ -130,6 +130,9 @@ namespace mu
 		*/
 		void ClearWorkingMemory();
 		
+		/** Set the amount of generated resources keys that will be stored for resource reusal. */
+		void SetGeneratedCacheSize(uint32 InCount);
+
         //! Set a new provider for external image data. This is only necessary if image parameters
         //! are used in the models.
         void SetImageParameterGenerator(const TSharedPtr<ImageParameterGenerator>& );
@@ -161,16 +164,16 @@ namespace mu
 
 		//! Only valid between BeginUpdate and EndUpdate
 		//! Calculate the description of an image, without generating it.
-		void GetImageDesc(Instance::ID InstanceID, RESOURCE_ID ImageId, FImageDesc& OutDesc );
+		void GetImageDesc(Instance::ID InstanceID, FResourceID ImageId, FImageDesc& OutDesc );
 
 		//! Only valid between BeginUpdate and EndUpdate
 		//! \param MipsToSkip Number of mips to skip compared from the full image.
 		//! If 0, all mip levels will be generated. If more levels than possible to discard are specified, 
 		//! the image will still contain a minimum number of mips specified at model compile time.
-		Ptr<const Image> GetImage(Instance::ID InstanceID, RESOURCE_ID ImageId, int32 MipsToSkip = 0, int32 LOD = 0);
+		Ptr<const Image> GetImage(Instance::ID InstanceID, FResourceID ImageId, int32 MipsToSkip = 0, int32 LOD = 0);
 
         //! Only valid between BeginUpdate and EndUpdate
-        Ptr<const Mesh> GetMesh(Instance::ID InstanceID, RESOURCE_ID MeshId);
+        Ptr<const Mesh> GetMesh(Instance::ID InstanceID, FResourceID MeshId);
 
         //! Invalidate and free the last Instance data returned by a call to BeginUpdate with
         //! the same instance index. After a call to this method, that Instance cannot be used any
diff --git a/Engine/Plugins/Experimental/Mutable/Source/MutableRuntime/Public/MuR/Types.h b/Engine/Plugins/Experimental/Mutable/Source/MutableRuntime/Public/MuR/Types.h
--- a/Engine/Plugins/Experimental/Mutable/Source/MutableRuntime/Public/MuR/Types.h
+++ b/Engine/Plugins/Experimental/Mutable/Source/MutableRuntime/Public/MuR/Types.h
@@ -8,12 +8,11 @@
 
 namespace mu
 {
+	/** Handle used to identify separate resources generated by mutable: images and meshes included in instances. */
+    typedef uint64 FResourceID;
 
-    typedef uint32            RESOURCE_ID;
-
-    //!
-    typedef FString           EXTERNAL_IMAGE_ID;
-
+    /** Unique identifier for images provided by code using mutable, of which mutable doesn't know much. */
+    typedef FString FExternalImageID;
 
 }
 
diff --git a/Engine/Plugins/Experimental/Mutable/Source/MutableTools/Private/MuT/CodeGenerator_Image.cpp b/Engine/Plugins/Experimental/Mutable/Source/MutableTools/Private/MuT/CodeGenerator_Image.cpp
--- a/Engine/Plugins/Experimental/Mutable/Source/MutableTools/Private/MuT/CodeGenerator_Image.cpp
+++ b/Engine/Plugins/Experimental/Mutable/Source/MutableTools/Private/MuT/CodeGenerator_Image.cpp
@@ -291,7 +291,7 @@ namespace mu
 			op->parameter.m_name = node.m_name;
 			op->parameter.m_uid = node.m_uid;
             op->parameter.m_type = PARAMETER_TYPE::T_IMAGE;
-        	op->parameter.m_defaultValue.Set<ParamImageType>(mu::EXTERNAL_IMAGE_ID());
+        	op->parameter.m_defaultValue.Set<ParamImageType>(mu::FExternalImageID());
 
 			// Generate the code for the ranges
 			for (int32 a = 0; a < node.m_ranges.Num(); ++a)
diff --git a/Engine/Plugins/Experimental/Mutable/Source/MutableTools/Private/MuT/Streams.cpp b/Engine/Plugins/Experimental/Mutable/Source/MutableTools/Private/MuT/Streams.cpp
--- a/Engine/Plugins/Experimental/Mutable/Source/MutableTools/Private/MuT/Streams.cpp
+++ b/Engine/Plugins/Experimental/Mutable/Source/MutableTools/Private/MuT/Streams.cpp
@@ -193,183 +193,73 @@ namespace mu
     //---------------------------------------------------------------------------------------------
     //---------------------------------------------------------------------------------------------
     //---------------------------------------------------------------------------------------------
-    class FileModelStreamer::Private : public Base
+    FileModelWriter::FileModelWriter(const FString& InModelLocation)
     {
-    public:
-        OPERATION_ID m_lastOperationID = 0;
-
-        IFileHandle* m_writeFile = nullptr;
-
-        // For debug
-        size_t m_totalWritten = 0;
-
-        static inline FString MakeDataFileName(
-                const FString& strModelName,
-                uint64_t key0
-                )
-        {
-			FString name = strModelName;
-			int extensionPos = 0;
-			if (name.FindLastChar('.', extensionPos))
-			{
-				name = name.Left(extensionPos);
-			}
-
-            FString temp = FString::Printf(TEXT(".%08x-%08x.mutable_data"),
-                               uint32_t((key0>>32) & 0xffffffff), uint32(key0 & 0xffffffff) );
-            name += temp;
-
-            return name;
-        }
-
-
-    };
-
-
-    //---------------------------------------------------------------------------------------------
-    FileModelStreamer::FileModelStreamer()
-    {
-        m_pD = new Private();
-    }
-
-
-    //---------------------------------------------------------------------------------------------
-    FileModelStreamer::~FileModelStreamer()
-    {
-        delete m_pD;
-    }
-
-
-    //---------------------------------------------------------------------------------------------
-    void FileModelStreamer::GetFileName( const FString& strModelName, uint64 key0, FString& outFileName )
-    {
-		outFileName = Private::MakeDataFileName( strModelName, key0 );
-    }
-
-
-    //---------------------------------------------------------------------------------------------
-    ModelStreamer::OPERATION_ID FileModelStreamer::BeginReadBlock(const mu::Model* Model, uint64 key0, void* pBuffer, uint64 size)
-    {
-        //UE_LOG(LogMutableCore,Log,"Streamer OpenRead");
-
-		check(Model);
-
-		OPERATION_ID id = ++m_pD->m_lastOperationID;
-
-        IFileHandle* pFile = nullptr;
-
-        if ( !key0 )
-        {
-            // It is the main model file
-            pFile = FPlatformFileManager::Get().GetPlatformFile().OpenRead( *FString(Model->GetLocation()) );
-
-            // Read the file identifying tag.
-            char tag[4];
-            size_t read = pFile->Read((uint8*)&tag[0], 4);
-            check( read==4 );
-            check( !memcmp(tag,MUTABLE_COMPILED_MODEL_FILETAG,4) );
-
-            uint32_t codeVersion = 0;
-            read = pFile->Read( (uint8*)&codeVersion, sizeof(uint32_t) );
-            check( read==4 );
-            check( codeVersion == MUTABLE_COMPILED_MODEL_CODE_VERSION  );
-
-        }
-        else
-        {
-            // It is a data file from a split model
-            FString name = m_pD->MakeDataFileName( Model->GetLocation(), key0);
-            pFile = FPlatformFileManager::Get().GetPlatformFile().OpenRead( *name );
-        }
-
-        if( !pFile )
-		{
-			return -1;
-		}
-
-        size_t read = 1;
-        size_t totalRead = 0;
-        while (read && totalRead<size)
-        {
-            read = pFile->Read( ((uint8*)pBuffer)+totalRead, (size_t)size-totalRead  );
-            totalRead += read;
-        }
-
-		delete pFile;
-
-        check( totalRead==size );
-
-        return id;
+		ModelLocation = InModelLocation;
     }
 
 
     //---------------------------------------------------------------------------------------------
-    bool FileModelStreamer::IsReadCompleted( OPERATION_ID )
+    void FileModelWriter::GetFileName( uint64 key0, FString& outFileName )
     {
-        return true;
+		outFileName = MakeDataFileName( ModelLocation, key0 );
     }
 
 
     //---------------------------------------------------------------------------------------------
-    void FileModelStreamer::EndRead( OPERATION_ID )
+    void FileModelWriter::OpenWriteFile(uint64 key0)
     {
-	}
-
-
-    //---------------------------------------------------------------------------------------------
-    void FileModelStreamer::OpenWriteFile( const char* strModelName, uint64 key0 )
-    {
-        check( !m_pD->m_writeFile );
+        check( !WriteFile );
 
         IFileHandle* pFile = nullptr;
-        m_pD->m_totalWritten = 0;
+        TotalWritten = 0;
 
         if ( !key0 )
         {
             // It is the main model file
-            pFile = FPlatformFileManager::Get().GetPlatformFile().OpenWrite( *FString(strModelName) );
+            pFile = FPlatformFileManager::Get().GetPlatformFile().OpenWrite( *ModelLocation );
 			check(pFile);
-            if (!pFile)
-                return;
+			if (!pFile)
+			{
+				return;
+			}
 
             // Write the file identifying tag.
-            bool success = pFile->Write( (const uint8*)MUTABLE_COMPILED_MODEL_FILETAG, 4 );
-            check(success);
+            bool bSuccess = pFile->Write( (const uint8*)MUTABLE_COMPILED_MODEL_FILETAG, 4 );
+            check(bSuccess);
 
-            uint32_t codeVersion = MUTABLE_COMPILED_MODEL_CODE_VERSION;
-			success = pFile->Write( (const uint8*)&codeVersion, sizeof(uint32_t) );
-            check(success);
+            uint32 codeVersion = MUTABLE_COMPILED_MODEL_CODE_VERSION;
+			bSuccess = pFile->Write( (const uint8*)&codeVersion, sizeof(uint32) );
+            check(bSuccess);
         }
         else
         {
             // It is a data file from a split model
-            FString name = m_pD->MakeDataFileName( strModelName, key0 );
+            FString name = MakeDataFileName(ModelLocation, key0 );
             pFile = FPlatformFileManager::Get().GetPlatformFile().OpenWrite( *name );
         }
 
         check( pFile );
 
-        m_pD->m_writeFile = pFile;
+        WriteFile = pFile;
     }
 
 
     //---------------------------------------------------------------------------------------------
-    void FileModelStreamer::Write( const void* pBuffer, uint64 size )
+    void FileModelWriter::Write( const void* pBuffer, uint64 size )
     {
-        check( m_pD->m_writeFile );
-
-        m_pD->m_writeFile->Write( ((const uint8*)pBuffer), size );
-        m_pD->m_totalWritten += size;
+        check( WriteFile );
+        WriteFile->Write( ((const uint8*)pBuffer), size );
+        TotalWritten += size;
     }
 
 
     //---------------------------------------------------------------------------------------------
-    void FileModelStreamer::CloseWriteFile()
+    void FileModelWriter::CloseWriteFile()
     {
-        check( m_pD->m_writeFile );
-
-        delete m_pD->m_writeFile;
-        m_pD->m_writeFile = 0;
+        check( WriteFile );
+        delete WriteFile;
+        WriteFile = nullptr;
     }
 
 
@@ -389,7 +279,6 @@ namespace mu
     class ProxyFactoryMutableSourceFile::Private : public Base
     {
     public:
-        // \todo: pimpl this out
         InputFileStream* m_pStream = nullptr;
         std::string m_filename;
     };
diff --git a/Engine/Plugins/Experimental/Mutable/Source/MutableTools/Public/MuT/Streams.h b/Engine/Plugins/Experimental/Mutable/Source/MutableTools/Public/MuT/Streams.h
--- a/Engine/Plugins/Experimental/Mutable/Source/MutableTools/Public/MuT/Streams.h
+++ b/Engine/Plugins/Experimental/Mutable/Source/MutableTools/Public/MuT/Streams.h
@@ -111,24 +111,19 @@ namespace mu
     //! Implementation of the streaming interface that actually loads the file synchronously when
     //! needed.
     //! \ingroup tools
-    class MUTABLETOOLS_API FileModelStreamer : public ModelStreamer
+    class MUTABLETOOLS_API FileModelWriter : public ModelWriter
     {
     public:
 
         //-----------------------------------------------------------------------------------------
         // Life cycle
         //-----------------------------------------------------------------------------------------
-        FileModelStreamer();
-        ~FileModelStreamer() override;
+		FileModelWriter( const FString& ModelLocation );
 
         //-----------------------------------------------------------------------------------------
         // ModelStreamer interface
         //-----------------------------------------------------------------------------------------
-		OPERATION_ID BeginReadBlock(const mu::Model*, uint64 key0, void* pBuffer, uint64 size ) override;
-        bool IsReadCompleted( OPERATION_ID ) override;
-        void EndRead( OPERATION_ID ) override;
-
-        void OpenWriteFile( const char* strModelName, uint64 key0 ) override;
+        void OpenWriteFile( uint64 key0 ) override;
         void Write( const void* pBuffer, uint64 size ) override;
         void CloseWriteFile() override;
 
@@ -137,13 +132,33 @@ namespace mu
         //-----------------------------------------------------------------------------------------
 
         //! Get the file name that this streamer would use for a specific model and data ID.
-        static void GetFileName(const FString& strModelName, uint64 key0, FString& outFileName );
+        void GetFileName(uint64 key0, FString& outFileName );
+
+		/** */
+		static inline FString MakeDataFileName(const FString& ModelName, uint64 key0)
+		{
+			FString Name = ModelName;
+			int32 ExtensionPos = 0;
+			if (Name.FindLastChar('.', ExtensionPos))
+			{
+				Name = Name.Left(ExtensionPos);
+			}
+
+			Name += FString::Printf(TEXT(".%08x-%08x.mutable_data"), uint32_t((key0 >> 32) & 0xffffffff), uint32(key0 & 0xffffffff));
+
+			return Name;
+		}
 
     private:
 
-        class Private;
-        Private* m_pD;
-    };
+		IFileHandle* WriteFile = nullptr;
+
+		FString ModelLocation;
+
+		// For debug
+		uint64 TotalWritten = 0;
+
+	};
 
 
     //! ProxyFactory that provides proxies for data stored in temporary files.
