[mutable] Use hash to generate binary blobs for image parameters.
Prevent some unnecesary string copies.
#rnx

diff --git a/Engine/Plugins/Experimental/Mutable/Source/CustomizableObject/Classes/MuCO/CustomizableObjectSystem.h b/Engine/Plugins/Experimental/Mutable/Source/CustomizableObject/Classes/MuCO/CustomizableObjectSystem.h
--- a/Engine/Plugins/Experimental/Mutable/Source/CustomizableObject/Classes/MuCO/CustomizableObjectSystem.h
+++ b/Engine/Plugins/Experimental/Mutable/Source/CustomizableObject/Classes/MuCO/CustomizableObjectSystem.h
@@ -166,17 +166,17 @@ public:
 
 	// Query that Mutable will run to find out if a texture will be provided as an Unreal UTexture2D,
 	// or as a raw data blob.
-	virtual ValueType HasTextureParameterValue(FString ID) { return ValueType::None; }
+	virtual ValueType HasTextureParameterValue(const FString& ID) { return ValueType::None; }
 
 	// In case IsTextureParameterValueUnreal returns false, this will be used to query the texture size data.
-	virtual FIntVector GetTextureParameterValueSize(FString ID) { return FIntVector(0, 0, 0); }
+	virtual FIntVector GetTextureParameterValueSize(const FString& ID) { return FIntVector(0, 0, 0); }
 
 	// In case IsTextureParameterValueUnreal returns false, this will be used to query the texture data that must
 	// be copied in the preallocated buffer. The pixel format is assumed to be 4-channel RGBA, uint8_t per channel.
-	virtual void GetTextureParameterValueData(FString ID, uint8* OutData) {}
+	virtual void GetTextureParameterValueData(const FString& ID, uint8* OutData) {}
 
 	// In case IsTextureParameterValueUnreal returns true, this will be used to query the texture.
-	virtual UTexture2D* GetTextureParameterValue(FString ID) { return nullptr; }
+	virtual UTexture2D* GetTextureParameterValue(const FString& ID) { return nullptr; }
 
 	// Used in the editor to show the list of available options.
 	// Only necessary if the images are required in editor previews.
diff --git a/Engine/Plugins/Experimental/Mutable/Source/CustomizableObject/Private/MuCO/CustomizableObject.cpp b/Engine/Plugins/Experimental/Mutable/Source/CustomizableObject/Private/MuCO/CustomizableObject.cpp
--- a/Engine/Plugins/Experimental/Mutable/Source/CustomizableObject/Private/MuCO/CustomizableObject.cpp
+++ b/Engine/Plugins/Experimental/Mutable/Source/CustomizableObject/Private/MuCO/CustomizableObject.cpp
@@ -1659,7 +1659,7 @@ FString UCustomizableObject::GetTextureParameterDefaultValue(const FString& InPa
 		return FCustomizableObjectTextureParameterValue::DEFAULT_PARAMETER_VALUE;
 	}
 	
-	return Model->GetImageDefaultValue(ParameterIndex);
+	return Model->GetImageDefaultValue(ParameterIndex).GetName();
 }
 
 
diff --git a/Engine/Plugins/Experimental/Mutable/Source/CustomizableObject/Private/MuCO/CustomizableObjectInstanceDescriptor.cpp b/Engine/Plugins/Experimental/Mutable/Source/CustomizableObject/Private/MuCO/CustomizableObjectInstanceDescriptor.cpp
--- a/Engine/Plugins/Experimental/Mutable/Source/CustomizableObject/Private/MuCO/CustomizableObjectInstanceDescriptor.cpp
+++ b/Engine/Plugins/Experimental/Mutable/Source/CustomizableObject/Private/MuCO/CustomizableObjectInstanceDescriptor.cpp
@@ -973,12 +973,12 @@ void FCustomizableObjectInstanceDescriptor::ReloadParameters()
 							}
 						}
 
-						Param.ParameterRangeValues[RangeIndex] = MutableParameters->GetImageValue(ParamIndex, RangeValueIdxPtr);
+						Param.ParameterRangeValues[RangeIndex] = MutableParameters->GetImageValue(ParamIndex, RangeValueIdxPtr).GetName();
 					}
 				}
 				else
 				{
-					Param.ParameterValue = MutableParameters->GetImageValue(ParamIndex);
+					Param.ParameterValue = MutableParameters->GetImageValue(ParamIndex).GetName();
 				}
 			}
 
diff --git a/Engine/Plugins/Experimental/Mutable/Source/CustomizableObject/Private/MuCO/DefaultImageProvider.cpp b/Engine/Plugins/Experimental/Mutable/Source/CustomizableObject/Private/MuCO/DefaultImageProvider.cpp
--- a/Engine/Plugins/Experimental/Mutable/Source/CustomizableObject/Private/MuCO/DefaultImageProvider.cpp
+++ b/Engine/Plugins/Experimental/Mutable/Source/CustomizableObject/Private/MuCO/DefaultImageProvider.cpp
@@ -34,7 +34,7 @@ namespace UDefaultImageProviderCVars
 }
 
 
-UCustomizableSystemImageProvider::ValueType UDefaultImageProvider::HasTextureParameterValue(FString ID)
+UCustomizableSystemImageProvider::ValueType UDefaultImageProvider::HasTextureParameterValue(const FString& ID)
 {
 	const TObjectPtr<UTexture2D>* Texture = Textures.Find(ID);
 	
@@ -44,7 +44,7 @@ UCustomizableSystemImageProvider::ValueType UDefaultImageProvider::HasTexturePar
 }
 
 
-UTexture2D* UDefaultImageProvider::GetTextureParameterValue(FString ID)
+UTexture2D* UDefaultImageProvider::GetTextureParameterValue(const FString& ID)
 {
 	return Textures[ID];
 }
diff --git a/Engine/Plugins/Experimental/Mutable/Source/CustomizableObject/Private/MuCO/UnrealMutableImageProvider.cpp b/Engine/Plugins/Experimental/Mutable/Source/CustomizableObject/Private/MuCO/UnrealMutableImageProvider.cpp
--- a/Engine/Plugins/Experimental/Mutable/Source/CustomizableObject/Private/MuCO/UnrealMutableImageProvider.cpp
+++ b/Engine/Plugins/Experimental/Mutable/Source/CustomizableObject/Private/MuCO/UnrealMutableImageProvider.cpp
@@ -117,9 +117,9 @@ mu::EImageFormat GetMutablePixelFormat(EPixelFormat InTextureFormat)
 
 //-------------------------------------------------------------------------------------------------
 #ifdef MUTABLE_USE_NEW_TASKGRAPH
-	TTuple<UE::Tasks::FTask, TFunction<void()>> FUnrealMutableImageProvider::GetImageAsync(mu::FExternalImageID id, uint8 MipmapsToSkip, TFunction<void(mu::Ptr<mu::Image>)>& ResultCallback)
+	TTuple<UE::Tasks::FTask, TFunction<void()>> FUnrealMutableImageProvider::GetImageAsync(const FExternalImageID& id, uint8 MipmapsToSkip, TFunction<void(mu::Ptr<mu::Image>)>& ResultCallback)
 #else
-	TTuple<FGraphEventRef, TFunction<void()>> FUnrealMutableImageProvider::GetImageAsync(mu::FExternalImageID id, uint8 MipmapsToSkip, TFunction<void(mu::Ptr<mu::Image>)>& ResultCallback)
+	TTuple<FGraphEventRef, TFunction<void()>> FUnrealMutableImageProvider::GetImageAsync(const FExternalImageID& id, uint8 MipmapsToSkip, TFunction<void(mu::Ptr<mu::Image>)>& ResultCallback)
 #endif
 {
 	// Thread: worker
@@ -393,7 +393,7 @@ mu::EImageFormat GetMutablePixelFormat(EPixelFormat InTextureFormat)
 
 
 // This should mantain parity with the descriptor of the images generated by GetImageAsync 
-mu::FImageDesc FUnrealMutableImageProvider::GetImageDesc(mu::FExternalImageID id, uint8 MipmapsToSkip) 
+mu::FImageDesc FUnrealMutableImageProvider::GetImageDesc(const FExternalImageID& id, uint8 MipmapsToSkip) 
 {
 	MUTABLE_CPUPROFILER_SCOPE(FUnrealMutableImageProvider::GetImageDesc);
 
@@ -463,7 +463,7 @@ mu::FImageDesc FUnrealMutableImageProvider::GetImageDesc(mu::FExternalImageID id
 }
 
 
-void FUnrealMutableImageProvider::CacheImage(const mu::FExternalImageID& Id, bool bUser)
+void FUnrealMutableImageProvider::CacheImage(const FExternalImageID& Id, bool bUser)
 {
 	check( IsInGameThread() );
 
@@ -498,27 +498,28 @@ void FUnrealMutableImageProvider::CacheImage(const mu::FExternalImageID& Id, boo
 
 				if (Provider.IsValid())
 				{
-					switch (Provider->HasTextureParameterValue(Id))
+					// \TODO: all these queries could probably be optimized into a single call.
+					switch (Provider->HasTextureParameterValue(Id.GetName()))
 					{
 
 					case UCustomizableSystemImageProvider::ValueType::Raw:
 					{
-						FIntVector desc = Provider->GetTextureParameterValueSize(Id);
+						FIntVector desc = Provider->GetTextureParameterValueSize(Id.GetName());
 						pResult = new mu::Image(desc[0], desc[1], 1, mu::EImageFormat::IF_RGBA_UBYTE, mu::EInitializationType::Black);
-						Provider->GetTextureParameterValueData(Id, pResult->GetData());
+						Provider->GetTextureParameterValueData(Id.GetName(), pResult->GetData());
 						break;
 					}
 
 					case UCustomizableSystemImageProvider::ValueType::Unreal:
 					{
-						UTexture2D* UnrealTexture = Provider->GetTextureParameterValue(Id);
+						UTexture2D* UnrealTexture = Provider->GetTextureParameterValue(Id.GetName());
 						pResult = ConvertTextureUnrealToMutable(UnrealTexture, 0);
 						break;
 					}
 
 					case UCustomizableSystemImageProvider::ValueType::Unreal_Deferred:
 					{
-						UnrealDeferredTexture = Provider->GetTextureParameterValue(Id);
+						UnrealDeferredTexture = Provider->GetTextureParameterValue(Id.GetName());
 						break;
 					}
 
@@ -530,7 +531,7 @@ void FUnrealMutableImageProvider::CacheImage(const mu::FExternalImageID& Id, boo
 
 			if (!pResult && !UnrealDeferredTexture)
 			{
-				UE_LOG(LogMutable, Warning, TEXT("Failed to cache external image %s."), *Id);
+				UE_LOG(LogMutable, Warning, TEXT("Failed to cache external image %s."), *Id.GetName());
 				pResult = CreateDummy();
 			}
 			
@@ -551,7 +552,7 @@ void FUnrealMutableImageProvider::CacheImage(const mu::FExternalImageID& Id, boo
 }
 
 
-void FUnrealMutableImageProvider::UnCacheImage(const mu::FExternalImageID& Id, bool bUser)
+void FUnrealMutableImageProvider::UnCacheImage(const FExternalImageID& Id, bool bUser)
 {
 	check(IsInGameThread());
 
@@ -585,7 +586,7 @@ void FUnrealMutableImageProvider::ClearCache(bool bUser)
 	check(IsInGameThread());
 
 	FScopeLock Lock(&ExternalImagesLock);
-	for (TTuple<FString, FUnrealMutableImageInfo> Tuple : GlobalExternalImages)
+	for (TTuple<FExternalImageID, FUnrealMutableImageInfo> Tuple : GlobalExternalImages)
 	{
 		UnCacheImage(Tuple.Key, bUser);
 	}	
@@ -603,7 +604,7 @@ void FUnrealMutableImageProvider::CacheImages(const mu::Parameters& Parameters)
 		}
 	
 		{
-			const FString TextureId = Parameters.GetImageValue(ParamIndex);
+			const FExternalImageID TextureId = Parameters.GetImageValue(ParamIndex);
 			CacheImage(TextureId, false);
 		}
 
@@ -611,7 +612,7 @@ void FUnrealMutableImageProvider::CacheImages(const mu::Parameters& Parameters)
 		for (int32 ValueIndex = 0; ValueIndex < NumValues; ++ValueIndex)
 		{
 			mu::Ptr<const mu::RangeIndex> Range = Parameters.GetValueIndex(ParamIndex, ValueIndex);
-			const FString TextureId = Parameters.GetImageValue(ParamIndex, Range);
+			const FExternalImageID TextureId = Parameters.GetImageValue(ParamIndex, Range);
 
 			CacheImage(TextureId, false);
 		}
@@ -630,7 +631,7 @@ void FUnrealMutableImageProvider::UnCacheImages(const mu::Parameters& Parameters
 		}
 	
 		{
-			const FString TextureId = Parameters.GetImageValue(ParamIndex);
+			const FExternalImageID TextureId = Parameters.GetImageValue(ParamIndex);
 			UnCacheImage(TextureId, false);				
 		}
 
@@ -638,7 +639,7 @@ void FUnrealMutableImageProvider::UnCacheImages(const mu::Parameters& Parameters
 		for (int32 ValueIndex = 0; ValueIndex < NumValues; ++ValueIndex)
 		{
 			mu::Ptr<const mu::RangeIndex> Range = Parameters.GetValueIndex(ParamIndex, ValueIndex);
-			const FString TextureId = Parameters.GetImageValue(ParamIndex, Range);
+			const FExternalImageID TextureId = Parameters.GetImageValue(ParamIndex, Range);
 
 			UnCacheImage(TextureId, false);
 		}
@@ -687,7 +688,7 @@ mu::FImageDesc FUnrealMutableImageProvider::CreateDummyDesc()
 
 void FUnrealMutableImageProvider::AddReferencedObjects(FReferenceCollector& Collector)
 {
-	for (TPair<mu::FExternalImageID, FUnrealMutableImageInfo>& Image : GlobalExternalImages)
+	for (TPair<FExternalImageID, FUnrealMutableImageInfo>& Image : GlobalExternalImages)
 	{
 		if (Image.Value.TextureToLoad)
 		{
diff --git a/Engine/Plugins/Experimental/Mutable/Source/CustomizableObject/Private/MuCO/UnrealMutableImageProvider.h b/Engine/Plugins/Experimental/Mutable/Source/CustomizableObject/Private/MuCO/UnrealMutableImageProvider.h
--- a/Engine/Plugins/Experimental/Mutable/Source/CustomizableObject/Private/MuCO/UnrealMutableImageProvider.h
+++ b/Engine/Plugins/Experimental/Mutable/Source/CustomizableObject/Private/MuCO/UnrealMutableImageProvider.h
@@ -23,23 +23,23 @@ public:
 	// mu::ImageParameterGenerator interface
 	// Thread: worker
 #ifdef MUTABLE_USE_NEW_TASKGRAPH
-	virtual TTuple<UE::Tasks::FTask, TFunction<void()>> GetImageAsync(mu::FExternalImageID Id, uint8 MipmapsToSkip, TFunction<void(mu::Ptr<mu::Image>)>& ResultCallback) override;
+	virtual TTuple<UE::Tasks::FTask, TFunction<void()>> GetImageAsync(const FExternalImageID& Id, uint8 MipmapsToSkip, TFunction<void(mu::Ptr<mu::Image>)>& ResultCallback) override;
 #else
-	virtual TTuple<FGraphEventRef, TFunction<void()>> GetImageAsync(mu::FExternalImageID Id, uint8 MipmapsToSkip, TFunction<void(mu::Ptr<mu::Image>)>& ResultCallback) override;
+	virtual TTuple<FGraphEventRef, TFunction<void()>> GetImageAsync(const FExternalImageID& Id, uint8 MipmapsToSkip, TFunction<void(mu::Ptr<mu::Image>)>& ResultCallback) override;
 #endif
 
-	virtual mu::FImageDesc GetImageDesc(mu::FExternalImageID Id, uint8 MipmapsToSkip) override;
+	virtual mu::FImageDesc GetImageDesc(const FExternalImageID& Id, uint8 MipmapsToSkip) override;
 
 	
 	// Own interface	
 	// Thread: Game
 	/** Add a reference to the image. If it was not cached it caches it.
 	 * @param bUser if true, adds a reference to the user reference counter. If false, adds a reference to the system reference counter. */
-	void CacheImage(const mu::FExternalImageID& Id, bool bUser);
+	void CacheImage(const FExternalImageID& Id, bool bUser);
 
 	/** Removes a reference to the image. If all references are removed, it uncaches the image.
 	 * @param bUser if true, removes a reference from the user reference counter. If false, removes a reference from the system reference counter. */
-	void UnCacheImage(const mu::FExternalImageID& Id, bool bUser);
+	void UnCacheImage(const FExternalImageID& Id, bool bUser);
 
 	/** Removes a reference to all images. All images which no longer have references will be uncached.
 	 * @param bUser if true, removes a references from the user reference counter. If false, removes a reference from the system reference counter. */
@@ -89,7 +89,7 @@ private:
 	* This is only safely written from the game thread protected by the following critical section, and it
 	* is safely read from the mutable thread during the update of the instance or texture mips
 	*/
-	TMap<FString, FUnrealMutableImageInfo> GlobalExternalImages;
+	TMap<FExternalImageID, FUnrealMutableImageInfo> GlobalExternalImages;
 	
 	/** Access to GlobalExternalImages must be protected with this because it may be accessed concurrently from the 
 	Game thread to modify it and from the Mutable thread to read it. */
diff --git a/Engine/Plugins/Experimental/Mutable/Source/CustomizableObject/Public/MuCO/DefaultImageProvider.h b/Engine/Plugins/Experimental/Mutable/Source/CustomizableObject/Public/MuCO/DefaultImageProvider.h
--- a/Engine/Plugins/Experimental/Mutable/Source/CustomizableObject/Public/MuCO/DefaultImageProvider.h
+++ b/Engine/Plugins/Experimental/Mutable/Source/CustomizableObject/Public/MuCO/DefaultImageProvider.h
@@ -24,8 +24,8 @@ class CUSTOMIZABLEOBJECT_API UDefaultImageProvider : public UCustomizableSystemI
 
 public:
 	// UCustomizableSystemImageProvider interface
-	virtual ValueType HasTextureParameterValue(FString ID) override;
-	virtual UTexture2D* GetTextureParameterValue(FString ID) override;
+	virtual ValueType HasTextureParameterValue(const FString& ID) override;
+	virtual UTexture2D* GetTextureParameterValue(const FString& ID) override;
 	virtual void GetTextureParameterValues(TArray<FCustomizableObjectExternalTexture>& OutValues) override;
 	
 	/** Add a Texture to the provider. */
diff --git a/Engine/Plugins/Experimental/Mutable/Source/CustomizableObjectEditor/Private/MuCOE/SMutableCodeViewer.cpp b/Engine/Plugins/Experimental/Mutable/Source/CustomizableObjectEditor/Private/MuCOE/SMutableCodeViewer.cpp
--- a/Engine/Plugins/Experimental/Mutable/Source/CustomizableObjectEditor/Private/MuCOE/SMutableCodeViewer.cpp
+++ b/Engine/Plugins/Experimental/Mutable/Source/CustomizableObjectEditor/Private/MuCOE/SMutableCodeViewer.cpp
@@ -2316,9 +2316,9 @@ namespace
 		
 	public:
 #ifdef MUTABLE_USE_NEW_TASKGRAPH
-		TTuple<UE::Tasks::FTask, TFunction<void()>> GetImageAsync(mu::FExternalImageID id, uint8 MipmapsToSkip, TFunction<void(mu::Ptr<mu::Image>)>& ResultCallback) override
+		TTuple<UE::Tasks::FTask, TFunction<void()>> GetImageAsync(const FExternalImageID& id, uint8 MipmapsToSkip, TFunction<void(mu::Ptr<mu::Image>)>& ResultCallback) override
 #else
-		TTuple<FGraphEventRef, TFunction<void()>> GetImageAsync(mu::FExternalImageID id, uint8 MipmapsToSkip, TFunction<void(mu::Ptr<mu::Image>)>& ResultCallback) override
+		TTuple<FGraphEventRef, TFunction<void()>> GetImageAsync(const FExternalImageID& id, uint8 MipmapsToSkip, TFunction<void(mu::Ptr<mu::Image>)>& ResultCallback) override
 #endif
 		{
 			MUTABLE_CPUPROFILER_SCOPE(TestImageProvider_GetImage);
@@ -2369,7 +2369,7 @@ namespace
 			return MakeTuple(CompletionEvent, []() -> void {});
 		}
 
-		mu::FImageDesc GetImageDesc(mu::FExternalImageID Id, uint8 MipmapsToSkip) override
+		mu::FImageDesc GetImageDesc(const FExternalImageID& Id, uint8 MipmapsToSkip) override
 		{
 			return IMAGE_DESC;
 		}
diff --git a/Engine/Plugins/Experimental/Mutable/Source/MutableRuntime/Private/MuR/System.cpp b/Engine/Plugins/Experimental/Mutable/Source/MutableRuntime/Private/MuR/System.cpp
--- a/Engine/Plugins/Experimental/Mutable/Source/MutableRuntime/Private/MuR/System.cpp
+++ b/Engine/Plugins/Experimental/Mutable/Source/MutableRuntime/Private/MuR/System.cpp
@@ -1559,11 +1559,11 @@ namespace mu
 				pos += dataSize;
 
 				// Multi-values
-				if (ParamIndex < int(Params->GetPrivate()->m_multiValues.Num()))
+				if (ParamIndex < Params->GetPrivate()->m_multiValues.Num())
 				{
-					const auto& multi = Params->GetPrivate()->m_multiValues[ParamIndex];
+					const TMap< TArray<int32>, PARAMETER_VALUE >& multi = Params->GetPrivate()->m_multiValues[ParamIndex];
 					pos = AddMultiValueKeys(Blob, pos, multi);
-					for (const auto& v : multi)
+					for (const TPair< TArray<int32>, PARAMETER_VALUE >& v : multi)
 					{
 						Blob.Add(v.Value.Get<ParamBoolType>() ? 1 : 0);
 						pos += dataSize;
@@ -1578,12 +1578,12 @@ namespace mu
 				pos += dataSize;
 
 				// Multi-values
-				if (ParamIndex < int(Params->GetPrivate()->m_multiValues.Num()))
+				if (ParamIndex < Params->GetPrivate()->m_multiValues.Num())
 				{
-					const auto& multi = Params->GetPrivate()->m_multiValues[ParamIndex];
+					const TMap< TArray<int32>, PARAMETER_VALUE >& multi = Params->GetPrivate()->m_multiValues[ParamIndex];
 					pos = AddMultiValueKeys(Blob, pos, multi);
 					Blob.SetNum(pos + multi.Num() * dataSize);
-					for (const auto& v : multi)
+					for (const TPair< TArray<int32>, PARAMETER_VALUE >& v : multi)
 					{
 						FMemory::Memcpy(&Blob[pos], &v.Value.Get<ParamIntType>(), dataSize);
 						pos += dataSize;
@@ -1601,10 +1601,10 @@ namespace mu
 				// Multi-values
 				if (ParamIndex < Params->GetPrivate()->m_multiValues.Num())
 				{
-					const auto& multi = Params->GetPrivate()->m_multiValues[ParamIndex];
+					const TMap< TArray<int32>, PARAMETER_VALUE >& multi = Params->GetPrivate()->m_multiValues[ParamIndex];
 					pos = AddMultiValueKeys(Blob, pos, multi);
 					Blob.SetNum(pos + multi.Num() * dataSize);
-					for (const auto& v : multi)
+					for (const TPair< TArray<int32>, PARAMETER_VALUE >& v : multi)
 					{
 						FMemory::Memcpy(&Blob[pos], &v.Value.Get<ParamFloatType>(), dataSize);
 						pos += dataSize;
@@ -1620,12 +1620,12 @@ namespace mu
 				pos += dataSize;
 
 				// Multi-values
-				if (ParamIndex < int(Params->GetPrivate()->m_multiValues.Num()))
+				if (ParamIndex < Params->GetPrivate()->m_multiValues.Num())
 				{
-					const auto& multi = Params->GetPrivate()->m_multiValues[ParamIndex];
+					const TMap< TArray<int32>, PARAMETER_VALUE >& multi = Params->GetPrivate()->m_multiValues[ParamIndex];
 					pos = AddMultiValueKeys(Blob, pos, multi);
 					Blob.SetNum(pos + multi.Num() * dataSize);
-					for (const auto& v : multi)
+					for (const TPair< TArray<int32>, PARAMETER_VALUE >& v : multi)
 					{
 						FMemory::Memcpy(&Blob[pos], &v.Value.Get<ParamColorType>(), dataSize);
 						pos += dataSize;
@@ -1644,12 +1644,12 @@ namespace mu
 				pos += dataSize;
 
 				// Multi-values
-				if (ParamIndex < int(Params->GetPrivate()->m_multiValues.Num()))
+				if (ParamIndex < Params->GetPrivate()->m_multiValues.Num())
 				{
-					const auto& multi = Params->GetPrivate()->m_multiValues[ParamIndex];
+					const TMap< TArray<int32>, PARAMETER_VALUE >& multi = Params->GetPrivate()->m_multiValues[ParamIndex];
 					pos = AddMultiValueKeys(Blob, pos, multi);
 					Blob.SetNum(pos + multi.Num() * dataSize);
-					for (const auto& v : multi)
+					for (const TPair< TArray<int32>, PARAMETER_VALUE >& v : multi)
 					{
 						FMemory::Memcpy(&Blob[pos], &v.Value.Get<ParamProjectorType>(), dataSize);
 						pos += dataSize;
@@ -1658,26 +1658,27 @@ namespace mu
 				break;
 
 			case PARAMETER_TYPE::T_IMAGE:
-				dataSize = sizeof(ParamImageType);
+			{
+				dataSize = sizeof(ParamImageType::HashType);
 				Blob.SetNum(pos + dataSize);
-				FMemory::Memcpy(&Blob[pos],
-					&Params->GetPrivate()->m_values[ParamIndex].Get<ParamImageType>(),
-					dataSize);
+				uint64 Hash = Params->GetPrivate()->m_values[ParamIndex].Get<ParamImageType>().GetHash();
+				FMemory::Memcpy(&Blob[pos], &Hash, dataSize);
 				pos += dataSize;
 
 				// Multi-values
-				if (ParamIndex < int(Params->GetPrivate()->m_multiValues.Num()))
+				if (ParamIndex < Params->GetPrivate()->m_multiValues.Num())
 				{
-					const auto& multi = Params->GetPrivate()->m_multiValues[ParamIndex];
+					const TMap< TArray<int32>, PARAMETER_VALUE >& multi = Params->GetPrivate()->m_multiValues[ParamIndex];
 					pos = AddMultiValueKeys(Blob, pos, multi);
 					Blob.SetNum(pos + multi.Num() * dataSize);
-					for (const auto& v : multi)
+					for (const TPair< TArray<int32>, PARAMETER_VALUE >& v : multi)
 					{
 						FMemory::Memcpy(&Blob[pos], &v.Value.Get<ParamImageType>(), dataSize);
 						pos += dataSize;
 					}
 				}
 				break;
+			}
 
 			default:
 				// unsupported parameter type
@@ -1716,12 +1717,10 @@ namespace mu
 		NewKey.LastRequestId = LastResourceResquestId;
 
 		// Don't MoveTemp so that we get a tight array on the destination.
-		//newKey.m_parameterValuesBlob = MoveTemp(Blob);
+		//NewKey.ParameterValuesBlob = MoveTemp(Blob);
 		NewKey.ParameterValuesBlob = Blob;
 
-		// TODO: Move the constant to settings?
-		const size_t maxGeneratedResourcesIDCacheSize = InMaxResourceKeys;
-		if (GeneratedResources.Num() >= maxGeneratedResourcesIDCacheSize)
+		if (GeneratedResources.Num() >= InMaxResourceKeys)
 		{
 			GeneratedResources[OldestCachePosition] = MoveTemp(NewKey);
 		}
diff --git a/Engine/Plugins/Experimental/Mutable/Source/MutableRuntime/Private/MuR/Types.cpp b/Engine/Plugins/Experimental/Mutable/Source/MutableRuntime/Private/MuR/Types.cpp
new file mode 100644
--- /dev/null
+++ b/Engine/Plugins/Experimental/Mutable/Source/MutableRuntime/Private/MuR/Types.cpp
@@ -0,0 +1,24 @@
+// Copyright Epic Games, Inc. All Rights Reserved.
+
+
+#include "MuR/Types.h"
+#include "MuR/Serialisation.h"
+#include "Hash/CityHash.h"
+
+
+FExternalImageID::FExternalImageID(const FString& InName)
+{
+	Name = InName;
+	NameHash = CityHash64(reinterpret_cast<const char*>(*Name),Name.Len()*sizeof(FString::ElementType));
+}
+
+void FExternalImageID::Serialise(mu::OutputArchive& arch) const
+{
+	arch << Name;
+}
+
+void FExternalImageID::Unserialise(mu::InputArchive& arch)
+{
+	arch >> Name;
+	NameHash = CityHash64(reinterpret_cast<const char*>(*Name), Name.Len() * sizeof(FString::ElementType));
+}
diff --git a/Engine/Plugins/Experimental/Mutable/Source/MutableRuntime/Public/MuR/Skeleton.h b/Engine/Plugins/Experimental/Mutable/Source/MutableRuntime/Public/MuR/Skeleton.h
--- a/Engine/Plugins/Experimental/Mutable/Source/MutableRuntime/Public/MuR/Skeleton.h
+++ b/Engine/Plugins/Experimental/Mutable/Source/MutableRuntime/Public/MuR/Skeleton.h
@@ -56,7 +56,7 @@ namespace mu
         //! Return the name of one of the bones used by this skeleton
 		//! \param boneIndex goes from 0 to GetBoneCount()-1
 		const char* GetBoneName(int32 boneIndex) const;
-        void SetBoneName( int b, const char* strName );
+        void SetBoneName( int32 b, const char* strName );
 
         //! Get and set the parent bone of each bone. The parent can be -1 if the bone is a root.
         int32 GetBoneParent(int32 boneIndex) const;
diff --git a/Engine/Plugins/Experimental/Mutable/Source/MutableRuntime/Public/MuR/System.h b/Engine/Plugins/Experimental/Mutable/Source/MutableRuntime/Public/MuR/System.h
--- a/Engine/Plugins/Experimental/Mutable/Source/MutableRuntime/Public/MuR/System.h
+++ b/Engine/Plugins/Experimental/Mutable/Source/MutableRuntime/Public/MuR/System.h
@@ -91,12 +91,12 @@ namespace mu
 
         //! Returns the completion event and a cleanup function that must be called once event is completed.
 #ifdef MUTABLE_USE_NEW_TASKGRAPH
-		virtual TTuple<UE::Tasks::FTask, TFunction<void()>> GetImageAsync(FExternalImageID Id, uint8 MipmapsToSkip, TFunction<void(Ptr<Image>)>& ResultCallback) = 0;
+		virtual TTuple<UE::Tasks::FTask, TFunction<void()>> GetImageAsync(const FExternalImageID& Id, uint8 MipmapsToSkip, TFunction<void(Ptr<Image>)>& ResultCallback) = 0;
 #else
-		virtual TTuple<FGraphEventRef, TFunction<void()>> GetImageAsync(FExternalImageID Id, uint8 MipmapsToSkip, TFunction<void(Ptr<Image>)>& ResultCallback) = 0;
+		virtual TTuple<FGraphEventRef, TFunction<void()>> GetImageAsync(const FExternalImageID& Id, uint8 MipmapsToSkip, TFunction<void(Ptr<Image>)>& ResultCallback) = 0;
 #endif
 
-        virtual mu::FImageDesc GetImageDesc(FExternalImageID Id, uint8 MipmapsToSkip) = 0;
+        virtual mu::FImageDesc GetImageDesc(const FExternalImageID& Id, uint8 MipmapsToSkip) = 0;
     };
 
 
diff --git a/Engine/Plugins/Experimental/Mutable/Source/MutableRuntime/Public/MuR/Types.h b/Engine/Plugins/Experimental/Mutable/Source/MutableRuntime/Public/MuR/Types.h
--- a/Engine/Plugins/Experimental/Mutable/Source/MutableRuntime/Public/MuR/Types.h
+++ b/Engine/Plugins/Experimental/Mutable/Source/MutableRuntime/Public/MuR/Types.h
@@ -5,15 +5,48 @@
 #include "HAL/Platform.h"
 #include "Containers/Map.h"
 #include "Containers/Array.h"
+#include "Containers/UnrealString.h"
 
 namespace mu
 {
+	class OutputArchive;
+	class InputArchive;
+
 	/** Handle used to identify separate resources generated by mutable: images and meshes included in instances. */
     typedef uint64 FResourceID;
+}
+
+/** Unique identifier for images provided by code using mutable, of which mutable doesn't know much. */
+struct FExternalImageID
+{
+public:
 
-    /** Unique identifier for images provided by code using mutable, of which mutable doesn't know much. */
-    typedef FString FExternalImageID;
+	typedef uint64 HashType;
 
+	FExternalImageID() {}
+	MUTABLERUNTIME_API FExternalImageID(const FString& InName);
+
+	inline bool operator==(const FExternalImageID& Other) const
+	{
+		return NameHash == Other.NameHash && Name == Other.Name;
+	}
+
+	void Serialise(mu::OutputArchive& arch) const;
+	void Unserialise(mu::InputArchive& arch);
+
+	inline const FString& GetName() const { return Name; }
+	inline HashType GetHash() const { return NameHash; }
+
+private:
+	FString Name;
+	HashType NameHash = 0;
+};
+
+
+/** */
+inline uint32 GetTypeHash(const FExternalImageID& Key)
+{
+	return GetTypeHash(Key.GetHash());
 }
 
 
diff --git a/Engine/Plugins/Experimental/Mutable/Source/MutableTools/Private/MuT/CodeGenerator_Image.cpp b/Engine/Plugins/Experimental/Mutable/Source/MutableTools/Private/MuT/CodeGenerator_Image.cpp
--- a/Engine/Plugins/Experimental/Mutable/Source/MutableTools/Private/MuT/CodeGenerator_Image.cpp
+++ b/Engine/Plugins/Experimental/Mutable/Source/MutableTools/Private/MuT/CodeGenerator_Image.cpp
@@ -291,7 +291,7 @@ namespace mu
 			op->parameter.m_name = node.m_name;
 			op->parameter.m_uid = node.m_uid;
             op->parameter.m_type = PARAMETER_TYPE::T_IMAGE;
-        	op->parameter.m_defaultValue.Set<ParamImageType>(mu::FExternalImageID());
+        	op->parameter.m_defaultValue.Set<ParamImageType>(FExternalImageID());
 
 			// Generate the code for the ranges
 			for (int32 a = 0; a < node.m_ranges.Num(); ++a)
