[mutable] Fix ompilation problem with log.
#rnx

diff --git a/Engine/Plugins/Experimental/Mutable/Source/CustomizableObject/Private/MuCO/UnrealMutableImageProvider.cpp b/Engine/Plugins/Experimental/Mutable/Source/CustomizableObject/Private/MuCO/UnrealMutableImageProvider.cpp
--- a/Engine/Plugins/Experimental/Mutable/Source/CustomizableObject/Private/MuCO/UnrealMutableImageProvider.cpp
+++ b/Engine/Plugins/Experimental/Mutable/Source/CustomizableObject/Private/MuCO/UnrealMutableImageProvider.cpp
@@ -164,7 +164,7 @@ mu::EImageFormat GetMutablePixelFormat(EPixelFormat InTextureFormat)
 		if (!GlobalExternalImages.Contains(id))
 		{
 			// Null case, no image was provided
-			UE_LOG(LogMutable, Warning, TEXT("Failed to get external image. GlobalExternalImage not found with id %s."), *id);
+			UE_LOG(LogMutable, Warning, TEXT("Failed to get external image. GlobalExternalImage not found."));
 
 			ResultCallback(CreateDummy());
 			return Invoke(TrivialReturn);
