[Mutable] Add normal morphing to morph operations.

- Now morphs use the normal data provided by unreal to morph the tangent space.
- Fix not released temporal mesh in mesh merge corner case (one of the meshes is null and the other does not  have verteices).

#rnx

diff --git a/Engine/Plugins/Experimental/Mutable/Source/CustomizableObjectEditor/Private/MuCOE/GenerateMutableSource/GenerateMutableSourceMesh.cpp b/Engine/Plugins/Experimental/Mutable/Source/CustomizableObjectEditor/Private/MuCOE/GenerateMutableSource/GenerateMutableSourceMesh.cpp
--- a/Engine/Plugins/Experimental/Mutable/Source/CustomizableObjectEditor/Private/MuCOE/GenerateMutableSource/GenerateMutableSourceMesh.cpp
+++ b/Engine/Plugins/Experimental/Mutable/Source/CustomizableObjectEditor/Private/MuCOE/GenerateMutableSource/GenerateMutableSourceMesh.cpp
@@ -2108,27 +2108,50 @@ mu::MeshPtr BuildMorphedMutableMesh(const UEdGraphPin* BaseSourcePin, const FStr
 
 			if (MorphTarget && MorphTarget->GetMorphLODModels().IsValidIndex(LODIndex))
 			{
-				int posBuf, posChannel;
-				MorphedSourceMesh->GetVertexBuffers().FindChannel(mu::MBS_POSITION, 0, &posBuf, &posChannel);
-				int posElemSize = MorphedSourceMesh->GetVertexBuffers().GetElementSize(posBuf);
-				int posOffset = MorphedSourceMesh->GetVertexBuffers().GetChannelOffset(posBuf, posChannel);
-				uint8* posBuffer = MorphedSourceMesh->GetVertexBuffers().GetBufferData(posBuf) + posOffset;
+				int32 PosBuf = -1;
+				int32 PosChannel = -1;
+				MorphedSourceMesh->GetVertexBuffers().FindChannel(mu::MBS_POSITION, 0, &PosBuf, &PosChannel);
+				check(PosBuf >= 0 && PosChannel >= 0);
 
-				uint32 MaterialVertexStart = (uint32)SkeletalMesh->GetImportedModel()->LODModels[LODIndex].Sections[SectionIndex].GetVertexBufferIndex();
-				uint32 MeshVertexCount = (uint32)MorphedSourceMesh->GetVertexBuffers().GetElementCount();
+				int32 PosElemSize = MorphedSourceMesh->GetVertexBuffers().GetElementSize(PosBuf);
+				int32 PosOffset = MorphedSourceMesh->GetVertexBuffers().GetChannelOffset(PosBuf, PosChannel);
+				uint8* PosBuffer = MorphedSourceMesh->GetVertexBuffers().GetBufferData(PosBuf) + PosOffset;
 
-				const TArray<FMorphTargetLODModel>& MorphLODModels = MorphTarget->GetMorphLODModels();
+				int32 NorBuf = -1;
+				int32 NorChannel = -1;
+				MorphedSourceMesh->GetVertexBuffers().FindChannel(mu::MBS_NORMAL, 0, &NorBuf, &NorChannel);
 
-				for (int v = 0; v < MorphLODModels[LODIndex].Vertices.Num(); ++v)
+				const bool bHasNormals = NorBuf >= 0 && NorChannel >= 0;
+
+				int32 NorElemSize = bHasNormals ? MorphedSourceMesh->GetVertexBuffers().GetElementSize(NorBuf) : 0;
+				int32 NorOffset	  = bHasNormals ? MorphedSourceMesh->GetVertexBuffers().GetChannelOffset(NorBuf, NorChannel) : 0;
+				uint8* NorBuffer  = bHasNormals ? MorphedSourceMesh->GetVertexBuffers().GetBufferData(NorBuf) + NorOffset : nullptr;
+
+				int32 MaterialVertexStart = SkeletalMesh->GetImportedModel()->LODModels[LODIndex].Sections[SectionIndex].GetVertexBufferIndex();
+				int32 MeshVertexCount = MorphedSourceMesh->GetVertexBuffers().GetElementCount();
+
+				const FMorphTargetLODModel& MorphLODModel = MorphTarget->GetMorphLODModels()[LODIndex];
+				for (const FMorphTargetDelta& MorphDelta : MorphLODModel.Vertices)
 				{
-					const FMorphTargetDelta& data = MorphLODModels[LODIndex].Vertices[v];
-					if ((data.SourceIdx >= MaterialVertexStart) &&
-						((data.SourceIdx - MaterialVertexStart) < MeshVertexCount))
+					const int32 VertexIndex = MorphDelta.SourceIdx - MaterialVertexStart;
+					if (VertexIndex >= 0 && VertexIndex < MeshVertexCount)
 					{
-						float* pPos = (float*)(posBuffer + posElemSize * (data.SourceIdx - MaterialVertexStart));
-						pPos[0] += data.PositionDelta[0];
-						pPos[1] += data.PositionDelta[1];
-						pPos[2] += data.PositionDelta[2];
+						{
+							float* const PosData = reinterpret_cast<float*>(PosBuffer + PosElemSize * VertexIndex);
+							const FVector3f MorphedPosition = FVector3f(PosData[0], PosData[1], PosData[2]) + MorphDelta.PositionDelta;
+							PosData[0] = MorphedPosition.X;
+							PosData[1] = MorphedPosition.Y;
+							PosData[2] = MorphedPosition.Z;
+						}
+
+						if (bHasNormals)
+						{
+							float* const NorData = reinterpret_cast<float*>(NorBuffer + NorElemSize * VertexIndex);
+							const FVector3f MorphedNormal = FVector3f(NorData[0], NorData[1], NorData[2]) + MorphDelta.TangentZDelta;
+							NorData[0] = MorphedNormal.X;
+							NorData[1] = MorphedNormal.Y;
+							NorData[2] = MorphedNormal.Z;
+						}
 					}
 				}
 			}
@@ -2705,6 +2728,7 @@ mu::NodeMeshPtr GenerateMorphMesh(const UEdGraphPin* Pin,
 			mu::NodeMeshMakeMorphPtr Morph = new mu::NodeMeshMakeMorph;
 			Morph->SetBase(BaseSourceMesh.get());
 			Morph->SetTarget(MorphedSourceMeshNodeTable.get());
+			Morph->SetOnlyPositionAndNormal(true);
 			Morph->SetMessageContext(MorphNode);
 
 			Result->SetMorph(Morph);
@@ -2725,6 +2749,7 @@ mu::NodeMeshPtr GenerateMorphMesh(const UEdGraphPin* Pin,
 			mu::NodeMeshMakeMorphPtr Morph = new mu::NodeMeshMakeMorph;
 			Morph->SetBase(BaseSourceMesh.get());
 			Morph->SetTarget(MorphedSourceMeshNode.get());
+			Morph->SetOnlyPositionAndNormal(true);
 			Morph->SetMessageContext(MorphNode);
 
 			Result->SetMorph(Morph);
@@ -2787,7 +2812,6 @@ mu::NodeMeshPtr GenerateMorphMesh(const UEdGraphPin* Pin,
 	return Result;
 }
 
-
 /** Create a default layout. Used when no layout is found. */
 mu::NodeLayoutBlocksPtr CreateDefaultLayout()
 {
diff --git a/Engine/Plugins/Experimental/Mutable/Source/MutableRuntime/Private/MuR/CodeRunner.cpp b/Engine/Plugins/Experimental/Mutable/Source/MutableRuntime/Private/MuR/CodeRunner.cpp
--- a/Engine/Plugins/Experimental/Mutable/Source/MutableRuntime/Private/MuR/CodeRunner.cpp
+++ b/Engine/Plugins/Experimental/Mutable/Source/MutableRuntime/Private/MuR/CodeRunner.cpp
@@ -1307,6 +1307,8 @@ namespace mu
                 }
                 else
                 {
+					Release(pA);
+					Release(pB);
 					StoreMesh(item, CreateMesh());
                 }
 
diff --git a/Engine/Plugins/Experimental/Mutable/Source/MutableRuntime/Private/MuR/OpMeshMorph.h b/Engine/Plugins/Experimental/Mutable/Source/MutableRuntime/Private/MuR/OpMeshMorph.h
--- a/Engine/Plugins/Experimental/Mutable/Source/MutableRuntime/Private/MuR/OpMeshMorph.h
+++ b/Engine/Plugins/Experimental/Mutable/Source/MutableRuntime/Private/MuR/OpMeshMorph.h
@@ -10,11 +10,11 @@
 
 namespace mu
 {
-	
+
 	//---------------------------------------------------------------------------------------------
 	//! Optimized linear factor version for morphing 2 targets
 	//---------------------------------------------------------------------------------------------
-    inline void MeshMorph2(Mesh* Result, const Mesh* pBase, const Mesh* pMin, const Mesh* pMax, const float factor, bool& bOutSuccess)
+    inline void MeshMorph2(Mesh* Result, const Mesh* pBase, const Mesh* pMin, const Mesh* pMax, const float Factor, bool& bOutSuccess)
     {
         MUTABLE_CPUPROFILER_SCOPE(MeshMorph2);
 		bOutSuccess = true;
@@ -25,97 +25,199 @@ namespace mu
 			return;
 		}
 
-        const auto ApplyMorph = []
-            ( auto BaseIdIter, const TArray< UntypedMeshBufferIterator >& baseChannelsIters, const int32 baseSize,
-              auto MorphIdIter, const TArray< UntypedMeshBufferIteratorConst >& morphChannelsIters, const int32 morphSize,
-              const float factor )
-        -> void
-        {
-            uint32 minBaseId = std::numeric_limits<uint32>::max();
-            uint32 maxBaseId = 0;
-
+		const auto MakeIndexMap = [](
+				UntypedMeshBufferIteratorConst BaseIdIter, int32 BaseNum,
+				MeshBufferIteratorConst<MBF_UINT32, uint32, 1> MorphIdIter, int32 MorphNum) 
+		-> SparseIndexMap
+		{	
+            uint32 MinBaseId = TNumericLimits<uint32>::Max();
+            uint32 MaxBaseId = 0;
 			{
-                auto limitsBaseIdIter = BaseIdIter;
-				for ( int32 i = 0; i < baseSize; ++i, ++limitsBaseIdIter )
+				for (int32 Index = 0; Index < BaseNum; ++Index, ++BaseIdIter)
 				{
-					const uint32 id = limitsBaseIdIter.GetAsUINT32();
-					minBaseId = FMath::Min( id, minBaseId );
-					maxBaseId = FMath::Max( id, maxBaseId );
+					const uint32 BaseId = BaseIdIter.GetAsUINT32();
+					MinBaseId = FMath::Min(BaseId, MinBaseId);
+					MaxBaseId = FMath::Max(BaseId, MaxBaseId);
 				}
 			}
 
-            SparseIndexMap IndexMap( minBaseId, maxBaseId );
-            
-            auto usedMorphIdIter = MorphIdIter;
-            for ( uint32 i = 0; i < static_cast<uint32>(morphSize); ++i, ++usedMorphIdIter)
+            SparseIndexMap IndexMap(MinBaseId, MaxBaseId);
+           
+            for (int32 Index = 0; Index < MorphNum; ++Index, ++MorphIdIter)
             {
-                const uint32 morphId = (*usedMorphIdIter)[0];
+                const uint32 MorphId = (*MorphIdIter)[0];
                 
-                IndexMap.Insert( morphId, i );
+                IndexMap.Insert(MorphId, Index);
             }
 
-            for ( int32 v = 0; v < baseSize; ++v )
+			return IndexMap;
+		};
+
+        const auto ApplyNormalMorph = [](
+				UntypedMeshBufferIteratorConst BaseIdIter, const TStaticArray<UntypedMeshBufferIterator, 3>& BaseTangentFrameIters, const int32 BaseNum,
+				MeshBufferIteratorConst<MBF_UINT32, uint32, 1> MorphIdIter, UntypedMeshBufferIteratorConst& MorphNormalIter, const int32 MorphNum,
+				const SparseIndexMap& IndexMap, const float Factor)
+        -> void
+        {
+			const UntypedMeshBufferIterator& BaseNormalIter = BaseTangentFrameIters[2];
+			const UntypedMeshBufferIterator& BaseTangentIter = BaseTangentFrameIters[1];
+			const UntypedMeshBufferIterator& BaseBiNormalIter = BaseTangentFrameIters[0];
+
+            const MESH_BUFFER_FORMAT NormalFormat = BaseNormalIter.GetFormat();
+            const int32 NormalComps = BaseNormalIter.GetComponents();
+
+            const MESH_BUFFER_FORMAT TangentFormat = BaseTangentIter.GetFormat();
+            const int32 TangentComps = BaseTangentIter.GetComponents();
+
+            const MESH_BUFFER_FORMAT BiNormalFormat = BaseBiNormalIter.GetFormat();
+            const int32 BiNormalComps = BaseBiNormalIter.GetComponents();
+
+			// When normal is packed, binormal channel is not expected. It is not a big deal if it's there but we would be doing extra unused work in that case. 
+			ensureAlways(!(NormalFormat == MBF_PACKEDDIR8_W_TANGENTSIGN || NormalFormat == MBF_PACKEDDIRS8_W_TANGENTSIGN) || !BaseBiNormalIter.ptr());
+
+            for (int32 VertexIndex = 0; VertexIndex < BaseNum; ++VertexIndex)
             {
-                const uint32 baseId = (BaseIdIter + v).GetAsUINT32();
-                const uint32 morphIdx = IndexMap.Find( baseId );
+                const uint32 BaseId = (BaseIdIter + VertexIndex).GetAsUINT32();
+                const int32 MorphIndex = static_cast<int32>(IndexMap.Find(BaseId));
 
-                if ( morphIdx == SparseIndexMap::NotFoundValue )
+                if (MorphIndex == SparseIndexMap::NotFoundValue)
                 {
                     continue;
                 }
 
-                const int32 m = static_cast<int32>( morphIdx );
+                // Find consecutive run.
+                UntypedMeshBufferIteratorConst RunBaseIter = BaseIdIter + VertexIndex;
+                MeshBufferIteratorConst<MBF_UINT32, uint32, 1> RunMorphIter = MorphIdIter + MorphIndex;
+
+                int32 RunSize = 0;
+                for ( ; VertexIndex + RunSize < BaseNum && MorphIndex + RunSize < MorphNum && RunBaseIter.GetAsUINT32() == (*RunMorphIter)[0];
+                        ++RunSize, ++RunBaseIter, ++RunMorphIter);
+
+				for (int32 RunIndex = 0; RunIndex < RunSize; ++RunIndex)
+				{
+					UntypedMeshBufferIterator NormalIter = BaseNormalIter + (VertexIndex + RunIndex);
+
+					const FVector3f BaseNormal = NormalIter.GetAsVec3f();
+					const FVector3f MorphNormal = (MorphNormalIter + (MorphIndex + RunIndex)).GetAsVec3f();
+ 
+					const FVector3f Normal = (BaseNormal + MorphNormal*Factor).GetSafeNormal();
+				
+					// Leave the tangent basis sign untouched for packed normals formats.
+					for (int32 C = 0; C < NormalComps && C < 3; ++C)
+					{
+						ConvertData(C, NormalIter.ptr(), NormalFormat, &Normal, MBF_FLOAT32);
+					}
+
+					// Tangent
+					if (BaseTangentIter.ptr())
+					{
+						UntypedMeshBufferIterator TangentIter = BaseTangentIter + (VertexIndex + RunIndex);
+
+						const FVector3f BaseTangent = TangentIter.GetAsVec3f();
+
+						// Orthogonalize Tangent based on new Normal. This assumes Normal and BaseTangent are normalized and different.
+						const FVector3f Tangent = (BaseTangent - FVector3f::DotProduct(Normal, BaseTangent) * Normal).GetSafeNormal();
+
+						for (int32 C = 0; C < TangentComps && C < 3; ++C)
+						{
+							ConvertData(C, TangentIter.ptr(), TangentFormat, &Tangent, MBF_FLOAT32);
+						}
+
+						// BiNormal
+						if (BaseBiNormalIter.ptr())
+						{
+							UntypedMeshBufferIterator BiNormalIter = BaseBiNormalIter + (VertexIndex + RunIndex);
+
+							const FVector3f& N = BaseNormal;
+							const FVector3f& T = BaseTangent;
+							const FVector3f  B = BiNormalIter.GetAsVec3f();
+
+							const float BaseTangentBasisDeterminant =  
+									B.X*T.Y*N.Z + B.Z*T.X*N.Y + B.Y*T.Z*N.Y -
+									B.Z*T.Y*N.X - B.Y*T.X*N.Z - B.X*T.Z*N.Y;
+
+							const float BaseTangentBasisDeterminantSign = BaseTangentBasisDeterminant >= 0 ? 1.0f : -1.0f;
+
+							const FVector3f BiNormal = FVector3f::CrossProduct(Tangent, Normal) * BaseTangentBasisDeterminantSign;
+
+							for (int32 C = 0; C < BiNormalComps && C < 3; ++C)
+							{
+								ConvertData(C, BiNormalIter.ptr(), BiNormalFormat, &BiNormal, MBF_FLOAT32);
+							}
+						}
+					}
+				}
+
+				VertexIndex += FMath::Max(RunSize - 1, 0);
+            }
+        };
+
+        const auto ApplyGenericMorph = []( 
+				UntypedMeshBufferIteratorConst BaseIdIter, const TArray<UntypedMeshBufferIterator>& BaseChannelsIters, const int32 BaseNum,
+				MeshBufferIteratorConst<MBF_UINT32, uint32, 1> MorphIdIter, const TArray<UntypedMeshBufferIteratorConst>& MorphChannelsIters, const int32 MorphNum,
+				const SparseIndexMap& IndexMap, const float Factor)
+        -> void
+        {
+            for (int32 VertexIndex = 0; VertexIndex < BaseNum; ++VertexIndex)
+            {
+                const uint32 BaseId = (BaseIdIter + VertexIndex).GetAsUINT32();
+                const uint32 MorphIdx = IndexMap.Find(BaseId);
+
+                if (MorphIdx == SparseIndexMap::NotFoundValue)
+                {
+                    continue;
+                }
+
+                const int32 MorphIndex = static_cast<int32>(MorphIdx);
 
                 // Find consecutive run.
-                auto RunBaseIter = BaseIdIter + v;
-                auto RunMorphIter = MorphIdIter + m;
+                UntypedMeshBufferIteratorConst RunBaseIter = BaseIdIter + VertexIndex;
+                MeshBufferIteratorConst<MBF_UINT32, uint32, 1> RunMorphIter = MorphIdIter + MorphIndex;
 
-                int32 runSize = 0;
-                for (  ; v + runSize < baseSize && m + runSize < morphSize && RunBaseIter.GetAsUINT32() == (*RunMorphIter)[0];
-                        ++runSize, ++RunBaseIter, ++RunMorphIter );
+                int32 RunSize = 0;
+                for ( ; VertexIndex + RunSize < BaseNum && MorphIndex + RunSize < MorphNum && RunBaseIter.GetAsUINT32() == (*RunMorphIter)[0];
+                        ++RunSize, ++RunBaseIter, ++RunMorphIter);
 
-                const size_t channelCount = morphChannelsIters.Num();
-                for ( size_t c = 1; c < channelCount; ++c )
+                const int32 ChannelNum = MorphChannelsIters.Num();
+                for (int32 ChannelIndex = 1; ChannelIndex < ChannelNum; ++ChannelIndex)
                 {
-                
-                	if ( !(baseChannelsIters[c].ptr() && morphChannelsIters[c].ptr()) )
+                	if (!(BaseChannelsIters[ChannelIndex].ptr() && MorphChannelsIters[ChannelIndex].ptr()))
                 	{
                 		continue;
                 	}
                 
-                    UntypedMeshBufferIterator channelBaseIter = baseChannelsIters[c] + v;
-					UntypedMeshBufferIteratorConst channelMorphIter = morphChannelsIters[c] + m;
+                    UntypedMeshBufferIterator ChannelBaseIter = BaseChannelsIters[ChannelIndex] + VertexIndex;
+					UntypedMeshBufferIteratorConst ChannelMorphIter = MorphChannelsIters[ChannelIndex] + MorphIndex;
                    
-                    const MESH_BUFFER_FORMAT dstChannelFormat = baseChannelsIters[c].GetFormat();
-                    const int dstChannelComps = baseChannelsIters[c].GetComponents();
+                    const MESH_BUFFER_FORMAT DestChannelFormat = BaseChannelsIters[ChannelIndex].GetFormat();
+                    const int32 DestChannelComps = BaseChannelsIters[ChannelIndex].GetComponents();
 
                     // Apply Morph to range found above.
-                    for ( int32 r = 0; r < runSize; ++r, ++channelBaseIter, ++channelMorphIter )
+                    for (int32 R = 0; R < RunSize; ++R, ++ChannelBaseIter, ++ChannelMorphIter)
                     {
-                        const FVector4f value = channelBaseIter.GetAsVec4f() + channelMorphIter.GetAsVec4f()*factor;
+                        const FVector4f Value = ChannelBaseIter.GetAsVec4f() + ChannelMorphIter.GetAsVec4f()*Factor;
 
                         // TODO: Optimize this for the specific components.
                         // Max 4 components
-                        for ( int32 comp = 0; comp < dstChannelComps && comp < 4; ++comp )
+                        for (int32 Comp = 0; Comp < DestChannelComps && Comp < 4; ++Comp)
                         {
-                            ConvertData( comp, channelBaseIter.ptr(), dstChannelFormat, &value, MBF_FLOAT32 );
+                            ConvertData(Comp, ChannelBaseIter.ptr(), DestChannelFormat, &Value, MBF_FLOAT32);
                         }
                     }
                 }
 
-				v += FMath::Max(runSize - 1, 0);
+				VertexIndex += FMath::Max(RunSize - 1, 0);
             }
         };
 
 
-		
 		// Number of vertices to modify
-		const uint32 minCount = pMin ? pMin->GetVertexBuffers().GetElementCount() : 0;
-		const uint32 maxCount = pMax ? pMax->GetVertexBuffers().GetElementCount() : 0;
-		const uint32 baseCount = pBase ? pBase->GetVertexBuffers().GetElementCount() : 0;
-		const Mesh* refTarget = minCount > 0 ? pMin : pMax;
+		const int32 MinNum = pMin ? pMin->GetVertexBuffers().GetElementCount() : 0;
+		const int32 MaxNum = pMax ? pMax->GetVertexBuffers().GetElementCount() : 0;
+		const int32 BaseNum = pBase ? pBase->GetVertexBuffers().GetElementCount() : 0;
+		const Mesh* RefTarget = MinNum > 0 ? pMin : pMax;
 
-		if ( baseCount == 0 || (minCount + maxCount) == 0)
+		if (BaseNum == 0 || (MinNum + MaxNum) == 0)
 		{
 			bOutSuccess = false; // Use the passed pBase as result.
 			return;
@@ -123,48 +225,90 @@ namespace mu
 
 		Result->CopyFrom(*pBase);
 
-		if (refTarget)
+		if (RefTarget)
 		{
-			uint32 ccount = refTarget->GetVertexBuffers().GetBufferChannelCount(0);
+			const int32 ChannelsNum = RefTarget->GetVertexBuffers().GetBufferChannelCount(0);
 
-			// Iterator of the vertex ids of the base vertices
-			UntypedMeshBufferIteratorConst itBaseId(pBase->GetVertexBuffers(), MBS_VERTEXINDEX);
+			TArray<UntypedMeshBufferIterator> BaseChannelsIters;
+			BaseChannelsIters.SetNum(ChannelsNum);
+			TArray<UntypedMeshBufferIteratorConst> MinChannelsIters;
+			MinChannelsIters.SetNum(ChannelsNum);
+			TArray<UntypedMeshBufferIteratorConst> MaxChannelsIters;
+			MaxChannelsIters.SetNum(ChannelsNum);
 
-			TArray< UntypedMeshBufferIterator > itBaseChannels;
-			itBaseChannels.SetNum(ccount);
-			TArray< UntypedMeshBufferIteratorConst > itMinChannels;
-			itMinChannels.SetNum(ccount);
-			TArray< UntypedMeshBufferIteratorConst > itMaxChannels;
-			itMaxChannels.SetNum(ccount);
+			// {BiNormal, Tangent, Normal}
+			TStaticArray<UntypedMeshBufferIterator, 3> BaseTangentFrameChannelsIters;
+			UntypedMeshBufferIteratorConst MinNormalChannelIter;
+			UntypedMeshBufferIteratorConst MaxNormalChannelIter;
 
-			for (size_t c = 1; c < ccount; ++c)
+			const bool bBaseHasNormals = UntypedMeshBufferIteratorConst(pBase->GetVertexBuffers(), MBS_NORMAL, 0).ptr() != nullptr;
+			for (int32 ChannelIndex = 1; ChannelIndex < ChannelsNum; ++ChannelIndex)
 			{
-				const FMeshBufferSet& MBSPriv = refTarget->GetVertexBuffers();
-				MESH_BUFFER_SEMANTIC sem = MBSPriv.m_buffers[0].m_channels[c].m_semantic;
-				int semIndex = MBSPriv.m_buffers[0].m_channels[c].m_semanticIndex;
-				
-				itBaseChannels[c] = UntypedMeshBufferIterator(Result->GetVertexBuffers(), sem, semIndex);
-				if (minCount > 0)
+				const FMeshBufferSet& MBSPriv = RefTarget->GetVertexBuffers();
+				MESH_BUFFER_SEMANTIC Sem = MBSPriv.m_buffers[0].m_channels[ChannelIndex].m_semantic;
+				int32 SemIndex = MBSPriv.m_buffers[0].m_channels[ChannelIndex].m_semanticIndex;
+			
+				if (Sem == MBS_NORMAL && bBaseHasNormals)
 				{
-					itMinChannels[c] = UntypedMeshBufferIteratorConst(pMin->GetVertexBuffers(), sem, semIndex);
-				}
+					BaseTangentFrameChannelsIters[2] = UntypedMeshBufferIterator(Result->GetVertexBuffers(), Sem, SemIndex);
+					if (MinNum > 0)
+					{
+						MinNormalChannelIter = UntypedMeshBufferIteratorConst(pMin->GetVertexBuffers(), Sem, SemIndex);
+					}
 
-				if (maxCount > 0)
+					if (MaxNum > 0)
+					{
+						MaxNormalChannelIter = UntypedMeshBufferIteratorConst(pMax->GetVertexBuffers(), Sem, SemIndex);
+					}
+				}
+				else if (Sem == MBS_TANGENT && bBaseHasNormals)
+				{
+					BaseTangentFrameChannelsIters[1] = UntypedMeshBufferIterator(Result->GetVertexBuffers(), Sem, SemIndex);
+				}
+				else if (Sem == MBS_BINORMAL && bBaseHasNormals)
 				{
-					itMaxChannels[c] = UntypedMeshBufferIteratorConst(pMax->GetVertexBuffers(), sem, semIndex);
+					BaseTangentFrameChannelsIters[0] = UntypedMeshBufferIterator(Result->GetVertexBuffers(), Sem, SemIndex);
+				}
+				else
+				{
+					BaseChannelsIters[ChannelIndex] = UntypedMeshBufferIterator(Result->GetVertexBuffers(), Sem, SemIndex);
+					if (MinNum > 0)
+					{
+						MinChannelsIters[ChannelIndex] = UntypedMeshBufferIteratorConst(pMin->GetVertexBuffers(), Sem, SemIndex);
+					}
+
+					if (MaxNum > 0)
+					{
+						MaxChannelsIters[ChannelIndex] = UntypedMeshBufferIteratorConst(pMax->GetVertexBuffers(), Sem, SemIndex);
+					}
 				}
 			}
 			
-			if (minCount > 0)
+			UntypedMeshBufferIteratorConst BaseIdIter(pBase->GetVertexBuffers(), MBS_VERTEXINDEX);
+
+			if (MinNum > 0)
 			{
-				MeshBufferIteratorConst<MBF_UINT32, uint32, 1> itMinId(pMin->GetVertexBuffers(), MBS_VERTEXINDEX);
-				ApplyMorph(itBaseId, itBaseChannels, baseCount, itMinId, itMinChannels, minCount, 1.0f - factor);
+				MeshBufferIteratorConst<MBF_UINT32, uint32, 1> MinIdIter(pMin->GetVertexBuffers(), MBS_VERTEXINDEX);
+				SparseIndexMap IndexMap = MakeIndexMap(BaseIdIter, BaseNum, MinIdIter, MinNum);
+
+				ApplyGenericMorph(BaseIdIter, BaseChannelsIters, BaseNum, MinIdIter, MinChannelsIters, MinNum, IndexMap, 1.0f - Factor);
+				if (MinNormalChannelIter.ptr())
+				{
+					ApplyNormalMorph(BaseIdIter, BaseTangentFrameChannelsIters, BaseNum, MinIdIter, MinNormalChannelIter, MinNum, IndexMap, 1.0f - Factor);
+				}
 			}
 
-			if (maxCount > 0)
+			if (MaxNum > 0)
 			{
-				MeshBufferIteratorConst<MBF_UINT32, uint32, 1> itMaxId(pMax->GetVertexBuffers(), MBS_VERTEXINDEX);
-				ApplyMorph(itBaseId, itBaseChannels, baseCount, itMaxId, itMaxChannels, maxCount, factor);
+				MeshBufferIteratorConst<MBF_UINT32, uint32, 1> MaxIdIter(pMax->GetVertexBuffers(), MBS_VERTEXINDEX);
+				SparseIndexMap IndexMap = MakeIndexMap(BaseIdIter, BaseNum, MaxIdIter, MaxNum);
+
+				ApplyGenericMorph(BaseIdIter, BaseChannelsIters, BaseNum, MaxIdIter, MaxChannelsIters, MaxNum, IndexMap, Factor);
+				
+				if (MaxNormalChannelIter.ptr())
+				{
+					ApplyNormalMorph(BaseIdIter, BaseTangentFrameChannelsIters, BaseNum, MaxIdIter, MaxNormalChannelIter, MaxNum, IndexMap, Factor);
+				}
 			}
 		}
     }
diff --git a/Engine/Plugins/Experimental/Mutable/Source/MutableTools/Private/MuT/CodeGenerator_Mesh.cpp b/Engine/Plugins/Experimental/Mutable/Source/MutableTools/Private/MuT/CodeGenerator_Mesh.cpp
--- a/Engine/Plugins/Experimental/Mutable/Source/MutableTools/Private/MuT/CodeGenerator_Mesh.cpp
+++ b/Engine/Plugins/Experimental/Mutable/Source/MutableTools/Private/MuT/CodeGenerator_Mesh.cpp
@@ -848,6 +848,12 @@ class Node;
 
         // \todo Texcoords are broken?
         op->bIgnoreTextureCoords = true;
+	
+		// UE only has position and normal morph data, optimize for this case if indicated. 
+		if (node.bOnlyPositionAndNormal)
+		{
+			op->Channels = { {static_cast<uint8>(MBS_POSITION), 0}, {static_cast<uint8>(MBS_NORMAL), 0} };
+		}
 
         // Base
         FMeshGenerationResult BaseResult;
diff --git a/Engine/Plugins/Experimental/Mutable/Source/MutableTools/Private/MuT/NodeMeshMakeMorph.cpp b/Engine/Plugins/Experimental/Mutable/Source/MutableTools/Private/MuT/NodeMeshMakeMorph.cpp
--- a/Engine/Plugins/Experimental/Mutable/Source/MutableTools/Private/MuT/NodeMeshMakeMorph.cpp
+++ b/Engine/Plugins/Experimental/Mutable/Source/MutableTools/Private/MuT/NodeMeshMakeMorph.cpp
@@ -88,6 +88,18 @@ namespace mu
 		m_pD->m_pTarget = p;
 	}
 
+	//---------------------------------------------------------------------------------------------
+	void NodeMeshMakeMorph::SetOnlyPositionAndNormal(bool bInOnlyPositionAndNormals)
+	{
+		m_pD->bOnlyPositionAndNormal = bInOnlyPositionAndNormals;
+	}
+
+
+	//---------------------------------------------------------------------------------------------
+	bool NodeMeshMakeMorph::GetOnlyPositionAndNormal() const
+	{
+		return m_pD->bOnlyPositionAndNormal;
+	}
 
 	//---------------------------------------------------------------------------------------------
     NodeLayoutPtr NodeMeshMakeMorph::Private::GetLayout( int index ) const
diff --git a/Engine/Plugins/Experimental/Mutable/Source/MutableTools/Private/MuT/NodeMeshMakeMorphPrivate.h b/Engine/Plugins/Experimental/Mutable/Source/MutableTools/Private/MuT/NodeMeshMakeMorphPrivate.h
--- a/Engine/Plugins/Experimental/Mutable/Source/MutableTools/Private/MuT/NodeMeshMakeMorphPrivate.h
+++ b/Engine/Plugins/Experimental/Mutable/Source/MutableTools/Private/MuT/NodeMeshMakeMorphPrivate.h
@@ -23,15 +23,18 @@ namespace mu
 
 		NodeMeshPtr m_pBase;
 		NodeMeshPtr m_pTarget;
+		
+		bool bOnlyPositionAndNormal = false;
 
 		//!
 		void Serialise( OutputArchive& arch ) const
 		{
-            uint32_t ver = 0;
+            uint32_t ver = 1;
 			arch << ver;
 
 			arch << m_pBase;
 			arch << m_pTarget;
+			arch << bOnlyPositionAndNormal;
 		}
 
 		//!
@@ -39,10 +42,15 @@ namespace mu
 		{
             uint32_t ver;
 			arch >> ver;
-			check(ver==0);
+			check(ver <= 1);
 
 			arch >> m_pBase;
 			arch >> m_pTarget;
+
+			if (ver == 1)
+			{
+				arch >> bOnlyPositionAndNormal;
+			}
 		}
 
 		// NodeMesh::Private interface
diff --git a/Engine/Plugins/Experimental/Mutable/Source/MutableTools/Public/MuT/NodeMeshMakeMorph.h b/Engine/Plugins/Experimental/Mutable/Source/MutableTools/Public/MuT/NodeMeshMakeMorph.h
--- a/Engine/Plugins/Experimental/Mutable/Source/MutableTools/Public/MuT/NodeMeshMakeMorph.h
+++ b/Engine/Plugins/Experimental/Mutable/Source/MutableTools/Public/MuT/NodeMeshMakeMorph.h
@@ -58,6 +58,9 @@ namespace mu
 		NodeMeshPtr GetTarget() const;
 		void SetTarget( NodeMesh* p );
 
+		bool GetOnlyPositionAndNormal() const;
+		void SetOnlyPositionAndNormal(bool bOnlyPositionAndNormal);
+
 		//-----------------------------------------------------------------------------------------
 		// Interface pattern
 		//-----------------------------------------------------------------------------------------
