[mutable] Fixed text conflict from LOCTEXT macro.

#rnx

diff --git a/Engine/Plugins/Experimental/Mutable/Source/CustomizableObjectEditor/Private/MuCOE/CustomizableObjectCompiler.cpp b/Engine/Plugins/Experimental/Mutable/Source/CustomizableObjectEditor/Private/MuCOE/CustomizableObjectCompiler.cpp
--- a/Engine/Plugins/Experimental/Mutable/Source/CustomizableObjectEditor/Private/MuCOE/CustomizableObjectCompiler.cpp
+++ b/Engine/Plugins/Experimental/Mutable/Source/CustomizableObjectEditor/Private/MuCOE/CustomizableObjectCompiler.cpp
@@ -1571,9 +1571,9 @@ void FCustomizableObjectCompiler::NotifyCompilationErrors() const
 	const FText Message = NoWarningsOrErrors ?
 		FText::Format(LOCTEXT("CompilationFinishedSuccessfully", "{0} finished compiling successfully."), Prefix) :
 		NumIgnoreds > 0 ?
-		FText::Format(LOCTEXT("CompilationFinished", "{0} finished compiling successfully with {1} {1}|plural(one=warning,other=warnings), {2} {2}|plural(one=error,other=errors) and {3} more similar warnings."), Prefix, NumWarnings, NumErrors, NumIgnoreds)
+		FText::Format(LOCTEXT("CompilationFinished_WithIgnoreds", "{0} finished compiling successfully with {1} {1}|plural(one=warning,other=warnings), {2} {2}|plural(one=error,other=errors) and {3} more similar warnings."), Prefix, NumWarnings, NumErrors, NumIgnoreds)
 		:
-		FText::Format(LOCTEXT("CompilationFinished", "{0} finished compiling successfully with {1} {1}|plural(one=warning,other=warnings) and {2} {2}|plural(one=error,other=errors)."), Prefix, NumWarnings, NumErrors);
+		FText::Format(LOCTEXT("CompilationFinished_WithoutIgnoreds", "{0} finished compiling successfully with {1} {1}|plural(one=warning,other=warnings) and {2} {2}|plural(one=error,other=errors)."), Prefix, NumWarnings, NumErrors);
 	
 	FCustomizableObjectEditorLogger::CreateLog(Message)
 	.Category(ELoggerCategory::Compilation)
