[Mutable] Fix crash with mutable.EnableImageCache 0.

Renamed CVar to make its purpose clear.

#rnx

diff --git a/Engine/Plugins/Experimental/Mutable/Source/CustomizableObject/Classes/MuCO/CustomizableObjectSystem.h b/Engine/Plugins/Experimental/Mutable/Source/CustomizableObject/Classes/MuCO/CustomizableObjectSystem.h
--- a/Engine/Plugins/Experimental/Mutable/Source/CustomizableObject/Classes/MuCO/CustomizableObjectSystem.h
+++ b/Engine/Plugins/Experimental/Mutable/Source/CustomizableObject/Classes/MuCO/CustomizableObjectSystem.h
@@ -63,7 +63,7 @@ DECLARE_DWORD_ACCUMULATOR_STAT_EXTERN(TEXT("Streaming Ops"), STAT_MutableStreami
 
 extern TAutoConsoleVariable<bool> CVarClearWorkingMemoryOnUpdateEnd;
 
-extern TAutoConsoleVariable<bool> CVarEnableImageCache;
+extern TAutoConsoleVariable<bool> CVarReuseImagesBetweenInstances;
 
 
 #if WITH_EDITOR
@@ -427,7 +427,6 @@ private:
 	TMap<FString, int64> PlatformMaxChunkSize;
 	
 #endif
-	void OnEnableImageCacheChanged(IConsoleVariable* CVar);
 	
 	// Friends
 	friend class FCustomizableObjectSystemPrivate;
diff --git a/Engine/Plugins/Experimental/Mutable/Source/CustomizableObject/Private/MuCO/CustomizableObjectInstance.cpp b/Engine/Plugins/Experimental/Mutable/Source/CustomizableObject/Private/MuCO/CustomizableObjectInstance.cpp
--- a/Engine/Plugins/Experimental/Mutable/Source/CustomizableObject/Private/MuCO/CustomizableObjectInstance.cpp
+++ b/Engine/Plugins/Experimental/Mutable/Source/CustomizableObject/Private/MuCO/CustomizableObjectInstance.cpp
@@ -5962,7 +5962,7 @@ void UCustomizableInstancePrivateData::BuildMaterials(const TSharedPtr<FMutableO
 											ReuseTexture(Texture);
 										}
 										else if (Texture)
-										{
+										{	
 											//if (!bNeverStream) // No need to check bNeverStream. In that case, the texture won't use 
 											// the MutableMipDataProviderFactory anyway and it's needed for detecting Mutable textures elsewhere
 											{
@@ -5990,11 +5990,7 @@ void UCustomizableInstancePrivateData::BuildMaterials(const TSharedPtr<FMutableO
 										}
 									}
 
-									// update the model resources cache
-									if (CVarEnableImageCache.GetValueOnGameThread())
-									{
-										Cache.Images.Add(ImageCacheKey, Texture);						
-									}
+									Cache.Images.Add(ImageCacheKey, Texture);						
 								}
 								else
 								{
diff --git a/Engine/Plugins/Experimental/Mutable/Source/CustomizableObject/Private/MuCO/CustomizableObjectSystem.cpp b/Engine/Plugins/Experimental/Mutable/Source/CustomizableObject/Private/MuCO/CustomizableObjectSystem.cpp
--- a/Engine/Plugins/Experimental/Mutable/Source/CustomizableObject/Private/MuCO/CustomizableObjectSystem.cpp
+++ b/Engine/Plugins/Experimental/Mutable/Source/CustomizableObject/Private/MuCO/CustomizableObjectSystem.cpp
@@ -95,10 +95,10 @@ TAutoConsoleVariable<bool> CVarClearWorkingMemoryOnUpdateEnd(
 	ECVF_Scalability);
 
 
-TAutoConsoleVariable<bool> CVarEnableImageCache(
-	TEXT("mutable.EnableImageCache"),
+TAutoConsoleVariable<bool> CVarReuseImagesBetweenInstances(
+	TEXT("mutable.ReuseImagesBetweenInstances"),
 	true,
-	TEXT("Enables or disables the Instance Image Cache. Avoids regenerating images when no parameters or mips have changed."),
+	TEXT("Enables or disables the reuse of images between instances."),
 	ECVF_Scalability);
 
 static TAutoConsoleVariable<int32> CVarGeneratedResourcesCacheSize(
@@ -381,8 +381,6 @@ void UCustomizableObjectSystem::InitSystem()
 	DefaultInstanceLODManagement = NewObject<UCustomizableInstanceLODManagement>();
 	check(DefaultInstanceLODManagement != nullptr);
 	CurrentInstanceLODManagement = DefaultInstanceLODManagement;
-
-	CVarEnableImageCache->SetOnChangedCallback(FConsoleVariableDelegate::CreateUObject(this, &UCustomizableObjectSystem::OnEnableImageCacheChanged));
 }
 
 
@@ -1547,16 +1545,8 @@ namespace impl
 				Image.FullImageSizeY = ImageDesc.m_size[1] / Reduction;
 			}
 
-			bool bCached = false;
-
-			// See if it is cached from this same instance (can happen with LODs)
-			bCached = ImagesInThisInstance.Contains(Image.ImageID);
-
-			// See if it is cached from another instance
-			if (!bCached)
-			{
-				bCached = CustomizableObjectSystemPrivateData->ProtectedObjectCachedImages.Contains(Image.ImageID);
-			}
+			const bool bCached = ImagesInThisInstance.Contains(Image.ImageID) || // See if it is cached from this same instance (can happen with LODs)
+				(CVarReuseImagesBetweenInstances.GetValueOnAnyThread() && CustomizableObjectSystemPrivateData->ProtectedObjectCachedImages.Contains(Image.ImageID)); // See if it is cached from another instance
 
 			if (bCached)
 			{
@@ -1633,7 +1623,6 @@ namespace impl
 				ImagesInThisInstance.Add(Image.ImageID);
 			}
 		}
-
 	}
 	
 
@@ -3288,18 +3277,6 @@ uint64 UCustomizableObjectSystem::GetMaxChunkSizeForPlatform(const ITargetPlatfo
 #endif // WITH_EDITOR
 
 
-void UCustomizableObjectSystem::OnEnableImageCacheChanged(IConsoleVariable* CVar)
-{
-	if (!CVar->GetBool())
-	{		
-		for (FMutableResourceCache& Cache : GetPrivate()->ModelResourcesCache)
-		{
-			Cache.Images.Empty();
-		}
-	}	
-}
-
-
 void UCustomizableObjectSystem::CacheImage(FString ImageId)
 {
 	GetPrivateChecked()->GetImageProviderChecked()->CacheImage(ImageId, true);
