[Mutable] Fix automatic compilation infinite loop when a compilation fails.
Fix excessive logging spam when there's an uncompiled CO.
Fix crash when setting requested LODs.
Tweak some logs.

#rnx

diff --git a/Engine/Plugins/Experimental/Mutable/Source/CustomizableObject/Classes/MuCO/CustomizableObject.h b/Engine/Plugins/Experimental/Mutable/Source/CustomizableObject/Classes/MuCO/CustomizableObject.h
--- a/Engine/Plugins/Experimental/Mutable/Source/CustomizableObject/Classes/MuCO/CustomizableObject.h
+++ b/Engine/Plugins/Experimental/Mutable/Source/CustomizableObject/Classes/MuCO/CustomizableObject.h
@@ -1625,6 +1625,8 @@ public:
 	UPROPERTY(Transient)
 	FGuid Identifier;
 
+	ECustomizableObjectCompilationState CompilationState = ECustomizableObjectCompilationState::None;
+
 	FPostCompileDelegate PostCompileDelegate;
 
 	void PostCompile();
diff --git a/Engine/Plugins/Experimental/Mutable/Source/CustomizableObject/Classes/MuCO/CustomizableObjectInstance.h b/Engine/Plugins/Experimental/Mutable/Source/CustomizableObject/Classes/MuCO/CustomizableObjectInstance.h
--- a/Engine/Plugins/Experimental/Mutable/Source/CustomizableObject/Classes/MuCO/CustomizableObjectInstance.h
+++ b/Engine/Plugins/Experimental/Mutable/Source/CustomizableObject/Classes/MuCO/CustomizableObjectInstance.h
@@ -274,6 +274,9 @@ public:
 	void SetSelectedParameterProfileDirty();
 	bool IsSelectedParameterProfileDirty() const;
 
+	/** Return true if the instance is not locked and if it's compiled */
+	bool CanUpdateInstance() const;
+
 	/** Update Skeletal Mesh asynchronously. */
 	UFUNCTION(BlueprintCallable, Category = CustomizableObjectInstance)
 	void UpdateSkeletalMeshAsync(bool bIgnoreCloseDist = false, bool bForceHighPriority = false);
diff --git a/Engine/Plugins/Experimental/Mutable/Source/CustomizableObject/Private/MuCO/CustomizableObject.cpp b/Engine/Plugins/Experimental/Mutable/Source/CustomizableObject/Private/MuCO/CustomizableObject.cpp
--- a/Engine/Plugins/Experimental/Mutable/Source/CustomizableObject/Private/MuCO/CustomizableObject.cpp
+++ b/Engine/Plugins/Experimental/Mutable/Source/CustomizableObject/Private/MuCO/CustomizableObject.cpp
@@ -727,6 +727,12 @@ bool UCustomizableObject::ConditionalAutoCompile()
 		return false;
 	}
 
+	// Don't re-compile objects if they failed to compile. 
+	if (CompilationState == ECustomizableObjectCompilationState::Failed)
+	{
+		return false;
+	}
+
 	// Don't compile if we're running game, Compile and/or if AutoCompile is disabled
 	if (IsRunningGame() || System->IsCompilationDisabled() || !System->IsAutoCompileEnabled())
 	{
@@ -1208,18 +1214,7 @@ int32 UCustomizableObject::GetNumLODs() const
 
 int32 UCustomizableObject::GetComponentCount() const
 {
-	if (!IsCompiled())
-	{
-		UE_LOG(LogMutable, Warning,
-		       TEXT(
-			       "You are trying to get the component count of the non-compiled CO %s. This will always return 0 as value."
-		       ), *GetName()
-		);
-
-		return 0;
-	}
-	
-	return NumMeshComponentsInRoot;
+	return IsCompiled() ? NumMeshComponentsInRoot : 0;
 }
 
 int UCustomizableObject::GetParameterCount() const
diff --git a/Engine/Plugins/Experimental/Mutable/Source/CustomizableObject/Private/MuCO/CustomizableObjectInstance.cpp b/Engine/Plugins/Experimental/Mutable/Source/CustomizableObject/Private/MuCO/CustomizableObjectInstance.cpp
--- a/Engine/Plugins/Experimental/Mutable/Source/CustomizableObject/Private/MuCO/CustomizableObjectInstance.cpp
+++ b/Engine/Plugins/Experimental/Mutable/Source/CustomizableObject/Private/MuCO/CustomizableObjectInstance.cpp
@@ -598,6 +598,11 @@ void UCustomizableInstancePrivateData::ReloadParameters(UCustomizableObjectInsta
 	if (bInvalidatePreviousData)
 	{
 		InvalidateGeneratedData();
+
+#if WITH_EDITOR
+		// Clear the Generated flag to trigger a new update after changing or compiling the CO,
+		ClearCOInstanceFlags(Generated);
+#endif
 	}
 
 	Public->GetDescriptor().ReloadParameters();
@@ -745,6 +750,22 @@ void UCustomizableInstancePrivateData::PostEditChangePropertyWithoutEditor(USkel
 }
 
 
+bool UCustomizableObjectInstance::CanUpdateInstance() const
+{
+	UCustomizableObject* CustomizableObject = GetCustomizableObject();
+	if (!CustomizableObject)
+	{
+		return false;
+	}
+
+#if WITH_EDITOR
+	return CustomizableObject->ConditionalAutoCompile();
+#else
+	return CustomizableObject->IsCompiled();
+#endif
+}
+
+
 void UCustomizableObjectInstance::UpdateSkeletalMeshAsync(bool bIgnoreCloseDist, bool bForceHighPriority)
 {
 	DoUpdateSkeletalMesh( false, false, bIgnoreCloseDist, bForceHighPriority, nullptr, nullptr);
@@ -763,22 +784,10 @@ EUpdateRequired UCustomizableObjectInstance::IsUpdateRequired(bool bIsCloseDistT
 	const UCustomizableInstancePrivateData* const Private = GetPrivate();
 	
 	UCustomizableObject* CustomizableObject = GetCustomizableObject();
-	if (!CustomizableObject || CustomizableObject->IsLocked())
-	{
-		return EUpdateRequired::NoUpdate;
-	}
-	
-#if WITH_EDITOR
-	if (!CustomizableObject->IsCompiled() && !CustomizableObject->ConditionalAutoCompile())
+	if (!CanUpdateInstance() || CustomizableObject->IsLocked())
 	{
 		return EUpdateRequired::NoUpdate;
 	}
-#else
-	if (!CustomizableObject->IsCompiled())
-	{
-		return EUpdateRequired::NoUpdate;		
-	}
-#endif
 
 	const bool bIsGenerated = Private->HasCOInstanceFlags(Generated);
 	const int32 NumGeneratedInstancesLimit = System->GetInstanceLODManagement()->GetNumGeneratedInstancesLimitFullLODs();
@@ -864,7 +873,12 @@ void UCustomizableObjectInstance::DoUpdateSkeletalMesh(bool bIsCloseDistTick, bo
 {
 	MUTABLE_CPUPROFILER_SCOPE(UCustomizableInstancePrivateData::DoUpdateSkeletalMesh);
 	check(IsInGameThread());
-	check(!OptionalUpdateRequired || *OptionalUpdateRequired != EUpdateRequired::NoUpdate) // If no update is required this functions must not be called.
+	check(!OptionalUpdateRequired || *OptionalUpdateRequired != EUpdateRequired::NoUpdate); // If no update is required this functions must not be called.
+
+	if (!CanUpdateInstance())
+	{
+		return;
+	}
 
 	UCustomizableObjectSystem* System = UCustomizableObjectSystem::GetInstance();
 
@@ -928,6 +942,11 @@ void UCustomizableObjectInstance::DoUpdateSkeletalMesh(bool bIsCloseDistTick, bo
 
 void UCustomizableInstancePrivateData::TickUpdateCloseCustomizableObjects(UCustomizableObjectInstance& Public, const bool bOnlyUpdatePrioForUpdates, FMutableInstanceUpdateMap& InOutRequestedUpdates)
 {
+	if (!Public.CanUpdateInstance())
+	{
+		return;
+	}
+
 	const EUpdateRequired UpdateRequired = Public.IsUpdateRequired(true, true, false);
 	if (UpdateRequired != EUpdateRequired::NoUpdate) // Since this is done in the tick, avoid starting an update that we know for sure that would not be performed. Once started it has some performance implications that we want to avoid.
 	{
@@ -973,25 +992,32 @@ void UCustomizableInstancePrivateData::TickUpdateCloseCustomizableObjects(UCusto
 
 void UCustomizableInstancePrivateData::UpdateInstanceIfNotGenerated(UCustomizableObjectInstance& Public, FMutableInstanceUpdateMap& InOutRequestedUpdates)
 {
-	if (!HasCOInstanceFlags(Generated))
+	if (HasCOInstanceFlags(Generated))
 	{
-		Public.DoUpdateSkeletalMesh(false, true, false, false, nullptr, nullptr);
+		return;
+	}
 
-		EQueuePriorityType Priority = Public.GetUpdatePriority(true, true, false, false);
-		FMutableUpdateCandidate* UpdateCandidate = InOutRequestedUpdates.Find(&Public);
+	if (!Public.CanUpdateInstance())
+	{
+		return;
+	}
 
-		if (UpdateCandidate)
-		{
-			UpdateCandidate->Priority = Priority;
-			UpdateCandidate->Issue();
-		}
-		else
-		{
-			FMutableUpdateCandidate Candidate(&Public);
-			Candidate.Priority = Priority;
-			Candidate.Issue();
-			InOutRequestedUpdates.Add(&Public, Candidate);
-		}
+	Public.DoUpdateSkeletalMesh(false, true, false, false, nullptr, nullptr);
+
+	EQueuePriorityType Priority = Public.GetUpdatePriority(true, true, false, false);
+	FMutableUpdateCandidate* UpdateCandidate = InOutRequestedUpdates.Find(&Public);
+
+	if (UpdateCandidate)
+	{
+		UpdateCandidate->Priority = Priority;
+		UpdateCandidate->Issue();
+	}
+	else
+	{
+		FMutableUpdateCandidate Candidate(&Public);
+		Candidate.Priority = Priority;
+		Candidate.Issue();
+		InOutRequestedUpdates.Add(&Public, Candidate);
 	}
 }
 
@@ -1988,7 +2014,7 @@ bool UCustomizableInstancePrivateData::DoComponentsNeedUpdate(UCustomizableObjec
 	{
 		bFoundEmptyLOD = true;
 		bUpdateMeshes = true;
-		UE_LOG(LogMutable, Warning, TEXT("Building instance: an LOD has no mesh geometry. This cannot be handled by Unreal."));
+		UE_LOG(LogMutable, Warning, TEXT("Building instance: An LOD has no mesh geometry. This cannot be handled by Unreal."));
 	}
 
 	bOutEmptyMesh = bFoundEmptyLOD || bFoundEmptyMesh;
@@ -2009,7 +2035,7 @@ bool UCustomizableInstancePrivateData::UpdateSkeletalMesh_PostBeginUpdate0(UCust
 	if (bEmptyMesh)
 	{
 		Public->SkeletalMeshStatus = ESkeletalMeshState::UpdateError;
-		UE_LOG(LogMutable, Warning, TEXT("Updating skeletal mesh error for mesh %s"), *Public->GetName());
+		UE_LOG(LogMutable, Warning, TEXT("Updating skeletal mesh error for CO Instance %s"), *Public->GetName());
 
 		Public->SkeletalMeshes.Reset(); // What about all the references
 
@@ -6437,6 +6463,11 @@ void UCustomizableObjectInstance::SetRequestedLODs(int32 InMinLOD, int32 InMaxLO
 		return;
 	}
 
+	if (!CanUpdateInstance())
+	{
+		return;
+	}
+
 #if !UE_BUILD_SHIPPING
 	// Ignore Min/Max LOD limits. Mainly used for debug
 	if (!bIgnoreMinMaxLOD)
@@ -6470,6 +6501,8 @@ void UCustomizableObjectInstance::SetRequestedLODs(int32 InMinLOD, int32 InMaxLO
 		// TODO Pere: With Requested LODs We could skip updates by increasing the lower limit of LODs mutable can generate if only MinLOD have changed.
 
 		const int32 ComponentCount = GetNumComponents();
+		MutableUpdateCandidate.RequestedLODLevels.SetNum(ComponentCount);
+
 		const TArray<uint16>& GeneratedLODsPerComponent = DescriptorRuntimeHash.GetRequestedLODs();
 
 		const bool bIgnoreGeneratedLODs = DescriptorRuntimeHash != UpdateDescriptorRuntimeHash || GeneratedLODsPerComponent.Num() != ComponentCount;
diff --git a/Engine/Plugins/Experimental/Mutable/Source/CustomizableObject/Private/MuCO/CustomizableObjectSystem.cpp b/Engine/Plugins/Experimental/Mutable/Source/CustomizableObject/Private/MuCO/CustomizableObjectSystem.cpp
--- a/Engine/Plugins/Experimental/Mutable/Source/CustomizableObject/Private/MuCO/CustomizableObjectSystem.cpp
+++ b/Engine/Plugins/Experimental/Mutable/Source/CustomizableObject/Private/MuCO/CustomizableObjectSystem.cpp
@@ -3019,42 +3019,45 @@ void UCustomizableObjectSystem::SetOnlyGenerateRequestedLODsEnabled(bool bIsEnab
 void UCustomizableObjectSystem::AddUncompiledCOWarning(const UCustomizableObject& InObject, FString const* OptionalLogInfo)
 {
 	FString Msg;
-	Msg += FString::Printf(TEXT("Warning: Customizable Object [%s] not compiled. Please compile and save the object."), *InObject.GetName());
+	Msg += FString::Printf(TEXT("Warning: Customizable Object [%s] not compiled."), *InObject.GetName());
 	GEngine->AddOnScreenDebugMessage((uint64)((PTRINT)&InObject), 10.0f, FColor::Red, Msg);
 
 #if WITH_EDITOR
+	// Mutable will spam these warnings constantly due to the tick and LOD manager checking for instances to update with every tick. Send only one message per CO in the editor.
+	if (UncompiledCustomizableObjectIds.Find(InObject.GetVersionId()) != INDEX_NONE)
+	{
+		return;
+	}
+	
+	// Add notification
+	UncompiledCustomizableObjectIds.Add(InObject.GetVersionId());
+
+	FMessageLog MessageLog("Mutable");
+	MessageLog.Warning(FText::FromString(Msg));
+
+	if (!UncompiledCustomizableObjectsNotificationPtr.IsValid())
+	{
+		FNotificationInfo Info(FText::FromString("Uncompiled Customizable Object/s found. Please, check the Message Log - Mutable for more information."));
+		Info.bFireAndForget = true;
+		Info.bUseThrobber = true;
+		Info.FadeOutDuration = 1.0f;
+		Info.ExpireDuration = 5.0f;
+
+		UncompiledCustomizableObjectsNotificationPtr = FSlateNotificationManager::Get().AddNotification(Info);
+	}
+
 	const FString ErrorString = FString::Printf(
 		TEXT("Customizable Object [%s] not compiled.  Compile via the editor or via code before instancing.  %s"),
-						   *InObject.GetName(), OptionalLogInfo ? **OptionalLogInfo : TEXT(""));
+		*InObject.GetName(), OptionalLogInfo ? **OptionalLogInfo : TEXT(""));
+
 #else // !WITH_EDITOR
 	const FString ErrorString = FString::Printf(
 		TEXT("Customizable Object [%s] not compiled.  This is not an Editor build, so this is an unrecoverable bad state; could be due to code or a cook failure.  %s"),
-						   *InObject.GetName(), OptionalLogInfo ? **OptionalLogInfo : TEXT(""));
+		*InObject.GetName(), OptionalLogInfo ? **OptionalLogInfo : TEXT(""));
 #endif
 
 	// Also log an error so if this happens as part of a bug report we'll have this info.
 	UE_LOG(LogMutable, Error, TEXT("%s"), *ErrorString);
-
-#if WITH_EDITOR
-	if (UncompiledCustomizableObjectIds.Find(InObject.GetVersionId()) == INDEX_NONE)
-	{
-		UncompiledCustomizableObjectIds.Add(InObject.GetVersionId());
-
-		FMessageLog MessageLog("Mutable");
-		MessageLog.Warning(FText::FromString(Msg));
-
-		if (!UncompiledCustomizableObjectsNotificationPtr.IsValid())
-		{
-			FNotificationInfo Info(FText::FromString("Uncompiled Customizable Object/s found. Please, check the Message Log - Mutable for more information."));
-			Info.bFireAndForget = true;
-			Info.bUseThrobber = true;
-			Info.FadeOutDuration = 1.0f;
-			Info.ExpireDuration = 5.0f;
-
-			UncompiledCustomizableObjectsNotificationPtr = FSlateNotificationManager::Get().AddNotification(Info);
-		}
-	}
-#endif
 }
 
 void UCustomizableObjectSystem::EnableBenchmark()
@@ -3125,10 +3128,7 @@ void UCustomizableObjectSystem::OnPreBeginPIE(const bool bIsSimulatingInEditor)
 
 void UCustomizableObjectSystem::StartNextRecompile()
 {
-	if (ObjectsToRecompile.Num() % 10 == 0) // TODO DRB: What?  If we can add more than one between calls, this might never be hit.  Why % 10?
-	{
-		CollectGarbage(GARBAGE_COLLECTION_KEEPFLAGS);
-	}
+	CollectGarbage(GARBAGE_COLLECTION_KEEPFLAGS);
 
 	FAssetData Itr = ObjectsToRecompile.Pop();
 	UCustomizableObject* CustomizableObject = Cast<UCustomizableObject>(Itr.GetAsset());
diff --git a/Engine/Plugins/Experimental/Mutable/Source/CustomizableObjectEditor/Private/MuCOE/CustomizableObjectCompiler.cpp b/Engine/Plugins/Experimental/Mutable/Source/CustomizableObjectEditor/Private/MuCOE/CustomizableObjectCompiler.cpp
--- a/Engine/Plugins/Experimental/Mutable/Source/CustomizableObjectEditor/Private/MuCOE/CustomizableObjectCompiler.cpp
+++ b/Engine/Plugins/Experimental/Mutable/Source/CustomizableObjectEditor/Private/MuCOE/CustomizableObjectCompiler.cpp
@@ -62,7 +62,7 @@ bool FCustomizableObjectCompiler::Tick()
 		else
 		{
 			bUpdated = true;
-			State = ECustomizableObjectCompilationState::Completed;
+			SetCompilationState(ECustomizableObjectCompilationState::Completed);
 
 			RemoveCompileNotification();
 
@@ -82,7 +82,8 @@ bool FCustomizableObjectCompiler::Tick()
 
 	if (SaveDDTask.IsValid() && SaveDDTask->IsCompleted())
 	{
-		State = ECustomizableObjectCompilationState::Completed;
+		SetCompilationState(ECustomizableObjectCompilationState::Completed);
+
 		bUpdated = true;
 
 		FinishSavingDerivedData();
@@ -889,11 +890,22 @@ void FCustomizableObjectCompiler::GetCompilationMessages(TArray<FText>& OutWarni
 }
 
 
+void FCustomizableObjectCompiler::SetCompilationState(ECustomizableObjectCompilationState InState)
+{
+	State = InState;
+	
+	if (CurrentObject)
+	{
+		CurrentObject->CompilationState = InState;
+	}
+}
+
+
 void FCustomizableObjectCompiler::CompileInternal(UCustomizableObject* Object, const FCompilationOptions& InOptions, bool bAsync)
 {
 	MUTABLE_CPUPROFILER_SCOPE(FCustomizableObjectCompiler::Compile)
 	
-	State = ECustomizableObjectCompilationState::Failed;
+	SetCompilationState(ECustomizableObjectCompilationState::Failed);
 
 	if (!Object) return;
 
@@ -1282,7 +1294,7 @@ void FCustomizableObjectCompiler::CompileInternal(UCustomizableObject* Object, c
 
 		CompileTask = MakeShareable(new FCustomizableObjectCompileRunnable(MutableRoot));
 		CompileTask->Options = Options;
-		State = ECustomizableObjectCompilationState::InProgress;
+		SetCompilationState(ECustomizableObjectCompilationState::InProgress);
 
 		if (GenerationContext.ArrayTextureUnrealToMutableTask.Num() > 0)
 		{
@@ -1311,7 +1323,7 @@ void FCustomizableObjectCompiler::CompileInternal(UCustomizableObject* Object, c
 
 			CurrentObject->PostCompile(); 
 
-			State = ECustomizableObjectCompilationState::Completed;
+			SetCompilationState(ECustomizableObjectCompilationState::Completed);
 		}
 		else
 		{
@@ -1569,11 +1581,11 @@ void FCustomizableObjectCompiler::NotifyCompilationErrors() const
 	const FText Prefix = FText::FromString(CurrentObject ? CurrentObject->GetName() : "Customizable Object");
 
 	const FText Message = NoWarningsOrErrors ?
-		FText::Format(LOCTEXT("CompilationFinishedSuccessfully", "{0} finished compiling successfully."), Prefix) :
+		FText::Format(LOCTEXT("CompilationFinishedSuccessfully", "{0} finished compiling."), Prefix) :
 		NumIgnoreds > 0 ?
-		FText::Format(LOCTEXT("CompilationFinished_WithIgnoreds", "{0} finished compiling successfully with {1} {1}|plural(one=warning,other=warnings), {2} {2}|plural(one=error,other=errors) and {3} more similar warnings."), Prefix, NumWarnings, NumErrors, NumIgnoreds)
+		FText::Format(LOCTEXT("CompilationFinished_WithIgnoreds", "{0} finished compiling with {1} {1}|plural(one=warning,other=warnings), {2} {2}|plural(one=error,other=errors) and {3} more similar warnings."), Prefix, NumWarnings, NumErrors, NumIgnoreds)
 		:
-		FText::Format(LOCTEXT("CompilationFinished_WithoutIgnoreds", "{0} finished compiling successfully with {1} {1}|plural(one=warning,other=warnings) and {2} {2}|plural(one=error,other=errors)."), Prefix, NumWarnings, NumErrors);
+		FText::Format(LOCTEXT("CompilationFinished_WithoutIgnoreds", "{0} finished compiling with {1} {1}|plural(one=warning,other=warnings) and {2} {2}|plural(one=error,other=errors)."), Prefix, NumWarnings, NumErrors);
 	
 	FCustomizableObjectEditorLogger::CreateLog(Message)
 	.Category(ELoggerCategory::Compilation)
diff --git a/Engine/Plugins/Experimental/Mutable/Source/CustomizableObjectEditor/Private/MuCOE/CustomizableObjectCompiler.h b/Engine/Plugins/Experimental/Mutable/Source/CustomizableObjectEditor/Private/MuCOE/CustomizableObjectCompiler.h
--- a/Engine/Plugins/Experimental/Mutable/Source/CustomizableObjectEditor/Private/MuCOE/CustomizableObjectCompiler.h
+++ b/Engine/Plugins/Experimental/Mutable/Source/CustomizableObjectEditor/Private/MuCOE/CustomizableObjectCompiler.h
@@ -100,6 +100,8 @@ private:
 	FCompilationMessageCache CompilationLogsContainer;
 	
 private:
+	void SetCompilationState(ECustomizableObjectCompilationState InState);
+
 	ECustomizableObjectCompilationState State = ECustomizableObjectCompilationState::None;
 	
 	void CompileInternal(UCustomizableObject* Object, const FCompilationOptions& Options, bool bAsync = false);
