[mutable] Now, edit material nodes will have all layout blocks selected by default. Also, all block will be selected on change the parent material or the layout index.

#rnx

diff --git a/Engine/Plugins/Experimental/Mutable/Source/CustomizableObjectEditor/Classes/MuCOE/Nodes/CustomizableObjectNodeEditMaterial.h b/Engine/Plugins/Experimental/Mutable/Source/CustomizableObjectEditor/Classes/MuCOE/Nodes/CustomizableObjectNodeEditMaterial.h
--- a/Engine/Plugins/Experimental/Mutable/Source/CustomizableObjectEditor/Classes/MuCOE/Nodes/CustomizableObjectNodeEditMaterial.h
+++ b/Engine/Plugins/Experimental/Mutable/Source/CustomizableObjectEditor/Classes/MuCOE/Nodes/CustomizableObjectNodeEditMaterial.h
@@ -60,6 +60,7 @@ public:
 
 	// UCustomizableObjectNodeEditMaterialBase interface
 	virtual UEdGraphPin* OutputPin() const override;
+	virtual void SetLayoutIndex(const int32 LayoutIndex) override;
 	
 	// ICustomizableObjectNodeParentedMaterial interface
 	virtual void SetParentNode(UCustomizableObject* Object, FGuid NodeId) override;
@@ -75,6 +76,9 @@ public:
 	 * @returns Always returns a valid pin if EditsImage(const FGuid&) returns true. */
 	const UEdGraphPin* GetUsedImageMaskPin(const FGuid& ImageId) const;
 
+	// Function to select all the layout blocks. Called when a parameter is reset (ParentMaterial, ParentLayoutIndex)
+	void SelectAllLayoutBlocks();
+
 private:
 	/** Relates a Parameter id to a Pin. Only used to improve performance. */
 	UPROPERTY()
diff --git a/Engine/Plugins/Experimental/Mutable/Source/CustomizableObjectEditor/Classes/MuCOE/Nodes/CustomizableObjectNodeEditMaterialBase.h b/Engine/Plugins/Experimental/Mutable/Source/CustomizableObjectEditor/Classes/MuCOE/Nodes/CustomizableObjectNodeEditMaterialBase.h
--- a/Engine/Plugins/Experimental/Mutable/Source/CustomizableObjectEditor/Classes/MuCOE/Nodes/CustomizableObjectNodeEditMaterialBase.h
+++ b/Engine/Plugins/Experimental/Mutable/Source/CustomizableObjectEditor/Classes/MuCOE/Nodes/CustomizableObjectNodeEditMaterialBase.h
@@ -25,6 +25,9 @@ public:
 	virtual FGuid GetParentNodeId() const override;
 	virtual UCustomizableObject* GetParentObject() const override;
 
+	// UCustomizableObjectNodeEditMaterialBase interface
+	virtual void SetLayoutIndex(const int32 LayoutIndex);
+
 	// Own interface
 	/** Index of the layout to use to patch blocks. */
 	UPROPERTY(EditAnywhere, Category = CustomizableObject)
diff --git a/Engine/Plugins/Experimental/Mutable/Source/CustomizableObjectEditor/Private/MuCOE/Nodes/CustomizableObjectNodeEditMaterial.cpp b/Engine/Plugins/Experimental/Mutable/Source/CustomizableObjectEditor/Private/MuCOE/Nodes/CustomizableObjectNodeEditMaterial.cpp
--- a/Engine/Plugins/Experimental/Mutable/Source/CustomizableObjectEditor/Private/MuCOE/Nodes/CustomizableObjectNodeEditMaterial.cpp
+++ b/Engine/Plugins/Experimental/Mutable/Source/CustomizableObjectEditor/Private/MuCOE/Nodes/CustomizableObjectNodeEditMaterial.cpp
@@ -88,8 +88,8 @@ void UCustomizableObjectNodeEditMaterial::SetParentNode(UCustomizableObject* Obj
 	Super::SetParentNode(Object, NodeId);
 	
 	PostSetParentNodeWork(Object, NodeId);
-	
-	BlockIds.Empty();
+
+	SelectAllLayoutBlocks();
 }
 
 
@@ -277,5 +277,35 @@ FCustomizableObjectNodeParentedMaterial& UCustomizableObjectNodeEditMaterial::Ge
 }
 
 
+void UCustomizableObjectNodeEditMaterial::SetLayoutIndex(const int32 LayoutIndex)
+{
+	ParentLayoutIndex = LayoutIndex;
+	SelectAllLayoutBlocks();
+}
+
+
+void UCustomizableObjectNodeEditMaterial::SelectAllLayoutBlocks()
+{
+	BlockIds.Empty();
+
+	// On reset parent material we want to select all nodes 
+	if (UCustomizableObjectNodeMaterialBase* ParentMaterialNode = GetParentMaterialNode())
+	{
+		TArray<UCustomizableObjectLayout*> Layouts = ParentMaterialNode->GetLayouts();
+
+		if (Layouts.IsValidIndex(ParentLayoutIndex))
+		{
+			UCustomizableObjectLayout* Layout = Layouts[ParentLayoutIndex];
+
+			for (const FCustomizableObjectLayoutBlock& Block : Layout->Blocks)
+			{
+				BlockIds.AddUnique(Block.Id);
+			}
+		}
+	}
+}
+
+
+
 #undef LOCTEXT_NAMESPACE
 
diff --git a/Engine/Plugins/Experimental/Mutable/Source/CustomizableObjectEditor/Private/MuCOE/Nodes/CustomizableObjectNodeEditMaterialBase.cpp b/Engine/Plugins/Experimental/Mutable/Source/CustomizableObjectEditor/Private/MuCOE/Nodes/CustomizableObjectNodeEditMaterialBase.cpp
--- a/Engine/Plugins/Experimental/Mutable/Source/CustomizableObjectEditor/Private/MuCOE/Nodes/CustomizableObjectNodeEditMaterialBase.cpp
+++ b/Engine/Plugins/Experimental/Mutable/Source/CustomizableObjectEditor/Private/MuCOE/Nodes/CustomizableObjectNodeEditMaterialBase.cpp
@@ -65,3 +65,8 @@ UCustomizableObject* UCustomizableObjectNodeEditMaterialBase::GetParentObject()
 	return ParentMaterialObject;
 }
 
+
+void UCustomizableObjectNodeEditMaterialBase::SetLayoutIndex(const int32 LayoutIndex)
+{
+	ParentLayoutIndex = LayoutIndex;
+}
diff --git a/Engine/Plugins/Experimental/Mutable/Source/CustomizableObjectEditor/Private/MuCOE/Nodes/CustomizableObjectNodeEditMaterialBaseDetails.cpp b/Engine/Plugins/Experimental/Mutable/Source/CustomizableObjectEditor/Private/MuCOE/Nodes/CustomizableObjectNodeEditMaterialBaseDetails.cpp
--- a/Engine/Plugins/Experimental/Mutable/Source/CustomizableObjectEditor/Private/MuCOE/Nodes/CustomizableObjectNodeEditMaterialBaseDetails.cpp
+++ b/Engine/Plugins/Experimental/Mutable/Source/CustomizableObjectEditor/Private/MuCOE/Nodes/CustomizableObjectNodeEditMaterialBaseDetails.cpp
@@ -93,7 +93,7 @@ TSharedPtr<FString> FCustomizableObjectNodeEditMaterialBaseDetails::GenerateLayo
 
 	TSharedPtr<FString> ItemToSelect = nullptr;
 
-	UCustomizableObjectNodeMaterialBase* ParentMaterialNode = NodeEditMaterialBase->GetParentMaterialNode();
+	UCustomizableObjectNodeMaterialBase* ParentMaterialNode = NodeEditMaterialBase->GetParentMaterialNodeIfPath();
 
 	if (ParentMaterialNode && ParentMaterialNode->GetLayouts().Num())
 	{
@@ -121,7 +121,7 @@ void FCustomizableObjectNodeEditMaterialBaseDetails::OnLayoutComboBoxSelectionCh
 	{
 		if (LayoutOptionNames[OptionIndex] == Selection)
 		{
-			LayoutProperty->SetValue(LayoutOptionReferences[OptionIndex]);
+			NodeEditMaterialBase->SetLayoutIndex(LayoutOptionReferences[OptionIndex]);
 			DetailBuilderPtr->ForceRefreshDetails();
 			break;	
 		}
diff --git a/Engine/Plugins/Experimental/Mutable/Source/CustomizableObjectEditor/Private/MuCOE/SCustomizableObjectNodeLayoutBlocksSelector.cpp b/Engine/Plugins/Experimental/Mutable/Source/CustomizableObjectEditor/Private/MuCOE/SCustomizableObjectNodeLayoutBlocksSelector.cpp
--- a/Engine/Plugins/Experimental/Mutable/Source/CustomizableObjectEditor/Private/MuCOE/SCustomizableObjectNodeLayoutBlocksSelector.cpp
+++ b/Engine/Plugins/Experimental/Mutable/Source/CustomizableObjectEditor/Private/MuCOE/SCustomizableObjectNodeLayoutBlocksSelector.cpp
@@ -71,33 +71,40 @@ void SCustomizableObjectNodeLayoutBlocksSelector::SetSelectedNode(UCustomizableO
 {
 	CurrentNode = Node;
 
-	UCustomizableObjectNodeMaterialBase* ParentMaterialNode = Node ? Node->GetParentMaterialNode() : nullptr;
+	UCustomizableObjectNodeMaterialBase* ParentMaterialNode = Node ? Node->GetParentMaterialNodeIfPath() : nullptr;
 
 	// Try to locate the source mesh
 	TArray<FVector2f> UVs;
-	if (CurrentNode && ParentMaterialNode)
+	if (CurrentNode)
 	{
-		const UEdGraphPin* SourceMeshPin = FindMeshBaseSource(*ParentMaterialNode->OutputPin(), false);
-		if (SourceMeshPin)
+		if (ParentMaterialNode)
 		{
-			const UCustomizableObjectNodeSkeletalMesh* MeshNode = Cast<UCustomizableObjectNodeSkeletalMesh>(SourceMeshPin->GetOwningNode());
-			if (MeshNode)
+			const UEdGraphPin* SourceMeshPin = FindMeshBaseSource(*ParentMaterialNode->OutputPin(), false);
+			if (SourceMeshPin)
 			{
-				if (ParentMaterialNode->GetLayouts().Num() > 0 && ParentMaterialNode->GetLayouts().IsValidIndex(CurrentNode->ParentLayoutIndex))
+				const UCustomizableObjectNodeSkeletalMesh* MeshNode = Cast<UCustomizableObjectNodeSkeletalMesh>(SourceMeshPin->GetOwningNode());
+				if (MeshNode)
 				{
-					MeshNode->GetUVChannelForPin(SourceMeshPin, UVs, CurrentNode->ParentLayoutIndex);
+					if (ParentMaterialNode->GetLayouts().Num() > 0 && ParentMaterialNode->GetLayouts().IsValidIndex(CurrentNode->ParentLayoutIndex))
+					{
+						MeshNode->GetUVChannelForPin(SourceMeshPin, UVs, CurrentNode->ParentLayoutIndex);
+					}
 				}
-			}
 
-			const UCustomizableObjectNodeTable* TableNode = Cast<UCustomizableObjectNodeTable>(SourceMeshPin->GetOwningNode());
-			if (TableNode)
-			{
-				if (ParentMaterialNode->GetLayouts().Num() > 0 && ParentMaterialNode->GetLayouts().IsValidIndex(CurrentNode->ParentLayoutIndex))
+				const UCustomizableObjectNodeTable* TableNode = Cast<UCustomizableObjectNodeTable>(SourceMeshPin->GetOwningNode());
+				if (TableNode)
 				{
-					TableNode->GetUVChannelForPin(SourceMeshPin, UVs, CurrentNode->ParentLayoutIndex);
+					if (ParentMaterialNode->GetLayouts().Num() > 0 && ParentMaterialNode->GetLayouts().IsValidIndex(CurrentNode->ParentLayoutIndex))
+					{
+						TableNode->GetUVChannelForPin(SourceMeshPin, UVs, CurrentNode->ParentLayoutIndex);
+					}
 				}
 			}
 		}
+		else
+		{
+			CurrentNode->BlockIds.Empty();
+		}
 	}
 
 	this->ChildSlot
@@ -191,7 +198,7 @@ TSharedRef<SWidget> SCustomizableObjectNodeLayoutBlocksSelector::BuildLayoutTool
 
 void SCustomizableObjectNodeLayoutBlocksSelector::OnSelectAll()
 {
-	UCustomizableObjectNodeMaterialBase* ParentMaterialNode = CurrentNode ? CurrentNode->GetParentMaterialNode() : nullptr;
+	UCustomizableObjectNodeMaterialBase* ParentMaterialNode = CurrentNode ? CurrentNode->GetParentMaterialNodeIfPath() : nullptr;
 
 	if ( CurrentNode && ParentMaterialNode)
 	{
@@ -232,9 +239,9 @@ void SCustomizableObjectNodeLayoutBlocksSelector::OnSelectNone()
 FIntPoint SCustomizableObjectNodeLayoutBlocksSelector::GetGridSize() const
 {
 	FIntPoint Result(1, 1);
-	UCustomizableObjectNodeMaterialBase* ParentMaterialNode = CurrentNode ? CurrentNode->GetParentMaterialNode() : nullptr;
+	UCustomizableObjectNodeMaterialBase* ParentMaterialNode = CurrentNode ? CurrentNode->GetParentMaterialNodeIfPath() : nullptr;
 
-	if ( CurrentNode && ParentMaterialNode)
+	if (ParentMaterialNode)
 	{
 		TArray<UCustomizableObjectLayout*> Layouts = ParentMaterialNode->GetLayouts();
 
@@ -266,7 +273,7 @@ TArray<FCustomizableObjectLayoutBlock> SCustomizableObjectNodeLayoutBlocksSelect
 {
 	TArray<FCustomizableObjectLayoutBlock> Blocks;
 
-	UCustomizableObjectNodeMaterialBase* ParentMaterialNode = CurrentNode ? CurrentNode->GetParentMaterialNode() : nullptr;
+	UCustomizableObjectNodeMaterialBase* ParentMaterialNode = CurrentNode ? CurrentNode->GetParentMaterialNodeIfPath() : nullptr;
 
 	if (ParentMaterialNode)
 	{
