[Mutable] Replace CustomizableInstanceLODManagment check in a pure virtual method with the macro PURE_VIRTUAL

#rnx

diff --git a/Engine/Plugins/Experimental/Mutable/Source/CustomizableObject/Classes/MuCO/CustomizableInstanceLODManagement.h b/Engine/Plugins/Experimental/Mutable/Source/CustomizableObject/Classes/MuCO/CustomizableInstanceLODManagement.h
--- a/Engine/Plugins/Experimental/Mutable/Source/CustomizableObject/Classes/MuCO/CustomizableInstanceLODManagement.h
+++ b/Engine/Plugins/Experimental/Mutable/Source/CustomizableObject/Classes/MuCO/CustomizableInstanceLODManagement.h
@@ -19,15 +19,15 @@ public:
 	virtual ~UCustomizableInstanceLODManagementBase() {};
 
 	// WARNING! The following methods must be overriden in derived classes
-	virtual void UpdateInstanceDistsAndLODs(FMutableInstanceUpdateMap& InOutRequestedUpdates) { check(0); };
+	virtual void UpdateInstanceDistsAndLODs(FMutableInstanceUpdateMap& InOutRequestedUpdates) PURE_VIRTUAL(UCustomizableInstanceLODManagementBase::UpdateInstanceDistsAndLODs, );
 	// Remove the following method after CL 25973936 is propagated to all streams
-	virtual void UpdateInstanceDistsAndLODs() { check(0); };
-	virtual int32 GetNumGeneratedInstancesLimitFullLODs() const { check(0); return 0; };
-	virtual int32 GetNumGeneratedInstancesLimitLOD1() const { check(0); return 0; };
-	virtual int32 GetNumGeneratedInstancesLimitLOD2() const { check(0); return 0; };
-	virtual float GetOnlyUpdateCloseCustomizableObjectsDist() const { check(0); return 0.f; };
-	virtual bool IsOnlyUpdateCloseCustomizableObjectsEnabled() const { check(0); return false; };
-	virtual bool IsOnlyGenerateRequestedLODLevelsEnabled() const { check(0); return false; };
+	virtual void UpdateInstanceDistsAndLODs() PURE_VIRTUAL(UCustomizableInstanceLODManagementBase::UpdateInstanceDistsAndLODs, );
+	virtual int32 GetNumGeneratedInstancesLimitFullLODs() const PURE_VIRTUAL(UCustomizableInstanceLODManagementBase::GetNumGeneratedInstancesLimitFullLODs, return 0;);
+	virtual int32 GetNumGeneratedInstancesLimitLOD1() const PURE_VIRTUAL(UCustomizableInstanceLODManagementBase::GetNumGeneratedInstancesLimitLOD1, return 0;);
+	virtual int32 GetNumGeneratedInstancesLimitLOD2() const PURE_VIRTUAL(UCustomizableInstanceLODManagementBase::GetNumGeneratedInstancesLimitLOD2, return 0;);
+	virtual float GetOnlyUpdateCloseCustomizableObjectsDist() const PURE_VIRTUAL(UCustomizableInstanceLODManagementBase::GetOnlyUpdateCloseCustomizableObjectsDist, return 0.f;);
+	virtual bool IsOnlyUpdateCloseCustomizableObjectsEnabled() const PURE_VIRTUAL(UCustomizableInstanceLODManagementBase::IsOnlyUpdateCloseCustomizableObjectsEnabled, return false;);
+	virtual bool IsOnlyGenerateRequestedLODLevelsEnabled() const PURE_VIRTUAL(UCustomizableInstanceLODManagementBase::IsOnlyGenerateRequestedLODLevelsEnabled, return false;);
 };
 
 
