[Mutable] Replace BoneNames with BoneIds that point to an FName stored in the CO to eliminate strings in the core.
The resources affected by this change are the Meshes, Skeletons, PhysicsBodies, and the bytecode of related operations.
Clean up some commented code.

#rnx

diff --git a/Engine/Plugins/Experimental/Mutable/Source/CustomizableObject/Classes/MuCO/CustomizableObject.h b/Engine/Plugins/Experimental/Mutable/Source/CustomizableObject/Classes/MuCO/CustomizableObject.h
--- a/Engine/Plugins/Experimental/Mutable/Source/CustomizableObject/Classes/MuCO/CustomizableObject.h
+++ b/Engine/Plugins/Experimental/Mutable/Source/CustomizableObject/Classes/MuCO/CustomizableObject.h
@@ -1528,6 +1528,7 @@ public:
 	void SerializeClothingDerivedData(FMemoryWriter64& Ar);
 	void DeserializeClothingDerivedData(FMemoryReaderView& Ar);
 
+	FGuid GetCompilationGuid() const;
 
 	int32 FindState( const FString& Name ) const;
 
@@ -1651,6 +1652,8 @@ public:
 
 #if WITH_EDITOR
 	void SetModel(TSharedPtr<mu::Model, ESPMode::ThreadSafe> Model);
+
+	void SetBoneNamesArray(const TArray<FName>& BoneNames);
 #endif
 
 	int32 GetNumLODs() const;
@@ -1661,6 +1664,17 @@ public:
 	/** Return the names used by mutable to identify which mu::Image should be considered of LowPriority. */
 	void GetLowPriorityTextureNames(TArray<FString>& OutTextureNames);
 
+	const TArray<FName>& GetBoneNamesArray() const;
+
+private:
+	
+	/** Stores the bone names of all the bones that can possibly use the generated meshes */
+	UPROPERTY()
+	TArray<FName> BoneNames;
+
+	/** Cache of merged skeletons */
+	TArray<FMergedSkeleton> MergedSkeletons;
+
 	/** Used to prevent GC of MaskOutCache and keep it in memory while it's needed */
 	UPROPERTY(Transient)
 	TObjectPtr<UMutableMaskOutCache> MaskOutCache_HardRef;
@@ -1669,13 +1683,6 @@ public:
 	UPROPERTY()
 	FGuid CompilationGuid;
 
-	/** Cache of merged skeletons */
-	TArray<FMergedSkeleton> MergedSkeletons;
-
-public:
-	FGuid GetCompilationGuid() const;
-
-private:
 	/** BulkData that stores all in-game resources used by Mutable when generating instances.
 	  * Only valid in packaged builds */
 	UPROPERTY()
diff --git a/Engine/Plugins/Experimental/Mutable/Source/CustomizableObject/Private/MuCO/CustomizableInstancePrivateData.h b/Engine/Plugins/Experimental/Mutable/Source/CustomizableObject/Private/MuCO/CustomizableInstancePrivateData.h
--- a/Engine/Plugins/Experimental/Mutable/Source/CustomizableObject/Private/MuCO/CustomizableInstancePrivateData.h
+++ b/Engine/Plugins/Experimental/Mutable/Source/CustomizableObject/Private/MuCO/CustomizableInstancePrivateData.h
@@ -270,21 +270,21 @@ private:
 
 	void InitSkeletalMeshData(const TSharedPtr<FMutableOperationData>& OperationData, USkeletalMesh* SkeletalMesh, const FMutableRefSkeletalMeshData* RefSkeletalMeshData, int32 ComponentIndex);
 
-	bool BuildSkeletonData(const TSharedPtr<FMutableOperationData>& OperationData, USkeletalMesh* SkeletalMesh, const FMutableRefSkeletalMeshData* RefSkeletalMeshData, UCustomizableObjectInstance* CustomizableObjectIntance, int32 ComponentIndex);
+	bool BuildSkeletonData(const TSharedPtr<FMutableOperationData>& OperationData, USkeletalMesh& SkeletalMesh, const FMutableRefSkeletalMeshData& RefSkeletalMeshData, UCustomizableObject& CustomizableObject, int32 ComponentIndex);
 	void BuildMeshSockets(const TSharedPtr<FMutableOperationData>& OperationData, USkeletalMesh* SkeletalMesh, const FMutableRefSkeletalMeshData* RefSkeletalMeshData, UCustomizableObjectInstance* CustomizableObjectInstance, mu::MeshPtrConst MutableMesh);
 	void BuildOrCopyElementData(const TSharedPtr<FMutableOperationData>& OperationData, USkeletalMesh* SkeletalMesh, UCustomizableObjectInstance* CustomizableObjectInstance, int32 ComponentIndex);
 	void BuildOrCopyMorphTargetsData(const TSharedPtr<FMutableOperationData>& OperationData, USkeletalMesh* SkeletalMesh, const USkeletalMesh* SrcSkeletalMesh, UCustomizableObjectInstance* CustomizableObjectInstance, int32 ComponentIndex);
 	bool BuildOrCopyRenderData(const TSharedPtr<FMutableOperationData>& OperationData, USkeletalMesh* SkeletalMesh, const USkeletalMesh* SrcSkeletalMesh, UCustomizableObjectInstance* CustomizableObjectInstance, int32 ComponentIndex);
 	void BuildOrCopyClothingData(const TSharedPtr<FMutableOperationData>& OperationData, USkeletalMesh* SkeletalMesh, const USkeletalMesh* SrcSkeletalMesh, UCustomizableObjectInstance* CustomizableObjectInstance, int32 ComponentIndex);
 
-	bool CopySkeletonData(const TSharedPtr<FMutableOperationData>& OperationData, USkeletalMesh* SrcSkeletalMesh, USkeletalMesh* DestSkeletalMesh, int32 ComponentIndex);
+	bool CopySkeletonData(const TSharedPtr<FMutableOperationData>& OperationData, USkeletalMesh* SrcSkeletalMesh, USkeletalMesh* DestSkeletalMesh, const UCustomizableObject& CustomizableObject, int32 ComponentIndex);
 	void CopyMeshSockets(USkeletalMesh* SrcSkeletalMesh, USkeletalMesh* DestSkeletalMesh);
 
 	//
-	USkeleton* MergeSkeletons(UCustomizableObjectInstance* Public, const FMutableRefSkeletalMeshData* RefSkeletalMeshData, int32 ComponentIndex);
+	USkeleton* MergeSkeletons(UCustomizableObject& CustomizableObject, const FMutableRefSkeletalMeshData& RefSkeletalMeshData, int32 ComponentIndex);
 
 	//
-	UPhysicsAsset* GetOrBuildMainPhysicsAsset(TObjectPtr<class UPhysicsAsset> TamplateAsset, const mu::PhysicsBody* PhysicsBody, int32 ComponentIndex, bool bDisableCollisionBetweenAssets);
+	UPhysicsAsset* GetOrBuildMainPhysicsAsset(TObjectPtr<class UPhysicsAsset> TamplateAsset, const mu::PhysicsBody* PhysicsBody, const UCustomizableObject& CustomizableObject, int32 ComponentIndex, bool bDisableCollisionBetweenAssets);
 	
 	// Create a transient texture and add it to the TextureTrackerArray
 	UTexture2D* CreateTexture();
diff --git a/Engine/Plugins/Experimental/Mutable/Source/CustomizableObject/Private/MuCO/CustomizableObject.cpp b/Engine/Plugins/Experimental/Mutable/Source/CustomizableObject/Private/MuCO/CustomizableObject.cpp
--- a/Engine/Plugins/Experimental/Mutable/Source/CustomizableObject/Private/MuCO/CustomizableObject.cpp
+++ b/Engine/Plugins/Experimental/Mutable/Source/CustomizableObject/Private/MuCO/CustomizableObject.cpp
@@ -334,6 +334,7 @@ void UCustomizableObject::ClearCompiledData()
 	ClothSharedConfigsData.Empty();
 	SkinWeightProfilesInfo.Empty();
 	AnimBpOverridePhysiscAssetsInfo.Empty();
+	BoneNames.Empty();
 
 #if WITH_EDITORONLY_DATA
 	CustomizableObjectPathMap.Empty();
@@ -431,6 +432,8 @@ void UCustomizableObject::SaveCompiledData(FArchive& MemoryWriter, bool bSkipEdi
 
 	MemoryWriter << AnimBpOverridePhysiscAssetsInfo;
 
+	MemoryWriter << BoneNames;
+
 	MemoryWriter << HashToStreamableBlock;
 
 	// All Editor Only data must be serialized here
@@ -530,6 +533,15 @@ void UCustomizableObject::LoadCompiledData(FArchive& MemoryReader, bool bSkipEdi
 
 		MemoryReader << AnimBpOverridePhysiscAssetsInfo;
 
+		TArray<FString> StringBoneNames;
+		MemoryReader << StringBoneNames;
+
+		BoneNames.Reserve(StringBoneNames.Num());
+		for (const FString& BoneName : StringBoneNames)
+		{
+			BoneNames.Add(FName(*BoneName));
+		}
+
 		MemoryReader << HashToStreamableBlock;
 
 		// All Editor Only data must be loaded here
@@ -1204,6 +1216,11 @@ void UCustomizableObject::SetModel(TSharedPtr<mu::Model, ESPMode::ThreadSafe> Mo
 	UpdateCompiledDataFromModel();
 }
 
+void UCustomizableObject::SetBoneNamesArray(const TArray<FName>& InBoneNames)
+{
+	BoneNames = InBoneNames;
+}
+
 #endif // End WITH_EDITOR
 
 
@@ -1748,6 +1765,12 @@ void UCustomizableObject::GetLowPriorityTextureNames(TArray<FString>& OutTexture
 }
 
 
+const TArray<FName>& UCustomizableObject::GetBoneNamesArray() const
+{
+	return BoneNames;
+}
+
+
 FGuid UCustomizableObject::GetCompilationGuid() const
 {
 	return CompilationGuid;
diff --git a/Engine/Plugins/Experimental/Mutable/Source/CustomizableObject/Private/MuCO/CustomizableObjectInstance.cpp b/Engine/Plugins/Experimental/Mutable/Source/CustomizableObject/Private/MuCO/CustomizableObjectInstance.cpp
--- a/Engine/Plugins/Experimental/Mutable/Source/CustomizableObject/Private/MuCO/CustomizableObjectInstance.cpp
+++ b/Engine/Plugins/Experimental/Mutable/Source/CustomizableObject/Private/MuCO/CustomizableObjectInstance.cpp
@@ -1123,15 +1123,10 @@ bool AreSkeletonsCompatible(const TArray<TObjectPtr<USkeleton>>& InSkeletons)
 #endif
 
 
-USkeleton* UCustomizableInstancePrivateData::MergeSkeletons(UCustomizableObjectInstance* Public, const FMutableRefSkeletalMeshData* RefSkeletalMeshData, int32 ComponentIndex)
+USkeleton* UCustomizableInstancePrivateData::MergeSkeletons(UCustomizableObject& CustomizableObject, const FMutableRefSkeletalMeshData& RefSkeletalMeshData, int32 ComponentIndex)
 {
 	MUTABLE_CPUPROFILER_SCOPE(BuildSkeletonData_MergeSkeletons);
 
-	UCustomizableObject* CustomizableObject = Public->GetCustomizableObject();
-
-	check(CustomizableObject);
-	check(RefSkeletalMeshData);
-	
 	FCustomizableInstanceComponentData* ComponentData = GetComponentData(ComponentIndex);
 	check(ComponentData);
 
@@ -1169,8 +1164,8 @@ USkeleton* UCustomizableInstancePrivateData::MergeSkeletons(UCustomizableObjectI
 	if (!FinalSkeleton)
 	{
 		FString Msg = FString::Printf(TEXT("MergeSkeletons failed for Customizable Object [%s], Instance [%s]. Skeletons involved: "),
-			*CustomizableObject->GetName(),
-			*Public->GetName());
+			*CustomizableObject.GetName(),
+			*GetOuter()->GetName());
 		
 		const int32 SkeletonCount = Params.SkeletonsToMerge.Num();
 		for (int32 SkeletonIndex = 0; SkeletonIndex < SkeletonCount; ++SkeletonIndex)
@@ -1198,7 +1193,7 @@ USkeleton* UCustomizableInstancePrivateData::MergeSkeletons(UCustomizableObjectI
 		}
 
 		// Add Skeleton to the cache
-		CustomizableObject->CacheMergedSkeleton(ComponentIndex, ReferencedSkeletons.SkeletonIds, FinalSkeleton);
+		CustomizableObject.CacheMergedSkeleton(ComponentIndex, ReferencedSkeletons.SkeletonIds, FinalSkeleton);
 		ReferencedSkeletons.SkeletonIds.Empty();
 	}
 	
@@ -1350,7 +1345,7 @@ namespace
 	};
 
 	TObjectPtr<UPhysicsAsset> MakePhysicsAssetFromTemplateAndMutableBody(
-		TObjectPtr<UPhysicsAsset> TemplateAsset, const mu::PhysicsBody* MutablePhysics) 
+		TObjectPtr<UPhysicsAsset> TemplateAsset, const mu::PhysicsBody* MutablePhysics, const UCustomizableObject& CustomizableObject) 
 	{
 		check(TemplateAsset);
 		TObjectPtr<UPhysicsAsset> Result = NewObject<UPhysicsAsset>();
@@ -1365,13 +1360,19 @@ namespace
 
 		Result->bNotForDedicatedServer = TemplateAsset->bNotForDedicatedServer;
 
+		const TArray<FName>& BoneNames = CustomizableObject.GetBoneNamesArray();
 		TMap<FName, int32> BonesInUse;
 
 		const int32 MutablePhysicsBodyCount = MutablePhysics->GetBodyCount();
 		BonesInUse.Reserve(MutablePhysicsBodyCount);
 		for ( int32 I = 0; I < MutablePhysicsBodyCount; ++I )
 		{
-			BonesInUse.Add(FName(MutablePhysics->GetBodyBoneName(I)), I);
+			const uint16 BoneNameId = MutablePhysics->GetBodyBoneId(I);
+			if (BoneNames.IsValidIndex(BoneNameId))
+			{
+				FName BoneName = BoneNames[BoneNameId];
+				BonesInUse.Add(BoneName, I);
+			}
 		}
 
 		const int32 PhysicsAssetBodySetupNum = TemplateAsset->SkeletalBodySetups.Num();
@@ -1471,7 +1472,9 @@ namespace
 
 UPhysicsAsset* UCustomizableInstancePrivateData::GetOrBuildMainPhysicsAsset(
 	TObjectPtr<UPhysicsAsset> TemplateAsset,
-	const mu::PhysicsBody* MutablePhysics, int32 ComponentId,
+	const mu::PhysicsBody* MutablePhysics,
+	const UCustomizableObject& CustomizableObject,
+	int32 ComponentId,
 	bool bDisableCollisionsBetweenDifferentAssets)
 {
 
@@ -1525,154 +1528,19 @@ UPhysicsAsset* UCustomizableInstancePrivateData::GetOrBuildMainPhysicsAsset(
 
 	Result->bNotForDedicatedServer = TemplateAsset->bNotForDedicatedServer;
 
-	//auto MakeAggGeomFromMutablePhysics = [](int32 BodyIndex, const mu::PhysicsBody* MutablePhysicsBody) -> FKAggregateGeom
-	//{
-	//	FKAggregateGeom BodyAggGeom;
-
-	//	auto GetCollisionEnabledFormFlags = [](uint32 Flags) -> ECollisionEnabled::Type
-	//	{
-	//		return ECollisionEnabled::Type( Flags & 0xFF ); 
-	//	};
-
-	//	auto GetContributeToMassFromFlags = [](uint32 Flags) -> bool
-	//	{
-	//		return static_cast<bool>( (Flags >> 8 ) & 1 );
-	//	};
-
-	//	const int32 NumSpheres = MutablePhysicsBody->GetSphereCount( BodyIndex );
-	//	TArray<FKSphereElem>& AggSpheres = BodyAggGeom.SphereElems;
-	//	AggSpheres.Empty(NumSpheres);
-	//	for (int32 I = 0; I < NumSpheres; ++I)
-	//	{
-	//		uint32 Flags = MutablePhysicsBody->GetSphereFlags( BodyIndex, I );
-	//		FString Name = MutablePhysicsBody->GetSphereName( BodyIndex, I );
-
-	//		FVector3f Position;
-	//		float Radius;
-
-	//		MutablePhysicsBody->GetSphere( BodyIndex, I, Position, Radius );
-	//		FKSphereElem& NewElem = AggSpheres.AddDefaulted_GetRef();
-	//		
-	//		NewElem.Center = FVector(Position);
-	//		NewElem.Radius = Radius;
-	//		NewElem.SetContributeToMass( GetContributeToMassFromFlags( Flags ) );
-	//		NewElem.SetCollisionEnabled( GetCollisionEnabledFormFlags( Flags ) );
-	//		NewElem.SetName(FName(*Name));
-	//	}
-	//	
-	//	const int32 NumBoxes = MutablePhysicsBody->GetBoxCount( BodyIndex );
-	//	TArray<FKBoxElem>& AggBoxes = BodyAggGeom.BoxElems;
-	//	AggBoxes.Empty(NumBoxes);
-	//	for (int32 I = 0; I < NumBoxes; ++I)
-	//	{
-	//		uint32 Flags = MutablePhysicsBody->GetBoxFlags( BodyIndex, I );
-	//		FString Name = MutablePhysicsBody->GetBoxName( BodyIndex, I );
-
-	//		FVector3f Position;
-	//		FQuat4f Orientation;
-	//		FVector3f Size;
-	//		MutablePhysicsBody->GetBox( BodyIndex, I, Position, Orientation, Size );
-
-	//		FKBoxElem& NewElem = AggBoxes.AddDefaulted_GetRef();
-	//		
-	//		NewElem.Center = FVector( Position );
-	//		NewElem.Rotation = FRotator( Orientation.Rotator() );
-	//		NewElem.X = Size.X;
-	//		NewElem.Y = Size.Y;
-	//		NewElem.Z = Size.Z;
-	//		NewElem.SetContributeToMass( GetContributeToMassFromFlags( Flags ) );
-	//		NewElem.SetCollisionEnabled( GetCollisionEnabledFormFlags( Flags ) );
-	//		NewElem.SetName( FName(*Name) );
-	//	}
-
-	//	//const int32 NumConvexes = MutablePhysicsBody->GetConvexCount( BodyIndex );
-	//	//TArray<FKConvexElem>& AggConvexes = BodyAggGeom.ConvexElems;
-	//	//AggConvexes.Empty();
-	//	//for (int32 I = 0; I < NumConvexes; ++I)
-	//	//{
-	//	//	uint32 Flags = MutablePhysicsBody->GetConvexFlags( BodyIndex, I );
-	//	//	FString Name = MutablePhysicsBody->GetConvexName( BodyIndex, I );
-
-	//	//	const FVector3f* Vertices;
-	//	//	const int32* Indices;
-	//	//	int32 NumVertices;
-	//	//	int32 NumIndices;
-	//	//	FTransform3f Transform;
-
-	//	//	MutablePhysicsBody->GetConvex( BodyIndex, I, Vertices, NumVertices, Indices, NumIndices, Transform );
-	//	//	
-	//	//	TArrayView<const FVector3f> VerticesView( Vertices, NumVertices );
-	//	//	TArrayView<const int32> IndicesView( Indices, NumIndices );
-	//	//}
-
-	//	TArray<FKSphylElem>& AggSphyls = BodyAggGeom.SphylElems;
-	//	const int32 NumSphyls = MutablePhysicsBody->GetSphylCount( BodyIndex );
-	//	AggSphyls.Empty(NumSphyls);
-
-	//	for (int32 I = 0; I < NumSphyls; ++I)
-	//	{
-	//		uint32 Flags = MutablePhysicsBody->GetSphylFlags( BodyIndex, I );
-	//		FString Name = MutablePhysicsBody->GetSphylName( BodyIndex, I );
-
-	//		FVector3f Position;
-	//		FQuat4f Orientation;
-	//		float Radius;
-	//		float Length;
-
-	//		MutablePhysicsBody->GetSphyl( BodyIndex, I, Position, Orientation, Radius, Length );
-
-	//		FKSphylElem& NewElem = AggSphyls.AddDefaulted_GetRef();
-	//		
-	//		NewElem.Center = FVector( Position );
-	//		NewElem.Rotation = FRotator( Orientation.Rotator() );
-	//		NewElem.Radius = Radius;
-	//		NewElem.Length = Length;
-	//		
-	//		NewElem.SetContributeToMass( GetContributeToMassFromFlags( Flags ) );
-	//		NewElem.SetCollisionEnabled( GetCollisionEnabledFormFlags( Flags ) );
-	//		NewElem.SetName( FName(*Name) );
-	//	}	
-
-	//	TArray<FKTaperedCapsuleElem>& AggTaperedCapsules = BodyAggGeom.TaperedCapsuleElems;
-	//	const int32 NumTaperedCapsules = MutablePhysicsBody->GetTaperedCapsuleCount( BodyIndex );
-	//	AggTaperedCapsules.Empty(NumTaperedCapsules);
-
-	//	for (int32 I = 0; I < NumTaperedCapsules; ++I)
-	//	{
-	//		uint32 Flags = MutablePhysicsBody->GetTaperedCapsuleFlags( BodyIndex, I );
-	//		FString Name = MutablePhysicsBody->GetTaperedCapsuleName( BodyIndex, I );
-
-	//		FVector3f Position;
-	//		FQuat4f Orientation;
-	//		float Radius0;
-	//		float Radius1;
-	//		float Length;
-
-	//		MutablePhysicsBody->GetTaperedCapsule( BodyIndex, I, Position, Orientation, Radius0, Radius1, Length );
-
-	//		FKTaperedCapsuleElem& NewElem = AggTaperedCapsules.AddDefaulted_GetRef();
-	//		
-	//		NewElem.Center = FVector( Position );
-	//		NewElem.Rotation = FRotator( Orientation.Rotator() );
-	//		NewElem.Radius0 = Radius0;
-	//		NewElem.Radius1 = Radius1;
-	//		NewElem.Length = Length;
-	//		
-	//		NewElem.SetContributeToMass( GetContributeToMassFromFlags( Flags ) );
-	//		NewElem.SetCollisionEnabled( GetCollisionEnabledFormFlags( Flags ) );
-	//		NewElem.SetName( FName(*Name) );	
-	//	}	
-
-	//	return BodyAggGeom;
-	//};
-
+	const TArray<FName>& BoneNames = CustomizableObject.GetBoneNamesArray();
 	TMap<FName, int32> BonesInUse;
 
 	const int32 MutablePhysicsBodyCount = MutablePhysics->GetBodyCount();
 	BonesInUse.Reserve(MutablePhysicsBodyCount);
 	for ( int32 I = 0; I < MutablePhysicsBodyCount; ++I )
 	{
-		BonesInUse.Add(FName(MutablePhysics->GetBodyBoneName(I)), I);
+		const uint16 BoneNameId = MutablePhysics->GetBodyBoneId(I);
+		if (BoneNames.IsValidIndex(BoneNameId))
+		{
+			FName BoneName = BoneNames[BoneNameId];
+			BonesInUse.Add(BoneName, I);
+		}
 	}
 
 	// Each array is a set of elements that can collide  
@@ -2116,7 +1984,7 @@ bool UCustomizableInstancePrivateData::UpdateSkeletalMesh_PostBeginUpdate0(UCust
 		if (Component.Mesh)
 		{
 			// Construct a new skeleton, fix up ActiveBones and Bonemap arrays and recompute the RefInvMatrices
-			bSuccess = BuildSkeletonData(OperationData, SkeletalMesh, RefSkeletalMeshData, Public, Component.Id);
+			bSuccess = BuildSkeletonData(OperationData, *SkeletalMesh, *RefSkeletalMeshData, *CustomizableObject, Component.Id);
 
 			if (!bSuccess)
 			{
@@ -2128,7 +1996,7 @@ bool UCustomizableInstancePrivateData::UpdateSkeletalMesh_PostBeginUpdate0(UCust
 			{
 				constexpr bool bDisallowCollisionBetweenAssets = true;
 				UPhysicsAsset* PhysicsAssetResult = GetOrBuildMainPhysicsAsset(
-					RefSkeletalMeshData->PhysicsAsset.Get(), MutablePhysics.get(), Component.Id, bDisallowCollisionBetweenAssets);
+					RefSkeletalMeshData->PhysicsAsset.Get(), MutablePhysics.get(), *CustomizableObject, Component.Id, bDisallowCollisionBetweenAssets);
 
 				SkeletalMesh->SetPhysicsAsset(PhysicsAssetResult);
 
@@ -2177,7 +2045,7 @@ bool UCustomizableInstancePrivateData::UpdateSkeletalMesh_PostBeginUpdate0(UCust
 					PhysicsAssetsUsedByAnimBp.AnimInstancePropertyIndexAndPhysicsAssets.Emplace_GetRef();
 
 				Entry.PropertyIndex = Info.PropertyIndex;
-				Entry.PhysicsAsset = MakePhysicsAssetFromTemplateAndMutableBody(PhysicsAssetTemplate, AdditionalPhysiscBody.get());
+				Entry.PhysicsAsset = MakePhysicsAssetFromTemplateAndMutableBody(PhysicsAssetTemplate, AdditionalPhysiscBody.get(), *CustomizableObject);
 			}
 
 			// Add sockets from the SkeletalMesh of reference and from the MutableMesh
@@ -2190,7 +2058,7 @@ bool UCustomizableInstancePrivateData::UpdateSkeletalMesh_PostBeginUpdate0(UCust
 			check(SrcSkeletalMesh);
 
 			// Set the Skeleton from the previously generated mesh and fix up ActiveBones and Bonemap arrays from new components
-			bSuccess = CopySkeletonData(OperationData, SrcSkeletalMesh, SkeletalMesh, Component.Id);
+			bSuccess = CopySkeletonData(OperationData, SrcSkeletalMesh, SkeletalMesh, *CustomizableObject, Component.Id);
 
 			if (!bSuccess)
 			{
@@ -3561,58 +3429,53 @@ void UCustomizableInstancePrivateData::InitSkeletalMeshData(const TSharedPtr<FMu
 }
 
 
-bool UCustomizableInstancePrivateData::BuildSkeletonData(const TSharedPtr<FMutableOperationData>& OperationData, USkeletalMesh* SkeletalMesh, const FMutableRefSkeletalMeshData* RefSkeletalMeshData, UCustomizableObjectInstance* Public, int32 ComponentIndex)
+bool UCustomizableInstancePrivateData::BuildSkeletonData(const TSharedPtr<FMutableOperationData>& OperationData, USkeletalMesh& SkeletalMesh, const FMutableRefSkeletalMeshData& RefSkeletalMeshData, UCustomizableObject& CustomizableObject, int32 ComponentIndex)
 {
 	MUTABLE_CPUPROFILER_SCOPE(UCustomizableInstancePrivateData::BuildSkeletonData);
 
-	check(SkeletalMesh);
-	check(RefSkeletalMeshData);
-
-	const TObjectPtr<USkeleton> Skeleton = MergeSkeletons(Public, RefSkeletalMeshData, ComponentIndex);
+	const TObjectPtr<USkeleton> Skeleton = MergeSkeletons(CustomizableObject, RefSkeletalMeshData, ComponentIndex);
 	if (!Skeleton)
 	{
 		return false;
 	}
 
-	SkeletalMesh->SetSkeleton(Skeleton);
+	SkeletalMesh.SetSkeleton(Skeleton);
 
-	SkeletalMesh->SetRefSkeleton(Skeleton->GetReferenceSkeleton());
-	FReferenceSkeleton& ReferenceSkeleton = SkeletalMesh->GetRefSkeleton();
+	SkeletalMesh.SetRefSkeleton(Skeleton->GetReferenceSkeleton());
+	FReferenceSkeleton& ReferenceSkeleton = SkeletalMesh.GetRefSkeleton();
 
 	// Check that the bones we need are present in the current Skeleton
 	FInstanceUpdateData::FSkeletonData& MutSkeletonData = OperationData->InstanceUpdateData.Skeletons[ComponentIndex];
-	const int32 MutBoneCount = MutSkeletonData.BoneNames.Num();
+	const int32 MutBoneCount = MutSkeletonData.BoneIds.Num();
 	
-	TMap<FName, uint16> BoneToFinalBoneIndexMap;
+	TMap<uint16, uint16> BoneToFinalBoneIndexMap;
 	BoneToFinalBoneIndexMap.Reserve(MutBoneCount);
 
 	{
 		MUTABLE_CPUPROFILER_SCOPE(BuildSkeletonData_EnsureBonesExist);
 
+		const TArray<FName>& BoneNames = CustomizableObject.GetBoneNamesArray();
+
 		// Ensure all the required bones are present in the skeleton
 		for (int32 BoneIndex = 0; BoneIndex < MutBoneCount; ++BoneIndex)
 		{
-			FName BoneName = MutSkeletonData.BoneNames[BoneIndex];
+			const uint16 BoneId = MutSkeletonData.BoneIds[BoneIndex];
+			check(BoneNames.IsValidIndex(BoneId));
+
+			const FName BoneName = BoneNames[BoneId]; 
 			check(BoneName != NAME_None);
 
-			int32 SourceBoneIndex = ReferenceSkeleton.FindRawBoneIndex(BoneName);
-			
+			const int32 SourceBoneIndex = ReferenceSkeleton.FindRawBoneIndex(BoneName);
 			if (SourceBoneIndex == INDEX_NONE)
 			{
-				if (MutSkeletonData.BoneMatricesWithScale.Find(BoneName))
-				{
-					// Merged skeleton is missing some bones! This happens if one of the skeletons involved in the merge is discarded due to being incompatible with the rest
-					// or if the source mesh is not in sync with the skeleton. 
-					UE_LOG(LogMutable, Warning, TEXT("Building instance: generated mesh has a bone [%s] not present in the reference mesh [%s]. Failing to generate mesh. "),
-						*BoneName.ToString(), *SkeletalMesh->GetName());
-					return false;
-				}
-
-				// The bone is missing but it's not required. Fix the SourceBoneIndex and continue.
-				SourceBoneIndex = 0;
+				// Merged skeleton is missing some bones! This happens if one of the skeletons involved in the merge is discarded due to being incompatible with the rest
+				// or if the source mesh is not in sync with the skeleton. 
+				UE_LOG(LogMutable, Warning, TEXT("Building instance: generated mesh has a bone [%s] not present in the reference mesh [%s]. Failing to generate mesh. "),
+					*BoneName.ToString(), *SkeletalMesh.GetName());
+				return false;
 			}
 
-			BoneToFinalBoneIndexMap.Add(BoneName, SourceBoneIndex);
+			BoneToFinalBoneIndexMap.Add(BoneId, SourceBoneIndex);
 		}
 	}
 
@@ -3629,14 +3492,13 @@ bool UCustomizableInstancePrivateData::BuildSkeletonData(const TSharedPtr<FMutab
 
 			for (uint32 BoneMapIndex = Component.FirstBoneMap; BoneMapIndex < Component.FirstBoneMap + Component.BoneMapCount; ++BoneMapIndex)
 			{
-				const int32 BoneIndex = OperationData->InstanceUpdateData.BoneMaps[BoneMapIndex];
-				const FName& BoneName = MutSkeletonData.BoneNames[BoneIndex];
-				OperationData->InstanceUpdateData.BoneMaps[BoneMapIndex] = BoneToFinalBoneIndexMap[BoneName];
+				const int32 BoneId = OperationData->InstanceUpdateData.BoneMaps[BoneMapIndex];
+				OperationData->InstanceUpdateData.BoneMaps[BoneMapIndex] = BoneToFinalBoneIndexMap[BoneId];
 			}
 
-			for (uint16& BoneIndex : Component.ActiveBones)
+			for (uint16& BoneId : Component.ActiveBones)
 			{
-				BoneIndex = BoneToFinalBoneIndexMap[MutSkeletonData.BoneNames[BoneIndex]];
+				BoneId = BoneToFinalBoneIndexMap[BoneId];
 			}
 			Component.ActiveBones.Sort();
 		}
@@ -3646,28 +3508,27 @@ bool UCustomizableInstancePrivateData::BuildSkeletonData(const TSharedPtr<FMutab
 		MUTABLE_CPUPROFILER_SCOPE(BuildSkeletonData_ApplyPose);
 		
 		const int32 RefRawBoneCount = ReferenceSkeleton.GetRawBoneNum();
-		const int32 MutBonePoseCount = MutSkeletonData.BoneMatricesWithScale.Num();
 
-		TArray<FMatrix44f>& RefBasesInvMatrix = SkeletalMesh->GetRefBasesInvMatrix();
+		TArray<FMatrix44f>& RefBasesInvMatrix = SkeletalMesh.GetRefBasesInvMatrix();
 		RefBasesInvMatrix.Empty(RefRawBoneCount);
 
 		// Initialize the base matrices
-		if (RefRawBoneCount == MutBonePoseCount)
+		if (RefRawBoneCount == MutBoneCount)
 		{
-			RefBasesInvMatrix.AddUninitialized(MutBonePoseCount);
+			RefBasesInvMatrix.AddUninitialized(MutBoneCount);
 		}
 		else
 		{
 			// Bad case, some bone poses are missing, calculate the InvRefMatrices to ensure all transforms are there for the second step 
 			MUTABLE_CPUPROFILER_SCOPE(BuildSkeletonData_CalcInvRefMatrices0);
-			SkeletalMesh->CalculateInvRefMatrices();
+			SkeletalMesh.CalculateInvRefMatrices();
 		}
 
 		// First step is to update the RefBasesInvMatrix for the bones.
-		for (const TPair<FName, FMatrix44f>& BoneMatrix : MutSkeletonData.BoneMatricesWithScale)
+		for (int32 BoneIndex = 0; BoneIndex < MutBoneCount; ++BoneIndex)
 		{
-			const int32 RefSkelBoneIndex = BoneToFinalBoneIndexMap[BoneMatrix.Key];
-			RefBasesInvMatrix[RefSkelBoneIndex] = BoneMatrix.Value;
+			const int32 RefSkelBoneIndex = BoneToFinalBoneIndexMap[MutSkeletonData.BoneIds[BoneIndex]];
+			RefBasesInvMatrix[RefSkelBoneIndex] = MutSkeletonData.BoneMatricesWithScale[BoneIndex];
 		}
 
 		// The second step is to update the pose transforms in the ref skeleton from the BasesInvMatrix
@@ -3690,15 +3551,14 @@ bool UCustomizableInstancePrivateData::BuildSkeletonData(const TSharedPtr<FMutab
 
 	{
 		MUTABLE_CPUPROFILER_SCOPE(BuildSkeletonData_CalcInvRefMatrices);
-		SkeletalMesh->CalculateInvRefMatrices();
+		SkeletalMesh.CalculateInvRefMatrices();
 	}
 
-
 	return true;
 }
 
 
-bool UCustomizableInstancePrivateData::CopySkeletonData(const TSharedPtr<FMutableOperationData>& OperationData, USkeletalMesh* SrcSkeletalMesh, USkeletalMesh* DestSkeletalMesh, int32 ComponentIndex)
+bool UCustomizableInstancePrivateData::CopySkeletonData(const TSharedPtr<FMutableOperationData>& OperationData, USkeletalMesh* SrcSkeletalMesh, USkeletalMesh* DestSkeletalMesh, const UCustomizableObject& CustomizableObject, int32 ComponentIndex)
 {
 	MUTABLE_CPUPROFILER_SCOPE(UCustomizableInstancePrivateData::CopySkeletonData);
 
@@ -3719,38 +3579,37 @@ bool UCustomizableInstancePrivateData::CopySkeletonData(const TSharedPtr<FMutabl
 
 	// Check that the bones we need are present in the current Skeleton
 	FInstanceUpdateData::FSkeletonData& MutSkeletonData = OperationData->InstanceUpdateData.Skeletons[ComponentIndex];
-	const int32 MutBoneCount = MutSkeletonData.BoneNames.Num();
+	const int32 MutBoneCount = MutSkeletonData.BoneIds.Num();
 	
-	TMap<FName, uint16> BoneToFinalBoneIndexMap;
+	TMap<uint16, uint16> BoneToFinalBoneIndexMap;
 	BoneToFinalBoneIndexMap.Reserve(MutBoneCount);
 
 	{
 		MUTABLE_CPUPROFILER_SCOPE(CopySkeletonData_CheckForMissingBones);
 
+		const TArray<FName>& BoneNames = CustomizableObject.GetBoneNamesArray();
+
 		// Ensure all the required bones are present in the skeleton
 		for (int32 BoneIndex = 0; BoneIndex < MutBoneCount; ++BoneIndex)
 		{
-			FName BoneName = MutSkeletonData.BoneNames[BoneIndex];
+			const uint16 BoneId = MutSkeletonData.BoneIds[BoneIndex];
+			check(BoneNames.IsValidIndex(BoneId));
+
+			const FName BoneName = BoneNames[BoneId];
 			check(BoneName != NAME_None);
 
-			int32 SourceBoneIndex = RefSkeleton.FindRawBoneIndex(BoneName);
+			const int32 SourceBoneIndex = RefSkeleton.FindRawBoneIndex(BoneName);
 			
 			if (SourceBoneIndex == INDEX_NONE)
 			{
-				if (MutSkeletonData.BoneMatricesWithScale.Find(BoneName))
-				{
-					// Merged skeleton is missing some bones! This happens if one of the skeletons involved in the merge is discarded due to being incompatible with the rest
-					// or if the source mesh is not in sync with the skeleton. 
-					UE_LOG(LogMutable, Warning, TEXT("Building instance: generated mesh has a bone [%s] not present in the reference mesh [%s]. Failing to generate mesh. "),
-						*BoneName.ToString(), *DestSkeletalMesh->GetName());
-					return false;
-				}
-
-				// The bone is missing but it's not required. Fix the SourceBoneIndex and continue.
-				SourceBoneIndex = 0;
+				// Merged skeleton is missing some bones! This happens if one of the skeletons involved in the merge is discarded due to being incompatible with the rest
+				// or if the source mesh is not in sync with the skeleton. 
+				UE_LOG(LogMutable, Warning, TEXT("Building instance: generated mesh has a bone [%s] not present in the reference mesh [%s]. Failing to generate mesh. "),
+					*BoneName.ToString(), *DestSkeletalMesh->GetName());
+				return false;
 			}
 
-			BoneToFinalBoneIndexMap.Add(BoneName, SourceBoneIndex);
+			BoneToFinalBoneIndexMap.Add(BoneId, SourceBoneIndex);
 		}
 	}
 
@@ -3767,15 +3626,14 @@ bool UCustomizableInstancePrivateData::CopySkeletonData(const TSharedPtr<FMutabl
 
 			for (uint32 BoneMapIndex = Component.FirstBoneMap; BoneMapIndex < Component.FirstBoneMap + Component.BoneMapCount; ++BoneMapIndex)
 			{
-				const int32 BoneIndex = OperationData->InstanceUpdateData.BoneMaps[BoneMapIndex];
-				const FName& BoneName = MutSkeletonData.BoneNames[BoneIndex];
-				OperationData->InstanceUpdateData.BoneMaps[BoneMapIndex] = BoneToFinalBoneIndexMap[BoneName];
+				const uint16 BoneId = OperationData->InstanceUpdateData.BoneMaps[BoneMapIndex];
+				OperationData->InstanceUpdateData.BoneMaps[BoneMapIndex] = BoneToFinalBoneIndexMap[BoneId];
 			}
 
 			TArray<uint16> UniqueActiveBones;
-			for (const uint16 BoneIndex : Component.ActiveBones)
+			for (const uint16 BoneId : Component.ActiveBones)
 			{
-				UniqueActiveBones.AddUnique(BoneToFinalBoneIndexMap[MutSkeletonData.BoneNames[BoneIndex]]);
+				UniqueActiveBones.AddUnique(BoneToFinalBoneIndexMap[BoneId]);
 			}
 
 			Exchange(Component.ActiveBones, UniqueActiveBones);
diff --git a/Engine/Plugins/Experimental/Mutable/Source/CustomizableObject/Private/MuCO/CustomizableObjectSystem.cpp b/Engine/Plugins/Experimental/Mutable/Source/CustomizableObject/Private/MuCO/CustomizableObjectSystem.cpp
--- a/Engine/Plugins/Experimental/Mutable/Source/CustomizableObject/Private/MuCO/CustomizableObjectSystem.cpp
+++ b/Engine/Plugins/Experimental/Mutable/Source/CustomizableObject/Private/MuCO/CustomizableObjectSystem.cpp
@@ -1679,8 +1679,8 @@ namespace impl
 
 			// Fill the data used to generate the RefSkeletalMesh
 			TArray<uint16>& SkeletonIds = OperationData->InstanceUpdateData.Skeletons[ComponentIndex].SkeletonIds;
-			TArray<FName>& BoneNames = OperationData->InstanceUpdateData.Skeletons[ComponentIndex].BoneNames;
-			TMap<FName, FMatrix44f>& BoneMatricesWithScale = OperationData->InstanceUpdateData.Skeletons[ComponentIndex].BoneMatricesWithScale;
+			TArray<uint16>& BoneIds = OperationData->InstanceUpdateData.Skeletons[ComponentIndex].BoneIds;
+			TArray<FMatrix44f>& BoneMatricesWithScale = OperationData->InstanceUpdateData.Skeletons[ComponentIndex].BoneMatricesWithScale;
 
 			// Use first valid LOD bone count as a potential total number of bones, used for pre-allocating data arrays
 			if (MinLODComponent.Mesh && MinLODComponent.Mesh->GetSkeleton())
@@ -1688,7 +1688,7 @@ namespace impl
 				const int32 TotalPossibleBones = MinLODComponent.Mesh->GetSkeleton()->GetBoneCount();
 
 				// Out Data
-				BoneNames.Reserve(TotalPossibleBones);
+				BoneIds.Reserve(TotalPossibleBones);
 				BoneMatricesWithScale.Reserve(TotalPossibleBones);
 			}
 
@@ -1700,7 +1700,7 @@ namespace impl
 				FInstanceUpdateData::FComponent& CurrentLODComponent = OperationData->InstanceUpdateData.Components[CurrentLOD.FirstComponent + ComponentIndex];
 				mu::MeshPtrConst Mesh = CurrentLODComponent.Mesh;
 
-				if (!Mesh || !Mesh->GetSkeleton())
+				if (!Mesh)
 				{
 					continue;
 				}
@@ -1712,39 +1712,28 @@ namespace impl
 					SkeletonIds.AddUnique(Mesh->GetSkeletonID(SkeletonIndex));
 				}
 
-				// Add bonemaps from this mesh
+				// Append BoneMap to the array of BoneMaps
 				const TArray<uint16>& BoneMap = Mesh->GetBoneMap();
 				CurrentLODComponent.FirstBoneMap = OperationData->InstanceUpdateData.BoneMaps.Num();
 				CurrentLODComponent.BoneMapCount = BoneMap.Num();
+				OperationData->InstanceUpdateData.BoneMaps.Append(BoneMap);
 
-				OperationData->InstanceUpdateData.BoneMaps.Reserve(CurrentLODComponent.FirstBoneMap + CurrentLODComponent.BoneMapCount);
-
-				const mu::SkeletonPtrConst Skeleton = Mesh->GetSkeleton();
-				for (const uint16& BoneIndex : BoneMap)
-				{
-					const FName BoneName = Skeleton->GetBoneName(BoneIndex);
-					check(BoneName != NAME_None);
-
-					const int32 FinalBoneIndex = BoneNames.AddUnique(BoneName);
-					OperationData->InstanceUpdateData.BoneMaps.Add(FinalBoneIndex);
-				}
-				
 				// Add active bone indices and poses
 				const int32 MaxBoneIndex = Mesh->GetBonePoseCount();
 				CurrentLODComponent.ActiveBones.Reserve(MaxBoneIndex);
 				for (int32 BonePoseIndex = 0; BonePoseIndex < MaxBoneIndex; ++BonePoseIndex)
 				{
-					const FName BoneName = Mesh->GetBonePoseName(BonePoseIndex);
-					check(BoneName != NAME_None);
+					const uint16 BoneId = Mesh->GetBonePoseBoneId(BonePoseIndex);
 
-					const int32 BoneIndex = BoneNames.AddUnique(BoneName); 
-					CurrentLODComponent.ActiveBones.Add(BoneIndex);
+					CurrentLODComponent.ActiveBones.Add(BoneId);
 
-					if (!BoneMatricesWithScale.Contains(BoneName))
+					if(BoneIds.Find(BoneId) == INDEX_NONE)
 					{
+						BoneIds.Add(BoneId);
+
 						FTransform3f Transform;
 						Mesh->GetBoneTransform(BonePoseIndex, Transform);
-						BoneMatricesWithScale.Emplace(BoneName, Transform.Inverse().ToMatrixWithScale());
+						BoneMatricesWithScale.Emplace(Transform.Inverse().ToMatrixWithScale());
 					}
 				}
 			}
diff --git a/Engine/Plugins/Experimental/Mutable/Source/CustomizableObject/Private/MuCO/CustomizableObjectSystemPrivate.h b/Engine/Plugins/Experimental/Mutable/Source/CustomizableObject/Private/MuCO/CustomizableObjectSystemPrivate.h
--- a/Engine/Plugins/Experimental/Mutable/Source/CustomizableObject/Private/MuCO/CustomizableObjectSystemPrivate.h
+++ b/Engine/Plugins/Experimental/Mutable/Source/CustomizableObject/Private/MuCO/CustomizableObjectSystemPrivate.h
@@ -495,8 +495,8 @@ struct FInstanceUpdateData
 
 		TArray<uint16> SkeletonIds;
 
-		TArray<FName> BoneNames;
-		TMap<FName, FMatrix44f> BoneMatricesWithScale;
+		TArray<uint16> BoneIds;
+		TArray<FMatrix44f> BoneMatricesWithScale;
 	};
 
 	TArray<FSkeletonData> Skeletons;
diff --git a/Engine/Plugins/Experimental/Mutable/Source/CustomizableObject/Private/MuCO/UnrealConversionUtils.cpp b/Engine/Plugins/Experimental/Mutable/Source/CustomizableObject/Private/MuCO/UnrealConversionUtils.cpp
--- a/Engine/Plugins/Experimental/Mutable/Source/CustomizableObject/Private/MuCO/UnrealConversionUtils.cpp
+++ b/Engine/Plugins/Experimental/Mutable/Source/CustomizableObject/Private/MuCO/UnrealConversionUtils.cpp
@@ -245,8 +245,7 @@ namespace UnrealConversionUtils
 
 		if (BoneWeightFormat == mu::MBF_NUINT16)
 		{
-			// TODO: 
-			unimplemented()
+			LODModel.SkinWeightVertexBuffer.SetUse16BitBoneWeight(true);
 		}
 
 		// Init skin weight buffer
diff --git a/Engine/Plugins/Experimental/Mutable/Source/CustomizableObjectEditor/Classes/MuCOE/Nodes/CustomizableObjectNodeAnimationPose.h b/Engine/Plugins/Experimental/Mutable/Source/CustomizableObjectEditor/Classes/MuCOE/Nodes/CustomizableObjectNodeAnimationPose.h
--- a/Engine/Plugins/Experimental/Mutable/Source/CustomizableObjectEditor/Classes/MuCOE/Nodes/CustomizableObjectNodeAnimationPose.h
+++ b/Engine/Plugins/Experimental/Mutable/Source/CustomizableObjectEditor/Classes/MuCOE/Nodes/CustomizableObjectNodeAnimationPose.h
@@ -40,6 +40,6 @@ public:
 	void AllocateDefaultPins(UCustomizableObjectNodeRemapPins* RemapPins) override;
 
 	// Retrieve the pose information from the PoseAsset
-	static void StaticRetrievePoseInformation(UPoseAsset* PoseAsset, class USkeletalMesh* RefSkeletalMesh, TArray<FString>& OutArrayBoneName, TArray<FTransform>& OutArrayTransform);
+	static void StaticRetrievePoseInformation(UPoseAsset* PoseAsset, class USkeletalMesh* RefSkeletalMesh, TArray<FName>& OutArrayBoneName, TArray<FTransform>& OutArrayTransform);
 };
 
diff --git a/Engine/Plugins/Experimental/Mutable/Source/CustomizableObjectEditor/Private/MuCOE/CustomizableObjectCompiler.cpp b/Engine/Plugins/Experimental/Mutable/Source/CustomizableObjectEditor/Private/MuCOE/CustomizableObjectCompiler.cpp
--- a/Engine/Plugins/Experimental/Mutable/Source/CustomizableObjectEditor/Private/MuCOE/CustomizableObjectCompiler.cpp
+++ b/Engine/Plugins/Experimental/Mutable/Source/CustomizableObjectEditor/Private/MuCOE/CustomizableObjectCompiler.cpp
@@ -1208,6 +1208,8 @@ void FCustomizableObjectCompiler::CompileInternal(UCustomizableObject* Object, c
 			Object->ReferencedSkeletons.Add(Skeleton);
 		}
 
+		Object->SetBoneNamesArray(GenerationContext.BoneNames);
+
 		// Pass-through textures
 		TArray<TSoftObjectPtr<UTexture2D>> NewReferencedPassThroughTextures;
 
diff --git a/Engine/Plugins/Experimental/Mutable/Source/CustomizableObjectEditor/Private/MuCOE/GenerateMutableSource/GenerateMutableSource.cpp b/Engine/Plugins/Experimental/Mutable/Source/CustomizableObjectEditor/Private/MuCOE/GenerateMutableSource/GenerateMutableSource.cpp
--- a/Engine/Plugins/Experimental/Mutable/Source/CustomizableObjectEditor/Private/MuCOE/GenerateMutableSource/GenerateMutableSource.cpp
+++ b/Engine/Plugins/Experimental/Mutable/Source/CustomizableObjectEditor/Private/MuCOE/GenerateMutableSource/GenerateMutableSource.cpp
@@ -574,7 +574,7 @@ UTexture2D* FindReferenceImage(const UEdGraphPin* Pin, FMutableGraphGenerationCo
 }
 
 
-mu::NodeMeshApplyPosePtr CreateNodeMeshApplyPose(mu::NodeMeshPtr InputMeshNode, UCustomizableObject * CustomizableObject, TArray<FString> ArrayBoneName, TArray<FTransform> ArrayTransform)
+mu::NodeMeshApplyPosePtr CreateNodeMeshApplyPose(FMutableGraphGenerationContext& GenerationContext, mu::NodeMeshPtr InputMeshNode, const TArray<FName>& ArrayBoneName, const TArray<FTransform>& ArrayTransform)
 {
 	check(ArrayBoneName.Num() == ArrayTransform.Num());
 
@@ -589,8 +589,12 @@ mu::NodeMeshApplyPosePtr CreateNodeMeshApplyPose(mu::NodeMeshPtr InputMeshNode,
 
 	for (int32 i = 0; i < ArrayBoneName.Num(); ++i)
 	{
-		MutableSkeleton->SetBoneName(i, StringCast<ANSICHAR>(*ArrayBoneName[i]).Get());
-		MutableMesh->SetBonePose(i, StringCast<ANSICHAR>(*ArrayBoneName[i]).Get(), (FTransform3f)ArrayTransform[i], mu::EBoneUsageFlags::Skinning);
+		const FName BoneName = ArrayBoneName[i];
+		const uint16 BoneId = uint16(GenerationContext.BoneNames.AddUnique(BoneName));
+
+		MutableSkeleton->SetBoneFName(i, BoneName);
+		MutableSkeleton->SetBoneId(i, BoneId);
+		MutableMesh->SetBonePose(i, BoneId, (FTransform3f)ArrayTransform[i], mu::EBoneUsageFlags::Skinning);
 	}
 
 	mu::NodeMeshApplyPosePtr NodeMeshApplyPose = new mu::NodeMeshApplyPose;
diff --git a/Engine/Plugins/Experimental/Mutable/Source/CustomizableObjectEditor/Private/MuCOE/GenerateMutableSource/GenerateMutableSource.h b/Engine/Plugins/Experimental/Mutable/Source/CustomizableObjectEditor/Private/MuCOE/GenerateMutableSource/GenerateMutableSource.h
--- a/Engine/Plugins/Experimental/Mutable/Source/CustomizableObjectEditor/Private/MuCOE/GenerateMutableSource/GenerateMutableSource.h
+++ b/Engine/Plugins/Experimental/Mutable/Source/CustomizableObjectEditor/Private/MuCOE/GenerateMutableSource/GenerateMutableSource.h
@@ -267,7 +267,7 @@ struct FTextureUnrealToMutableTask
 
 struct FPoseBoneData
 {
-	TArray<FString> ArrayBoneName;
+	TArray<FName> ArrayBoneName;
 	TArray<FTransform> ArrayTransform;
 };
 
@@ -613,6 +613,9 @@ struct FMutableGraphGenerationContext
 
 	TArray<const USkeleton*> ReferencedSkeletons;
 
+	// Array of unique BoneNames. All bones from mu::Skeletons will point to the names of it.
+	TArray<FName> BoneNames;
+
 	// Used to aviod Nodes with duplicated ids
 	TMap<FGuid, TArray<const UCustomizableObjectNode*>> NodeIdsMap;
 	TMultiMap<const UCustomizableObject*, FGroupNodeIdsTempData> DuplicatedGroupNodeIds;
@@ -779,7 +782,7 @@ void CheckNumOutputs(const UEdGraphPin& Pin, const FMutableGraphGenerationContex
 UTexture2D* FindReferenceImage(const UEdGraphPin* Pin, FMutableGraphGenerationContext& GenerationContext);
 
 
-mu::NodeMeshApplyPosePtr CreateNodeMeshApplyPose(mu::NodeMeshPtr InputMeshNode, UCustomizableObject * CustomizableObject, TArray<FString> ArrayBoneName, TArray<FTransform> ArrayTransform);
+mu::NodeMeshApplyPosePtr CreateNodeMeshApplyPose(FMutableGraphGenerationContext& GenerationContext, mu::NodeMeshPtr InputMeshNode, const TArray<FName>& ArrayBoneName, const TArray<FTransform>& ArrayTransform);
 
 
 /** Adds Tag to MutableMesh uniquely, returns the index were the tag has been inserted or the index where an intance of the tag has been found */
diff --git a/Engine/Plugins/Experimental/Mutable/Source/CustomizableObjectEditor/Private/MuCOE/GenerateMutableSource/GenerateMutableSourceGroupProjector.cpp b/Engine/Plugins/Experimental/Mutable/Source/CustomizableObjectEditor/Private/MuCOE/GenerateMutableSource/GenerateMutableSourceGroupProjector.cpp
--- a/Engine/Plugins/Experimental/Mutable/Source/CustomizableObjectEditor/Private/MuCOE/GenerateMutableSource/GenerateMutableSourceGroupProjector.cpp
+++ b/Engine/Plugins/Experimental/Mutable/Source/CustomizableObjectEditor/Private/MuCOE/GenerateMutableSource/GenerateMutableSourceGroupProjector.cpp
@@ -153,7 +153,7 @@ mu::NodeImagePtr GenerateMutableGroupProjection(const int32 NodeLOD, const int32
 			{
 				if (ProjectorTempData.CustomizableObjectNodeGroupProjectorParameter->OptionPoses[SelectorIndex].OptionPose)
 				{
-					mu::NodeMeshApplyPosePtr NodeMeshApplyPose = CreateNodeMeshApplyPose(MeshNode, GenerationContext.Object,
+					mu::NodeMeshApplyPosePtr NodeMeshApplyPose = CreateNodeMeshApplyPose(GenerationContext, MeshNode,
 						ProjectorTempData.PoseBoneDataArray[SelectorIndex].ArrayBoneName, ProjectorTempData.PoseBoneDataArray[SelectorIndex].ArrayTransform);
 
 					if (!NodeMeshApplyPose)
diff --git a/Engine/Plugins/Experimental/Mutable/Source/CustomizableObjectEditor/Private/MuCOE/GenerateMutableSource/GenerateMutableSourceMesh.cpp b/Engine/Plugins/Experimental/Mutable/Source/CustomizableObjectEditor/Private/MuCOE/GenerateMutableSource/GenerateMutableSourceMesh.cpp
--- a/Engine/Plugins/Experimental/Mutable/Source/CustomizableObjectEditor/Private/MuCOE/GenerateMutableSource/GenerateMutableSourceMesh.cpp
+++ b/Engine/Plugins/Experimental/Mutable/Source/CustomizableObjectEditor/Private/MuCOE/GenerateMutableSource/GenerateMutableSourceMesh.cpp
@@ -378,10 +378,7 @@ bool IsSkeletalMeshCompatibleWithRefSkeleton(FMutableComponentInfo& ComponentInf
 	return true;
 }
 
-void SetAndPropagatePoseBoneUsage(
-		mu::Mesh& MutableMesh, int32 PoseIndex, mu::EBoneUsageFlags Usage, 
-		const TMap<FName, int32>* BoneNameToSkeletonIndexMap = nullptr,
-		const TMap<FName, int32>* BoneNameToPoseIndexMap = nullptr)
+void SetAndPropagatePoseBoneUsage(mu::Mesh& MutableMesh, int32 PoseIndex, mu::EBoneUsageFlags Usage)
 {
 	if (!MutableMesh.GetSkeleton())
 	{
@@ -390,39 +387,17 @@ void SetAndPropagatePoseBoneUsage(
 
 	const mu::Skeleton& MutableSkeleton = *MutableMesh.GetSkeleton();
 
-	if (PoseIndex < 0 || PoseIndex >= MutableMesh.BonePoses.Num())
+	if (!MutableMesh.BonePoses.IsValidIndex(PoseIndex))
 	{
 		check(false);
 		return;
 	}
 
-	int32 BoneIndex = [&]() 
-	{
-		if (BoneNameToSkeletonIndexMap)
-		{
-			const int32* Found = BoneNameToSkeletonIndexMap->Find(FName(MutableMesh.BonePoses[PoseIndex].BoneName.c_str()));
-			return Found ? *Found : INDEX_NONE;
-		}
-		else
-		{
-			return MutableSkeleton.FindBone(MutableMesh.BonePoses[PoseIndex].BoneName.c_str());
-		}
-	}();
+	int32 BoneIndex = MutableSkeleton.FindBone(MutableMesh.BonePoses[PoseIndex].BoneId);
 
 	while (BoneIndex != INDEX_NONE)
 	{
-		PoseIndex = [&]()
-		{
-			if (BoneNameToPoseIndexMap)
-			{
-				const int32* Found = BoneNameToPoseIndexMap->Find(FName(MutableSkeleton.GetBoneName(BoneIndex)));
-				return Found ? *Found : INDEX_NONE;
-			}
-			else
-			{
-				return MutableMesh.FindBonePose(MutableSkeleton.GetBoneName(BoneIndex));
-			}
-		}();
+		PoseIndex = MutableMesh.FindBonePose(MutableSkeleton.GetBoneId(BoneIndex));
 
 		if (PoseIndex == INDEX_NONE)
 		{
@@ -472,7 +447,7 @@ TArray<TTuple<UPhysicsAsset*, int32>> GetPhysicsAssetsFromAnimInstance(const TSo
 	return Result;
 }
 
-TArray<uint8> MakePhysicsAssetBodySetupRelevancyMap(UPhysicsAsset* Asset, const mu::Ptr<mu::Mesh>& Mesh)
+TArray<uint8> MakePhysicsAssetBodySetupRelevancyMap(const FMutableGraphGenerationContext& GenerationContext, UPhysicsAsset* Asset, const mu::Ptr<mu::Mesh>& Mesh)
 {
 	const int32 BodySetupsNum = Asset->SkeletalBodySetups.Num();
 
@@ -486,15 +461,14 @@ TArray<uint8> MakePhysicsAssetBodySetupRelevancyMap(UPhysicsAsset* Asset, const
 
 	for (int32 BodyIndex = 0; BodyIndex < BodySetupsNum; ++BodyIndex)
 	{
-		FString BodyBoneName = Asset->SkeletalBodySetups[BodyIndex]->BoneName.ToString();
-
-		RelevancyMap[BodyIndex] = Mesh->GetSkeleton()->FindBone(TCHAR_TO_ANSI(*BodyBoneName)) >= 0;
+		const int32 BoneNameId = GenerationContext.BoneNames.Find(Asset->SkeletalBodySetups[BodyIndex]->BoneName);
+		RelevancyMap[BodyIndex] = BoneNameId != INDEX_NONE && Mesh->GetSkeleton()->FindBone(uint16(BoneNameId)) != INDEX_NONE;
 	}
 
 	return RelevancyMap;
 }
 
-mu::Ptr<mu::PhysicsBody> MakePhysicsBodyFromAsset(UPhysicsAsset* Asset, const TArray<uint8>& BodySetupRelevancyMap)
+mu::Ptr<mu::PhysicsBody> MakePhysicsBodyFromAsset(FMutableGraphGenerationContext& GenerationContext, UPhysicsAsset* Asset, const TArray<uint8>& BodySetupRelevancyMap)
 {
 	check(Asset);
 	check(Asset->SkeletalBodySetups.Num() == BodySetupRelevancyMap.Num());
@@ -527,8 +501,8 @@ mu::Ptr<mu::PhysicsBody> MakePhysicsBodyFromAsset(UPhysicsAsset* Asset, const TA
 
 		TObjectPtr<USkeletalBodySetup>& BodySetup = SkeletalBodySetups[SourceBodyIndex++];
 
-		FString BodyBoneName = BodySetup->BoneName.ToString();
-		PhysicsBody->SetBodyBoneName(B, TCHAR_TO_ANSI(*BodyBoneName));
+		const uint16 BodyBoneId = GenerationContext.BoneNames.AddUnique(BodySetup->BoneName);
+		PhysicsBody->SetBodyBoneId(B, BodyBoneId);
 		
 		const int32 NumSpheres = BodySetup->AggGeom.SphereElems.Num();
 		PhysicsBody->SetSphereCount(B, NumSpheres);
@@ -784,7 +758,7 @@ mu::MeshPtr ConvertSkeletalMeshToMutable(const USkeletalMesh* InSkeletalMesh, co
 			RequiredBones.AddUnique(RemappedBones[RequiredBoneIndex]);
 		}
 
-		// Build BoneMap
+		// Rebuild BoneMap
 		const TArray<uint16>& SourceBoneMap = ImportedModel->LODModels[LODIndex].Sections[SectionIndex].BoneMap;
 		const int32 NumBonesInBoneMap = SourceBoneMap.Num();
 		const int32 NumRemappedBones = RemappedBones.Num();
@@ -800,58 +774,55 @@ mu::MeshPtr ConvertSkeletalMeshToMutable(const USkeletalMesh* InSkeletalMesh, co
 			bBoneMapModified = bBoneMapModified || SourceBoneMap[BoneIndex] != FinalBoneIndex;
 		}
 
-		const int32 NumBonesBoneMap = BoneMap.Num();
-		const int32 NumRequiredBones = RequiredBones.Num();
-
-		// BoneMap mapping the BoneMap indices to those of the mu::skeleton
-		TArray<uint16> MutableBoneMap;
-		MutableBoneMap.Reserve(NumBonesBoneMap);
-
-		// Mapping of RequiredBones from RefSkeletonIndex to mu::Skeleton BoneIndex
-		TMap<int32, int32> InverseBoneMap;
-		InverseBoneMap.Reserve(NumRequiredBones);
-
-		for (int32 BoneIndex = 0; BoneIndex < NumBonesBoneMap; ++BoneIndex)
-		{
-			int32 RefSkeletonIndex = BoneMap[BoneIndex];
-			InverseBoneMap.Add(RefSkeletonIndex, BoneIndex);
-			MutableBoneMap.Add(BoneIndex);
-		}
-
-		MutableMesh->SetBoneMap(MutableBoneMap);
-
-		// Create the skeleton and poses for this mesh
+		// Create the skeleton, poses, and BoneMap for this mesh
 		mu::SkeletonPtr MutableSkeleton = new mu::Skeleton;
 		MutableMesh->SetSkeleton(MutableSkeleton);
 
+		const int32 NumRequiredBones = RequiredBones.Num();
 		MutableMesh->SetBonePoseCount(NumRequiredBones);
 		MutableSkeleton->SetBoneCount(NumRequiredBones);
 
+		// MutableBoneMap will not keep an index to the Skeleton, but to the BoneName
+		TArray<uint16> MutableBoneMap;
+		MutableBoneMap.SetNum(BoneMap.Num());
+		
 		const TArray<FMeshBoneInfo>& RefBoneInfo = InSkeletalMesh->GetRefSkeleton().GetRefBoneInfo();
-		for (int32 RequiredBoneIndex = 0; RequiredBoneIndex < NumRequiredBones; ++RequiredBoneIndex)
+		for (int32 BoneIndex = 0; BoneIndex < NumRequiredBones; ++BoneIndex)
 		{
-			const int32 RefSkelIndex = RequiredBones[RequiredBoneIndex];
-			const FMeshBoneInfo& BoneInfo = RefBoneInfo[RefSkelIndex];
-			
-			const int32 BoneIndex = InverseBoneMap.FindOrAdd(RefSkelIndex, InverseBoneMap.Num());
-			const int32 ParentBoneIndex = BoneInfo.ParentIndex != INDEX_NONE ? InverseBoneMap.FindOrAdd(BoneInfo.ParentIndex, InverseBoneMap.Num()) : INDEX_NONE;
+			const int32 RefSkeletonBoneIndex = RequiredBones[BoneIndex];
+
+			const FMeshBoneInfo& BoneInfo = RefBoneInfo[RefSkeletonBoneIndex];
+			const int32 ParentBoneIndex = RequiredBones.Find(BoneInfo.ParentIndex);
 
 			// Set bone hierarchy
-			const FString BoneName = BoneInfo.Name.ToString();
-			MutableSkeleton->SetBoneName(BoneIndex, StringCast<ANSICHAR>(*BoneName).Get());
+			const uint16 BoneId = uint16(GenerationContext.BoneNames.AddUnique(BoneInfo.Name));
+
+			MutableSkeleton->SetBoneId(BoneIndex, BoneId);
 			MutableSkeleton->SetBoneParent(BoneIndex, ParentBoneIndex);
 
+			// Debug. Will not be serialized
+			MutableSkeleton->SetBoneFName(BoneIndex, BoneInfo.Name);
+
+			// BoneMap: Convert RefSkeletonBoneIndex to BoneId
+			const int32 BoneMapIndex = BoneMap.Find(RefSkeletonBoneIndex);
+			if (BoneMapIndex != INDEX_NONE)
+			{
+				MutableBoneMap[BoneMapIndex] = BoneId;
+			}
+
 			// Set bone pose
-			FMatrix44f BaseInvMatrix = InSkeletalMesh->GetRefBasesInvMatrix()[RefSkelIndex];
+			FMatrix44f BaseInvMatrix = InSkeletalMesh->GetRefBasesInvMatrix()[RefSkeletonBoneIndex];
 			FTransform3f BaseInvTransform;
 			BaseInvTransform.SetFromMatrix(BaseInvMatrix);
 
 			mu::EBoneUsageFlags BoneUsageFlags = mu::EBoneUsageFlags::None;
-			EnumAddFlags(BoneUsageFlags, BoneIndex < NumBonesBoneMap ? mu::EBoneUsageFlags::Skinning : mu::EBoneUsageFlags::None);
-			EnumAddFlags(BoneUsageFlags, BoneInfo.ParentIndex == INDEX_NONE ? mu::EBoneUsageFlags::Root : mu::EBoneUsageFlags::None);
+			EnumAddFlags(BoneUsageFlags, BoneMapIndex != INDEX_NONE ? mu::EBoneUsageFlags::Skinning : mu::EBoneUsageFlags::None);
+			EnumAddFlags(BoneUsageFlags, ParentBoneIndex == INDEX_NONE ? mu::EBoneUsageFlags::Root : mu::EBoneUsageFlags::None);
 
-			MutableMesh->SetBonePose(BoneIndex, StringCast<ANSICHAR>(*BoneName).Get(), BaseInvTransform.Inverse(), BoneUsageFlags);
+			MutableMesh->SetBonePose(BoneIndex, BoneId, BaseInvTransform.Inverse(), BoneUsageFlags);
 		}
+
+		MutableMesh->SetBoneMap(MutableBoneMap);
 	}
 
 	// Vertices
@@ -1569,7 +1540,7 @@ mu::MeshPtr ConvertSkeletalMeshToMutable(const USkeletalMesh* InSkeletalMesh, co
 		TArray<uint8> DiscardedBodySetups;
 		DiscardedBodySetups.Init(1, SkeletalBodySetups.Num());
 
-		for (int32 BodySetupIndex = 0; BodySetupIndex < SkeletalBodySetups.Num(); ++BodySetupIndex )
+		for (int32 BodySetupIndex = 0; BodySetupIndex < SkeletalBodySetups.Num(); ++BodySetupIndex)
 		{
 			TObjectPtr<USkeletalBodySetup>& BodySetup = SkeletalBodySetups[BodySetupIndex];
 			if (!BodySetup)
@@ -1577,22 +1548,17 @@ mu::MeshPtr ConvertSkeletalMeshToMutable(const USkeletalMesh* InSkeletalMesh, co
 				continue;
 			}
 
-			FString BodyBoneName = BodySetup->BoneName.ToString();
+			const int32 BoneNameIndex = GenerationContext.BoneNames.Find(BodySetup->BoneName);
+			const int32 BonePoseIndex = MutableMesh->FindBonePose(uint16(BoneNameIndex));
 
-			const int32 SkeletonBoneCount = MutableMesh->GetSkeleton()->GetBoneCount();
-			for ( int32 I = 0; I < SkeletonBoneCount; ++I )
+			if (BonePoseIndex == INDEX_NONE)
 			{
-				FString SkeletonBoneName = MutableMesh->GetSkeleton()->GetBoneName(I);
-				
-				if (SkeletonBoneName.Equals(BodyBoneName))
-				{
-					RelevantBodySetups.Add(BodySetup);
-					DiscardedBodySetups[BodySetupIndex] = 0;
-					int32 BonePoseIndex = MutableMesh->FindBonePose(StringCast<ANSICHAR>(*BodyBoneName).Get());
-					
-					EnumAddFlags(MutableMesh->BonePoses[BonePoseIndex].BoneUsageFlags, mu::EBoneUsageFlags::Physics);
-				}
+				continue;
 			}
+		
+			RelevantBodySetups.Add(BodySetup);
+			DiscardedBodySetups[BodySetupIndex] = 0;
+			EnumAddFlags(MutableMesh->BonePoses[BonePoseIndex].BoneUsageFlags, mu::EBoneUsageFlags::Physics);
 		}
 
 		const int32 NumDiscardedSetups = Algo::CountIf(DiscardedBodySetups, [](const uint8& V) { return V; });
@@ -1646,8 +1612,8 @@ mu::MeshPtr ConvertSkeletalMeshToMutable(const USkeletalMesh* InSkeletalMesh, co
 		{
 			TObjectPtr<USkeletalBodySetup>& BodySetup = RelevantBodySetups[B];
 			
-			FString BodyBoneName = BodySetup->BoneName.ToString();
-			PhysicsBody->SetBodyBoneName( B, StringCast<ANSICHAR>(*BodyBoneName).Get() );
+			const uint16 BoneId = GenerationContext.BoneNames.AddUnique(BodySetup->BoneName);
+			PhysicsBody->SetBodyBoneId( B, BoneId);
 			
 			const int32 NumSpheres = BodySetup->AggGeom.SphereElems.Num();
 			PhysicsBody->SetSphereCount( B, NumSpheres );
@@ -1748,14 +1714,14 @@ mu::MeshPtr ConvertSkeletalMeshToMutable(const USkeletalMesh* InSkeletalMesh, co
 				mu::EBoneUsageFlags::Skinning | mu::EBoneUsageFlags::Physics | mu::EBoneUsageFlags::Deform;
 		if (EnumHasAnyFlags(BonePose.BoneUsageFlags, FlagsToPropagate))
 		{
-			const int32 BoneIndex = MutableMesh->GetSkeleton()->FindBone(MutableMesh->GetBonePoseName(I));
+			const int32 Index = MutableMesh->GetSkeleton()->FindBone(MutableMesh->GetBonePoseBoneId(I));
 
-			if (BoneIndex == INDEX_NONE)
+			if (Index == INDEX_NONE)
 			{
 				continue;
 			}
 
-			const int32 ParentIndex = MutableMesh->GetSkeleton()->GetBoneParent(BoneIndex);
+			const int32 ParentIndex = MutableMesh->GetSkeleton()->GetBoneParent(Index);
 
 			if (ParentIndex == INDEX_NONE)
 			{	
@@ -1795,9 +1761,9 @@ mu::MeshPtr ConvertSkeletalMeshToMutable(const USkeletalMesh* InSkeletalMesh, co
 
 			const int32 PhysicsAssetId = GenerationContext.AnimBpOverridePhysicsAssetsInfo.AddUnique(Info);
 
-			mu::Ptr<mu::PhysicsBody> MutableBody = MakePhysicsBodyFromAsset(
+			mu::Ptr<mu::PhysicsBody> MutableBody = MakePhysicsBodyFromAsset(GenerationContext,
 					PropertyAsset,
-					MakePhysicsAssetBodySetupRelevancyMap(PropertyAsset, MutableMesh));
+					MakePhysicsAssetBodySetupRelevancyMap(GenerationContext, PropertyAsset, MutableMesh));
 			MutableBody->CustomId = PhysicsAssetId;
 
 			MutableMesh->AddAdditionalPhysicsBody(MutableBody);
@@ -2252,7 +2218,7 @@ TArray<TTuple<USkeletalMesh*, TSoftClassPtr<UAnimInstance>>> GetSkeletalMeshesIn
 
 
 bool GetAndValidateReshapeBonesToDeform(
-	TArray<FString>& OutBonesToDeform,
+	TArray<FName>& OutBonesToDeform,
 	const TArray<FMeshReshapeBoneReference>& InBonesToDeform,
 	const TArray<TTuple<USkeletalMesh*, TSoftClassPtr<UAnimInstance>>>& SkeletalMeshesInfo,
 	const UCustomizableObjectNode* Node,
@@ -2284,7 +2250,7 @@ bool GetAndValidateReshapeBonesToDeform(
 				{
 					if (SkeletalMesh->GetRefSkeleton().GetParentIndex(BoneIndex) != INDEX_NONE)
 					{
-						OutBonesToDeform.AddUnique(BoneName.ToString());
+						OutBonesToDeform.AddUnique(BoneName);
 					}
 
 					MissingBones[InBoneIndex] &= false;
@@ -2361,7 +2327,7 @@ bool GetAndValidateReshapeBonesToDeform(
 
 				if (!bFound && Mesh.Get<USkeletalMesh*>()->GetRefSkeleton().GetParentIndex(BoneIndex) != INDEX_NONE)
 				{
-					OutBonesToDeform.AddUnique(BoneName.ToString());
+					OutBonesToDeform.AddUnique(BoneName);
 				}
 			}
 		}
@@ -2377,7 +2343,7 @@ bool GetAndValidateReshapeBonesToDeform(
 		{
 			if (RefSkeleton.GetParentIndex(BoneIndex) != INDEX_NONE)
 			{
-				OutBonesToDeform.AddUnique(RefSkeleton.GetBoneName(BoneIndex).ToString());
+				OutBonesToDeform.AddUnique(RefSkeleton.GetBoneName(BoneIndex));
 			}
 		}
 	}
@@ -2400,7 +2366,7 @@ bool GetAndValidateReshapeBonesToDeform(
 				if (RefSkeleton.FindBoneIndex(BoneName) == INDEX_NONE 
 					&& SkeletalMesh->GetRefSkeleton().GetParentIndex(BoneIndex) != INDEX_NONE)
 				{
-					OutBonesToDeform.AddUnique(BoneName.ToString());
+					OutBonesToDeform.AddUnique(BoneName);
 				}
 			}
 		}
@@ -2411,7 +2377,7 @@ bool GetAndValidateReshapeBonesToDeform(
 
 
 bool GetAndValidateReshapePhysicsToDeform(
-	TArray<FString>& OutPhysiscsToDeform,
+	TArray<FName>& OutPhysiscsToDeform,
 	const TArray<FMeshReshapeBoneReference>& InPhysicsToDeform,
 	const TArray<TTuple<USkeletalMesh*, TSoftClassPtr<UAnimInstance>>>& SkeletalMeshesInfo,
 	EBoneDeformSelectionMethod SelectionMethod,
@@ -2566,8 +2532,7 @@ bool GetAndValidateReshapePhysicsToDeform(
 		{
 			if (BoneInclusionSet[I])
 			{
-				FName SetupBoneName = PhysicsAsset->SkeletalBodySetups[I]->BoneName;
-				OutPhysiscsToDeform.AddUnique(SetupBoneName.ToString());
+				OutPhysiscsToDeform.AddUnique(PhysicsAsset->SkeletalBodySetups[I]->BoneName);
 			}
 		}
 	}
@@ -2769,19 +2734,19 @@ mu::NodeMeshPtr GenerateMorphMesh(const UEdGraphPin* Pin,
 					bool bWarningFound = false;
 					if (TypedMorphNode->bReshapeSkeleton)
 					{
-						TArray<FString> BonesToDeform;
+						TArray<FName> BonesToDeform;
 						bWarningFound = GetAndValidateReshapeBonesToDeform(
 							BonesToDeform, TypedMorphNode->BonesToDeform, SkeletalMeshesToDeform, TypedMorphNode, TypedMorphNode->SelectionMethod, GenerationContext);
 
-						for (const FString& BoneName : BonesToDeform)
+						for (const FName BoneName : BonesToDeform)
 						{
-							Result->AddBoneToDeform(StringCast<ANSICHAR>(*BoneName).Get());
+							Result->AddBoneToDeform(GenerationContext.BoneNames.AddUnique(BoneName));
 						}
 					}
 
 					if (TypedMorphNode->bReshapePhysicsVolumes)
 					{
-						TArray<FString> PhysicsToDeform;
+						TArray<FName> PhysicsToDeform;
 
 						const EBoneDeformSelectionMethod SelectionMethod = TypedMorphNode->PhysicsSelectionMethod;
 						bWarningFound = bWarningFound || GetAndValidateReshapePhysicsToDeform(
@@ -2789,9 +2754,9 @@ mu::NodeMeshPtr GenerateMorphMesh(const UEdGraphPin* Pin,
 							TypedMorphNode->PhysicsBodiesToDeform, SkeletalMeshesToDeform, SelectionMethod, 
 							TypedMorphNode, GenerationContext);
 	
-						for (const FString& PhysicsBoneName : PhysicsToDeform)
+						for (const FName& PhysicsBoneName : PhysicsToDeform)
 						{
-							Result->AddPhysicsBodyToDeform(StringCast<ANSICHAR>(*PhysicsBoneName).Get());
+							Result->AddPhysicsBodyToDeform(GenerationContext.BoneNames.AddUnique(PhysicsBoneName));
 						}
 					}
 						
@@ -3422,28 +3387,28 @@ mu::NodeMeshPtr GenerateMutableSourceMesh(const UEdGraphPin* Pin,
 			bool bWarningFound = false;
 			if (TypedNodeReshape->bReshapeSkeleton)
 			{
-				TArray<FString> BonesToDeform;
+				TArray<FName> BonesToDeform;
 				bWarningFound = GetAndValidateReshapeBonesToDeform(
 					BonesToDeform, TypedNodeReshape->BonesToDeform, SkeletalMeshesToDeform, TypedNodeReshape, TypedNodeReshape->SelectionMethod, GenerationContext);
 				
-				for (const FString& BoneName : BonesToDeform)
+				for (const FName& BoneName : BonesToDeform)
 				{
-					MeshNode->AddBoneToDeform(StringCast<ANSICHAR>(*BoneName).Get());
+					MeshNode->AddBoneToDeform(GenerationContext.BoneNames.Find(BoneName));
 				}
 			}
 
 			if (TypedNodeReshape->bReshapePhysicsVolumes)
 			{
 				EBoneDeformSelectionMethod SelectionMethod = TypedNodeReshape->PhysicsSelectionMethod;
-				TArray<FString> PhysicsToDeform;
+				TArray<FName> PhysicsToDeform;
 				bWarningFound = bWarningFound || GetAndValidateReshapePhysicsToDeform(
 					PhysicsToDeform, 
 					TypedNodeReshape->PhysicsBodiesToDeform, SkeletalMeshesToDeform, SelectionMethod, 
 					TypedNodeReshape, GenerationContext);
 
-				for (const FString& PhysicsBoneName : PhysicsToDeform)
+				for (const FName& PhysicsBoneName : PhysicsToDeform)
 				{
-					MeshNode->AddPhysicsBodyToDeform(StringCast<ANSICHAR>(*PhysicsBoneName).Get());
+					MeshNode->AddPhysicsBodyToDeform(GenerationContext.BoneNames.Find(PhysicsBoneName));
 				}	
 			}
 			
@@ -3525,10 +3490,10 @@ mu::NodeMeshPtr GenerateMutableSourceMesh(const UEdGraphPin* Pin,
 
 			if (TypedNode->PoseAsset && GenerationContext.GetCurrentComponentInfo().RefSkeletalMesh)
 			{
-				TArray<FString> ArrayBoneName;
+				TArray<FName> ArrayBoneName;
 				TArray<FTransform> ArrayTransform;
 				UCustomizableObjectNodeAnimationPose::StaticRetrievePoseInformation(TypedNode->PoseAsset, GenerationContext.GetCurrentComponentInfo().RefSkeletalMesh, ArrayBoneName, ArrayTransform);
-				mu::NodeMeshApplyPosePtr NodeMeshApplyPose = CreateNodeMeshApplyPose(InputMeshNode, GenerationContext.Object, ArrayBoneName, ArrayTransform);
+				mu::NodeMeshApplyPosePtr NodeMeshApplyPose = CreateNodeMeshApplyPose(GenerationContext, InputMeshNode, ArrayBoneName, ArrayTransform);
 
 				if (NodeMeshApplyPose)
 				{
diff --git a/Engine/Plugins/Experimental/Mutable/Source/CustomizableObjectEditor/Private/MuCOE/MutableMeshPreviewUtils.cpp b/Engine/Plugins/Experimental/Mutable/Source/CustomizableObjectEditor/Private/MuCOE/MutableMeshPreviewUtils.cpp
--- a/Engine/Plugins/Experimental/Mutable/Source/CustomizableObjectEditor/Private/MuCOE/MutableMeshPreviewUtils.cpp
+++ b/Engine/Plugins/Experimental/Mutable/Source/CustomizableObjectEditor/Private/MuCOE/MutableMeshPreviewUtils.cpp
@@ -370,11 +370,11 @@ namespace MutableMeshPreviewUtils
 			OutSkeletalMesh->SetHasVertexColors(false);
 
 			// Find how many bones the bonemap could have
-			int32 NumBoneMapIndices = !InMutableMesh->GetBoneMap().IsEmpty() ? InMutableMesh->GetBoneMap().Num() : InMutableMesh->GetBonePoseCount();
+			const int32 NumBonesInBoneMap = !InMutableMesh->GetBoneMap().IsEmpty() ? InMutableMesh->GetBoneMap().Num() : InMutableMesh->GetBonePoseCount();
 			
 			// Fill the bonemap with zeros
 			TArray<uint16> BoneMap;
-			BoneMap.SetNumZeroed(FMath::Max(NumBoneMapIndices, 1));
+			BoneMap.SetNumZeroed(FMath::Max(NumBonesInBoneMap, 1));
 
 			// Load buffer data found on the mutable model onto the out skeletal mesh
 			// It includes vertex and index buffers
diff --git a/Engine/Plugins/Experimental/Mutable/Source/CustomizableObjectEditor/Private/MuCOE/Nodes/CustomizableObjectNodeAnimationPose.cpp b/Engine/Plugins/Experimental/Mutable/Source/CustomizableObjectEditor/Private/MuCOE/Nodes/CustomizableObjectNodeAnimationPose.cpp
--- a/Engine/Plugins/Experimental/Mutable/Source/CustomizableObjectEditor/Private/MuCOE/Nodes/CustomizableObjectNodeAnimationPose.cpp
+++ b/Engine/Plugins/Experimental/Mutable/Source/CustomizableObjectEditor/Private/MuCOE/Nodes/CustomizableObjectNodeAnimationPose.cpp
@@ -73,7 +73,7 @@ FLinearColor UCustomizableObjectNodeAnimationPose::GetNodeTitleColor() const
 }
 
 
-void UCustomizableObjectNodeAnimationPose::StaticRetrievePoseInformation(UPoseAsset* PoseAsset, USkeletalMesh* RefSkeletalMesh, TArray<FString>& OutArrayBoneName, TArray<FTransform>& OutArrayTransform)
+void UCustomizableObjectNodeAnimationPose::StaticRetrievePoseInformation(UPoseAsset* PoseAsset, USkeletalMesh* RefSkeletalMesh, TArray<FName>& OutArrayBoneName, TArray<FTransform>& OutArrayTransform)
 {
 	if (PoseAsset && RefSkeletalMesh)
 	{
@@ -159,7 +159,7 @@ void UCustomizableObjectNodeAnimationPose::StaticRetrievePoseInformation(UPoseAs
 			FTransform BoneToComponentTranform = SkeletalMeshComponent->GetEditableComponentSpaceTransforms()[BoneIndex];
 			FTransform TransformToAdd = BoneToComponentTranform.Inverse() * CumulativePoseTransform;
 
-			OutArrayBoneName.Add(BoneName.ToString());
+			OutArrayBoneName.Add(BoneName);
 			OutArrayTransform.Add(TransformToAdd);
 		}
 	}
diff --git a/Engine/Plugins/Experimental/Mutable/Source/CustomizableObjectEditor/Private/MuCOE/SMutableSkeletonViewer.cpp b/Engine/Plugins/Experimental/Mutable/Source/CustomizableObjectEditor/Private/MuCOE/SMutableSkeletonViewer.cpp
--- a/Engine/Plugins/Experimental/Mutable/Source/CustomizableObjectEditor/Private/MuCOE/SMutableSkeletonViewer.cpp
+++ b/Engine/Plugins/Experimental/Mutable/Source/CustomizableObjectEditor/Private/MuCOE/SMutableSkeletonViewer.cpp
@@ -35,7 +35,7 @@ public:
 			+ SHorizontalBox::Slot()
 			[
 				SNew(STextBlock).
-				Text(RowItem->BoneName)
+				Text(FText::FromName(InRowItem->BoneName))
 			]
 		];
 
@@ -118,49 +118,36 @@ void SMutableSkeletonViewer::FillBoneDefinitionArrays()
 	const int32 BoneCount = Skeleton->GetBoneCount();
 	for (int32 BoneIndex = 0; BoneIndex < BoneCount; BoneIndex++)
 	{
-		const char* BoneName = Skeleton->GetBoneName(BoneIndex);
-
-		// Get the parent node
-		const int32 ParentBoneIndex = Skeleton->GetBoneParent(BoneIndex);
-		const char* ParentBoneName = nullptr;
-		if (ParentBoneIndex != -1)
-		{
-			ParentBoneName = Skeleton->GetBoneName(ParentBoneIndex);
-		}
-
 		// Generate a new definition object
 		TSharedPtr<FBoneDefinition> NewBoneDefinition = MakeShareable(new FBoneDefinition());
-		{
-			// Save the bone data
-			NewBoneDefinition->BoneIndex = BoneIndex;
-			NewBoneDefinition->BoneName = FText::FromString(FString(BoneName));
 
-			NewBoneDefinition->ParentBoneIndex = ParentBoneIndex;
-			if (ParentBoneName)
-			{
-				NewBoneDefinition->ParentBoneName = FText::FromString(FString(ParentBoneName));
-			}
-		}
+		// Save the bone data
+		NewBoneDefinition->BoneIndex = BoneIndex;
+		NewBoneDefinition->BoneName = Skeleton->GetBoneFName(BoneIndex);
 
-		// Store the root bones on their own data structure separated from the child bones
-		if (NewBoneDefinition->ParentBoneIndex == -1)
+		// Get the parent node if any
+		const int32 ParentBoneIndex = Skeleton->GetBoneParent(BoneIndex);
+		if (ParentBoneIndex != INDEX_NONE)
 		{
-			RootBones.Add(NewBoneDefinition);
+			NewBoneDefinition->ParentBoneIndex = ParentBoneIndex;
+			NewBoneDefinition->ParentBoneName = Skeleton->GetBoneFName(ParentBoneIndex);
+
+			// Store child bones
+			ChildBones.Add(NewBoneDefinition);
 		}
 		else
 		{
-			ChildBones.Add(NewBoneDefinition);
+			// Store the root bones on their own data structure separated from the child bones
+			RootBones.Add(NewBoneDefinition);
 		}
 	}
 
 	// Add default definition to tell the user no bone has been found
 	if (RootBones.Num() == 0)
 	{
+		// Save the bone data
 		const TSharedPtr<FBoneDefinition> NewBoneDefinition = MakeShareable(new FBoneDefinition());
-		{
-			// Save the bone data
-			NewBoneDefinition->BoneName = FText(INVTEXT("No bones found..."));
-		}
+		NewBoneDefinition->BoneName = FName("No bones found...");
 		RootBones.Add(NewBoneDefinition);
 	}
 
@@ -181,7 +168,7 @@ void SMutableSkeletonViewer::OnGetChildrenFromBoneDefinition(TSharedPtr<FBoneDef
 	// Iterate over the bone definitions and save the ones that do have as parent the bone we have selected
 	for (TSharedPtr< FBoneDefinition>& PossibleChildBoneDefinition : ChildBones)
 	{
-		if (PossibleChildBoneDefinition.Get()->ParentBoneIndex == TargetBoneIndex)
+		if (PossibleChildBoneDefinition->ParentBoneIndex == TargetBoneIndex)
 		{
 			OutChildrenBones.Add(PossibleChildBoneDefinition);
 		}
diff --git a/Engine/Plugins/Experimental/Mutable/Source/CustomizableObjectEditor/Private/MuCOE/SMutableSkeletonViewer.h b/Engine/Plugins/Experimental/Mutable/Source/CustomizableObjectEditor/Private/MuCOE/SMutableSkeletonViewer.h
--- a/Engine/Plugins/Experimental/Mutable/Source/CustomizableObjectEditor/Private/MuCOE/SMutableSkeletonViewer.h
+++ b/Engine/Plugins/Experimental/Mutable/Source/CustomizableObjectEditor/Private/MuCOE/SMutableSkeletonViewer.h
@@ -13,11 +13,11 @@ class ITableRow;
 /** Container with the information required to generate the tree of bones */
 struct FBoneDefinition
 {
-	FText ParentBoneName = FText(INVTEXT(""));
-	int32 ParentBoneIndex = -1;
+	FName ParentBoneName = FName("");
+	int32 ParentBoneIndex = INDEX_NONE;
 
-	FText BoneName = FText(INVTEXT(""));
-	int32 BoneIndex = -1;
+	FName BoneName = FName("");
+	int32 BoneIndex = INDEX_NONE;
 };
 
 /** Object responsible of showing with a tree view slate the bones present on the provided mutable skeleton */
diff --git a/Engine/Plugins/Experimental/Mutable/Source/MutableRuntime/Private/MuR/CodeRunner.cpp b/Engine/Plugins/Experimental/Mutable/Source/MutableRuntime/Private/MuR/CodeRunner.cpp
--- a/Engine/Plugins/Experimental/Mutable/Source/MutableRuntime/Private/MuR/CodeRunner.cpp
+++ b/Engine/Plugins/Experimental/Mutable/Source/MutableRuntime/Private/MuR/CodeRunner.cpp
@@ -1998,7 +1998,7 @@ namespace mu
 					Ptr<Mesh> Result = CreateMesh(Source ? Source->GetDataSize() : 0);
 
 					bool bOutSuccess = false;
-					MeshClipMorphPlane(Result.get(), Source.get(), origin, normal, args.dist, args.factor, morphShape.size[0], morphShape.size[1], morphShape.size[2], selectionShape, bOutSuccess, mu::string(), -1);
+					MeshClipMorphPlane(Result.get(), Source.get(), origin, normal, args.dist, args.factor, morphShape.size[0], morphShape.size[1], morphShape.size[2], selectionShape, bOutSuccess, INDEX_NONE, -1);
 					
 					if (!bOutSuccess)
 					{
@@ -2014,16 +2014,13 @@ namespace mu
 
 				else if (args.vertexSelectionType == OP::MeshClipMorphPlaneArgs::VS_BONE_HIERARCHY)
 				{
-					check(args.vertexSelectionShapeOrBone < (uint32)pModel->GetPrivate()->m_program.m_constantStrings.Num());
-
 					FShape selectionShape;
 					selectionShape.type = (uint8)FShape::Type::None;
-					const string& selectionBone = pModel->GetPrivate()->m_program.m_constantStrings[args.vertexSelectionShapeOrBone];
 
 					Ptr<Mesh> Result = CreateMesh(Source->GetDataSize());
 
 					bool bOutSuccess = false;
-					MeshClipMorphPlane(Result.get(), Source.get(), origin, normal, args.dist, args.factor, morphShape.size[0], morphShape.size[1], morphShape.size[2], selectionShape, bOutSuccess, selectionBone, args.maxBoneRadius);
+					MeshClipMorphPlane(Result.get(), Source.get(), origin, normal, args.dist, args.factor, morphShape.size[0], morphShape.size[1], morphShape.size[2], selectionShape, bOutSuccess, args.vertexSelectionShapeOrBone, args.maxBoneRadius);
 
 					if (!bOutSuccess)
 					{
@@ -2045,7 +2042,7 @@ namespace mu
 					Ptr<Mesh> Result = CreateMesh(Source ? Source->GetDataSize() : 0);
 
 					bool bOutSuccess = false;
-					MeshClipMorphPlane(Result.get(), Source.get(), origin, normal, args.dist, args.factor, morphShape.size[0], morphShape.size[1], morphShape.size[2], selectionShape, bOutSuccess, mu::string(), -1.0f);
+					MeshClipMorphPlane(Result.get(), Source.get(), origin, normal, args.dist, args.factor, morphShape.size[0], morphShape.size[1], morphShape.size[2], selectionShape, bOutSuccess, INDEX_NONE, -1.0f);
 
 					if (!bOutSuccess)
 					{
@@ -2413,29 +2410,19 @@ namespace mu
 					FMemory::Memcpy(&NumBones, Data, sizeof(int32)); 
 					Data += sizeof(int32);
 					
-					TArray<string> BonesToDeform;
-					BonesToDeform.SetNum(NumBones);
-					for (int32 b = 0; b < NumBones; ++b)
-					{
-						int32 ConstStringIndex;
-						FMemory::Memcpy(&ConstStringIndex, Data, sizeof(int32));
-						BonesToDeform[b] = pModel->GetPrivate()->m_program.m_constantStrings[ConstStringIndex].c_str();
-						Data += sizeof(int32);
-					}
+					TArray<uint16> BonesToDeform;
+					BonesToDeform.SetNumUninitialized(NumBones);
+					FMemory::Memcpy(BonesToDeform.GetData(), Data, NumBones * sizeof(uint16));
+					Data += NumBones * sizeof(uint16);
 
 					int32 NumPhysicsBodies;
 					FMemory::Memcpy(&NumPhysicsBodies, Data, sizeof(int32)); 
 					Data += sizeof(int32);
 
-					TArray<string> PhysicsToDeform;
-					PhysicsToDeform.SetNum(NumPhysicsBodies);
-					for (int32 b = 0; b < NumPhysicsBodies; ++b)
-					{
-						int32 ConstStringIndex;
-						FMemory::Memcpy(&ConstStringIndex, Data, sizeof(int32));
-						PhysicsToDeform[b] = pModel->GetPrivate()->m_program.m_constantStrings[ConstStringIndex].c_str();
-						Data += sizeof(int32);
-					}
+					TArray<uint16> PhysicsToDeform;
+					PhysicsToDeform.SetNumUninitialized(NumPhysicsBodies);
+					FMemory::Memcpy(PhysicsToDeform.GetData(), Data, NumPhysicsBodies * sizeof(uint16));
+					Data += NumPhysicsBodies * sizeof(uint16);
 
 					const EMeshBindShapeFlags BindFlags = static_cast<EMeshBindShapeFlags>(Args.flags);
 
@@ -2668,7 +2655,9 @@ namespace mu
                 // Only if both are valid.
                 if (Source && pSkeleton)
                 {
-                    if (Source->GetSkeleton() && !Source->GetSkeleton()->m_bones.IsEmpty())
+                    if ( Source->GetSkeleton()
+                         &&
+                         !Source->GetSkeleton()->BoneIds.IsEmpty() )
                     {
                         // For some reason we already have bone data, so we can't just overwrite it
                         // or the skinning may break. This may happen because of a problem in the
diff --git a/Engine/Plugins/Experimental/Mutable/Source/MutableRuntime/Private/MuR/Mesh.cpp b/Engine/Plugins/Experimental/Mutable/Source/MutableRuntime/Private/MuR/Mesh.cpp
--- a/Engine/Plugins/Experimental/Mutable/Source/MutableRuntime/Private/MuR/Mesh.cpp
+++ b/Engine/Plugins/Experimental/Mutable/Source/MutableRuntime/Private/MuR/Mesh.cpp
@@ -71,7 +71,7 @@ MeshPtr Mesh::Clone() const
 	
 	// Clone bone poses
 	pResult->BonePoses = BonePoses;
-	pResult->BoneMapIndices = BoneMapIndices;
+	pResult->BoneMap = BoneMap;
 
 	// Clone SkeletonIDs
 	pResult->SkeletonIDs = SkeletonIDs;
@@ -156,7 +156,7 @@ MeshPtr Mesh::Clone(EMeshCopyFlags Flags) const
 	// Clone BoneMap
 	if (EnumHasAnyFlags(Flags, EMeshCopyFlags::WithBoneMap))
 	{
-		pResult->BoneMapIndices = BoneMapIndices;
+		pResult->BoneMap = BoneMap;
 	}
 
 	// Clone SkeletonIDs
@@ -248,7 +248,7 @@ void Mesh::CopyFrom(const Mesh& From, EMeshCopyFlags Flags)
 	// Copy BoneMap
 	if (EnumHasAnyFlags(Flags, EMeshCopyFlags::WithBoneMap))
 	{
-		BoneMapIndices = From.BoneMapIndices;
+		BoneMap = From.BoneMap;
 	}
 
 	// Copy SkeletonIDs
@@ -420,7 +420,7 @@ void Mesh::GetSurface( int32 surfaceIndex,
             if (firstIndex) *firstIndex = 0;
             if (indexCount) *indexCount = GetIndexCount();
 			if (BoneIndex) *BoneIndex = 0;
-			if (BoneCount) *BoneCount = BoneMapIndices.Num();
+			if (BoneCount) *BoneCount = BoneMap.Num();
         }
     }
     else
@@ -591,21 +591,9 @@ void Mesh::SetTag( int tagIndex, const char* strName )
 
 
 //---------------------------------------------------------------------------------------------
-const char* Mesh::GetBonePoseName(int32 BoneIndex) const
+int32 Mesh::FindBonePose(const uint16 BoneId) const
 {
-	check(BoneIndex >= 0 && BoneIndex < GetBonePoseCount());
-	if (BoneIndex >= 0 && BoneIndex < GetBonePoseCount())
-	{
-		return BonePoses[BoneIndex].BoneName.c_str();
-	}
-	return "";
-}
-
-
-//---------------------------------------------------------------------------------------------
-int32 Mesh::FindBonePose(const char* StrName) const
-{
-	return BonePoses.IndexOfByPredicate([StrName](const FBonePose& Pose) { return Pose.BoneName == StrName; });
+	return BonePoses.IndexOfByPredicate([BoneId](const FBonePose& Pose) { return Pose.BoneId == BoneId; });
 }
 
 
@@ -625,16 +613,26 @@ int32 mu::Mesh::GetBonePoseCount() const
 
 
 //---------------------------------------------------------------------------------------------
-void mu::Mesh::SetBonePose(int32 BoneIndex, const char* StrName, FTransform3f Transform, EBoneUsageFlags BoneUsageFlags)
+void mu::Mesh::SetBonePose(int32 Index, uint16 BoneId, FTransform3f Transform, EBoneUsageFlags BoneUsageFlags)
 {
-	if (BoneIndex >= 0 && BoneIndex < BonePoses.Num())
+	check(BonePoses.IsValidIndex(Index));
+	if (BonePoses.IsValidIndex(Index))
 	{
-		BonePoses[BoneIndex] = FBonePose{ StrName, BoneUsageFlags, Transform };
+		BonePoses[Index] = FBonePose{ BoneId, BoneUsageFlags, Transform };
 	}
-	else 
+}
+
+
+//---------------------------------------------------------------------------------------------
+int32 Mesh::GetBonePoseBoneId(int32 Index) const
+{
+	check(BonePoses.IsValidIndex(Index));
+	if (BonePoses.IsValidIndex(Index))
 	{
-		check(false);
+		return BonePoses[Index].BoneId;
 	}
+
+	return INDEX_NONE;
 }
 
 
@@ -657,14 +655,14 @@ EBoneUsageFlags Mesh::GetBoneUsageFlags(int32 BoneIndex) const
 //---------------------------------------------------------------------------------------------
 void Mesh::SetBoneMap(const TArray<uint16>& InBoneMap)
 {
-	BoneMapIndices = InBoneMap;
+	BoneMap = InBoneMap;
 }
 
 
 //---------------------------------------------------------------------------------------------
 const TArray<uint16>& Mesh::GetBoneMap() const
 {
-	return BoneMapIndices;
+	return BoneMap;
 }
 
 
@@ -1104,7 +1102,7 @@ void Mesh::EnsureSurfaceData()
 		MESH_SURFACE s;
 		s.m_vertexCount = m_VertexBuffers.GetElementCount();
 		s.m_indexCount = m_IndexBuffers.GetElementCount();
-		s.BoneMapCount = BoneMapIndices.Num();
+		s.BoneMapCount = BoneMap.Num();
 		m_surfaces.Add(s);
 	}
 }
@@ -1226,10 +1224,10 @@ void Mesh::FACE_GROUP::Unserialise(InputArchive& arch)
 //-------------------------------------------------------------------------------------------------
 void Mesh::FBonePose::Serialise(OutputArchive& arch) const
 {
-	const int32 ver = 1;
+	const int32 ver = 2;
 	arch << ver;
 
-	arch << BoneName;
+	arch << BoneId;
 	arch << BoneUsageFlags;
 	arch << BoneTransform;
 }
@@ -1240,9 +1238,19 @@ void Mesh::FBonePose::Unserialise(InputArchive& arch)
 {
 	int32 ver = 0;
 	arch >> ver;
-	check(ver <= 1);
+	check(ver <= 2);
+
+	if (ver <= 1)
+	{
+		string BoneName;
+		arch >> BoneName;
 
-	arch >> BoneName;
+		BoneId = 0;
+	}
+	else
+	{
+		arch >> BoneId;
+	}
 
 	if (ver == 0)
 	{
@@ -1283,7 +1291,7 @@ void Mesh::Serialise(OutputArchive& arch) const
 	arch << m_tags;
 
 	arch << BonePoses;
-	arch << BoneMapIndices;
+	arch << BoneMap;
 
 	arch << AdditionalPhysicsBodies;
 }
@@ -1344,7 +1352,7 @@ void Mesh::Unserialise(InputArchive& arch)
 
 		for (int32 BoneIndex = 0; BoneIndex < NumBones; ++BoneIndex)
 		{
-			BonePoses[BoneIndex].BoneName = m_pSkeleton->GetBoneName(BoneIndex);
+			BonePoses[BoneIndex].BoneId = BoneIndex;
 			BonePoses[BoneIndex].BoneUsageFlags = EBoneUsageFlags::Skinning;
 			BonePoses[BoneIndex].BoneTransform = m_pSkeleton->m_boneTransforms_DEPRECATED[BoneIndex];
 		}
@@ -1352,15 +1360,15 @@ void Mesh::Unserialise(InputArchive& arch)
 
 	if (ver >= 16)
 	{
-		arch >> BoneMapIndices;
+		arch >> BoneMap;
 	}
 	else
 	{
 		const int32 NumBonePoses = BonePoses.Num();
-		BoneMapIndices.SetNum(NumBonePoses);
+		BoneMap.SetNum(NumBonePoses);
 		for (int32 BoneIndex = 0; BoneIndex < NumBonePoses; ++BoneIndex)
 		{
-			BoneMapIndices[BoneIndex] = BoneIndex;
+			BoneMap[BoneIndex] = BoneIndex;
 		}
 
 		for (MESH_SURFACE& Surface : m_surfaces)
diff --git a/Engine/Plugins/Experimental/Mutable/Source/MutableRuntime/Private/MuR/OpMeshApplyPose.h b/Engine/Plugins/Experimental/Mutable/Source/MutableRuntime/Private/MuR/OpMeshApplyPose.h
--- a/Engine/Plugins/Experimental/Mutable/Source/MutableRuntime/Private/MuR/OpMeshApplyPose.h
+++ b/Engine/Plugins/Experimental/Mutable/Source/MutableRuntime/Private/MuR/OpMeshApplyPose.h
@@ -38,7 +38,7 @@ namespace mu
 		skinTransforms.SetNum(pSkeleton->GetBoneCount());
         for ( int32 b=0; b< pSkeleton->GetBoneCount(); ++b )
         {
-            int32 poseBoneIndex = pPose->FindBonePose(pSkeleton->GetBoneName(b));
+            int32 poseBoneIndex = pPose->FindBonePose(pSkeleton->GetBoneId(b));
 
             if( poseBoneIndex != INDEX_NONE )
             {
diff --git a/Engine/Plugins/Experimental/Mutable/Source/MutableRuntime/Private/MuR/OpMeshBind.h b/Engine/Plugins/Experimental/Mutable/Source/MutableRuntime/Private/MuR/OpMeshBind.h
--- a/Engine/Plugins/Experimental/Mutable/Source/MutableRuntime/Private/MuR/OpMeshBind.h
+++ b/Engine/Plugins/Experimental/Mutable/Source/MutableRuntime/Private/MuR/OpMeshBind.h
@@ -491,7 +491,7 @@ namespace mu
 
 	inline TTuple<TArray<FReshapePointBindingData>, TArray<int32>, TArray<int32>> BindPhysicsBodies( 
 			TArray<const PhysicsBody*> PhysicsBodies, FShapeMeshTree& ShapeMeshTree, const Mesh* pMesh, 
-			const TArray<string>& PhysicsToDeform )
+			const TArray<uint16>& PhysicsToDeform )
 	{
 		TTuple<TArray<FReshapePointBindingData>, TArray<int32>, TArray<int32>> ReturnValue;
 
@@ -511,7 +511,7 @@ namespace mu
 				const int32 NumBodies = PhysicsBodies[PhysicsBodyIndex]->GetBodyCount();
 				for (int32 BodyIndex = 0; BodyIndex < NumBodies; ++BodyIndex)
 				{
-					if (PhysicsToDeform.Contains(Body->GetBodyBoneName(BodyIndex)))
+					if (PhysicsToDeform.Contains(Body->GetBodyBoneId(BodyIndex)))
 					{
 						BodiesToDeformIndices.Add(BodyIndex);
 					}
@@ -573,49 +573,48 @@ namespace mu
 			const int32 IndicesEnd = BodiesToDeformOffsets[PhysicsBodyIndex + 1];
 
 			TArrayView<int32> BodyIndices = TArrayView<int32>(
-					BodiesToDeformIndices.GetData() + IndicesBegin, IndicesEnd - IndicesBegin);
+				BodiesToDeformIndices.GetData() + IndicesBegin, IndicesEnd - IndicesBegin);
 
 			for (const int32 B : BodyIndices)
-		{
-				const char* BoneName = Body->GetBodyBoneName(B);
-			int32 BoneIdx = pMesh->FindBonePose(BoneName);
-
-			FTransform3f T = FTransform3f::Identity;
-			if (BoneIdx > 0)
 			{
-				pMesh->GetBoneTransform(BoneIdx, T);
-			}
-			
+				int32 BoneIndex = pMesh->FindBonePose(Body->GetBodyBoneId(B));
+
+				FTransform3f T = FTransform3f::Identity;
+				if (BoneIndex > 0)
+				{
+					pMesh->GetBoneTransform(BoneIndex, T);
+				}
+
 				const int32 SphereCount = Body->GetSphereCount(B);
 				for (int32 I = 0; I < SphereCount; ++I, AddedPoints += 6)
-			{
-				FVector3f P;
-				float R;
+				{
+					FVector3f P;
+					float R;
 
 					Body->GetSphere(B, I, P, R);
 
-				Points[AddedPoints + 0] = T.TransformPosition(P + FVector3f(R, 0.0f, 0.0f));
-				Points[AddedPoints + 1] = T.TransformPosition(P - FVector3f(R, 0.0f, 0.0f));
+					Points[AddedPoints + 0] = T.TransformPosition(P + FVector3f(R, 0.0f, 0.0f));
+					Points[AddedPoints + 1] = T.TransformPosition(P - FVector3f(R, 0.0f, 0.0f));
 
-				Points[AddedPoints + 2] = T.TransformPosition(P + FVector3f(0.0f, R, 0.0f));
-				Points[AddedPoints + 3] = T.TransformPosition(P - FVector3f(0.0f, R, 0.0f));
+					Points[AddedPoints + 2] = T.TransformPosition(P + FVector3f(0.0f, R, 0.0f));
+					Points[AddedPoints + 3] = T.TransformPosition(P - FVector3f(0.0f, R, 0.0f));
 
-				Points[AddedPoints + 4] = T.TransformPosition(P + FVector3f(0.0f, 0.0f, R));
-				Points[AddedPoints + 5] = T.TransformPosition(P - FVector3f(0.0f, 0.0f, R));
-			}
+					Points[AddedPoints + 4] = T.TransformPosition(P + FVector3f(0.0f, 0.0f, R));
+					Points[AddedPoints + 5] = T.TransformPosition(P - FVector3f(0.0f, 0.0f, R));
+				}
 
 				const int32 BoxCount = Body->GetBoxCount(B);
 				for (int32 I = 0; I < BoxCount; ++I, AddedPoints += 14)
-			{
-				FVector3f P;
-				FQuat4f Q;
-				FVector3f S;
-			
+				{
+					FVector3f P;
+					FQuat4f Q;
+					FVector3f S;
+
 					Body->GetBox(B, I, P, Q, S);
 
-				const FVector3f BasisX = Q.RotateVector(FVector3f::UnitX());
-				const FVector3f BasisY = Q.RotateVector(FVector3f::UnitY());
-				const FVector3f BasisZ = Q.RotateVector(FVector3f::UnitZ());
+					const FVector3f BasisX = Q.RotateVector(FVector3f::UnitX());
+					const FVector3f BasisY = Q.RotateVector(FVector3f::UnitY());
+					const FVector3f BasisZ = Q.RotateVector(FVector3f::UnitZ());
 
 					Points[AddedPoints + 0] = T.TransformPosition(P + BasisX * S.X + BasisY * S.Y + BasisZ * S.Z);
 					Points[AddedPoints + 1] = T.TransformPosition(P + BasisX * S.X - BasisY * S.Y + BasisZ * S.Z);
@@ -634,107 +633,107 @@ namespace mu
 					Points[AddedPoints + 11] = T.TransformPosition(P - BasisX * S.X);
 					Points[AddedPoints + 12] = T.TransformPosition(P - BasisY * S.Y);
 					Points[AddedPoints + 13] = T.TransformPosition(P - BasisZ * S.Z);
-			}
+				}
 
 				const int32 SphylCount = Body->GetSphylCount(B);
 				for (int32 I = 0; I < SphylCount; ++I, AddedPoints += 14)
-			{
-				FVector3f P;
-				FQuat4f Q;
-				float R;
-				float L;
-			
+				{
+					FVector3f P;
+					FQuat4f Q;
+					float R;
+					float L;
+
 					Body->GetSphyl(B, I, P, Q, R, L);
-				
-				const float H = L * 0.5f;
 
-				const FVector3f BasisX = Q.RotateVector(FVector3f::UnitX());
-				const FVector3f BasisY = Q.RotateVector(FVector3f::UnitY());
-				const FVector3f BasisZ = Q.RotateVector(FVector3f::UnitZ());
+					const float H = L * 0.5f;
+
+					const FVector3f BasisX = Q.RotateVector(FVector3f::UnitX());
+					const FVector3f BasisY = Q.RotateVector(FVector3f::UnitY());
+					const FVector3f BasisZ = Q.RotateVector(FVector3f::UnitZ());
 
-				// Top and Bottom
+					// Top and Bottom
 					Points[AddedPoints + 0] = T.TransformPosition(P + BasisZ * (H + R));
 					Points[AddedPoints + 1] = T.TransformPosition(P - BasisZ * (H + R));
 
-				// Top ring
+					// Top ring
 					Points[AddedPoints + 2] = T.TransformPosition(P + BasisX * R + BasisZ * H);
 					Points[AddedPoints + 3] = T.TransformPosition(P - BasisX * R + BasisZ * H);
 					Points[AddedPoints + 4] = T.TransformPosition(P + BasisY * R + BasisZ * H);
 					Points[AddedPoints + 5] = T.TransformPosition(P - BasisY * R + BasisZ * H);
 
-				// Center ring
+					// Center ring
 					Points[AddedPoints + 6] = T.TransformPosition(P + BasisX * R);
 					Points[AddedPoints + 7] = T.TransformPosition(P - BasisX * R);
 					Points[AddedPoints + 8] = T.TransformPosition(P + BasisY * R);
 					Points[AddedPoints + 9] = T.TransformPosition(P - BasisY * R);
-	
-				// Bottom ring
+
+					// Bottom ring
 					Points[AddedPoints + 10] = T.TransformPosition(P + BasisX * R - BasisZ * H);
 					Points[AddedPoints + 11] = T.TransformPosition(P - BasisX * R - BasisZ * H);
 					Points[AddedPoints + 12] = T.TransformPosition(P + BasisY * R - BasisZ * H);
 					Points[AddedPoints + 13] = T.TransformPosition(P - BasisY * R - BasisZ * H);
 
-			}
-	
+				}
+
 				const int32 TaperedCapsuleCount = Body->GetTaperedCapsuleCount(B);
 				for (int32 I = 0; I < TaperedCapsuleCount; ++I, AddedPoints += 14)
-			{
-				FVector3f P;
-				FQuat4f Q;
-				float R0;
-				float R1;
-				float L;
-			
+				{
+					FVector3f P;
+					FQuat4f Q;
+					float R0;
+					float R1;
+					float L;
+
 					Body->GetTaperedCapsule(B, I, P, Q, R0, R1, L);
-		
+
 					const float H = L * 0.5f;
-				const float RCenter = (R0 + R1) * 0.5f;
+					const float RCenter = (R0 + R1) * 0.5f;
 
-				const FVector3f BasisX = Q.RotateVector(FVector3f::UnitX());
-				const FVector3f BasisY = Q.RotateVector(FVector3f::UnitY());
-				const FVector3f BasisZ = Q.RotateVector(FVector3f::UnitZ());
+					const FVector3f BasisX = Q.RotateVector(FVector3f::UnitX());
+					const FVector3f BasisY = Q.RotateVector(FVector3f::UnitY());
+					const FVector3f BasisZ = Q.RotateVector(FVector3f::UnitZ());
 
-				// Top and Bottom
+					// Top and Bottom
 					Points[AddedPoints + 0] = T.TransformPosition(P + BasisZ * (H + R0));
 					Points[AddedPoints + 1] = T.TransformPosition(P - BasisZ * (H + R1));
 
-				// Top ring
+					// Top ring
 					Points[AddedPoints + 2] = T.TransformPosition(P + BasisX * R0 + BasisZ * H);
 					Points[AddedPoints + 3] = T.TransformPosition(P - BasisX * R0 + BasisZ * H);
 					Points[AddedPoints + 4] = T.TransformPosition(P + BasisY * R0 + BasisZ * H);
 					Points[AddedPoints + 5] = T.TransformPosition(P - BasisY * R0 + BasisZ * H);
 
-				// Center ring
+					// Center ring
 					Points[AddedPoints + 6] = T.TransformPosition(P + BasisX * RCenter);
 					Points[AddedPoints + 7] = T.TransformPosition(P - BasisX * RCenter);
 					Points[AddedPoints + 8] = T.TransformPosition(P + BasisY * RCenter);
 					Points[AddedPoints + 9] = T.TransformPosition(P - BasisY * RCenter);
-	
-				// Bottom ring
+
+					// Bottom ring
 					Points[AddedPoints + 10] = T.TransformPosition(P + BasisX * R1 - BasisZ * H);
 					Points[AddedPoints + 11] = T.TransformPosition(P - BasisX * R1 - BasisZ * H);
 					Points[AddedPoints + 12] = T.TransformPosition(P + BasisY * R1 - BasisZ * H);
 					Points[AddedPoints + 13] = T.TransformPosition(P - BasisY * R1 - BasisZ * H);
-			}
+				}
 
 				const int32 ConvexCount = Body->GetConvexCount(B);
 				for (int32 I = 0; I < ConvexCount; ++I)
-			{
-			
-				TArrayView<const FVector3f> VerticesView;
-				TArrayView<const int32> IndicesView;
-				FTransform3f ConvexT;
+				{
+
+					TArrayView<const FVector3f> VerticesView;
+					TArrayView<const int32> IndicesView;
+					FTransform3f ConvexT;
 
 					Body->GetConvex(B, I, VerticesView, IndicesView, ConvexT);
 
-				ConvexT = T * ConvexT;
+					ConvexT = T * ConvexT;
 					for (const FVector3f& P : VerticesView)
-				{
-					Points[AddedPoints++] = ConvexT.TransformPosition(P);
-				}	
+					{
+						Points[AddedPoints++] = ConvexT.TransformPosition(P);
+					}
+				}
 			}
 		}
-		}
 		
 		TArray<FReshapePointBindingData>& PhysicsBodyBindData = ReturnValue.Get<0>();
 		PhysicsBodyBindData.SetNumUninitialized(TotalNumPoints);
@@ -830,7 +829,7 @@ namespace mu
 	}
 
 	inline TTuple<TArray<FReshapePointBindingData>, TArray<int32>> BindPose(
-			const Mesh* Mesh, FShapeMeshTree& ShapeMeshTree, const TArray<string>& BonesToDeform )
+			const Mesh* Mesh, FShapeMeshTree& ShapeMeshTree, const TArray<uint16>& BonesToDeform )
 	{
 		UE::Geometry::FAxisAlignedBox3d ShapeAABBox = ShapeMeshTree.GetBoundingBox();
 
@@ -856,7 +855,7 @@ namespace mu
 				continue;
 			}
 
-			if (!BonesToDeform.Contains(Mesh->GetBonePoseName(BoneIndex)))
+			if (!BonesToDeform.Contains(Mesh->BonePoses[BoneIndex].BoneId))
 			{
 				continue;
 			}
@@ -884,7 +883,7 @@ namespace mu
     inline void MeshBindShapeReshape(
 			Mesh* Result,
 			const Mesh* BaseMesh, const Mesh* ShapeMesh, 
-			const TArray<string>& BonesToDeform, const TArray<string>& PhysicsToDeform, 
+			const TArray<uint16>& BonesToDeform, const TArray<uint16>& PhysicsToDeform, 
 			EMeshBindShapeFlags BindFlags,
 			bool& bOutSuccess)
     {
diff --git a/Engine/Plugins/Experimental/Mutable/Source/MutableRuntime/Private/MuR/OpMeshClipMorphPlane.h b/Engine/Plugins/Experimental/Mutable/Source/MutableRuntime/Private/MuR/OpMeshClipMorphPlane.h
--- a/Engine/Plugins/Experimental/Mutable/Source/MutableRuntime/Private/MuR/OpMeshClipMorphPlane.h
+++ b/Engine/Plugins/Experimental/Mutable/Source/MutableRuntime/Private/MuR/OpMeshClipMorphPlane.h
@@ -70,7 +70,7 @@ namespace mu
 	//! Reference version
 	//---------------------------------------------------------------------------------------------
 	inline void MeshClipMorphPlane(Mesh* Result, const Mesh* pBase, const vec3f& origin, const vec3f& normal, float dist, float factor, float radius,
-		float radius2, float angle, const FShape& selectionShape, bool& bOutSuccess, const mu::string& boneName = mu::string(), float vertexSelectionBoneMaxRadius = -1.f)
+		float radius2, float angle, const FShape& selectionShape, bool& bOutSuccess, const int32 BoneId = INDEX_NONE, float vertexSelectionBoneMaxRadius = -1.f)
 	{
 		bOutSuccess = true;
 		//float radius = 8.f;
@@ -96,33 +96,32 @@ namespace mu
 			return;
 		}
 
-		TArray<bool> bone_is_affected;
+		TArray<bool> AffectedBones;
 		TArray<vertex_bone_info> vertex_info;
 
-        auto pBaseSkeleton = pBase->GetSkeleton();
-        if (!boneName.empty() && pBaseSkeleton )
+        Ptr<const Skeleton> BaseSkeleton = pBase->GetSkeleton();
+
+		const int32 BaseBoneIndex = BaseSkeleton ? BaseSkeleton->FindBone(BoneId) : INDEX_NONE;
+        if (BaseBoneIndex != INDEX_NONE)
 		{
-            bone_is_affected.SetNum(pBaseSkeleton->GetBoneCount());
+			AffectedBones.SetNum(BaseSkeleton->GetBoneCount());
+			AffectedBones[BaseBoneIndex] = true;
 
-            for (int32 bone_idx = 0; bone_idx < pBaseSkeleton->GetBoneCount(); ++bone_idx)
+            for (int32 BoneIndex = 0; BoneIndex < BaseSkeleton->GetBoneCount(); ++BoneIndex)
 			{
-				int32 current_idx = bone_idx;
-				bool found_bone = false;
-
-				while (current_idx >= 0)
+				int32 CurrentBoneIndex = BoneIndex;
+				while (CurrentBoneIndex >= 0)
 				{
-                    if (pBaseSkeleton->GetBoneName(current_idx) == boneName)
+                    if (BaseSkeleton->GetBoneId(CurrentBoneIndex) == BaseBoneIndex)
 					{
-						found_bone = true;
+						AffectedBones[BoneIndex] = true;
 						break;
 					}
 					else
 					{
-                        current_idx = pBaseSkeleton->GetBoneParent(current_idx);
+						CurrentBoneIndex = BaseSkeleton->GetBoneParent(CurrentBoneIndex);
 					}
 				}
-
-				bone_is_affected[bone_idx] = found_bone;
 			}
 
 			// Look for affected vertex indices
@@ -303,8 +302,8 @@ namespace mu
 						}
 
 						if (
-							(  !boneName.empty() && VertexIsAffectedByBone(v, bone_is_affected, vertex_info) && VertexIsInMaxRadius(vertex, origin, vertexSelectionBoneMaxRadius))
-                            || (boneName.empty() && selectionShape.type == (uint8_t)FShape::Type::None)
+							(  BaseBoneIndex != INDEX_NONE && VertexIsAffectedByBone(v, AffectedBones, vertex_info) && VertexIsInMaxRadius(vertex, origin, vertexSelectionBoneMaxRadius))
+                            || (BaseBoneIndex == INDEX_NONE && selectionShape.type == (uint8_t)FShape::Type::None)
                             || (selectionShape.type == (uint8_t)FShape::Type::AABox && PointInBoundingBox(vertex, selectionShape))
 							)
 						{
diff --git a/Engine/Plugins/Experimental/Mutable/Source/MutableRuntime/Private/MuR/OpMeshFormat.cpp b/Engine/Plugins/Experimental/Mutable/Source/MutableRuntime/Private/MuR/OpMeshFormat.cpp
--- a/Engine/Plugins/Experimental/Mutable/Source/MutableRuntime/Private/MuR/OpMeshFormat.cpp
+++ b/Engine/Plugins/Experimental/Mutable/Source/MutableRuntime/Private/MuR/OpMeshFormat.cpp
@@ -686,7 +686,7 @@ namespace mu
 		Result->m_AdditionalBuffers = pSource->m_AdditionalBuffers;
 
 		Result->BonePoses = pSource->BonePoses;
-		Result->BoneMapIndices = pSource->BoneMapIndices;
+		Result->BoneMap = pSource->BoneMap;
 
 		Result->SkeletonIDs = pSource->SkeletonIDs;
 
diff --git a/Engine/Plugins/Experimental/Mutable/Source/MutableRuntime/Private/MuR/OpMeshMerge.h b/Engine/Plugins/Experimental/Mutable/Source/MutableRuntime/Private/MuR/OpMeshMerge.h
--- a/Engine/Plugins/Experimental/Mutable/Source/MutableRuntime/Private/MuR/OpMeshMerge.h
+++ b/Engine/Plugins/Experimental/Mutable/Source/MutableRuntime/Private/MuR/OpMeshMerge.h
@@ -436,7 +436,6 @@ namespace mu
 			bMergeSkeletons = !(*pFirst->GetSkeleton() == *pSecond->GetSkeleton());
 		}
 
-		TArray<uint16> SecondToResultBoneIndices;
 
 		if (bMergeSkeletons)
 		{
@@ -452,40 +451,39 @@ namespace mu
 			pResultSkeleton = pFirstSkeleton ? pFirstSkeleton->Clone() : new Skeleton;
 			Result->SetSkeleton(pResultSkeleton);
 
+			TArray<uint16> SecondToResultBoneIndices;
 			SecondToResultBoneIndices.SetNumUninitialized(NumBonesSecond);
 
 			// Merge pSecond and build the remap array 
 			for (int32 SecondBoneIndex = 0; SecondBoneIndex < NumBonesSecond; ++SecondBoneIndex)
 			{
-				int16 BoneIndex = INDEX_NONE;
-				for (int32 FirstBoneIndex = 0; FirstBoneIndex < NumBonesFirst; ++FirstBoneIndex)
-				{
-					if (pResultSkeleton->m_bones[FirstBoneIndex] == pSecondSkeleton->GetBoneName(SecondBoneIndex))
-					{
-						BoneIndex = FirstBoneIndex;
-						break;
-					}
-				}
+				const uint16 BoneNameId = pSecondSkeleton->BoneIds[SecondBoneIndex];
+				int32 Index = pResultSkeleton->FindBone(BoneNameId);
 
 				// Add a new bone
-				if (BoneIndex == INDEX_NONE)
+				if (Index == INDEX_NONE)
 				{
-					BoneIndex = pResultSkeleton->m_bones.Add(pSecondSkeleton->m_bones[SecondBoneIndex]);
+					Index = pResultSkeleton->BoneIds.Add(BoneNameId);
 
 					// Add an incorrect index, to be fixed below in case the parent index is later in the bone array.
-					pResultSkeleton->m_boneParents.Add(pSecondSkeleton->m_boneParents[SecondBoneIndex]);
+					pResultSkeleton->BoneParents.Add(pSecondSkeleton->BoneParents[SecondBoneIndex]);
+
+					if (pSecondSkeleton->BoneNames.IsValidIndex(SecondBoneIndex))
+					{
+						pResultSkeleton->BoneNames.Add(pSecondSkeleton->BoneNames[SecondBoneIndex]);
+					}
 				}
 
-				SecondToResultBoneIndices[SecondBoneIndex] = (uint16)BoneIndex;
+				SecondToResultBoneIndices[SecondBoneIndex] = (uint16)Index;
 			}
 
 			// Fix second mesh bone parents
-			for (int32 ob = NumBonesFirst; ob < pResultSkeleton->m_boneParents.Num(); ++ob)
+			for (int32 ob = NumBonesFirst; ob < pResultSkeleton->BoneParents.Num(); ++ob)
 			{
-				int16_t secondMeshIndex = pResultSkeleton->m_boneParents[ob];
+				int16 secondMeshIndex = pResultSkeleton->BoneParents[ob];
 				if (secondMeshIndex != INDEX_NONE)
 				{
-					pResultSkeleton->m_boneParents[ob] = SecondToResultBoneIndices[secondMeshIndex];
+					pResultSkeleton->BoneParents[ob] = SecondToResultBoneIndices[secondMeshIndex];
 				}
 			}
 		}
@@ -503,25 +501,23 @@ namespace mu
 		TArray<uint16> RemappedBoneMapIndices;
 
 		// Used to know the format of the bone index buffer
-		uint32 MaxBoneMapIndices = 0;
-		const int32 NumSecondBoneMapIndices = pSecond->BoneMapIndices.Num();
+		uint32 MaxNumBonesInBoneMaps = 0;
+		const int32 NumSecondBonesInBoneMap = pSecond->BoneMap.Num();
 
 		{
 			MUTABLE_CPUPROFILER_SCOPE(Surfaces);
 			
-			const int32 NumFirstBoneMapIndices = pFirst->BoneMapIndices.Num();
-			Result->BoneMapIndices = pFirst->BoneMapIndices;
+			const int32 NumFirstBonesInBoneMap = pFirst->BoneMap.Num();
+			Result->BoneMap = pFirst->BoneMap;
 
 			if (bMergeSurfaces)
 			{
 				// Merge BoneMaps
-				RemappedBoneMapIndices.SetNumUninitialized(NumSecondBoneMapIndices);
+				RemappedBoneMapIndices.SetNumUninitialized(NumSecondBonesInBoneMap);
 
-				for (uint16 SecondBoneMapIndex = 0; SecondBoneMapIndex < NumSecondBoneMapIndices; ++SecondBoneMapIndex)
+				for (uint16 SecondBoneMapIndex = 0; SecondBoneMapIndex < NumSecondBonesInBoneMap; ++SecondBoneMapIndex)
 				{
-					const int32 SecondBoneIndex = pSecond->BoneMapIndices[SecondBoneMapIndex];
-					const int32 BoneIndex = bMergeSkeletons ? SecondToResultBoneIndices[SecondBoneIndex] : SecondBoneIndex;
-					const int32 BoneMapIndex = Result->BoneMapIndices.AddUnique(BoneIndex);
+					const int32 BoneMapIndex = Result->BoneMap.AddUnique(pSecond->BoneMap[SecondBoneMapIndex]);
 					RemappedBoneMapIndices[SecondBoneMapIndex] = BoneMapIndex;
 
 					bRemapBoneIndices = bRemapBoneIndices || BoneMapIndex != SecondBoneMapIndex;
@@ -530,18 +526,12 @@ namespace mu
 				MESH_SURFACE& NewSurface = Result->m_surfaces.AddDefaulted_GetRef();
 				NewSurface.m_vertexCount = pFirst->GetVertexCount() + pSecond->GetVertexCount();
 				NewSurface.m_indexCount = pFirst->GetIndexCount() + pSecond->GetIndexCount();
-				NewSurface.BoneMapCount = Result->BoneMapIndices.Num();
+				NewSurface.BoneMapCount = Result->BoneMap.Num();
 			}
 			else
 			{
-				// Add the bonemaps of the second mesh with the fixed BoneIndices
-				const int32 NumBoneMapIndices = NumFirstBoneMapIndices + NumSecondBoneMapIndices;
-				Result->BoneMapIndices.Reserve(NumBoneMapIndices);
-
-				for (const uint16& BoneIndex : pSecond->BoneMapIndices)
-				{
-					Result->BoneMapIndices.Add(bMergeSkeletons ? SecondToResultBoneIndices[BoneIndex] : BoneIndex);
-				}
+				// Add the BoneMap of the second mesh
+				Result->BoneMap.Append(pSecond->BoneMap);
 
 				// Add pFirst surfaces
 				Result->m_surfaces = pFirst->m_surfaces;
@@ -553,15 +543,15 @@ namespace mu
 				MESH_SURFACE& NewSurface = Result->m_surfaces.Add_GetRef(pSecond->m_surfaces[0]);
 				NewSurface.m_firstVertex += FirstVertexIndex;
 				NewSurface.m_firstIndex += FirstIndexIndex;
-				NewSurface.BoneMapIndex += NumFirstBoneMapIndices;
+				NewSurface.BoneMapIndex += NumFirstBonesInBoneMap;
 			}
 
 			for (const MESH_SURFACE& Surface : Result->m_surfaces)
 			{
-				MaxBoneMapIndices = FMath::Max(MaxBoneMapIndices, Surface.BoneMapCount);
+				MaxNumBonesInBoneMaps = FMath::Max(MaxNumBonesInBoneMaps, Surface.BoneMapCount);
 			}
 
-			Result->BoneMapIndices.Shrink();
+			Result->BoneMap.Shrink();
 		}
 
 
@@ -578,7 +568,7 @@ namespace mu
 			// Add or override bone poses
 			for (const Mesh::FBonePose& SecondBonePose : pSecond->BonePoses)
 			{
-				const int32 ResultBoneIndex = Result->FindBonePose(SecondBonePose.BoneName.c_str());
+				const int32 ResultBoneIndex = Result->FindBonePose(SecondBonePose.BoneId);
 
 				if (ResultBoneIndex != INDEX_NONE)
 				{
@@ -618,11 +608,11 @@ namespace mu
 			// Appends InPhysicsBody to OutPhysicsBody removing Bodies that are equal, have same bone and customId and its properies are identical.	
 			auto AppendPhysicsBodiesUnique = [](PhysicsBody& OutPhysicsBody, const PhysicsBody& InPhysicsBody) -> bool
 			{
-				TArray<string>& OutBones = OutPhysicsBody.Bones;
+				TArray<uint16>& OutBones = OutPhysicsBody.BoneIds;
 				TArray<int32>& OutCustomIds = OutPhysicsBody.BodiesCustomIds;
 				TArray<FPhysicsBodyAggregate>& OutBodies = OutPhysicsBody.Bodies;
 
-				const TArray<string>& InBones = InPhysicsBody.Bones;
+				const TArray<uint16>& InBones = InPhysicsBody.BoneIds;
 				const TArray<int32>& InCustomIds = InPhysicsBody.BodiesCustomIds;
 				const TArray<FPhysicsBodyAggregate>& InBodies = InPhysicsBody.Bodies;
 
@@ -764,7 +754,7 @@ namespace mu
 
 			// Check if the format of the BoneIndex buffer has to change
 			bool bChangeBoneIndicesFormat = false;
-			MESH_BUFFER_FORMAT BoneIndexFormat = MaxBoneMapIndices > MAX_uint8 ? MBF_UINT16 : MBF_UINT8;
+			MESH_BUFFER_FORMAT BoneIndexFormat = MaxNumBonesInBoneMaps > MAX_uint8 ? MBF_UINT16 : MBF_UINT8;
 
 			// Iterate all vertex buffers to check if we need to format bone indices
 			{
@@ -1061,7 +1051,7 @@ namespace mu
 									uint8 BoneMapIndex = pD[ComponentIndex];
 
 									// be defensive
-									if (BoneMapIndex < NumSecondBoneMapIndices)
+									if (BoneMapIndex < NumSecondBonesInBoneMap)
 									{
 										pD[ComponentIndex] = (uint8)RemappedBoneMapIndices[BoneMapIndex];
 									}
@@ -1086,7 +1076,7 @@ namespace mu
 									uint16 BoneMapIndex = pD[ComponentIndex];
 
 									// be defensive
-									if (BoneMapIndex < NumSecondBoneMapIndices)
+									if (BoneMapIndex < NumSecondBonesInBoneMap)
 									{
 										pD[ComponentIndex] = (uint16)RemappedBoneMapIndices[BoneMapIndex];
 									}
@@ -1137,15 +1127,20 @@ namespace mu
         int initialBones = pBase->GetBoneCount();
         for ( int b=0; pOther && b<pOther->GetBoneCount(); ++b)
         {
-            int resultBoneIndex = pBase->FindBone( pOther->GetBoneName(b) );
+            int32 resultBoneIndex = pBase->FindBone( pOther->GetBoneId(b) );
             if ( resultBoneIndex<0 )
             {
-                int32 newIndex = pBase->m_bones.Num();
+                int32 newIndex = pBase->BoneIds.Num();
                 otherToResult.Add(b,newIndex);
-                pBase->m_bones.Add( pOther->m_bones[b] );
+                pBase->BoneIds.Add( pOther->BoneIds[b] );
 
                 // Will be remapped below
-                pBase->m_boneParents.Add(pOther->m_boneParents[b] );
+                pBase->BoneParents.Add(pOther->BoneParents[b] );
+
+				if (pOther->BoneNames.IsValidIndex(b))
+				{
+					pBase->BoneNames.Add(pOther->BoneNames[b]);
+				}
             }
             else
             {
@@ -1156,10 +1151,10 @@ namespace mu
         // Fix bone parent indices of the bones added from pOther
         for ( int b=initialBones;b<pBase->GetBoneCount(); ++b)
         {
-            int16_t sourceIndex = pBase->m_boneParents[b];
+            int16_t sourceIndex = pBase->BoneParents[b];
             if (sourceIndex>=0)
             {
-                pBase->m_boneParents[b] = (int16_t)otherToResult[ sourceIndex ];
+                pBase->BoneParents[b] = (int16_t)otherToResult[ sourceIndex ];
             }
         }
     }
@@ -1178,36 +1173,34 @@ namespace mu
             return;
         }
 
-        mu::SkeletonPtrConst SourceSkeleton = SourceMesh->GetSkeleton();
-		const TArray<uint16>& SourceBoneMaps = SourceMesh->GetBoneMap();
+  //      mu::SkeletonPtrConst SourceSkeleton = SourceMesh->GetSkeleton();
+		//const TArray<uint16>& SourceBoneMaps = SourceMesh->GetBoneMap();
 
-        // Remap the indices of the bonemap to those of the new skeleton
-		TArray<uint16> RemappedBoneMap;
+  //      // Remap the indices of the bonemap to those of the new skeleton
+		//TArray<uint16> RemappedBoneMap;
 
-		const int32 NumBonesBoneMap = SourceBoneMaps.Num();
-		RemappedBoneMap.Reserve(NumBonesBoneMap);
+		//const int32 NumBonesBoneMap = SourceBoneMaps.Num();
+		//RemappedBoneMap.Reserve(NumBonesBoneMap);
 
-		bool bBonesRemapped = false;
+		//bool bBonesRemapped = false;
 
-		for (const uint16& SourceBoneIndex : SourceBoneMaps)
-		{
-			const char* BoneName = SourceSkeleton->GetBoneName(SourceBoneIndex);
-			const uint16 BoneIndex = Skeleton->FindBone(BoneName);
+		//for (const uint16& SourceBoneIndex : SourceBoneMaps)
+		//{
+		//	const uint16 BoneIndex = Skeleton->FindBone(SourceSkeleton->GetBoneId(SourceBoneIndex));
 
-			bBonesRemapped = bBonesRemapped || BoneIndex != SourceBoneIndex;
-			RemappedBoneMap.Add(BoneIndex);
-		}
+		//	bBonesRemapped = bBonesRemapped || BoneIndex != SourceBoneIndex;
+		//	RemappedBoneMap.Add(BoneIndex);
+		//}
 
-        if (!bBonesRemapped)
-        {
-			bOutSuccess = false;
-            return;
-        }
+  //      if (!bBonesRemapped)
+  //      {
+		//	bOutSuccess = false;
+  //          return;
+  //      }
 
 		Result->CopyFrom(*SourceMesh);
-
 		Result->SetSkeleton(Skeleton);
-		Result->SetBoneMap(RemappedBoneMap);
+		//Result->SetBoneMap(RemappedBoneMap);
     }
 	
 }
diff --git a/Engine/Plugins/Experimental/Mutable/Source/MutableRuntime/Private/MuR/OpMeshReshape.h b/Engine/Plugins/Experimental/Mutable/Source/MutableRuntime/Private/MuR/OpMeshReshape.h
--- a/Engine/Plugins/Experimental/Mutable/Source/MutableRuntime/Private/MuR/OpMeshReshape.h
+++ b/Engine/Plugins/Experimental/Mutable/Source/MutableRuntime/Private/MuR/OpMeshReshape.h
@@ -537,7 +537,7 @@ namespace mu
 		// Retrieve them in the same order the boxes where put in, so they can be linked to the physics body volumes.
 		for ( const int32 B : UsedIndices )
 		{	
-			int32 BoneIdx = BaseMesh.FindBonePose(PBody.GetBodyBoneName(B));
+			int32 BoneIdx = BaseMesh.FindBonePose(PBody.GetBodyBoneId(B));
 			FTransform3f BoneTransform = FTransform3f::Identity;
 			if (BoneIdx >= 0)
 			{
diff --git a/Engine/Plugins/Experimental/Mutable/Source/MutableRuntime/Private/MuR/PhysicsBody.cpp b/Engine/Plugins/Experimental/Mutable/Source/MutableRuntime/Private/MuR/PhysicsBody.cpp
--- a/Engine/Plugins/Experimental/Mutable/Source/MutableRuntime/Private/MuR/PhysicsBody.cpp
+++ b/Engine/Plugins/Experimental/Mutable/Source/MutableRuntime/Private/MuR/PhysicsBody.cpp
@@ -233,7 +233,7 @@ namespace mu
 		pResult->CustomId = CustomId;
 
 		pResult->Bodies = Bodies;
-		pResult->Bones = Bones;
+		pResult->BoneIds = BoneIds;
 		pResult->BodiesCustomIds = BodiesCustomIds;
 
 		pResult->bBodiesModified = bBodiesModified;
@@ -258,7 +258,7 @@ namespace mu
 	void PhysicsBody::SetBodyCount(int32 Count)
 	{
 		Bodies.SetNum(Count);
-		Bones.SetNum(Count);
+		BoneIds.SetNum(Count);
 		BodiesCustomIds.Init(-1, Count);
 	}
 
@@ -269,19 +269,17 @@ namespace mu
 	}
 
 	//-------------------------------------------------------------------------------------------------
-	const char* PhysicsBody::GetBodyBoneName(int32 B) const
+	uint16 PhysicsBody::GetBodyBoneId(int32 B) const
 	{
-		check(B >= 0 && B < Bones.Num());
-
-		return Bones[B].c_str();
+		check(BoneIds.IsValidIndex(B));
+		return BoneIds[B];
 	}
 
 	//-------------------------------------------------------------------------------------------------
-	void PhysicsBody::SetBodyBoneName(int32 B, const char* BoneName)
+	void PhysicsBody::SetBodyBoneId(int32 B, const uint16 BoneId)
 	{
-		check(B >= 0 && B < Bones.Num());
-
-		Bones[B] = BoneName;
+		check(BoneIds.IsValidIndex(B));
+		BoneIds[B] = BoneId;
 	}
 
 	//-------------------------------------------------------------------------------------------------
@@ -711,12 +709,12 @@ namespace mu
 	//-------------------------------------------------------------------------------------------------
 	void PhysicsBody::Serialise(OutputArchive& arch) const
 	{
-		uint32 ver = 2;
+		uint32 ver = 3;
 		arch << ver;
         	
 		arch << CustomId;
 		arch << Bodies;
-		arch << Bones;
+		arch << BoneIds;
 		arch << BodiesCustomIds;
 		arch << bBodiesModified;
 	}
@@ -727,7 +725,7 @@ namespace mu
 	{
 		uint32 ver;
 		arch >> ver;
-		check(ver <= 2);
+		check(ver <= 3);
         	
 		if (ver >= 2)
 		{
@@ -735,7 +733,25 @@ namespace mu
 		}
         	
 		arch >> Bodies;
-		arch >> Bones;
+
+		if (ver >= 3)
+		{
+			arch >> BoneIds;
+		}
+		else
+		{
+			TArray<string> BoneNames;
+			arch >> BoneNames;
+
+			const int32 NumBoneNames = BoneNames.Num();
+			TArray<string> UniqueBoneNames;
+			UniqueBoneNames.Reserve(NumBoneNames);
+			BoneIds.Reserve(NumBoneNames);
+			for (int32 BoneIndex = 0; BoneIndex < NumBoneNames; ++BoneIndex)
+			{
+				BoneIds.Add((uint16)UniqueBoneNames.AddUnique(BoneNames[BoneIndex]));
+			}
+		}
 		arch >> BodiesCustomIds;
 
 		if (ver >= 1)
diff --git a/Engine/Plugins/Experimental/Mutable/Source/MutableRuntime/Private/MuR/Skeleton.cpp b/Engine/Plugins/Experimental/Mutable/Source/MutableRuntime/Private/MuR/Skeleton.cpp
--- a/Engine/Plugins/Experimental/Mutable/Source/MutableRuntime/Private/MuR/Skeleton.cpp
+++ b/Engine/Plugins/Experimental/Mutable/Source/MutableRuntime/Private/MuR/Skeleton.cpp
@@ -31,8 +31,11 @@ namespace mu
 	{
 		SkeletonPtr pResult = new Skeleton();
 
-		pResult->m_bones = m_bones;
-		pResult->m_boneParents = m_boneParents;
+		pResult->BoneIds = BoneIds;
+		pResult->BoneParents = BoneParents;
+
+		// For debug
+		pResult->BoneNames = BoneNames;
 
 		return pResult;
 	}
@@ -41,74 +44,100 @@ namespace mu
 	//---------------------------------------------------------------------------------------------
 	int32 Skeleton::GetBoneCount() const
 	{
-		return m_bones.Num();
+		return BoneIds.Num();
 	}
 
 
 	//---------------------------------------------------------------------------------------------
-	void Skeleton::SetBoneCount(int32 b)
+	void Skeleton::SetBoneCount(int32 NumBones)
 	{
-		m_bones.SetNum(b);
-		m_boneParents.Init(INDEX_NONE, b);
+		BoneNames.SetNum(NumBones);
+		BoneIds.SetNum(NumBones);
+		BoneParents.Init(INDEX_NONE, NumBones);
 	}
 
 
 	//---------------------------------------------------------------------------------------------
-	const char* Skeleton::GetBoneName(int32 boneIndex) const
+	const FName Skeleton::GetBoneFName(int32 Index) const
 	{
-		check(boneIndex >= 0 && boneIndex < GetBoneCount());
-		if (boneIndex >= 0 && boneIndex < GetBoneCount())
+		if (BoneNames.IsValidIndex(Index))
 		{
-			return m_bones[boneIndex].c_str();
+			return BoneNames[Index];
 		}
-		return "";
+
+		return FName("Unknown Bone");
 	}
 
 
 	//---------------------------------------------------------------------------------------------
-	void Skeleton::SetBoneName(int32 boneIndex, const char* strName)
+	void Skeleton::SetBoneFName(const int32 Index, const FName BoneName)
 	{
-		check(boneIndex >= 0 && boneIndex < GetBoneCount());
-		if (boneIndex >= 0 && boneIndex < GetBoneCount())
+		if (BoneNames.IsValidIndex(Index))
 		{
-			m_bones[boneIndex] = strName;
+			BoneNames[Index] = BoneName;
 		}
 	}
 
 
 	//---------------------------------------------------------------------------------------------
-	int Skeleton::GetBoneParent(int32 boneIndex) const
+	int32 Skeleton::GetBoneParent(int32 Index) const
 	{
-		if (boneIndex >= 0 && boneIndex < m_boneParents.Num())
+		if (BoneParents.IsValidIndex(Index))
 		{
-			return m_boneParents[boneIndex];
+			return BoneParents[Index];
 		}
 
-		return -1;
+		return INDEX_NONE;
+	}
+
+
+	//-------------------------------------------------------------------------------------------------
+	void Skeleton::SetBoneParent(int32 Index, int32 ParentIndex)
+	{
+		check(ParentIndex >= -1 && ParentIndex < GetBoneCount() && ParentIndex < 0xffff);
+		check(BoneParents.IsValidIndex(Index));
+		
+		if (BoneParents.IsValidIndex(Index))
+		{
+			BoneParents[Index] = (int16)ParentIndex;
+		}
 	}
 
 
 	//-------------------------------------------------------------------------------------------------
-	void Skeleton::SetBoneParent(int32 boneIndex, int32 parentBoneIndex)
+	uint16 Skeleton::GetBoneId(const int32 Index) const
+	{
+		check(BoneIds.IsValidIndex(Index));
+		return BoneIds[Index];
+	}
+
+	
+	//-------------------------------------------------------------------------------------------------
+	void Skeleton::SetBoneId(const int32 Index, const uint16 BoneId)
 	{
-		check(boneIndex >= 0 && boneIndex < GetBoneCount());
-		check(parentBoneIndex >= -1 && parentBoneIndex < GetBoneCount() &&
-			parentBoneIndex < 0xffff);
-		if (boneIndex >= 0 && boneIndex < m_boneParents.Num())
+		check(BoneIds.IsValidIndex(Index));
+		if (BoneIds.IsValidIndex(Index))
 		{
-			m_boneParents[boneIndex] = (int16_t)parentBoneIndex;
+			BoneIds[Index] = BoneId;
 		}
 	}
 
 
+	//-------------------------------------------------------------------------------------------------
+	int32 Skeleton::FindBone(const uint16 BoneId) const
+	{
+		return BoneIds.Find(BoneId);
+	}
+
+
 	//-------------------------------------------------------------------------------------------------
 	void Skeleton::Serialise(OutputArchive& arch) const
 	{
-		uint32 ver = 5;
+		uint32 ver = 6;
 		arch << ver;
 
-		arch << m_bones;
-		arch << m_boneParents;
+		arch << BoneIds;
+		arch << BoneParents;
 	}
 
 
@@ -119,19 +148,35 @@ namespace mu
 		arch >> ver;
 		check(ver >= 3);
 
-		arch >> m_bones;
+		if(ver >= 6)
+		{
+			arch >> BoneIds;
+		}
+		else
+		{
+			TArray<string> OldBoneNames;
+			arch >> OldBoneNames;
+
+			const int32 NumBones = OldBoneNames.Num();
+
+			BoneIds.SetNum(NumBones);
+			for (int32 Index = 0; Index < NumBones; ++Index)
+			{
+				BoneIds[Index] = Index;
+			}
+		}
 
 		if (ver == 3)
 		{
 			arch >> m_boneTransforms_DEPRECATED;
 		}
 
-		arch >> m_boneParents;
+		arch >> BoneParents;
 
 		if (ver <= 4)
 		{
-			TArray<int32> boneIds;
-			arch >> boneIds;
+			TArray<int32> BoneIds_DEPRECATED;
+			arch >> BoneIds_DEPRECATED;
 		}
 
 		if (ver == 3)
diff --git a/Engine/Plugins/Experimental/Mutable/Source/MutableRuntime/Public/MuR/Mesh.h b/Engine/Plugins/Experimental/Mutable/Source/MutableRuntime/Public/MuR/Mesh.h
--- a/Engine/Plugins/Experimental/Mutable/Source/MutableRuntime/Public/MuR/Mesh.h
+++ b/Engine/Plugins/Experimental/Mutable/Source/MutableRuntime/Public/MuR/Mesh.h
@@ -295,10 +295,7 @@ namespace mu
         void SetTag( int tagIndex, const char* strName );
 
 		//!
-		const char* GetBonePoseName( int32 boneIndex ) const;
-
-		//!
-		int32 FindBonePose(const char* strName) const;
+		int32 FindBonePose(uint16 BoneId) const;
 		
 		//!
 		void SetBonePoseCount(int32 count);
@@ -307,7 +304,10 @@ namespace mu
 		int32 GetBonePoseCount() const;
 
 		//!
-		void SetBonePose(int32 BoneIndex, const char* StrName, FTransform3f Transform, EBoneUsageFlags BoneUsageFlags);
+		void SetBonePose(int32 Index, uint16 BoneId, FTransform3f Transform, EBoneUsageFlags BoneUsageFlags);
+
+		//! Return BoneId (uint16) of the pose at index PoseIndex or INDEX_NONE
+		int32 GetBonePoseBoneId(int32 PoseIndex) const;
 
 		//! Return a matrix stored per bone. It is a set of 16-float values.
 		void GetBoneTransform(int32 BoneIndex, FTransform3f& Transform) const;
@@ -409,7 +409,9 @@ namespace mu
 
 		struct FBonePose
 		{
-			string BoneName;
+			// Index of the bone in the CO BoneNames array
+			uint16 BoneId;
+
 			EBoneUsageFlags BoneUsageFlags = EBoneUsageFlags::None;
 			FTransform3f BoneTransform;
 
@@ -423,15 +425,15 @@ namespace mu
 			inline bool operator==(const FBonePose& Other) const
 			{
 				return BoneUsageFlags == Other.BoneUsageFlags
-					&& BoneName == Other.BoneName;
+					&& BoneId == Other.BoneId;
 			}
 		};
 		// This is the pose used by this mesh fragment, used to update the transforms of the final skeleton
 		// taking into consideration the meshes being used.
 		TArray<FBonePose> BonePoses;
 
-		// Array containing the bonemap indices of all surfaces in the mesh.
-		TArray<uint16> BoneMapIndices;
+		// Array containing the bonemaps of all surfaces in the mesh.
+		TArray<uint16> BoneMap;
 
 		//!
 		inline void Serialise(OutputArchive& arch) const;
@@ -448,7 +450,7 @@ namespace mu
 			if (equal) equal = (m_FaceBuffers == o.m_FaceBuffers);
 			if (equal) equal = (m_layouts.Num() == o.m_layouts.Num());
 			if (equal) equal = (BonePoses.Num() == o.BonePoses.Num());
-			if (equal) equal = (BoneMapIndices.Num() == o.BoneMapIndices.Num());
+			if (equal) equal = (BoneMap.Num() == o.BoneMap.Num());
 			if (equal && m_pSkeleton != o.m_pSkeleton)
 			{
 				if (m_pSkeleton && o.m_pSkeleton)
@@ -482,7 +484,7 @@ namespace mu
 				equal &= BonePoses[i] == o.BonePoses[i];
 			}
 
-			if (equal) equal = BoneMapIndices == o.BoneMapIndices;
+			if (equal) equal = BoneMap == o.BoneMap;
 
 			equal &= AdditionalPhysicsBodies.Num() == o.AdditionalPhysicsBodies.Num();
 			for (int32 i = 0; equal && i < AdditionalPhysicsBodies.Num(); ++i)
diff --git a/Engine/Plugins/Experimental/Mutable/Source/MutableRuntime/Public/MuR/Model.h b/Engine/Plugins/Experimental/Mutable/Source/MutableRuntime/Public/MuR/Model.h
--- a/Engine/Plugins/Experimental/Mutable/Source/MutableRuntime/Public/MuR/Model.h
+++ b/Engine/Plugins/Experimental/Mutable/Source/MutableRuntime/Public/MuR/Model.h
@@ -13,7 +13,7 @@
 //! data structures. Compiled models are not necessarily compatible when the runtime is updated,
 //! so this version number can be used externally to verify this. It is not used internally, and
 //! serializing models from different versions than this runtime will probably result in a crash.
-#define MUTABLE_COMPILED_MODEL_CODE_VERSION		uint32( 70 )
+#define MUTABLE_COMPILED_MODEL_CODE_VERSION		uint32( 71 )
 #define MUTABLE_PARAMETERS_VERSION              uint32( 2 )
 
 
diff --git a/Engine/Plugins/Experimental/Mutable/Source/MutableRuntime/Public/MuR/PhysicsBody.h b/Engine/Plugins/Experimental/Mutable/Source/MutableRuntime/Public/MuR/PhysicsBody.h
--- a/Engine/Plugins/Experimental/Mutable/Source/MutableRuntime/Public/MuR/PhysicsBody.h
+++ b/Engine/Plugins/Experimental/Mutable/Source/MutableRuntime/Public/MuR/PhysicsBody.h
@@ -189,8 +189,8 @@ namespace mu
 		void SetBodyCount(int32 B);
 		int32 GetBodyCount() const;
 	
-		void SetBodyBoneName(int32 B, const char* BoneName);
-		const char* GetBodyBoneName(int32 B) const;
+		void SetBodyBoneId(int32 B, uint16 BoneId);
+		uint16 GetBodyBoneId(int32 B) const;
 		
 		void SetBodyCustomId(int32 B, int32 BodyCustomId);
 		int32 GetBodyCustomId(int32 B) const;
@@ -297,7 +297,7 @@ namespace mu
 		int32 CustomId = -1;
 
 		// Bone name the physics volume aggregate is bound to. 
-		TArray<string> Bones;
+		TArray<uint16> BoneIds;
 		TArray<FPhysicsBodyAggregate> Bodies;
 		TArray<int32> BodiesCustomIds;
 
@@ -314,7 +314,7 @@ namespace mu
 		inline bool operator==(const PhysicsBody& Other) const
         {
         	return CustomId        == Other.CustomId        &&
-				   Bones		   == Other.Bones           &&
+				   BoneIds == Other.BoneIds					&&
         		   Bodies          == Other.Bodies          &&
         	       BodiesCustomIds == Other.BodiesCustomIds &&
 				   bBodiesModified == Other.bBodiesModified;
diff --git a/Engine/Plugins/Experimental/Mutable/Source/MutableRuntime/Public/MuR/Skeleton.h b/Engine/Plugins/Experimental/Mutable/Source/MutableRuntime/Public/MuR/Skeleton.h
--- a/Engine/Plugins/Experimental/Mutable/Source/MutableRuntime/Public/MuR/Skeleton.h
+++ b/Engine/Plugins/Experimental/Mutable/Source/MutableRuntime/Public/MuR/Skeleton.h
@@ -53,15 +53,22 @@ namespace mu
 		int32 GetBoneCount() const;
         void SetBoneCount(int32 c);
 
-        //! Return the name of one of the bones used by this skeleton
-		//! \param boneIndex goes from 0 to GetBoneCount()-1
-		const char* GetBoneName(int32 boneIndex) const;
-        void SetBoneName( int32 b, const char* strName );
+		//! DEBUG. Return the FName of the bone at index BoneIndex. Only valid in the editor
+		const FName GetBoneFName(int32 Index) const;
+		void SetBoneFName(const int32 Index, const FName BoneName);
 
         //! Get and set the parent bone of each bone. The parent can be -1 if the bone is a root.
         int32 GetBoneParent(int32 boneIndex) const;
         void SetBoneParent(int32 boneIndex, int32 parentBoneIndex);
 
+		//! Return the BoneIndex of the bone at index from the BoneIndices array
+		uint16 GetBoneId(int32 Index) const;
+		void SetBoneId(int32 Index, uint16 BoneIndex);
+
+		//! Return the index of BoneIndex inside the BoneIndices array.
+		int32 FindBone(const uint16 BoneIndex) const;
+		
+
 	protected:
 
 		//! Forbidden. Manage with the Ptr<> template.
@@ -69,14 +76,21 @@ namespace mu
 
 	public:
 
-		//!
-		TArray<string> m_bones;
+		//! Deprecated
+		TArray<string> m_bones_DEPRECATED;
 		TArray<FTransform3f> m_boneTransforms_DEPRECATED;
 
+
+
+		//! DEBUG. FNames of the bones. Only valid in the editor. Do not serialize.
+		TArray<FName> BoneNames;
+
+		//! Ids of the bones. The Id is the BoneName index in the CO BoneNames array
+		TArray<uint16> BoneIds;
+
 		//! For each bone, index of the parent bone in the bone vectors. -1 means no parent.
 		//! This array must have the same size than the m_bones array.
-		TArray<int16> m_boneParents;
-
+		TArray<int16> BoneParents;
 
 		//!
 		void Serialise(OutputArchive& arch) const;
@@ -84,30 +98,12 @@ namespace mu
 		//!
 		void Unserialise(InputArchive& arch);
 
-
 		//!
 		inline bool operator==(const Skeleton& o) const
 		{
-			return m_bones == o.m_bones
-				&& m_boneParents == o.m_boneParents;
-		}
-
-
-		//-----------------------------------------------------------------------------------------
-		//! Fins a bone index by name
-		//-----------------------------------------------------------------------------------------
-		int FindBone(const char* strName) const
-		{
-			for (int32 i = 0; i < m_bones.Num(); ++i)
-			{
-				if (m_bones[i] == strName)
-				{
-					return (int)i;
-				}
-			}
-			return -1;
+			return BoneIds == o.BoneIds
+				&& BoneParents == o.BoneParents;
 		}
-
 	};
 
 }
diff --git a/Engine/Plugins/Experimental/Mutable/Source/MutableTools/Private/MuT/ASTOpMeshBindShape.cpp b/Engine/Plugins/Experimental/Mutable/Source/MutableTools/Private/MuT/ASTOpMeshBindShape.cpp
--- a/Engine/Plugins/Experimental/Mutable/Source/MutableTools/Private/MuT/ASTOpMeshBindShape.cpp
+++ b/Engine/Plugins/Experimental/Mutable/Source/MutableTools/Private/MuT/ASTOpMeshBindShape.cpp
@@ -69,12 +69,12 @@ namespace mu
 		hash_combine(res, bool(bReshapeVertices));
 		hash_combine(res, bool(BindingMethod));
 
-		for (const string& S : BonesToDeform)
+		for (const uint16 S : BonesToDeform)
 		{
 			hash_combine(res, S);
 		}
 
-		for (const string& S : PhysicsToDeform)
+		for (const uint16 S : PhysicsToDeform)
 		{
 			hash_combine(res, S);
 		}
@@ -142,17 +142,15 @@ namespace mu
 			AppendCode(program.m_byteCode, Args);
 
 			AppendCode(program.m_byteCode, (int32)BonesToDeform.Num());
-			for (const string& S : BonesToDeform)
+			for (const uint16 S : BonesToDeform)
 			{
-				const uint32 index = program.AddConstant(S);
-				AppendCode(program.m_byteCode, index);
+				AppendCode(program.m_byteCode, S);
 			}
 
 			AppendCode(program.m_byteCode, (int32)PhysicsToDeform.Num());
-			for (const string& S : PhysicsToDeform)
+			for (const uint16 S : PhysicsToDeform)
 			{
-				const uint32 index = program.AddConstant(S);
-				AppendCode(program.m_byteCode, index);
+				AppendCode(program.m_byteCode, S);
 			}
 		}
 	}
diff --git a/Engine/Plugins/Experimental/Mutable/Source/MutableTools/Private/MuT/ASTOpMeshBindShape.h b/Engine/Plugins/Experimental/Mutable/Source/MutableTools/Private/MuT/ASTOpMeshBindShape.h
--- a/Engine/Plugins/Experimental/Mutable/Source/MutableTools/Private/MuT/ASTOpMeshBindShape.h
+++ b/Engine/Plugins/Experimental/Mutable/Source/MutableTools/Private/MuT/ASTOpMeshBindShape.h
@@ -24,8 +24,8 @@ struct FProgram;
 		ASTChild Mesh;
 		ASTChild Shape;
 
-		TArray<string> BonesToDeform;
-		TArray<string> PhysicsToDeform;
+		TArray<uint16> BonesToDeform;
+		TArray<uint16> PhysicsToDeform;
 
 		uint32 BindingMethod = 0;
 		
diff --git a/Engine/Plugins/Experimental/Mutable/Source/MutableTools/Private/MuT/CodeGenerator_Mesh.cpp b/Engine/Plugins/Experimental/Mutable/Source/MutableTools/Private/MuT/CodeGenerator_Mesh.cpp
--- a/Engine/Plugins/Experimental/Mutable/Source/MutableTools/Private/MuT/CodeGenerator_Mesh.cpp
+++ b/Engine/Plugins/Experimental/Mutable/Source/MutableTools/Private/MuT/CodeGenerator_Mesh.cpp
@@ -1786,9 +1786,9 @@ class Node;
 
 		OpBind->bReshapeSkeleton = node.m_reshapeSkeleton;
     	OpBind->bEnableRigidParts = node.m_enableRigidParts;
-		OpBind->BonesToDeform = node.m_bonesToDeform;
+		OpBind->BonesToDeform = node.BonesToDeform;
     	OpBind->bReshapePhysicsVolumes = node.m_reshapePhysicsVolumes;
-		OpBind->PhysicsToDeform = node.m_physicsToDeform;
+		OpBind->PhysicsToDeform = node.PhysicsToDeform;
 		OpBind->bReshapeVertices = node.m_reshapeVertices;
 		OpBind->BindingMethod = static_cast<uint32>(EShapeBindingMethod::ReshapeClosestProject);
 
diff --git a/Engine/Plugins/Experimental/Mutable/Source/MutableTools/Private/MuT/NodeMeshMorph.cpp b/Engine/Plugins/Experimental/Mutable/Source/MutableTools/Private/MuT/NodeMeshMorph.cpp
--- a/Engine/Plugins/Experimental/Mutable/Source/MutableTools/Private/MuT/NodeMeshMorph.cpp
+++ b/Engine/Plugins/Experimental/Mutable/Source/MutableTools/Private/MuT/NodeMeshMorph.cpp
@@ -142,14 +142,14 @@ namespace mu
 	}	
 
 	//---------------------------------------------------------------------------------------------
-	void NodeMeshMorph::AddBoneToDeform(const char* BoneName)
+	void NodeMeshMorph::AddBoneToDeform(const uint16 BoneId)
 	{
-		m_pD->BonesToDeform.Emplace(BoneName);
+		m_pD->BonesToDeform.Emplace(BoneId);
 	}
 
-	void NodeMeshMorph::AddPhysicsBodyToDeform(const char* BoneName)
+	void NodeMeshMorph::AddPhysicsBodyToDeform(const uint16 BoneId)
 	{
-		m_pD->PhysicsToDeform.Emplace(BoneName);
+		m_pD->PhysicsToDeform.Emplace(BoneId);
 	}
 
 	//---------------------------------------------------------------------------------------------
diff --git a/Engine/Plugins/Experimental/Mutable/Source/MutableTools/Private/MuT/NodeMeshMorphPrivate.h b/Engine/Plugins/Experimental/Mutable/Source/MutableTools/Private/MuT/NodeMeshMorphPrivate.h
--- a/Engine/Plugins/Experimental/Mutable/Source/MutableTools/Private/MuT/NodeMeshMorphPrivate.h
+++ b/Engine/Plugins/Experimental/Mutable/Source/MutableTools/Private/MuT/NodeMeshMorphPrivate.h
@@ -28,13 +28,13 @@ namespace mu
 		bool bReshapeSkeleton = false;
 		bool bReshapePhysicsVolumes = false;
 		
-		TArray<string> BonesToDeform;
-		TArray<string> PhysicsToDeform;
+		TArray<uint16> BonesToDeform;
+		TArray<uint16> PhysicsToDeform;
 
         //!
 		void Serialise( OutputArchive& arch ) const
 		{
-            uint32_t ver = 6;
+            uint32_t ver = 7;
 			arch << ver;
 
 			arch << Factor;
@@ -52,7 +52,7 @@ namespace mu
 		{
             uint32_t ver;
 			arch >> ver;
-			check(ver <= 6);
+			check(ver <= 7);
 
 			arch >> Factor;
 			arch >> Base;
@@ -89,7 +89,23 @@ namespace mu
 				{
 					arch >> bReshapePhysicsVolumes;
 				}
-				arch >> BonesToDeform;			
+
+				if (ver >= 7)
+				{
+					arch >> BonesToDeform;
+				}
+				else
+				{
+					TArray<string> BonesToDeform_DEPRECATED;
+					arch >> BonesToDeform_DEPRECATED;
+
+					const int32 NumBonesToDeform = BonesToDeform.Num();
+					BonesToDeform.SetNumUninitialized(NumBonesToDeform);
+					for (int32 Index = 0; Index < NumBonesToDeform; ++Index)
+					{
+						BonesToDeform[Index] = Index;
+					}
+				}
 			}
 			else
 			{
@@ -110,10 +126,22 @@ namespace mu
 				arch >> bDeformAllPhysics_DEPRECATED;
 			}
 
-			if (ver >= 3)
+			if (ver >= 7)
 			{
 				arch >> PhysicsToDeform;
 			}
+			else if (ver >= 3)
+			{
+				TArray<string> PhysicsToDeform_DEPRECATED;
+				arch >> PhysicsToDeform_DEPRECATED;
+
+				const int32 NumPhysicsToDeform = PhysicsToDeform.Num();
+				PhysicsToDeform.SetNumUninitialized(NumPhysicsToDeform);
+				for (int32 Index = 0; Index < NumPhysicsToDeform; ++Index)
+				{
+					PhysicsToDeform[Index] = Index;
+				}
+			}
 			else
 			{
 				PhysicsToDeform.Empty();
diff --git a/Engine/Plugins/Experimental/Mutable/Source/MutableTools/Private/MuT/NodeMeshReshape.cpp b/Engine/Plugins/Experimental/Mutable/Source/MutableTools/Private/MuT/NodeMeshReshape.cpp
--- a/Engine/Plugins/Experimental/Mutable/Source/MutableTools/Private/MuT/NodeMeshReshape.cpp
+++ b/Engine/Plugins/Experimental/Mutable/Source/MutableTools/Private/MuT/NodeMeshReshape.cpp
@@ -136,15 +136,15 @@ namespace mu
 	}	
 
 	//---------------------------------------------------------------------------------------------
-	void NodeMeshReshape::AddBoneToDeform(const char* BoneName)
+	void NodeMeshReshape::AddBoneToDeform(const uint16 BoneId)
 	{
-		m_pD->m_bonesToDeform.Emplace(BoneName);
+		m_pD->BonesToDeform.Emplace(BoneId);
 	}
 
 	//---------------------------------------------------------------------------------------------
-	void NodeMeshReshape::AddPhysicsBodyToDeform(const char* BoneName)
+	void NodeMeshReshape::AddPhysicsBodyToDeform(const uint16 BoneId)
 	{
-		m_pD->m_physicsToDeform.Emplace(BoneName);
+		m_pD->PhysicsToDeform.Emplace(BoneId);
 	}
 
 	void NodeMeshReshape::SetReshapePhysicsVolumes(bool bEnable)
diff --git a/Engine/Plugins/Experimental/Mutable/Source/MutableTools/Private/MuT/NodeMeshReshapePrivate.h b/Engine/Plugins/Experimental/Mutable/Source/MutableTools/Private/MuT/NodeMeshReshapePrivate.h
--- a/Engine/Plugins/Experimental/Mutable/Source/MutableTools/Private/MuT/NodeMeshReshapePrivate.h
+++ b/Engine/Plugins/Experimental/Mutable/Source/MutableTools/Private/MuT/NodeMeshReshapePrivate.h
@@ -30,12 +30,12 @@ namespace mu
 		bool m_reshapePhysicsVolumes = false;
 	
 
-		TArray<string> m_bonesToDeform;
-		TArray<string> m_physicsToDeform;
+		TArray<uint16> BonesToDeform;
+		TArray<uint16> PhysicsToDeform;
         //!
 		void Serialise( OutputArchive& arch ) const
 		{
-            uint32_t ver = 8;
+            uint32_t ver = 9;
 			arch << ver;
 
 			arch << m_pBaseMesh;
@@ -44,9 +44,9 @@ namespace mu
 			arch << m_reshapeVertices;
 			arch << m_reshapeSkeleton;
 			arch << m_enableRigidParts;
-			arch << m_bonesToDeform;
+			arch << BonesToDeform;
 			arch << m_reshapePhysicsVolumes;
-			arch << m_physicsToDeform;
+			arch << PhysicsToDeform;
         }
 
 		//!
@@ -54,7 +54,7 @@ namespace mu
 		{
             uint32_t ver;
 			arch >> ver;
-			check(ver <= 8);
+			check(ver <= 9);
 
 			arch >> m_pBaseMesh;
 			arch >> m_pBaseShape;
@@ -74,8 +74,23 @@ namespace mu
 				arch >> bDeformAllBones_DEPRECATED;
 			}
 
-			arch >> m_bonesToDeform;
-			
+			if (ver >= 9)
+			{
+				arch >> BonesToDeform;
+			}
+			else
+			{
+				TArray<string> BonesToDeform_DEPRECATED;
+				arch >> BonesToDeform_DEPRECATED;
+
+				const int32 NumBonesToDeform = BonesToDeform_DEPRECATED.Num();
+				BonesToDeform.Reserve(NumBonesToDeform);
+				for (int32 BoneIndex = 0; BoneIndex < NumBonesToDeform; ++BoneIndex)
+				{
+					BonesToDeform.Add(BoneIndex);
+				}
+			}
+
 		  	if (ver >= 4)
 		 	{
 				arch >> m_reshapePhysicsVolumes;
@@ -87,9 +102,21 @@ namespace mu
 				arch >> bDeformAllPhysics_DEPRECATED;
 			}
 
-			if (ver >= 5)
+			if (ver >= 9)
+			{
+				arch >> PhysicsToDeform;
+			}
+			else if (ver >= 5)
 			{
-				arch >> m_physicsToDeform;
+				TArray<string> PhysicsToDeform_DEPRECATED;
+				arch >> PhysicsToDeform_DEPRECATED;
+
+				const int32 NumPhysicsToDeform = PhysicsToDeform_DEPRECATED.Num();
+				PhysicsToDeform.Reserve(NumPhysicsToDeform);
+				for (int32 BoneIndex = 0; BoneIndex < NumPhysicsToDeform; ++BoneIndex)
+				{
+					PhysicsToDeform.Add(BoneIndex);
+				}
 			}
 		}
 
diff --git a/Engine/Plugins/Experimental/Mutable/Source/MutableTools/Public/MuT/NodeMeshMorph.h b/Engine/Plugins/Experimental/Mutable/Source/MutableTools/Public/MuT/NodeMeshMorph.h
--- a/Engine/Plugins/Experimental/Mutable/Source/MutableTools/Public/MuT/NodeMeshMorph.h
+++ b/Engine/Plugins/Experimental/Mutable/Source/MutableTools/Public/MuT/NodeMeshMorph.h
@@ -65,10 +65,10 @@ namespace mu
 		void SetReshapePhysicsVolumes(bool);
 
 		/** Sets the number of bones that will be deform */
-		void AddBoneToDeform(const char* BoneName);	
+		void AddBoneToDeform(const uint16 BoneName);	
 
 		/** Add a Physics Body to deform */
-		void AddPhysicsBodyToDeform(const char* BoneName);
+		void AddPhysicsBodyToDeform(const uint16 BoneName);
 
         //-----------------------------------------------------------------------------------------
 		// Interface pattern
diff --git a/Engine/Plugins/Experimental/Mutable/Source/MutableTools/Public/MuT/NodeMeshReshape.h b/Engine/Plugins/Experimental/Mutable/Source/MutableTools/Public/MuT/NodeMeshReshape.h
--- a/Engine/Plugins/Experimental/Mutable/Source/MutableTools/Public/MuT/NodeMeshReshape.h
+++ b/Engine/Plugins/Experimental/Mutable/Source/MutableTools/Public/MuT/NodeMeshReshape.h
@@ -65,10 +65,10 @@ namespace mu
 		void SetReshapePhysicsVolumes(bool);
 
 		/** Sets the number of bones that will be deform */
-		void AddBoneToDeform(const char* BoneName);
+		void AddBoneToDeform(const uint16 BoneId);
 	
 		/** Add a Physics Body to deform */
-		void AddPhysicsBodyToDeform(const char* BoneName);
+		void AddPhysicsBodyToDeform(const uint16 BoneId);
         
 		//-----------------------------------------------------------------------------------------
 		// Interface pattern
