[mutable] Added support for Texture2DArrays (as passthrough textures) to table nodes.

#rnx

diff --git a/Engine/Plugins/Experimental/Mutable/Source/CustomizableObject/Classes/MuCO/CustomizableObject.h b/Engine/Plugins/Experimental/Mutable/Source/CustomizableObject/Classes/MuCO/CustomizableObject.h
--- a/Engine/Plugins/Experimental/Mutable/Source/CustomizableObject/Classes/MuCO/CustomizableObject.h
+++ b/Engine/Plugins/Experimental/Mutable/Source/CustomizableObject/Classes/MuCO/CustomizableObject.h
@@ -1139,7 +1139,7 @@ public:
 	TArray<TSoftObjectPtr<USkeleton>> ReferencedSkeletons;
 
 	UPROPERTY()
-	TArray<TSoftObjectPtr<UTexture2D>> ReferencedPassThroughTextures;
+	TArray<TSoftObjectPtr<UTexture>> ReferencedPassThroughTextures;
 
 	UPROPERTY(EditAnywhere, Category = CustomizableObject, meta = (DisplayName = "LOD Settings"))
 	FMutableLODSettings LODSettings;
@@ -1439,7 +1439,7 @@ private:
 	// This is a manual version number for the binary blobs in this asset.
 	// Increasing it invalidates all the previously compiled models.
 	// Warning: If while merging code both versions have changed, take the highest+1.
-	static const int32 CurrentSupportedVersion = 393;
+	static const int32 CurrentSupportedVersion = 394;
 
 public:
 
diff --git a/Engine/Plugins/Experimental/Mutable/Source/CustomizableObject/Private/MuCO/CustomizableInstancePrivateData.h b/Engine/Plugins/Experimental/Mutable/Source/CustomizableObject/Private/MuCO/CustomizableInstancePrivateData.h
--- a/Engine/Plugins/Experimental/Mutable/Source/CustomizableObject/Private/MuCO/CustomizableInstancePrivateData.h
+++ b/Engine/Plugins/Experimental/Mutable/Source/CustomizableObject/Private/MuCO/CustomizableInstancePrivateData.h
@@ -61,7 +61,7 @@ struct FGeneratedTexture
 	FString Name;
 
 	UPROPERTY(Category = CustomizableObjectInstance, EditAnywhere) 
-	TObjectPtr<UTexture2D> Texture = nullptr;
+	TObjectPtr<UTexture> Texture = nullptr;
 };
 
 
diff --git a/Engine/Plugins/Experimental/Mutable/Source/CustomizableObject/Private/MuCO/CustomizableObjectInstance.cpp b/Engine/Plugins/Experimental/Mutable/Source/CustomizableObject/Private/MuCO/CustomizableObjectInstance.cpp
--- a/Engine/Plugins/Experimental/Mutable/Source/CustomizableObject/Private/MuCO/CustomizableObjectInstance.cpp
+++ b/Engine/Plugins/Experimental/Mutable/Source/CustomizableObject/Private/MuCO/CustomizableObjectInstance.cpp
@@ -12,6 +12,7 @@
 #include "ClothConfig.h"
 #include "ClothingAsset.h"
 #include "Engine/SkeletalMeshSocket.h"
+#include "Engine/Texture2DArray.h"
 #include "MaterialDomain.h"
 #include "Materials/Material.h"
 #include "Materials/MaterialInstanceDynamic.h"
@@ -356,7 +357,7 @@ void UCustomizableObjectInstance::ReleaseMutableResources(bool bCalledFromBeginD
 				// instance's remaining sk meshes and GC is being performed anyway so it will free the textures if needed
 				if (!bCalledFromBeginDestroy && CustomizableObjectSystem->bReleaseTexturesImmediately)
 				{
-					UCustomizableInstancePrivateData::ReleaseMutableTexture(Texture.Id, Texture.Texture, Cache);
+					UCustomizableInstancePrivateData::ReleaseMutableTexture(Texture.Id, Cast<UTexture2D>(Texture.Texture), Cache);
 				}
 			}
 		}
@@ -509,7 +510,7 @@ bool UCustomizableObjectInstance::IsParamMultidimensional(const FString& ParamNa
 // Only safe to call if the Mutable texture ref count system returns 0 and absolutely sure nobody holds a reference to the texture
 void UCustomizableInstancePrivateData::ReleaseMutableTexture(int32 MutableTextureId, UTexture2D* Texture, FMutableResourceCache& Cache)
 {
-	if (Texture && Texture->IsValidLowLevel())
+	if (ensure(Texture) && Texture->IsValidLowLevel())
 	{
 		Texture->ConditionalBeginDestroy();
 
@@ -5707,7 +5708,8 @@ void UCustomizableInstancePrivateData::BuildMaterials(const TSharedPtr<FMutableO
 						FString KeyName = Image.Name;
 						mu::ImagePtrConst MutableImage = Image.Image;
 
-						UTexture2D* Texture = nullptr;
+						UTexture2D* MutableTexture = nullptr; // Texture generated by mutable
+						UTexture* PassThroughTexture = nullptr; // Texture not generated by mutable
 
 						// \TODO: Change this key to a struct.
 						FString TextureReuseCacheRef = bReuseTextures ? FString::Printf(TEXT("%d-%d-%d-%d"), Image.BaseLOD, ComponentIndex, Surface.SurfaceId, ImageIndex) : FString();
@@ -5720,10 +5722,10 @@ void UCustomizableInstancePrivateData::BuildMaterials(const TSharedPtr<FMutableO
 							if (CachedPointerPtr)
 							{
 								ensure(!CachedPointerPtr->IsStale());
-								Texture = CachedPointerPtr->Get();
+								MutableTexture = CachedPointerPtr->Get();
 							}
 
-							check(Texture);
+							check(MutableTexture);
 						}
 
 						// Check if the image is a reference to an engine texture
@@ -5732,10 +5734,10 @@ void UCustomizableInstancePrivateData::BuildMaterials(const TSharedPtr<FMutableO
 							uint32 ReferenceID = MutableImage->GetReferencedTexture();
 							if (CustomizableObject->ReferencedPassThroughTextures.IsValidIndex(ReferenceID))
 							{
-								TSoftObjectPtr<UTexture2D> Ref = CustomizableObject->ReferencedPassThroughTextures[ReferenceID];
+								TSoftObjectPtr<UTexture> Ref = CustomizableObject->ReferencedPassThroughTextures[ReferenceID];
 
 								// \TODO: This will force the load of the reference texture, potentially causing a hich. 
-								Texture = Ref.Get();
+								PassThroughTexture = Ref.Get();
 							}
 							else
 							{
@@ -5750,7 +5752,7 @@ void UCustomizableInstancePrivateData::BuildMaterials(const TSharedPtr<FMutableO
 						{
 							const FMutableModelImageProperties& Props = CustomizableObject->ImageProperties[ImageKey];
 
-							if (!Texture && MutableImage)
+							if (!MutableTexture && !PassThroughTexture && MutableImage)
 							{
 								TWeakObjectPtr<UTexture2D>* ReusedTexture = bReuseTextures ? TextureReuseCache.Find(TextureReuseCacheRef) : nullptr;
 
@@ -5763,48 +5765,48 @@ void UCustomizableInstancePrivateData::BuildMaterials(const TSharedPtr<FMutableO
 
 									if (PixelFormat == EPixelFormat::PF_R8G8B8A8)
 									{
-										Texture = (*ReusedTexture).Get();
-										check(Texture != nullptr);
+										MutableTexture = (*ReusedTexture).Get();
+										check(MutableTexture != nullptr);
 									}
 									else
 									{
 										ReusedTexture = nullptr;
-										Texture = CreateTexture();
+										MutableTexture = CreateTexture();
 
 #if WITH_EDITOR
 										UE_LOG(LogMutable, Warning,
 											TEXT("Tried to reuse an uncompressed texture with name %s. Make sure the selected Mutable state disables texture compression/streaming, that one of the state's runtime parameters affects the texture and that the CO is compiled with max. optimization settings."),
-											*Texture->GetName());
+											*MutableTexture->GetName());
 #endif
 									}
 								}
 								else
 								{
 									ReusedTexture = nullptr;
-									Texture = CreateTexture();
+									MutableTexture = CreateTexture();
 								}
 
-								if (Texture)
+								if (MutableTexture)
 								{
 									if (OperationData->ImageToPlatformDataMap.Contains(Image.ImageID))
 									{
-										SetTexturePropertiesFromMutableImageProps(Texture, Props, bNeverStream);
+										SetTexturePropertiesFromMutableImageProps(MutableTexture, Props, bNeverStream);
 
 										FTexturePlatformData* PlatformData = OperationData->ImageToPlatformDataMap[Image.ImageID];
 
 										if (ReusedTexture)
 										{
-											check(PlatformData->Mips.Num() == Texture->GetPlatformData()->Mips.Num());
-											check(PlatformData->Mips[0].SizeX == Texture->GetPlatformData()->Mips[0].SizeX);
-											check(PlatformData->Mips[0].SizeY == Texture->GetPlatformData()->Mips[0].SizeY);
+											check(PlatformData->Mips.Num() == MutableTexture->GetPlatformData()->Mips.Num());
+											check(PlatformData->Mips[0].SizeX == MutableTexture->GetPlatformData()->Mips[0].SizeX);
+											check(PlatformData->Mips[0].SizeY == MutableTexture->GetPlatformData()->Mips[0].SizeY);
 										}
 
-										if (Texture->GetPlatformData())
+										if (MutableTexture->GetPlatformData())
 										{
-											delete Texture->GetPlatformData();
+											delete MutableTexture->GetPlatformData();
 										}
 
-										Texture->SetPlatformData(PlatformData);
+										MutableTexture->SetPlatformData(PlatformData);
 										OperationData->ImageToPlatformDataMap.Remove(Image.ImageID);
 									}
 									else
@@ -5815,19 +5817,19 @@ void UCustomizableInstancePrivateData::BuildMaterials(const TSharedPtr<FMutableO
 									if (bNeverStream)
 									{
 										// To prevent LogTexture Error "Loading non-streamed mips from an external bulk file."
-										for (int32 i = 0; i < Texture->GetPlatformData()->Mips.Num(); ++i)
+										for (int32 i = 0; i < MutableTexture->GetPlatformData()->Mips.Num(); ++i)
 										{
-											Texture->GetPlatformData()->Mips[i].BulkData.ClearBulkDataFlags(BULKDATA_PayloadInSeperateFile);
+											MutableTexture->GetPlatformData()->Mips[i].BulkData.ClearBulkDataFlags(BULKDATA_PayloadInSeperateFile);
 										}
 									}
 
 									{
 										MUTABLE_CPUPROFILER_SCOPE(UpdateResource);
 #if !PLATFORM_DESKTOP && !PLATFORM_SWITCH // Switch does this in FTexture2DResource::InitRHI()
-										for (int32 i = 0; i < Texture->GetPlatformData()->Mips.Num(); ++i)
+										for (int32 i = 0; i < MutableTexture->GetPlatformData()->Mips.Num(); ++i)
 										{
-											uint32 DataFlags = Texture->GetPlatformData()->Mips[i].BulkData.GetBulkDataFlags();
-											Texture->GetPlatformData()->Mips[i].BulkData.SetBulkDataFlags(DataFlags | BULKDATA_SingleUse);
+											uint32 DataFlags = MutableTexture->GetPlatformData()->Mips[i].BulkData.GetBulkDataFlags();
+											MutableTexture->GetPlatformData()->Mips[i].BulkData.SetBulkDataFlags(DataFlags | BULKDATA_SingleUse);
 										}
 #endif
 
@@ -5836,21 +5838,21 @@ void UCustomizableInstancePrivateData::BuildMaterials(const TSharedPtr<FMutableO
 											// Must remove texture from cache since it will be reused with a different ImageID
 											for (TPair<FMutableImageCacheKey, TWeakObjectPtr<UTexture2D>>& CachedTexture : Cache.Images)
 											{
-												if (CachedTexture.Value == Texture)
+												if (CachedTexture.Value == MutableTexture)
 												{
 													Cache.Images.Remove(CachedTexture.Key);
 													break;
 												}
 											}
 
-											ReuseTexture(Texture);
+											ReuseTexture(MutableTexture);
 										}
-										else if (Texture)
+										else if (MutableTexture)
 										{	
 											//if (!bNeverStream) // No need to check bNeverStream. In that case, the texture won't use 
 											// the MutableMipDataProviderFactory anyway and it's needed for detecting Mutable textures elsewhere
 											{
-												UMutableTextureMipDataProviderFactory* MutableMipDataProviderFactory = Cast<UMutableTextureMipDataProviderFactory>(Texture->GetAssetUserDataOfClass(UMutableTextureMipDataProviderFactory::StaticClass()));
+												UMutableTextureMipDataProviderFactory* MutableMipDataProviderFactory = Cast<UMutableTextureMipDataProviderFactory>(MutableTexture->GetAssetUserDataOfClass(UMutableTextureMipDataProviderFactory::StaticClass()));
 												if (!MutableMipDataProviderFactory)
 												{
 													MutableMipDataProviderFactory = NewObject<UMutableTextureMipDataProviderFactory>();
@@ -5865,16 +5867,16 @@ void UCustomizableInstancePrivateData::BuildMaterials(const TSharedPtr<FMutableO
 														MutableMipDataProviderFactory->ImageRef.Component = uint8(ComponentIndex);
 														MutableMipDataProviderFactory->ImageRef.Image = uint8(ImageIndex);
 														MutableMipDataProviderFactory->UpdateContext = UpdateContext;
-														Texture->AddAssetUserData(MutableMipDataProviderFactory);
+														MutableTexture->AddAssetUserData(MutableMipDataProviderFactory);
 													}
 												}
 											}
 
-											Texture->UpdateResource();
+											MutableTexture->UpdateResource();
 										}
 									}
 
-									Cache.Images.Add(ImageCacheKey, Texture);						
+									Cache.Images.Add(ImageCacheKey, MutableTexture);						
 								}
 								else
 								{
@@ -5885,8 +5887,13 @@ void UCustomizableInstancePrivateData::BuildMaterials(const TSharedPtr<FMutableO
 							FGeneratedTexture TextureData;
 							TextureData.Id = Image.ImageID;
 							TextureData.Name = Props.TextureParameterName;
-							TextureData.Texture = Texture;
-							NewGeneratedTextures.Add(TextureData);
+							TextureData.Texture = MutableTexture ? MutableTexture : PassThroughTexture;
+							
+							// Only add textures generated by mutable to the cache
+							if (MutableTexture)
+							{
+								NewGeneratedTextures.Add(TextureData);
+							}
 
 							// Decoding Material Layer from Mutable parameter name
 							FString ImageName = Image.Name;
@@ -5919,9 +5926,9 @@ void UCustomizableInstancePrivateData::BuildMaterials(const TSharedPtr<FMutableO
 
 						if (bReuseTextures)
 						{
-							if (Texture)
+							if (MutableTexture)
 							{
-								TextureReuseCache.Add(TextureReuseCacheRef, Texture);
+								TextureReuseCache.Add(TextureReuseCacheRef, MutableTexture);
 							}
 							else
 							{
diff --git a/Engine/Plugins/Experimental/Mutable/Source/CustomizableObject/Private/MuCO/CustomizableObjectSystem.cpp b/Engine/Plugins/Experimental/Mutable/Source/CustomizableObject/Private/MuCO/CustomizableObjectSystem.cpp
--- a/Engine/Plugins/Experimental/Mutable/Source/CustomizableObject/Private/MuCO/CustomizableObjectSystem.cpp
+++ b/Engine/Plugins/Experimental/Mutable/Source/CustomizableObject/Private/MuCO/CustomizableObjectSystem.cpp
@@ -846,7 +846,7 @@ void UpdateSkeletalMesh(UCustomizableObjectInstance& CustomizableObjectInstance,
 
 		for (TPair<uint32, FGeneratedTexture>& Item : CustomizableObjectInstancePrivateData->TexturesToRelease)
 		{
-			UCustomizableInstancePrivateData::ReleaseMutableTexture(Item.Value.Id, Item.Value.Texture, Cache);
+			UCustomizableInstancePrivateData::ReleaseMutableTexture(Item.Value.Id, Cast<UTexture2D>(Item.Value.Texture), Cache);
 		}
 
 		CustomizableObjectInstancePrivateData->TexturesToRelease.Empty();
diff --git a/Engine/Plugins/Experimental/Mutable/Source/CustomizableObject/Private/MuCO/LogInformationUtil.cpp b/Engine/Plugins/Experimental/Mutable/Source/CustomizableObject/Private/MuCO/LogInformationUtil.cpp
--- a/Engine/Plugins/Experimental/Mutable/Source/CustomizableObject/Private/MuCO/LogInformationUtil.cpp
+++ b/Engine/Plugins/Experimental/Mutable/Source/CustomizableObject/Private/MuCO/LogInformationUtil.cpp
@@ -45,7 +45,7 @@ void LogInformationUtil::PrintGeneratedTextures(const TArray<struct FGeneratedTe
 		Log += FString::Printf(TEXT("Texture name=%s, "), *GeneratedTextures[i].Name);
 		Log += FString::Printf(TEXT("id=%d, "), GeneratedTextures[i].Id);
 		Log += "\n\t\t\t";
-		PrintTextureData(GeneratedTextures[i].Texture, Log, false);
+		PrintTextureData(Cast<UTexture2D>(GeneratedTextures[i].Texture), Log, false);
 		if (i < (Max - 1))
 		{
 			Log += "\n\t\t";
@@ -335,6 +335,11 @@ void LogInformationUtil::PrintProjectorParameters(const TArray<struct FCustomiza
 
 void LogInformationUtil::PrintTextureData(const UTexture2D* Texture, FString& Log, bool UseTabulation)
 {
+	if (!Texture)
+	{
+		return;
+	}
+
 	FString MessageChunk;
 	MessageChunk += FString::Printf(TEXT("        Name=%s\n"), *Texture->GetName());
 	if (UseTabulation)
diff --git a/Engine/Plugins/Experimental/Mutable/Source/CustomizableObjectEditor/Classes/MuCOE/Nodes/CustomizableObjectNodeTable.h b/Engine/Plugins/Experimental/Mutable/Source/CustomizableObjectEditor/Classes/MuCOE/Nodes/CustomizableObjectNodeTable.h
--- a/Engine/Plugins/Experimental/Mutable/Source/CustomizableObjectEditor/Classes/MuCOE/Nodes/CustomizableObjectNodeTable.h
+++ b/Engine/Plugins/Experimental/Mutable/Source/CustomizableObjectEditor/Classes/MuCOE/Nodes/CustomizableObjectNodeTable.h
@@ -3,6 +3,7 @@
 #pragma once
 
 #include "Engine/DataTable.h"
+#include "Engine/Texture2DArray.h"
 #include "MuCOE/Nodes/CustomizableObjectNode.h"
 #include "MuCOE/RemapPins/CustomizableObjectNodeRemapPinsByName.h"
 
@@ -17,6 +18,7 @@ class UEdGraphPin;
 class UObject;
 class USkeletalMesh;
 class UTexture2D;
+class UTexture2DArray;
 struct FGuid;
 
 
@@ -64,6 +66,9 @@ public:
 	bool IsDefaultImageMode() { return bIsDefault; }
 	void SetDefaultImageMode(bool bValue) { bIsDefault = bValue; }
 
+	bool IsArrayTexture() { return bIsArrayTexture; }
+	void SetIsArrayTexture(bool bValue) { bIsArrayTexture = bValue; }
+
 	// Pin Type
 	UPROPERTY()
 	ETableTextureType ImageMode = ETableTextureType::MUTABLE_TEXTURE;
@@ -73,6 +78,9 @@ private:
 	UPROPERTY()
 	bool bIsDefault = true;
 
+	UPROPERTY()
+	bool bIsArrayTexture = false;
+
 };
 
 
@@ -300,6 +308,9 @@ public:
 	// Returns true if the pin is in the default mode (same as node)
 	bool IsImagePinDefault(UEdGraphPin* Pin);
 
+	// Returns true if the pin is a texture array pin
+	bool IsImageArrayPin(UEdGraphPin* Pin);
+
 	// Returns the image mode of the column
 	ETableTextureType GetColumnImageMode(const FString& ColumnName) const;
 
diff --git a/Engine/Plugins/Experimental/Mutable/Source/CustomizableObjectEditor/Private/MuCOE/CustomizableObjectCompiler.cpp b/Engine/Plugins/Experimental/Mutable/Source/CustomizableObjectEditor/Private/MuCOE/CustomizableObjectCompiler.cpp
--- a/Engine/Plugins/Experimental/Mutable/Source/CustomizableObjectEditor/Private/MuCOE/CustomizableObjectCompiler.cpp
+++ b/Engine/Plugins/Experimental/Mutable/Source/CustomizableObjectEditor/Private/MuCOE/CustomizableObjectCompiler.cpp
@@ -1211,9 +1211,9 @@ void FCustomizableObjectCompiler::CompileInternal(UCustomizableObject* Object, c
 		Object->SetBoneNamesArray(GenerationContext.BoneNames);
 
 		// Pass-through textures
-		TArray<TSoftObjectPtr<UTexture2D>> NewReferencedPassThroughTextures;
+		TArray<TSoftObjectPtr<UTexture>> NewReferencedPassThroughTextures;
 
-		for (const TPair<TSoftObjectPtr<UTexture2D>, uint32>& Pair : GenerationContext.PassThroughTextureToIndexMap)
+		for (const TPair<TSoftObjectPtr<UTexture>, uint32>& Pair : GenerationContext.PassThroughTextureToIndexMap)
 		{
 			check(Pair.Value == NewReferencedPassThroughTextures.Num());
 			NewReferencedPassThroughTextures.Add(Pair.Key);
diff --git a/Engine/Plugins/Experimental/Mutable/Source/CustomizableObjectEditor/Private/MuCOE/EdGraphSchema_CustomizableObject.cpp b/Engine/Plugins/Experimental/Mutable/Source/CustomizableObjectEditor/Private/MuCOE/EdGraphSchema_CustomizableObject.cpp
--- a/Engine/Plugins/Experimental/Mutable/Source/CustomizableObjectEditor/Private/MuCOE/EdGraphSchema_CustomizableObject.cpp
+++ b/Engine/Plugins/Experimental/Mutable/Source/CustomizableObjectEditor/Private/MuCOE/EdGraphSchema_CustomizableObject.cpp
@@ -901,7 +901,7 @@ void UEdGraphSchema_CustomizableObject::GetContextMenuActions(UToolMenu* Menu, U
 		UCustomizableObjectNodeTable* TableNode = Cast<UCustomizableObjectNodeTable>(Context->Node);
 		UEdGraphPin* TexturePin = (UEdGraphPin*)Context->Pin;
 
-		if (TableNode && TexturePin && !TexturePin->LinkedTo.Num() && (TexturePin->PinType.PinCategory == PC_Image || TexturePin->PinType.PinCategory == PC_PassThroughImage))
+		if (TableNode && TexturePin && !TableNode->IsImageArrayPin(TexturePin) && !TexturePin->LinkedTo.Num() && (TexturePin->PinType.PinCategory == PC_Image || TexturePin->PinType.PinCategory == PC_PassThroughImage))
 		{
 			FText ActionText = FText::Format(LOCTEXT("ChangeTexturePinMode_Label", 
 				"Set as {0} texture"), Context->Pin->PinType.PinCategory == PC_PassThroughImage ? FText::FromString("Mutable") : FText::FromString("Pass Through"));
diff --git a/Engine/Plugins/Experimental/Mutable/Source/CustomizableObjectEditor/Private/MuCOE/GenerateMutableSource/GenerateMutableSource.h b/Engine/Plugins/Experimental/Mutable/Source/CustomizableObjectEditor/Private/MuCOE/GenerateMutableSource/GenerateMutableSource.h
--- a/Engine/Plugins/Experimental/Mutable/Source/CustomizableObjectEditor/Private/MuCOE/GenerateMutableSource/GenerateMutableSource.h
+++ b/Engine/Plugins/Experimental/Mutable/Source/CustomizableObjectEditor/Private/MuCOE/GenerateMutableSource/GenerateMutableSource.h
@@ -525,7 +525,7 @@ struct FMutableGraphGenerationContext
 	TMap<FGeneratedImageKey, mu::NodeImageConstantPtr> GeneratedImages;
 
 	// Cache of pass-through images and their IDs used in the core to indentify them
-	TMap<TSoftObjectPtr<UTexture2D>, uint32> PassThroughTextureToIndexMap;
+	TMap<TSoftObjectPtr<UTexture>, uint32> PassThroughTextureToIndexMap;
 
     // Global morph selection overrides.
     TArray<FRealTimeMorphSelectionOverride> RealTimeMorphTargetsOverrides;
diff --git a/Engine/Plugins/Experimental/Mutable/Source/CustomizableObjectEditor/Private/MuCOE/GenerateMutableSource/GenerateMutableSourceImage.cpp b/Engine/Plugins/Experimental/Mutable/Source/CustomizableObjectEditor/Private/MuCOE/GenerateMutableSource/GenerateMutableSourceImage.cpp
--- a/Engine/Plugins/Experimental/Mutable/Source/CustomizableObjectEditor/Private/MuCOE/GenerateMutableSource/GenerateMutableSourceImage.cpp
+++ b/Engine/Plugins/Experimental/Mutable/Source/CustomizableObjectEditor/Private/MuCOE/GenerateMutableSource/GenerateMutableSourceImage.cpp
@@ -968,9 +968,9 @@ mu::NodeImagePtr GenerateMutableSourceImage(const UEdGraphPin* Pin, FMutableGrap
 					bSuccess = false;
 				}
 
-				if (bSuccess && Pin->PinType.PinCategory != Schema->PC_MaterialAsset && !TypedNodeTable->GetColumnDefaultAssetByType<UTexture2D>(Pin))
+				if (bSuccess && Pin->PinType.PinCategory != Schema->PC_MaterialAsset && !TypedNodeTable->GetColumnDefaultAssetByType<UTexture2D>(Pin) && !TypedNodeTable->GetColumnDefaultAssetByType<UTexture2DArray>(Pin))
 				{
-					FString Msg = FString::Printf(TEXT("Couldn't find a default value in the data table's struct for the column [%s]. The default value is null or not a Texture2D"), *ColumnName);
+					FString Msg = FString::Printf(TEXT("Couldn't find a default value in the data table's struct for the column [%s]. The default value is null or not a supported Texture"), *ColumnName);
 					GenerationContext.Compiler->CompilerLog(FText::FromString(Msg), Node);
 					
 					bSuccess = false;
diff --git a/Engine/Plugins/Experimental/Mutable/Source/CustomizableObjectEditor/Private/MuCOE/GenerateMutableSource/GenerateMutableSourceTable.cpp b/Engine/Plugins/Experimental/Mutable/Source/CustomizableObjectEditor/Private/MuCOE/GenerateMutableSource/GenerateMutableSourceTable.cpp
--- a/Engine/Plugins/Experimental/Mutable/Source/CustomizableObjectEditor/Private/MuCOE/GenerateMutableSource/GenerateMutableSourceTable.cpp
+++ b/Engine/Plugins/Experimental/Mutable/Source/CustomizableObjectEditor/Private/MuCOE/GenerateMutableSource/GenerateMutableSourceTable.cpp
@@ -4,6 +4,7 @@
 
 #include "Animation/AnimInstance.h"
 #include "Engine/StaticMesh.h"
+#include "Engine/Texture2DArray.h"
 #include "GameplayTagContainer.h"
 #include "Materials/Material.h"
 #include "Materials/MaterialInstance.h"
@@ -336,17 +337,24 @@ bool FillTableColumn(const UCustomizableObjectNodeTable* TableNode,	mu::TablePtr
 
 		else if (SoftObjectProperty->PropertyClass->IsChildOf(UTexture::StaticClass()))
 		{
-			UTexture2D* Texture = Cast<UTexture2D>(Object);
+			UTexture* Texture = nullptr;
 
-			if (!Texture)
+			// Two supported texture types
+			UTexture2D* Texture2D = Cast<UTexture2D>(Object);
+			UTexture2DArray* TextureArray = Cast<UTexture2DArray>(Object);
+
+			if (!Texture2D && !TextureArray)
 			{
-				Texture = TableNode->GetColumnDefaultAssetByType<UTexture2D>(ColumnName);
-				FString Message = Cast<UObject>(Object) ? "not a Texture2D" : "null";
+				Texture2D = TableNode->GetColumnDefaultAssetByType<UTexture2D>(ColumnName);
+				TextureArray = TableNode->GetColumnDefaultAssetByType<UTexture2DArray>(ColumnName);
 
-				FString msg = FString::Printf(TEXT("Texture from column [%s] row [%s] is %s. The default texture will be used instead."), *ColumnName, *RowName, *Message);
-				GenerationContext.Compiler->CompilerLog(FText::FromString(msg), TableNode);
+				FString Message = Cast<UObject>(Object) ? "not a suported Texture" : "null";
+				FString WarningMessage = FString::Printf(TEXT("Texture from column [%s] row [%s] is %s. The default texture will be used instead."), *ColumnName, *RowName, *Message);
+				GenerationContext.Compiler->CompilerLog(FText::FromString(WarningMessage), TableNode);
 			}
 
+			check(Texture2D || TextureArray);
+
 			// Getting column index from column name
 			CurrentColumn = MutableTable->FindColumn(StringCast<ANSICHAR>(*ColumnName).Get());
 
@@ -357,6 +365,9 @@ bool FillTableColumn(const UCustomizableObjectNodeTable* TableNode,	mu::TablePtr
 
 			if (TableNode->GetColumnImageMode(ColumnName) == ETableTextureType::PASSTHROUGH_TEXTURE)
 			{
+				// There will be always one of the two options
+				Texture = Texture2D ? Cast<UTexture>(Texture2D) : Cast<UTexture>(TextureArray);
+
 				uint32* FoundIndex = GenerationContext.PassThroughTextureToIndexMap.Find(Texture);
 				uint32 ImageReferenceID;
 
@@ -374,7 +385,7 @@ bool FillTableColumn(const UCustomizableObjectNodeTable* TableNode,	mu::TablePtr
 			}
 			else
 			{
-				GenerationContext.ArrayTextureUnrealToMutableTask.Add(FTextureUnrealToMutableTask(MutableTable, Texture, TableNode, CurrentColumn, RowIdx));
+				GenerationContext.ArrayTextureUnrealToMutableTask.Add(FTextureUnrealToMutableTask(MutableTable, Texture2D, TableNode, CurrentColumn, RowIdx));
 			}
 		}
 
diff --git a/Engine/Plugins/Experimental/Mutable/Source/CustomizableObjectEditor/Private/MuCOE/Nodes/CustomizableObjectNodeTable.cpp b/Engine/Plugins/Experimental/Mutable/Source/CustomizableObjectEditor/Private/MuCOE/Nodes/CustomizableObjectNodeTable.cpp
--- a/Engine/Plugins/Experimental/Mutable/Source/CustomizableObjectEditor/Private/MuCOE/Nodes/CustomizableObjectNodeTable.cpp
+++ b/Engine/Plugins/Experimental/Mutable/Source/CustomizableObjectEditor/Private/MuCOE/Nodes/CustomizableObjectNodeTable.cpp
@@ -3,6 +3,7 @@
 #include "MuCOE/Nodes/CustomizableObjectNodeTable.h"
 
 #include "Engine/StaticMesh.h"
+#include "Engine/Texture2DArray.h"
 #include "Materials/Material.h"
 #include "Materials/MaterialInstance.h"
 #include "MuCO/UnrealPortabilityHelpers.h"
@@ -171,7 +172,7 @@ void UCustomizableObjectNodeTable::PinConnectionListChanged(UEdGraphPin* Pin)
 	{
 		if (UCustomizableObjectNodeTableImagePinData* ImagePinData = Cast<UCustomizableObjectNodeTableImagePinData>(GetPinData(*(Pin))))
 		{
-			if (ImagePinData->IsDefaultImageMode() && ImagePinData->ImageMode != DefaultImageMode)
+			if (!ImagePinData->IsArrayTexture() && ImagePinData->IsDefaultImageMode() && ImagePinData->ImageMode != DefaultImageMode)
 			{
 				ImagePinData->ImageMode = DefaultImageMode;
 				ReconstructNode();
@@ -201,7 +202,7 @@ void UCustomizableObjectNodeTable::PostEditChangeProperty(FPropertyChangedEvent&
 			{
 				if (UCustomizableObjectNodeTableImagePinData* TexturePinData = Cast<UCustomizableObjectNodeTableImagePinData>(GetPinData(*(Pin))))
 				{
-					if (!Pin->LinkedTo.Num() && TexturePinData->IsDefaultImageMode())
+					if (!Pin->LinkedTo.Num() && !TexturePinData->IsArrayTexture() && TexturePinData->IsDefaultImageMode())
 					{
 						ETableTextureType Mode = TexturePinData->ImageMode == ETableTextureType::PASSTHROUGH_TEXTURE ? ETableTextureType::MUTABLE_TEXTURE : ETableTextureType::PASSTHROUGH_TEXTURE;
 						TexturePinData->ImageMode = Mode;
@@ -288,6 +289,8 @@ void UCustomizableObjectNodeTable::AllocateDefaultPins(UCustomizableObjectNodeRe
 				{
 					UCustomizableObjectNodeTableImagePinData* PinData = NewObject<UCustomizableObjectNodeTableImagePinData>(this);
 					PinData->ColumnName = ColumnName;
+					PinData->SetIsArrayTexture(false);
+
 					FName PinCategory = DefaultImageMode == ETableTextureType::PASSTHROUGH_TEXTURE ? Schema->PC_PassThroughImage : Schema->PC_Image;
 
 					for (UEdGraphPin* Pin : OldPins)
@@ -306,6 +309,16 @@ void UCustomizableObjectNodeTable::AllocateDefaultPins(UCustomizableObjectNodeRe
 					OutPin = CustomCreatePin(EGPD_Output, PinCategory, FName(*PinName), PinData);
 				}
 
+				else if (Object->IsA(UTexture2DArray::StaticClass()))
+				{
+					UCustomizableObjectNodeTableImagePinData* PinData = NewObject<UCustomizableObjectNodeTableImagePinData>(this);
+					PinData->ColumnName = ColumnName;
+					PinData->ImageMode = ETableTextureType::PASSTHROUGH_TEXTURE;
+					PinData->SetIsArrayTexture(true);
+
+					OutPin = CustomCreatePin(EGPD_Output, Schema->PC_PassThroughImage, FName(*PinName), PinData);
+				}
+
 				else if (Object->IsA(UMaterialInstance::StaticClass()))
 				{
 					UCustomizableObjectNodeTableObjectPinData* PinData = NewObject<UCustomizableObjectNodeTableObjectPinData>(this);
@@ -566,6 +579,17 @@ bool UCustomizableObjectNodeTable::IsNodeOutDatedAndNeedsRefresh()
 
 					NumPins++;
 				}
+				else if (Object->IsA(UTexture2DArray::StaticClass()))
+				{
+					FString PinName = DataTableUtils::GetPropertyExportName(ColumnProperty);
+
+					if (CheckPinUpdated(PinName, Schema->PC_PassThroughImage))
+					{
+						bNeedsUpdate = true;
+					}
+
+					NumPins++;
+				}
 				else if (Object->IsA(UMaterialInstance::StaticClass()))
 				{
 					FString PinName = DataTableUtils::GetPropertyExportName(ColumnProperty);
@@ -696,6 +720,7 @@ void UCustomizableObjectNodeTable::RemapPinsData(const TMap<UEdGraphPin*, UEdGra
 			{
 				PinDataNewPin->ImageMode = PinDataOldPin->ImageMode;
 				PinDataNewPin->SetDefaultImageMode(PinDataOldPin->IsDefaultImageMode());
+				PinDataNewPin->SetIsArrayTexture(PinDataOldPin->IsArrayTexture());
 			}
 		}
 	}
@@ -1120,6 +1145,11 @@ void UCustomizableObjectNodeTable::ChangeImagePinMode(UEdGraphPin* Pin, bool bSe
 {
 	if (UCustomizableObjectNodeTableImagePinData* PinData = Cast<UCustomizableObjectNodeTableImagePinData>(GetPinData(*(Pin))))
 	{
+		if (PinData->IsArrayTexture())
+		{
+			return;
+		}
+
 		if (PinData->IsDefaultImageMode() && !bSetDefault)
 		{
 			ETableTextureType Mode = DefaultImageMode == ETableTextureType::PASSTHROUGH_TEXTURE ? ETableTextureType::MUTABLE_TEXTURE : ETableTextureType::PASSTHROUGH_TEXTURE;
@@ -1164,6 +1194,17 @@ bool UCustomizableObjectNodeTable::IsImagePinDefault(UEdGraphPin* Pin)
 }
 
 
+bool UCustomizableObjectNodeTable::IsImageArrayPin(UEdGraphPin* Pin)
+{
+	if (UCustomizableObjectNodeTableImagePinData* PinData = Cast<UCustomizableObjectNodeTableImagePinData>(GetPinData(*(Pin))))
+	{
+		return PinData->IsArrayTexture();
+	}
+
+	return false;
+}
+
+
 ETableTextureType UCustomizableObjectNodeTable::GetColumnImageMode(const FString& ColumnName) const
 {
 	for (UEdGraphPin* Pin : Pins)
