[Mutable] Remove MarkForReconstruct.

- Removed due to causing problems with the undo system. Replaced with the simple ReconstructNode function.
- Minor code cleanup.

#rnx

diff --git a/Engine/Plugins/Experimental/Mutable/Source/CustomizableObjectEditor/Classes/MuCOE/Nodes/CustomizableObjectNode.h b/Engine/Plugins/Experimental/Mutable/Source/CustomizableObjectEditor/Classes/MuCOE/Nodes/CustomizableObjectNode.h
--- a/Engine/Plugins/Experimental/Mutable/Source/CustomizableObjectEditor/Classes/MuCOE/Nodes/CustomizableObjectNode.h
+++ b/Engine/Plugins/Experimental/Mutable/Source/CustomizableObjectEditor/Classes/MuCOE/Nodes/CustomizableObjectNode.h
@@ -192,10 +192,6 @@ public:
 
 	void PostEditChangeProperty(FPropertyChangedEvent& PropertyChangedEvent) override;
 
-
-	// Use this function to reconstruct the graph in the next tick
-	void MarkForReconstruct();
-
 	/** Custom post duplicate function. Called at the beginning of duplication, before the nodes have their Guid updated. */
 	virtual void BeginPostDuplicate(bool bDuplicateForPIE) {};
 	
diff --git a/Engine/Plugins/Experimental/Mutable/Source/CustomizableObjectEditor/Classes/MuCOE/Nodes/CustomizableObjectNodeMeshMorphStackApplication.h b/Engine/Plugins/Experimental/Mutable/Source/CustomizableObjectEditor/Classes/MuCOE/Nodes/CustomizableObjectNodeMeshMorphStackApplication.h
--- a/Engine/Plugins/Experimental/Mutable/Source/CustomizableObjectEditor/Classes/MuCOE/Nodes/CustomizableObjectNodeMeshMorphStackApplication.h
+++ b/Engine/Plugins/Experimental/Mutable/Source/CustomizableObjectEditor/Classes/MuCOE/Nodes/CustomizableObjectNodeMeshMorphStackApplication.h
@@ -17,37 +17,21 @@ class CUSTOMIZABLEOBJECTEDITOR_API UCustomizableObjectNodeMeshMorphStackApplicat
 public:
 	GENERATED_BODY()
 
-	UCustomizableObjectNodeMeshMorphStackApplication();
-
-	// Begin EdGraphNode interface
-	FText GetNodeTitle(ENodeTitleType::Type TittleType)const override;
-	FLinearColor GetNodeTitleColor() const override;
-	FText GetTooltipText() const override;
+	// EdGraphNode interface
+	virtual FText GetNodeTitle(ENodeTitleType::Type TittleType)const override;
+	virtual FLinearColor GetNodeTitleColor() const override;
+	virtual FText GetTooltipText() const override;
 
 	// UCustomizableObjectNode interface
-	void AllocateDefaultPins(UCustomizableObjectNodeRemapPins* RemapPins) override;
-	void PinConnectionListChanged(UEdGraphPin* Pin) override;
-	void ReconstructNode(UCustomizableObjectNodeRemapPins* RemapPins) override;
-	bool IsNodeOutDatedAndNeedsRefresh() override;
-
-	// Fills the list with all the morphs
-	void UpdateMorphList();
-
-	// Returns the mesh pin
-	UEdGraphPin* GetMeshPin() const
-	{
-		return FindPin(TEXT("InMesh"), EGPD_Input);
-	}
-	
-	// Returns the stack pin
-	UEdGraphPin* GetStackPin() const
-	{
-		return FindPin(TEXT("Stack"), EGPD_Input);
-	}
+	virtual void AllocateDefaultPins(UCustomizableObjectNodeRemapPins* RemapPins) override;
 
-public:
+	// Own interface
+	/** Fills the list with all the morphs */
+	TArray<FString> GetMorphList() const;
+	
+	/** Returns the mesh pin. */
+	UEdGraphPin* GetMeshPin() const;
 
-	// List with all the morphs of the linked skeletal mesh
-	UPROPERTY()
-	TArray<FString> MorphNames;
+	/** Returns the stack pin. */
+	UEdGraphPin* GetStackPin() const;
 };
diff --git a/Engine/Plugins/Experimental/Mutable/Source/CustomizableObjectEditor/Classes/MuCOE/Nodes/CustomizableObjectNodeMeshMorphStackDefinition.h b/Engine/Plugins/Experimental/Mutable/Source/CustomizableObjectEditor/Classes/MuCOE/Nodes/CustomizableObjectNodeMeshMorphStackDefinition.h
--- a/Engine/Plugins/Experimental/Mutable/Source/CustomizableObjectEditor/Classes/MuCOE/Nodes/CustomizableObjectNodeMeshMorphStackDefinition.h
+++ b/Engine/Plugins/Experimental/Mutable/Source/CustomizableObjectEditor/Classes/MuCOE/Nodes/CustomizableObjectNodeMeshMorphStackDefinition.h
@@ -17,44 +17,35 @@ class CUSTOMIZABLEOBJECTEDITOR_API UCustomizableObjectNodeMeshMorphStackDefiniti
 public:
 	GENERATED_BODY()
 
-	UCustomizableObjectNodeMeshMorphStackDefinition();
-
-	// Begin EdGraphNode interface
-	FText GetNodeTitle(ENodeTitleType::Type TittleType)const override;
-	FLinearColor GetNodeTitleColor() const override;
-	FText GetTooltipText() const override;
-	void PinConnectionListChanged(UEdGraphPin* Pin) override;
-	bool IsNodeOutDatedAndNeedsRefresh() override;
+	// EdGraphNode interface
+	virtual FText GetNodeTitle(ENodeTitleType::Type TittleType)const override;
+	virtual FLinearColor GetNodeTitleColor() const override;
+	virtual FText GetTooltipText() const override;
+	virtual void PinConnectionListChanged(UEdGraphPin* Pin) override;
+	virtual bool IsNodeOutDatedAndNeedsRefresh() override;
 
 	// UCustomizableObjectNode interface
-	void AllocateDefaultPins(UCustomizableObjectNodeRemapPins* RemapPins) override;
-	void ReconstructNode(UCustomizableObjectNodeRemapPins* RemapPins) override;
-
-	// Fills the list with all the morphs
-	void UpdateMorphList();
+	virtual void AllocateDefaultPins(UCustomizableObjectNodeRemapPins* RemapPins) override;
+	virtual void ReconstructNode(UCustomizableObjectNodeRemapPins* RemapPins) override;
 
-	// Returns the mesh pin of the node
-	UEdGraphPin* GetMeshPin() const
-	{
-		return FindPin(TEXT("Mesh"), EGPD_Input);
-	}
+	// Own interface
+	/** Returns the mesh pin of the node. */
+	UEdGraphPin* GetMeshPin() const;
 
-	// Returns the stack pin of the node
-	UEdGraphPin* GetStackPin() const
-	{
-		return FindPin(TEXT("Stack"), EGPD_Output);
-	}
+	/** Returns the stack pin of the node. */
+	UEdGraphPin* GetStackPin() const;
 
-	// Returns the morph pin at Index
+	/** Returns the morph pin at Index. */
 	UEdGraphPin* GetMorphPin(int32 Index) const;
 	
-	// Returns the index of the next connected stack node. Returns -1 if there is none.
+	/** Returns the index of the next connected stack node. Returns -1 if there is none. */
 	int32 NextConnectedPin(int32 Index, TArray<FString> AvailableMorphs)const;
 
-public:
+private:
+	/** Fills the list with all the morphs. */
+	void UpdateMorphList();
 
-	// List with all the morphs of the linked skeletal mesh
+	/** List with all the morphs of the linked skeletal mesh. */
 	UPROPERTY()
 	TArray<FString> MorphNames;
-
 };
diff --git a/Engine/Plugins/Experimental/Mutable/Source/CustomizableObjectEditor/Classes/MuCOE/Nodes/CustomizableObjectNodeSwitchBase.cpp b/Engine/Plugins/Experimental/Mutable/Source/CustomizableObjectEditor/Classes/MuCOE/Nodes/CustomizableObjectNodeSwitchBase.cpp
--- a/Engine/Plugins/Experimental/Mutable/Source/CustomizableObjectEditor/Classes/MuCOE/Nodes/CustomizableObjectNodeSwitchBase.cpp
+++ b/Engine/Plugins/Experimental/Mutable/Source/CustomizableObjectEditor/Classes/MuCOE/Nodes/CustomizableObjectNodeSwitchBase.cpp
@@ -124,8 +124,7 @@ void UCustomizableObjectNodeSwitchBase::EnumParameterPostEditChangeProperty(FPos
 	{
 		if (const UEdGraphPin* ConnectedPin = FollowInputPin(*SwitchPin); ConnectedPin && ConnectedPin->GetOwningNode() == Parameters.Node)
 		{
-			// Using MarkForReconstruct instead of Super::ReconstructNode because when we Copy Paste this node, it crashes sometimes due to reconctructing the node while constructing it.
-			MarkForReconstruct();
+			Super::ReconstructNode();
 		}
 		else if (UCustomizableObjectNode* EnumNode = Cast<UCustomizableObjectNode>(Parameters.Node))
 		{
@@ -194,9 +193,9 @@ void UCustomizableObjectNodeSwitchBase::LinkPostEditChangePropertyDelegate(const
 			LastNodeEnumParameterConnected->PostEditChangePropertyDelegate.AddUniqueDynamic(this, &UCustomizableObjectNodeSwitchBase::EnumParameterPostEditChangeProperty);
 		}
 
-		if (UCustomizableObjectNodeEnumParameter* Node = Cast<UCustomizableObjectNodeEnumParameter>(ConnectedPin->GetOwningNode()))
+		if (Cast<UCustomizableObjectNodeEnumParameter>(ConnectedPin->GetOwningNode()))
 		{
-			MarkForReconstruct();
+			Super::ReconstructNode();
 		}
 	}
 }
diff --git a/Engine/Plugins/Experimental/Mutable/Source/CustomizableObjectEditor/Private/MuCOE/CustomizableObjectEditor.cpp b/Engine/Plugins/Experimental/Mutable/Source/CustomizableObjectEditor/Private/MuCOE/CustomizableObjectEditor.cpp
--- a/Engine/Plugins/Experimental/Mutable/Source/CustomizableObjectEditor/Private/MuCOE/CustomizableObjectEditor.cpp
+++ b/Engine/Plugins/Experimental/Mutable/Source/CustomizableObjectEditor/Private/MuCOE/CustomizableObjectEditor.cpp
@@ -1696,20 +1696,6 @@ void FCustomizableObjectEditor::Tick( float InDeltaTime )
             FCoreUObjectDelegates::BroadcastOnObjectModified(PreviewInstance);
         }
 	}
-
-	// Reconstruct marked nodes 
-	if (GraphEditor.IsValid())
-	{
-		for (UEdGraphNode* Node: ReconstructNodes)
-		{
-			GraphEditor->GetCurrentGraph()->Schema.GetDefaultObject()->ReconstructNode(*Node);
-		}
-		if (ReconstructNodes.Num()) 
-		{
-			GraphEditor->NotifyGraphChanged();
-			ReconstructNodes.Empty();
-		}
-	}
 }
 
 
@@ -2766,12 +2752,6 @@ void FCustomizableObjectEditor::GetExternalChildObjects(const UCustomizableObjec
 }
 
 
-void FCustomizableObjectEditor::MarkForReconstruct(UEdGraphNode* Node)
-{
-	ReconstructNodes.Add(Node);
-}
-
-
 void FCustomizableObjectEditor::CreatePreviewComponents()
 {
 	// Getting the number of mesh components from the root node
diff --git a/Engine/Plugins/Experimental/Mutable/Source/CustomizableObjectEditor/Private/MuCOE/CustomizableObjectEditor.h b/Engine/Plugins/Experimental/Mutable/Source/CustomizableObjectEditor/Private/MuCOE/CustomizableObjectEditor.h
--- a/Engine/Plugins/Experimental/Mutable/Source/CustomizableObjectEditor/Private/MuCOE/CustomizableObjectEditor.h
+++ b/Engine/Plugins/Experimental/Mutable/Source/CustomizableObjectEditor/Private/MuCOE/CustomizableObjectEditor.h
@@ -243,9 +243,6 @@ public:
 	/** Getter of CustomizableObjectEditorAdvancedPreviewSettings */
 	TSharedPtr<class SCustomizableObjectEditorAdvancedPreviewSettings> GetCustomizableObjectEditorAdvancedPreviewSettings();
 
-	/** Add a node to be reconstructed at the end of the tick. */
-	void MarkForReconstruct(UEdGraphNode* Node);
-
 	TSharedPtr<class SCustomizableObjectNodeLayoutBlocksEditor> GetLayoutBlocksEditor() { return LayoutBlocksEditor; }
 
 private:
@@ -474,9 +471,6 @@ private:
 	/** Adds the customizable Object Editor commands to the default toolbar */
 	void ExtendToolbar();
 
-	/** List of nodes to be reconstructed at the end of the tick. */
-	TArray<UEdGraphNode*> ReconstructNodes;
-
 	/** URL to open when pressing the documentation button generated by UE */
 	const FString DocumentationURL{ TEXT("https://work.anticto.com/w/mutable/unreal-engine-4/") };
 
diff --git a/Engine/Plugins/Experimental/Mutable/Source/CustomizableObjectEditor/Private/MuCOE/GenerateMutableSource/GenerateMutableSourceMesh.cpp b/Engine/Plugins/Experimental/Mutable/Source/CustomizableObjectEditor/Private/MuCOE/GenerateMutableSource/GenerateMutableSourceMesh.cpp
--- a/Engine/Plugins/Experimental/Mutable/Source/CustomizableObjectEditor/Private/MuCOE/GenerateMutableSource/GenerateMutableSourceMesh.cpp
+++ b/Engine/Plugins/Experimental/Mutable/Source/CustomizableObjectEditor/Private/MuCOE/GenerateMutableSource/GenerateMutableSourceMesh.cpp
@@ -3102,6 +3102,8 @@ mu::NodeMeshPtr GenerateMutableSourceMesh(const UEdGraphPin* Pin,
 
 	else if (const UCustomizableObjectNodeMeshMorphStackApplication* TypedNodeMeshMorphStackApp = Cast< UCustomizableObjectNodeMeshMorphStackApplication >(Node))
 	{
+		const TArray<FString> MorphNames = TypedNodeMeshMorphStackApp->GetMorphList();
+		
 		if (const UEdGraphPin* ConnectedPin = FollowInputPin(*TypedNodeMeshMorphStackApp->GetStackPin()))
 		{
 			UEdGraphNode* OwningNode = ConnectedPin->GetOwningNode();
@@ -3136,7 +3138,7 @@ mu::NodeMeshPtr GenerateMutableSourceMesh(const UEdGraphPin* Pin,
 
 					// Cheking if the morph exists in the application node
 					FString MorphName = MorphPin->PinFriendlyName.ToString();
-					if (!TypedNodeMeshMorphStackApp->MorphNames.Contains(MorphName))
+					if (!MorphNames.Contains(MorphName))
 					{
 						continue;
 					}
diff --git a/Engine/Plugins/Experimental/Mutable/Source/CustomizableObjectEditor/Private/MuCOE/Nodes/CustomizableObjectNode.cpp b/Engine/Plugins/Experimental/Mutable/Source/CustomizableObjectEditor/Private/MuCOE/Nodes/CustomizableObjectNode.cpp
--- a/Engine/Plugins/Experimental/Mutable/Source/CustomizableObjectEditor/Private/MuCOE/Nodes/CustomizableObjectNode.cpp
+++ b/Engine/Plugins/Experimental/Mutable/Source/CustomizableObjectEditor/Private/MuCOE/Nodes/CustomizableObjectNode.cpp
@@ -408,17 +408,6 @@ bool UCustomizableObjectNode::CanConnect( const UEdGraphPin* InOwnedInputPin, co
 	return bOutArePinsCompatible;
 }
 
-void UCustomizableObjectNode::MarkForReconstruct()
-{
-	TSharedPtr<FCustomizableObjectEditor> Editor = StaticCastSharedPtr<FCustomizableObjectEditor>(GetGraphEditor());
-
-	// This could be called if a skeletal mesh asset in the graph is deleted without the editor being open. 
-	if (Editor)
-	{
-		Editor->MarkForReconstruct(this);
-	}
-}
-
 
 UCustomizableObjectNodeRemapPins* UCustomizableObjectNode::CreateRemapPinsDefault() const
 {
diff --git a/Engine/Plugins/Experimental/Mutable/Source/CustomizableObjectEditor/Private/MuCOE/Nodes/CustomizableObjectNodeMeshMorphStackApplication.cpp b/Engine/Plugins/Experimental/Mutable/Source/CustomizableObjectEditor/Private/MuCOE/Nodes/CustomizableObjectNodeMeshMorphStackApplication.cpp
--- a/Engine/Plugins/Experimental/Mutable/Source/CustomizableObjectEditor/Private/MuCOE/Nodes/CustomizableObjectNodeMeshMorphStackApplication.cpp
+++ b/Engine/Plugins/Experimental/Mutable/Source/CustomizableObjectEditor/Private/MuCOE/Nodes/CustomizableObjectNodeMeshMorphStackApplication.cpp
@@ -15,28 +15,16 @@ class UCustomizableObjectNodeRemapPins;
 #define LOCTEXT_NAMESPACE "CustomizableObjectEditor"
 
 
-
-UCustomizableObjectNodeMeshMorphStackApplication::UCustomizableObjectNodeMeshMorphStackApplication() : Super()
-{
-
-}
-
-
 void UCustomizableObjectNodeMeshMorphStackApplication::AllocateDefaultPins(UCustomizableObjectNodeRemapPins* RemapPins)
 {
 	const UEdGraphSchema_CustomizableObject* Schema = GetDefault<UEdGraphSchema_CustomizableObject>();
 
 	// Input pins
-	FString PinName = TEXT("InMesh");
-	UEdGraphPin* InputMeshPin = CustomCreatePin(EGPD_Input, Schema->PC_Mesh, FName(*PinName));
-
-	PinName = TEXT("Stack");
-	UEdGraphPin* StackPin = CustomCreatePin(EGPD_Input, Schema->PC_Stack, FName(*PinName));
+	CustomCreatePin(EGPD_Input, Schema->PC_Mesh, TEXT("InMesh"));
+	CustomCreatePin(EGPD_Input, Schema->PC_Stack, TEXT("Stack"));
 
 	// Output pins
-	PinName = TEXT("Result Mesh");
-	UEdGraphPin* OutputMeshPin = CustomCreatePin(EGPD_Output, Schema->PC_Mesh, FName(*PinName));
-
+	CustomCreatePin(EGPD_Output, Schema->PC_Mesh, TEXT("Result Mesh"));
 }
 
 
@@ -59,97 +47,32 @@ FText UCustomizableObjectNodeMeshMorphStackApplication::GetTooltipText() const
 }
 
 
-void UCustomizableObjectNodeMeshMorphStackApplication::PinConnectionListChanged(UEdGraphPin* Pin)
+UEdGraphPin* UCustomizableObjectNodeMeshMorphStackApplication::GetMeshPin() const
 {
-	Super::PinConnectionListChanged(Pin);
-
-	if (Pin->Direction == EGPD_Input && Pin->PinName == "InMesh")
-	{
-		UpdateMorphList();
-
-		TSharedPtr<ICustomizableObjectEditor> Editor = GetGraphEditor();
-
-		if (Editor.IsValid())
-		{
-			Editor->UpdateGraphNodeProperties();
-			MarkForReconstruct();
-		}
-	}
+	return FindPin(TEXT("InMesh"), EGPD_Input);
 }
 
 
-void UCustomizableObjectNodeMeshMorphStackApplication::ReconstructNode(UCustomizableObjectNodeRemapPins* RemapPins)
+UEdGraphPin* UCustomizableObjectNodeMeshMorphStackApplication::GetStackPin() const
 {
-	UpdateMorphList();
-
-	Super::ReconstructNode(RemapPins);
+	return FindPin(TEXT("Stack"), EGPD_Input);
 }
 
-
-void UCustomizableObjectNodeMeshMorphStackApplication::UpdateMorphList()
+TArray<FString> UCustomizableObjectNodeMeshMorphStackApplication::GetMorphList() const
 {
-	MorphNames.Empty();
-
-	UEdGraphPin* MeshPin = GetMeshPin();
-
+	const UEdGraphPin* MeshPin = GetMeshPin();
 	if (!MeshPin)
 	{
-		return;
+		return {};
 	}
 
-	UEdGraphPin* OutputMeshPin = FollowInputPin(*MeshPin);
-
+	const UEdGraphPin* OutputMeshPin = FollowInputPin(*MeshPin);
 	if (!OutputMeshPin)
 	{
-		return;
+		return {};
 	}
 
-	const UEdGraphPin* MeshNodePin = FindMeshBaseSource(*OutputMeshPin, false);
-
-	if (MeshNodePin && MeshNodePin->GetOwningNode())
-	{
-		USkeletalMesh* SkeletalMesh = nullptr;
-
-		if (const UCustomizableObjectNodeSkeletalMesh* SkeletalMeshNode = Cast< UCustomizableObjectNodeSkeletalMesh >(MeshNodePin->GetOwningNode()))
-		{
-			SkeletalMesh = SkeletalMeshNode->SkeletalMesh;
-		}
-		else if (const UCustomizableObjectNodeTable* TableNode = Cast< UCustomizableObjectNodeTable >(MeshNodePin->GetOwningNode()))
-		{
-			SkeletalMesh = TableNode->GetColumnDefaultAssetByType<USkeletalMesh>(MeshNodePin);
-		}
-
-		if (SkeletalMesh)
-		{
-			for (int32 i = 0; i < SkeletalMesh->GetMorphTargets().Num(); ++i)
-			{
-				MorphNames.Add(SkeletalMesh->GetMorphTargets()[i]->GetName());
-			}
-		}
-	}
-}
-
-
-bool UCustomizableObjectNodeMeshMorphStackApplication::IsNodeOutDatedAndNeedsRefresh()
-{
-	UEdGraphPin* MeshPin = GetMeshPin();
-
-	if (!MeshPin)
-	{
-		return false;
-	}
-
-	UEdGraphPin* OutputMeshPin = FollowInputPin(*MeshPin);
-
-	if (!OutputMeshPin)
-	{
-		if (MorphNames.Num())
-		{
-			return true;
-		}
-
-		return false;
-	}
+	TArray<FString> MorphNames;
 
 	const UEdGraphPin* MeshNodePin = FindMeshBaseSource(*OutputMeshPin, false);
 
@@ -168,22 +91,14 @@ bool UCustomizableObjectNodeMeshMorphStackApplication::IsNodeOutDatedAndNeedsRef
 
 		if (SkeletalMesh)
 		{
-			if (SkeletalMesh->GetMorphTargets().Num() != MorphNames.Num())
-			{
-				return true;
-			}
-
 			for (int32 i = 0; i < SkeletalMesh->GetMorphTargets().Num(); ++i)
 			{
-				if (!MorphNames.Contains(SkeletalMesh->GetMorphTargets()[i]->GetName()))
-				{
-					return true;
-				}
+				MorphNames.Add(SkeletalMesh->GetMorphTargets()[i]->GetName());
 			}
 		}
 	}
 
-	return false;
+	return MorphNames;
 }
 
 
diff --git a/Engine/Plugins/Experimental/Mutable/Source/CustomizableObjectEditor/Private/MuCOE/Nodes/CustomizableObjectNodeMeshMorphStackDefinition.cpp b/Engine/Plugins/Experimental/Mutable/Source/CustomizableObjectEditor/Private/MuCOE/Nodes/CustomizableObjectNodeMeshMorphStackDefinition.cpp
--- a/Engine/Plugins/Experimental/Mutable/Source/CustomizableObjectEditor/Private/MuCOE/Nodes/CustomizableObjectNodeMeshMorphStackDefinition.cpp
+++ b/Engine/Plugins/Experimental/Mutable/Source/CustomizableObjectEditor/Private/MuCOE/Nodes/CustomizableObjectNodeMeshMorphStackDefinition.cpp
@@ -15,17 +15,12 @@ class UCustomizableObjectNodeRemapPins;
 #define LOCTEXT_NAMESPACE "CustomizableObjectEditor"
 
 
-
-UCustomizableObjectNodeMeshMorphStackDefinition::UCustomizableObjectNodeMeshMorphStackDefinition() : Super()
-{
-
-}
-
-
 void UCustomizableObjectNodeMeshMorphStackDefinition::AllocateDefaultPins(UCustomizableObjectNodeRemapPins* RemapPins)
 {
 	const UEdGraphSchema_CustomizableObject* Schema = GetDefault<UEdGraphSchema_CustomizableObject>();
 
+	UpdateMorphList();
+	
 	// Input pins
 	UEdGraphPin* MeshPin = CustomCreatePin(EGPD_Input, Schema->PC_Mesh, FName(TEXT("Mesh")));
 	MeshPin->bDefaultValueIsIgnored = true;
@@ -69,23 +64,13 @@ void UCustomizableObjectNodeMeshMorphStackDefinition::PinConnectionListChanged(U
 
 	if (Pin->Direction == EGPD_Input && Pin->PinName == "Mesh")
 	{
-		UpdateMorphList();
-
-		TSharedPtr<ICustomizableObjectEditor> Editor = GetGraphEditor();
-
-		if (Editor.IsValid())
-		{
-			Editor->UpdateGraphNodeProperties();
-			Super::ReconstructNode();
-		}
+		Super::ReconstructNode();
 	}
 }
 
 
 void UCustomizableObjectNodeMeshMorphStackDefinition::ReconstructNode(UCustomizableObjectNodeRemapPins* RemapPins)
 {
-	UpdateMorphList();
-
 	Super::ReconstructNode(RemapPins);
 }
 
@@ -134,6 +119,18 @@ void UCustomizableObjectNodeMeshMorphStackDefinition::UpdateMorphList()
 }
 
 
+UEdGraphPin* UCustomizableObjectNodeMeshMorphStackDefinition::GetMeshPin() const
+{
+	return FindPin(TEXT("Mesh"), EGPD_Input);
+}
+
+
+UEdGraphPin* UCustomizableObjectNodeMeshMorphStackDefinition::GetStackPin() const
+{
+	return FindPin(TEXT("Stack"), EGPD_Output);
+}
+
+
 UEdGraphPin* UCustomizableObjectNodeMeshMorphStackDefinition::GetMorphPin(int32 Index) const
 {
 	FString PinName = MorphNames[Index] + FString("_Morph");
diff --git a/Engine/Plugins/Experimental/Mutable/Source/CustomizableObjectEditor/Private/MuCOE/Nodes/CustomizableObjectNodeSkeletalMesh.cpp b/Engine/Plugins/Experimental/Mutable/Source/CustomizableObjectEditor/Private/MuCOE/Nodes/CustomizableObjectNodeSkeletalMesh.cpp
--- a/Engine/Plugins/Experimental/Mutable/Source/CustomizableObjectEditor/Private/MuCOE/Nodes/CustomizableObjectNodeSkeletalMesh.cpp
+++ b/Engine/Plugins/Experimental/Mutable/Source/CustomizableObjectEditor/Private/MuCOE/Nodes/CustomizableObjectNodeSkeletalMesh.cpp
@@ -35,7 +35,7 @@ void UCustomizableObjectNodeSkeletalMesh::PostEditChangeProperty(FPropertyChange
 	FProperty* PropertyThatChanged = PropertyChangedEvent.Property;
 	if (PropertyThatChanged && PropertyThatChanged->GetName() == TEXT("SkeletalMesh"))
 	{
-		MarkForReconstruct();
+		ReconstructNode();
 	}
 }
 
diff --git a/Engine/Plugins/Experimental/Mutable/Source/CustomizableObjectEditor/Private/MuCOE/Nodes/CustomizableObjectNodeTable.cpp b/Engine/Plugins/Experimental/Mutable/Source/CustomizableObjectEditor/Private/MuCOE/Nodes/CustomizableObjectNodeTable.cpp
--- a/Engine/Plugins/Experimental/Mutable/Source/CustomizableObjectEditor/Private/MuCOE/Nodes/CustomizableObjectNodeTable.cpp
+++ b/Engine/Plugins/Experimental/Mutable/Source/CustomizableObjectEditor/Private/MuCOE/Nodes/CustomizableObjectNodeTable.cpp
@@ -192,7 +192,7 @@ void UCustomizableObjectNodeTable::PostEditChangeProperty(FPropertyChangedEvent&
 	{
 		if (PropertyThatChanged->GetName() == TEXT("Table"))
 		{
-			MarkForReconstruct();
+			ReconstructNode();
 		}
 		else if (PropertyThatChanged->GetName() == TEXT("DefaultImageMode"))
 		{
