[Mutable] Fix deadlock when loading roms.

- Removed the locking wait in the ModelDiskStreamer.
- Fixed a out of range check with RomWeights.
- ModelDiskStreamer replaced potentially unsafe raw pointers with a shared pointers.

#rnx

diff --git a/Engine/Plugins/Experimental/Mutable/Source/CustomizableObject/Classes/MuCO/CustomizableObject.h b/Engine/Plugins/Experimental/Mutable/Source/CustomizableObject/Classes/MuCO/CustomizableObject.h
--- a/Engine/Plugins/Experimental/Mutable/Source/CustomizableObject/Classes/MuCO/CustomizableObject.h
+++ b/Engine/Plugins/Experimental/Mutable/Source/CustomizableObject/Classes/MuCO/CustomizableObject.h
@@ -1045,7 +1045,7 @@ public:
 
 	/** Creates and returns an array with IAsyncReadFileHandles for each BulkData file.
 	 * Used by Mutable to stream in resources when generating instances. Must be deleted by the caller. */
-	TArray<IAsyncReadFileHandle*> GetAsyncReadFileHandles() const;
+	TArray<TSharedPtr<IAsyncReadFileHandle>> GetAsyncReadFileHandles() const;
 
 #if WITH_EDITOR
 
diff --git a/Engine/Plugins/Experimental/Mutable/Source/CustomizableObject/Private/MuCO/CustomizableObject.cpp b/Engine/Plugins/Experimental/Mutable/Source/CustomizableObject/Private/MuCO/CustomizableObject.cpp
--- a/Engine/Plugins/Experimental/Mutable/Source/CustomizableObject/Private/MuCO/CustomizableObject.cpp
+++ b/Engine/Plugins/Experimental/Mutable/Source/CustomizableObject/Private/MuCO/CustomizableObject.cpp
@@ -1834,14 +1834,14 @@ void UCustomizableObjectBulk::PostLoad()
 }
 
 
-TArray<IAsyncReadFileHandle*> UCustomizableObjectBulk::GetAsyncReadFileHandles() const
+TArray<TSharedPtr<IAsyncReadFileHandle>> UCustomizableObjectBulk::GetAsyncReadFileHandles() const
 {
-	TArray<IAsyncReadFileHandle*> ReadFileHandles;
+	TArray<TSharedPtr<IAsyncReadFileHandle>> ReadFileHandles;
 	ReadFileHandles.Reserve(BulkDataFileNames.Num());
 
 	for (const FString& FilePath : BulkDataFileNames)
 	{
-		IAsyncReadFileHandle* ReadFileHandle = FPlatformFileManager::Get().GetPlatformFile().OpenAsyncRead(*FilePath);
+		TSharedPtr<IAsyncReadFileHandle> ReadFileHandle = MakeShareable(FPlatformFileManager::Get().GetPlatformFile().OpenAsyncRead(*FilePath));
 		
 		if (!ReadFileHandle)
 		{
@@ -1854,11 +1854,6 @@ TArray<IAsyncReadFileHandle*> UCustomizableObjectBulk::GetAsyncReadFileHandles()
 
 	if (ReadFileHandles.Num() != BulkDataFileNames.Num())
 	{
-		for (IAsyncReadFileHandle* ReadFileHandle : ReadFileHandles)
-		{
-			delete ReadFileHandle;
-		}
-
 		ReadFileHandles.Empty();
 	}
 
diff --git a/Engine/Plugins/Experimental/Mutable/Source/CustomizableObject/Private/MuCO/UnrealMutableModelDiskStreamer.cpp b/Engine/Plugins/Experimental/Mutable/Source/CustomizableObject/Private/MuCO/UnrealMutableModelDiskStreamer.cpp
--- a/Engine/Plugins/Experimental/Mutable/Source/CustomizableObject/Private/MuCO/UnrealMutableModelDiskStreamer.cpp
+++ b/Engine/Plugins/Experimental/Mutable/Source/CustomizableObject/Private/MuCO/UnrealMutableModelDiskStreamer.cpp
@@ -67,15 +67,9 @@ bool FUnrealMutableModelBulkReader::PrepareStreamingForObject(UCustomizableObjec
 		if (!Objects[ObjectIndex].Model.Pin())
 		{
 			// The CustomizableObject is gone, so we won't be streaming for it anymore.
-			for (TPair<OPERATION_ID, IAsyncReadRequest*>& it : Objects[ObjectIndex].CurrentReadRequests)
+			for (TPair<OPERATION_ID, FReadRequest>& it : Objects[ObjectIndex].CurrentReadRequests)
 			{
-				it.Value->WaitCompletion();
-				delete it.Value;
-			}
-
-			for (IAsyncReadFileHandle* ReadFileHandle : Objects[ObjectIndex].ReadFileHandles)
-			{
-				delete ReadFileHandle;
+				it.Value.ReadRequest->WaitCompletion();
 			}
 
 			Objects.RemoveAtSwap(ObjectIndex);
@@ -102,7 +96,7 @@ bool FUnrealMutableModelBulkReader::PrepareStreamingForObject(UCustomizableObjec
 		FString FolderPath = CustomizableObject->GetCompiledDataFolderPath(true);
 		FString FullFileName = FolderPath + CustomizableObject->GetCompiledDataFileName(false, nullptr, true);
 
-		IAsyncReadFileHandle* ReadFileHandle = FPlatformFileManager::Get().GetPlatformFile().OpenAsyncRead(*FullFileName);
+		const TSharedPtr<IAsyncReadFileHandle> ReadFileHandle = MakeShareable(FPlatformFileManager::Get().GetPlatformFile().OpenAsyncRead(*FullFileName));
 		if (!ReadFileHandle)
 		{
 			UE_LOG(LogMutable, Warning, TEXT("Streaming: Customizable Object %s is missing the Editor BulkData."), *CustomizableObject->GetName());
@@ -164,15 +158,9 @@ void FUnrealMutableModelBulkReader::CancelStreamingForObject(const UCustomizable
 	{
 		if (Objects[ObjectIndex].Model.Pin() == CustomizableObject->GetModel())
 		{
-			for (TPair<OPERATION_ID, IAsyncReadRequest*>& it : Objects[ObjectIndex].CurrentReadRequests)
+			for (TPair<OPERATION_ID, FReadRequest>& it : Objects[ObjectIndex].CurrentReadRequests)
 			{
-				it.Value->WaitCompletion();
-				delete it.Value;
-			}
-
-			for (IAsyncReadFileHandle* ReadFileHandle : Objects[ObjectIndex].ReadFileHandles)
-			{
-				delete ReadFileHandle;
+				it.Value.ReadRequest->WaitCompletion();
 			}
 
 			Objects.RemoveAtSwap(ObjectIndex);
@@ -214,15 +202,9 @@ void FUnrealMutableModelBulkReader::EndStreaming()
 {
 	for (FObjectData& o : Objects)
 	{
-		for (TPair<OPERATION_ID, IAsyncReadRequest*>& it : o.CurrentReadRequests)
-		{
-			it.Value->WaitCompletion();
-			delete it.Value;
-		}
-
-		for (IAsyncReadFileHandle* ReadFileHandle : o.ReadFileHandles)
+		for (TPair<OPERATION_ID, FReadRequest>& it : o.CurrentReadRequests)
 		{
-			delete ReadFileHandle;
+			it.Value.ReadRequest->WaitCompletion();
 		}
 	}
 	Objects.Empty();
@@ -236,11 +218,11 @@ FAutoConsoleVariableRef CVarStreamPriority(
 	TEXT(""));
 
 
-mu::ModelReader::OPERATION_ID FUnrealMutableModelBulkReader::BeginReadBlock(const mu::Model* Model, uint64 Key, void* pBuffer, uint64 size)
+mu::ModelReader::OPERATION_ID FUnrealMutableModelBulkReader::BeginReadBlock(const mu::Model* Model, uint64 Key, void* pBuffer, uint64 size, TFunction<void(bool bSuccess)>* CompletionCallback)
 {
 	MUTABLE_CPUPROFILER_SCOPE(FUnrealMutableModelBulkStreamer::OpenReadFile);
 
-	UE_LOG(LogMutable, VeryVerbose, TEXT("Streaming: reading data %08lx."), Key);
+	UE_LOG(LogMutable, VeryVerbose, TEXT("Streaming: reading data %08llu."), Key);
 
 	// Find the object we are streaming for
 	FObjectData* ObjectData = Objects.FindByPredicate(
@@ -251,34 +233,55 @@ mu::ModelReader::OPERATION_ID FUnrealMutableModelBulkReader::BeginReadBlock(cons
 		// The object has been unloaded. Streaming is not possible. 
 		// This may happen in the editor if we are recompiling an object but we still have instances from the old
 		// object that have progressive mip generation.
+		if (CompletionCallback)
+		{
+			(*CompletionCallback)(false);    		
+		}
 		return -1;
 	}
 
-	mu::ModelReader::OPERATION_ID Result = 0;
-
-	check(!ObjectData->ReadFileHandles.IsEmpty());
-
 	// this generally cannot fail because it is async
-	if (ObjectData->StreamableBlocks.Contains(Key))
+	if (!ObjectData->StreamableBlocks.Contains(Key))
 	{
-		Result = ++LastOperationID;
+		// File Handle not found! This shouldn't really happen.
+		UE_LOG(LogMutable, Error, TEXT("Streaming Block not found!"));
+		check(false);
 
-		const FMutableStreamableBlock& Block = ObjectData->StreamableBlocks[Key];
+		if (CompletionCallback)
+		{
+			(*CompletionCallback)(false);    		
+		}
+		return -1;
+	}
 
-		int32 BulkDataOffsetInFile = 0;
+	OPERATION_ID Result = ++LastOperationID;
+
+	const FMutableStreamableBlock& Block = ObjectData->StreamableBlocks[Key];
+
+	int32 BulkDataOffsetInFile = 0;
 #if WITH_EDITOR
-		BulkDataOffsetInFile = sizeof(MutableCompiledDataStreamHeader);
+	BulkDataOffsetInFile = sizeof(MutableCompiledDataStreamHeader);
 #endif
 
-		IAsyncReadRequest* ReadRequest = ObjectData->ReadFileHandles[Block.FileIndex]->ReadRequest(BulkDataOffsetInFile + Block.Offset, size, (EAsyncIOPriorityAndFlags)StreamPriority, nullptr, reinterpret_cast<uint8*>(pBuffer));
-		ObjectData->CurrentReadRequests.Add(Result, ReadRequest);
-	}
-	else
+	FReadRequest ReadRequest;
+
+	if (CompletionCallback)
 	{
-		// File Handle not found! This shouldn't really happen.
-		UE_LOG(LogMutable, Error, TEXT("Streaming Block not found!"));
-		check(false);
+		ReadRequest.FileCallback = MakeShared<TFunction<void(bool, IAsyncReadRequest*)>>([CompletionCallbackCapture = *CompletionCallback](bool bWasCancelled, IAsyncReadRequest*) -> void
+		{
+			CompletionCallbackCapture(!bWasCancelled);			
+		});		
 	}
+	
+	check(!ObjectData->ReadFileHandles.IsEmpty());
+	ReadRequest.ReadRequest = MakeShareable(ObjectData->ReadFileHandles[Block.FileIndex]->ReadRequest(
+		BulkDataOffsetInFile + Block.Offset,
+		size,
+		(EAsyncIOPriorityAndFlags)StreamPriority,
+		ReadRequest.FileCallback.Get(),
+		reinterpret_cast<uint8*>(pBuffer)));
+
+	ObjectData->CurrentReadRequests.Add(Result, ReadRequest);
 
 	INC_DWORD_STAT(STAT_MutableStreamingOps);
 
@@ -292,9 +295,9 @@ bool FUnrealMutableModelBulkReader::IsReadCompleted(mu::ModelReader::OPERATION_I
 
 	for (FObjectData& o : Objects)
 	{
-		if (IAsyncReadRequest** ReadRequest = o.CurrentReadRequests.Find(OperationId))
+		if (FReadRequest* ReadRequest = o.CurrentReadRequests.Find(OperationId))
 		{
-			return (*ReadRequest)->PollCompletion();
+			return ReadRequest->ReadRequest->PollCompletion();
 		}
 	}
 
@@ -311,19 +314,17 @@ void FUnrealMutableModelBulkReader::EndRead(mu::ModelReader::OPERATION_ID Operat
 	bool bFound = false;
 	for (FObjectData& o : Objects)
 	{
-		IAsyncReadRequest** ReadRequest = o.CurrentReadRequests.Find(OperationId);
+		FReadRequest* ReadRequest = o.CurrentReadRequests.Find(OperationId);
 		if (ReadRequest)
 		{
-			if (*ReadRequest)
+			if (ReadRequest->ReadRequest)
 			{
-				bool bCompleted = (*ReadRequest)->WaitCompletion();
+				bool bCompleted = ReadRequest->ReadRequest->WaitCompletion();
 				if (!bCompleted)
 				{
 					UE_LOG(LogMutable, Error, TEXT("Operation failed to complete in EndRead."));
 					check(false);
 				}
-
-				delete* ReadRequest;
 			}
 			o.CurrentReadRequests.Remove(OperationId);
 			bFound = true;
diff --git a/Engine/Plugins/Experimental/Mutable/Source/CustomizableObject/Public/MuCO/UnrealMutableModelDiskStreamer.h b/Engine/Plugins/Experimental/Mutable/Source/CustomizableObject/Public/MuCO/UnrealMutableModelDiskStreamer.h
--- a/Engine/Plugins/Experimental/Mutable/Source/CustomizableObject/Public/MuCO/UnrealMutableModelDiskStreamer.h
+++ b/Engine/Plugins/Experimental/Mutable/Source/CustomizableObject/Public/MuCO/UnrealMutableModelDiskStreamer.h
@@ -57,18 +57,24 @@ public:
 	void EndStreaming();
 
 	// mu::ModelReader interface
-	OPERATION_ID BeginReadBlock(const mu::Model*, uint64 key0, void* pBuffer, uint64 size) override;
+	OPERATION_ID BeginReadBlock(const mu::Model*, uint64 key0, void* pBuffer, uint64 size, TFunction<void(bool bSuccess)>* CompletionCallback) override;
 	bool IsReadCompleted(OPERATION_ID) override;
 	void EndRead(OPERATION_ID) override;
 
 protected:
 
+	struct FReadRequest
+	{
+		TSharedPtr<IAsyncReadRequest> ReadRequest = nullptr;
+		TSharedPtr<FAsyncFileCallBack> FileCallback;
+	};
+	
 	/** Streaming data for one object. */
 	struct FObjectData
 	{
 		TWeakPtr<const mu::Model> Model;
-		TArray<IAsyncReadFileHandle*> ReadFileHandles;
-		TMap<OPERATION_ID, IAsyncReadRequest*> CurrentReadRequests;
+		TArray<TSharedPtr<IAsyncReadFileHandle>> ReadFileHandles;
+		TMap<OPERATION_ID, FReadRequest> CurrentReadRequests;
 		TMap<uint64, FMutableStreamableBlock> StreamableBlocks;
 	};
 
diff --git a/Engine/Plugins/Experimental/Mutable/Source/MutableRuntime/Private/MuR/CodeRunner.h b/Engine/Plugins/Experimental/Mutable/Source/MutableRuntime/Private/MuR/CodeRunner.h
--- a/Engine/Plugins/Experimental/Mutable/Source/MutableRuntime/Private/MuR/CodeRunner.h
+++ b/Engine/Plugins/Experimental/Mutable/Source/MutableRuntime/Private/MuR/CodeRunner.h
@@ -30,6 +30,20 @@
 #endif
 
 
+using EventType = 
+#ifdef MUTABLE_USE_NEW_TASKGRAPH
+	UE::Tasks::FTaskEvent;
+#else
+	FGraphEventRef;
+#endif
+
+using TaskType = 
+#ifdef MUTABLE_USE_NEW_TASKGRAPH
+	UE::Tasks::FTask;
+#else
+	FGraphEventRef;
+#endif
+
 namespace  mu
 {
 	class Model;
@@ -205,6 +219,13 @@ namespace  mu
 			}
 		};
 
+		struct FRomLoadOp
+		{
+			ModelReader::OPERATION_ID m_streamID = -1;
+			TArray<uint8> m_streamBuffer;
+			TaskType Event;
+		};
+
 		class FLoadMeshRomTask : public CodeRunner::FIssuedTask
 		{
 		public:
@@ -217,7 +238,6 @@ namespace  mu
 			// FIssuedTask interface
 			bool Prepare(CodeRunner*, bool& bOutFailed) override;
 			void Complete(CodeRunner*) override;
-			bool IsComplete(CodeRunner*) override;
 
 		private:
 			int32 RomIndex = -1;
@@ -261,11 +281,12 @@ namespace  mu
 			// FIssuedTask interface
 			bool Prepare(CodeRunner*, bool& bOutFailed) override;
 			void Complete(CodeRunner*) override;
-			bool IsComplete(CodeRunner*) override;
 
 		private:
  			int32 LODIndexIndex = -1;
 			int32 LODIndexCount = -1;
+
+			TArray<int32> RomIndices;
 		};
 
 		void AddOp(const FScheduledOp& op)
@@ -355,6 +376,19 @@ namespace  mu
 			}
 		}
 
+    	/** Calculate an approximation of memory used by streaming buffers in this class. */
+		int32 GetStreamingMemoryBytes() const
+    	{
+			int32 Result = RomLoadOps.GetAllocatedSize();
+			
+			for (const TTuple<int32, TSharedPtr<FRomLoadOp>>& RomLoadOp : RomLoadOps)
+			{
+				Result += RomLoadOp.Value->m_streamBuffer.GetAllocatedSize();
+			}
+			
+			return Result;
+    	}
+    	
 		/** Calculate an approximation of memory used by manging structures in this class. */
 		int32 GetInternalMemoryBytes() const
 		{
@@ -567,16 +601,9 @@ namespace  mu
 		const Parameters* m_pParams = nullptr;
 		uint32 m_lodMask = 0;
 
-		// Async rom loading control
-		struct FRomLoadOp
-		{
-			int32 m_romIndex = 0;
-			DATATYPE ConstantType = DT_NONE;
-			ModelReader::OPERATION_ID m_streamID;
-			TArray<uint8> m_streamBuffer;
-		};
-		TArray<FRomLoadOp> m_romLoadOps;
-
+    	/** Rom read operations already in progress. Key is the rom index. */
+    	TMap<int32, TSharedPtr<FRomLoadOp>> RomLoadOps;
+    	
 	private:
 
 		inline void AddChildren(const FScheduledOp& dep)
@@ -603,9 +630,6 @@ namespace  mu
 		/** Calculate an heuristic to select op execution based on memory usage. */
 		int32 GetOpEstimatedMemoryDelta(const FScheduledOp& Candidate, const FProgram& Program);
 
-		/** */
-		void CompleteRomLoadOp(FRomLoadOp& o);
-		
 		/** Update debug stats. */
 		void UpdateTraces();
 
diff --git a/Engine/Plugins/Experimental/Mutable/Source/MutableRuntime/Private/MuR/CodeRunnerTasked.cpp b/Engine/Plugins/Experimental/Mutable/Source/MutableRuntime/Private/MuR/CodeRunnerTasked.cpp
--- a/Engine/Plugins/Experimental/Mutable/Source/MutableRuntime/Private/MuR/CodeRunnerTasked.cpp
+++ b/Engine/Plugins/Experimental/Mutable/Source/MutableRuntime/Private/MuR/CodeRunnerTasked.cpp
@@ -57,54 +57,6 @@ static FAutoConsoleVariableRef CVarCoreRunnerTaskPriority(
 namespace mu
 {
 
-	void CodeRunner::CompleteRomLoadOp( FRomLoadOp& o )
-	{
-		if (DebugRom && (DebugRomAll || o.m_romIndex == DebugRomIndex))
-			UE_LOG(LogMutableCore, Log, TEXT("CodeRunner::CompleteRomLoadOp for rom %d."), o.m_romIndex);
-
-		m_pSystem->StreamInterface->EndRead(o.m_streamID);
-
-		FProgram& program = m_pModel->GetPrivate()->m_program;
-		{
-			MUTABLE_CPUPROFILER_SCOPE(Unserialise);
-
-			InputMemoryStream stream(o.m_streamBuffer.GetData(), o.m_streamBuffer.Num());
-			InputArchive arch(&stream);
-
-			int32 ResIndex = program.m_roms[o.m_romIndex].ResourceIndex;
-			switch (o.ConstantType)
-			{
-			case DATATYPE::DT_MESH:
-			{
-				check(!program.m_constantMeshes[ResIndex].Value);
-				Ptr<Mesh> Value = Mesh::StaticUnserialise(arch);
-				program.m_constantMeshes[ResIndex].Value = Value;
-				check(program.m_constantMeshes[ResIndex].Value);
-				m_pSystem->WorkingMemoryManager.TrackedBudgetBytes_Rom += Value->GetDataSize();
-				break;
-			}
-			case DATATYPE::DT_IMAGE:
-			{
-				check(!program.m_constantImageLODs[ResIndex].Value);
-				// TODO: Try to reuse buffer from PooledImages.
-				Ptr<Image> Value = Image::StaticUnserialise(arch);
-				program.m_constantImageLODs[ResIndex].Value = Value;
-				check(program.m_constantImageLODs[ResIndex].Value);
-				m_pSystem->WorkingMemoryManager.TrackedBudgetBytes_Rom += Value->GetDataSize();
-				break;
-			}
-			default:
-				check(false);
-				break;
-			}
-
-			m_pSystem->WorkingMemoryManager.TrackedBudgetBytes_Stream -= o.m_streamBuffer.GetAllocatedSize();
-
-			o.m_romIndex = -1;
-			o.m_streamBuffer.Empty();
-		}
-	}
-
 	//---------------------------------------------------------------------------------------------
 	TRACE_DECLARE_INT_COUNTER(MutableRuntime_OpenTask, TEXT("MutableRuntime/OpenTask"));
 	TRACE_DECLARE_INT_COUNTER(MutableRuntime_ClosedTasks, TEXT("MutableRuntime/ClosedTasks"));
@@ -529,15 +481,6 @@ namespace mu
 				}
 			}
 
-			// Look for completed streaming ops and complete the rom loading
-			for (FRomLoadOp& o : m_romLoadOps)
-			{
-				if (o.m_romIndex>=0 && m_pSystem->StreamInterface->IsReadCompleted(o.m_streamID))
-				{
-					CompleteRomLoadOp(o);
-				}
-			}
-
 			// Look for a closed task with dependencies satisfied and move them to the open task list.
 			bool bSomeWasReady = false;
 			for (int Index = 0; Index<ClosedTasks.Num(); )
@@ -613,19 +556,6 @@ namespace mu
 					}
 				}
 
-				// If we reached here it means we didn't find an op to wait for. Try to wait for a loading op.
-				// \todo: unify loading ops with normal ones?
-				for (FRomLoadOp& o : m_romLoadOps)
-				{
-					if (o.m_romIndex >= 0)
-					{
-						CompleteRomLoadOp(o);
-						break;
-					}
-
-					// We should never reach this, since it would mean we deadlocked.
-					//check(false);
-				}
 			}
 		}
 
@@ -662,6 +592,7 @@ namespace mu
 		}
     }
 
+	
 	//---------------------------------------------------------------------------------------------
 	void CodeRunner::GetImageDescResult(FImageDesc& OutDesc)
 	{
@@ -2038,91 +1969,120 @@ namespace mu
 	//---------------------------------------------------------------------------------------------
 	//---------------------------------------------------------------------------------------------
 	//---------------------------------------------------------------------------------------------
-	bool CodeRunner::FLoadMeshRomTask::Prepare(CodeRunner* Runner, bool& bOutFailed )
+	bool CodeRunner::FLoadMeshRomTask::Prepare(CodeRunner* Runner, bool& bOutFailed)
 	{
+		// This runs in the mutable Runner thread
+		MUTABLE_CPUPROFILER_SCOPE(FLoadMeshRomTask_Prepare);
+
 		if (!Runner || !Runner->m_pSystem)
 		{
 			return false;
 		}
 
-		// This runs in the mutable Runner thread
-		MUTABLE_CPUPROFILER_SCOPE(FLoadMeshRomTask_Prepare);
 		bOutFailed = false;
 
-		FProgram& program = Runner->m_pModel->GetPrivate()->m_program;
-
-		check(RomIndex < program.m_roms.Num());
-		bool bRomIsLoaded = program.IsRomLoaded(RomIndex);
+		const FProgram& Program = Runner->m_pModel->GetPrivate()->m_program;
 
+		check(RomIndex < Program.m_roms.Num());
+		
 		FWorkingMemoryManager::FModelCacheEntry* ModelCache = Runner->m_pSystem->WorkingMemoryManager.FindModelCache(Runner->m_pModel.Get());
-		bool bRomHasPendingOps = ModelCache->PendingOpsPerRom[RomIndex]!=0;
-		if (!bRomIsLoaded && !bRomHasPendingOps)
-		{
-			check(Runner->m_pSystem->StreamInterface);
+		++ModelCache->PendingOpsPerRom[RomIndex];
 
-			uint32 RomSize = program.m_roms[RomIndex].Size;
-			check(RomSize > 0);
+		if (Program.IsRomLoaded(RomIndex))
+		{
+			return false;
+		}
 
-			// Free roms if necessary
-			{
-				MUTABLE_CPUPROFILER_SCOPE(FreeingRoms);
+		if (const TSharedPtr<FRomLoadOp>* Result = Runner->RomLoadOps.Find(RomIndex))
+		{
+			Event = (*Result)->Event;  // Wait for the read operation started by other task
+			return false;
+		}
 
-				Runner->m_pSystem->WorkingMemoryManager.MarkRomUsed(RomIndex, Runner->m_pModel);
-				Runner->m_pSystem->WorkingMemoryManager.EnsureBudgetBelow(RomSize);
-			}
+		const TSharedPtr<FRomLoadOp> RomLoadOp = MakeShared<FRomLoadOp>();
+		Runner->RomLoadOps.Add(RomIndex, RomLoadOp);
+		
+		check(Runner->m_pSystem->StreamInterface);
 
-			CodeRunner::FRomLoadOp* op = nullptr;
-			for (FRomLoadOp& o : Runner->m_romLoadOps)
-			{
-				if (o.m_romIndex < 0)
-				{
-					op = &o;
-				}
-			}
+		const uint32 RomSize = Program.m_roms[RomIndex].Size;
+		check(RomSize > 0);
 
-			if (!op)
-			{
-				Runner->m_romLoadOps.Add(CodeRunner::FRomLoadOp());
-				op = &Runner->m_romLoadOps.Last();
-			}
+		// Free roms if necessary
+		{
+			MUTABLE_CPUPROFILER_SCOPE(FreeingRoms);
 
-			op->m_romIndex = RomIndex;
-			op->ConstantType = DT_MESH;
+			Runner->m_pSystem->WorkingMemoryManager.MarkRomUsed(RomIndex, Runner->m_pModel);
+			Runner->m_pSystem->WorkingMemoryManager.EnsureBudgetBelow(RomSize);
+		}
 
-			if (uint32(op->m_streamBuffer.Num()) < RomSize)
-			{
-				op->m_streamBuffer.SetNum((size_t)RomSize);
-				Runner->m_pSystem->WorkingMemoryManager.TrackedBudgetBytes_Stream += op->m_streamBuffer.GetAllocatedSize();
-			}
+		RomLoadOp->m_streamBuffer.SetNum(RomSize);
 
-			uint32 RomId = program.m_roms[RomIndex].Id;
-			op->m_streamID = Runner->m_pSystem->StreamInterface->BeginReadBlock(Runner->m_pModel.Get(), RomId, op->m_streamBuffer.GetData(), RomSize);
-			if (op->m_streamID < 0)
-			{
-				bOutFailed = true;
-				return false;
-			}
+		EventType EventTyped =
+#ifdef MUTABLE_USE_NEW_TASKGRAPH
+			UE::Tasks::FTaskEvent(TEXT("FLoadMeshRomTask"));
+#else
+			FGraphEvent::CreateGraphEvent();		
+#endif
+		
+		RomLoadOp->Event = EventTyped;
+		
+		TFunction<void(bool)> Callback = [EventTyped](bool bSuccess) mutable
+		{
+#ifdef MUTABLE_USE_NEW_TASKGRAPH
+			EventTyped.Trigger();
+#else
+			EventTyped->DispatchSubsequents();
+#endif				
+		};
+		
+		const uint32 RomId = Program.m_roms[RomIndex].Id;
+		RomLoadOp->m_streamID = Runner->m_pSystem->StreamInterface->BeginReadBlock(Runner->m_pModel.Get(), RomId, RomLoadOp->m_streamBuffer.GetData(), RomSize, &Callback);
+		if (RomLoadOp->m_streamID < 0)
+		{
+			bOutFailed = true;
+			return false;
 		}
-		++ModelCache->PendingOpsPerRom[RomIndex];
 
-		//UE_LOG(LogMutableCore, Log, TEXT("FLoadMeshRomTask::Prepare romindex %d pending %d."), RomAt.m_rom, Runner->m_romPendingOps[RomAt.m_rom]);
-
-		// No worker thread work
-		return false;
+		Event = EventTyped; // Wait for read operation to end
+		return false; // No worker thread work
 	}
-
+	
 
 	//---------------------------------------------------------------------------------------------
 	void CodeRunner::FLoadMeshRomTask::Complete(CodeRunner* Runner)
 	{
+		// This runs in the Runner thread
+		MUTABLE_CPUPROFILER_SCOPE(FLoadMeshRomTask_Complete);
+
 		if (!Runner || !Runner->m_pSystem)
 		{
 			return;
 		}
 
-		MUTABLE_CPUPROFILER_SCOPE(FLoadMeshRomTask_Complete);
+		FProgram& Program = Runner->m_pModel->GetPrivate()->m_program;
 
-		// This runs in the Runner thread
+		// Since task could be reordered, we need to make sure we end the rom read before continuing
+		if (const TSharedPtr<FRomLoadOp>* Result = Runner->RomLoadOps.Find(RomIndex))
+		{
+			const TSharedPtr<FRomLoadOp>& RomLoadOp = *Result;
+
+			Runner->m_pSystem->StreamInterface->EndRead(RomLoadOp->m_streamID);									
+
+			const int32 ResIndex = Program.m_roms[RomIndex].ResourceIndex;
+			check(!Program.m_constantMeshes[ResIndex].Value)
+			MUTABLE_CPUPROFILER_SCOPE(Unserialise);
+
+			InputMemoryStream stream(RomLoadOp->m_streamBuffer.GetData(), RomLoadOp->m_streamBuffer.Num());
+			InputArchive arch(&stream);
+
+			check(!Program.m_constantMeshes[ResIndex].Value);
+			Ptr<Mesh> Value = Mesh::StaticUnserialise(arch);
+			Program.m_constantMeshes[ResIndex].Value = Value;
+			check(Program.m_constantMeshes[ResIndex].Value);
+			Runner->m_pSystem->WorkingMemoryManager.TrackedBudgetBytes_Rom += Value->GetDataSize();
+
+			Runner->RomLoadOps.Remove(RomIndex);
+		}
 
 		// Process the constant op normally, now that the rom is loaded.
 		Runner->RunCode(Op, Runner->m_pParams, Runner->m_pModel.Get(), Runner->m_lodMask);
@@ -2134,14 +2094,6 @@ namespace mu
 	}
 
 
-	//---------------------------------------------------------------------------------------------
-	bool CodeRunner::FLoadMeshRomTask::IsComplete(CodeRunner* Runner)
-	{ 
-		FProgram& program = Runner->m_pModel->GetPrivate()->m_program;
-		return program.IsRomLoaded(RomIndex);
-	}
-
-
 	//---------------------------------------------------------------------------------------------
 	//---------------------------------------------------------------------------------------------
 	//---------------------------------------------------------------------------------------------
@@ -2242,10 +2194,18 @@ namespace mu
 
 		FWorkingMemoryManager::FModelCacheEntry* ModelCache = Runner->m_pSystem->WorkingMemoryManager.FindModelCache(Runner->m_pModel.Get());
 
-		for (int32 i = 0; i<LODIndexCount; ++i )
+#ifdef MUTABLE_USE_NEW_TASKGRAPH
+		TArray<UE::Tasks::FTask> ReadCompleteEvents;
+#else
+		FGraphEventArray ReadCompleteEvents;
+#endif
+
+		ReadCompleteEvents.Reserve(LODIndexCount); 
+
+		for (int32 LODIndex = 0; LODIndex < LODIndexCount; ++LODIndex)
 		{
-			int32 CurrentIndexIndex = LODIndexIndex + i;
-			int32 CurrentIndex = program.m_constantImageLODIndices[CurrentIndexIndex];
+			const int32 CurrentIndexIndex = LODIndexIndex + LODIndex;
+			const int32 CurrentIndex = program.m_constantImageLODIndices[CurrentIndexIndex];
 
 			if (program.m_constantImageLODs[CurrentIndex].Key<0)
 			{
@@ -2256,21 +2216,30 @@ namespace mu
 			int32 RomIndex = program.m_constantImageLODs[CurrentIndex].Key;
 			check(RomIndex < program.m_roms.Num());
 
-			bool bRomIsLoaded = program.IsRomLoaded(RomIndex);
-			bool bRomHasAlreadyBeenRequested = ModelCache->PendingOpsPerRom[RomIndex] != 0;
 			++ModelCache->PendingOpsPerRom[RomIndex];
 
 			if (DebugRom && (DebugRomAll || RomIndex == DebugRomIndex))
 				UE_LOG(LogMutableCore, Log, TEXT("Preparing rom %d, now peding ops is %d."), RomIndex, ModelCache->PendingOpsPerRom[RomIndex]);
 
-			if (bRomIsLoaded || bRomHasAlreadyBeenRequested)
+			if (program.IsRomLoaded(RomIndex))
 			{
 				continue;
 			}
 
+			RomIndices.Add(RomIndex);
+
+			if (const TSharedPtr<FRomLoadOp>* Result = Runner->RomLoadOps.Find(RomIndex))
+			{
+				ReadCompleteEvents.Add((*Result)->Event); // Wait for the read operation started by other task
+				continue;
+			}
+
+			TSharedPtr<FRomLoadOp> RomLoadOp = MakeShared<FRomLoadOp>();
+			Runner->RomLoadOps.Add(RomIndex, RomLoadOp);
+			
 			check(Runner->m_pSystem->StreamInterface);
 
-			uint32 RomSize = program.m_roms[RomIndex].Size;
+			const uint32 RomSize = program.m_roms[RomIndex].Size;
 			check(RomSize > 0);
 
 			// Free roms if necessary
@@ -2279,113 +2248,110 @@ namespace mu
 				Runner->m_pSystem->WorkingMemoryManager.EnsureBudgetBelow(RomSize);
 			}
 
-			CodeRunner::FRomLoadOp* op = nullptr;
-			for (FRomLoadOp& o : Runner->m_romLoadOps)
-			{
-				if (o.m_romIndex < 0)
-				{
-					op = &o;
-				}
-			}
+			RomLoadOp->m_streamBuffer.SetNum(RomSize);
 
-			if (!op)
-			{
-				Runner->m_romLoadOps.Add(CodeRunner::FRomLoadOp());
-				op = &Runner->m_romLoadOps.Last();
-			}
+			EventType EventTyped = 
+#ifdef MUTABLE_USE_NEW_TASKGRAPH
+				UE::Tasks::FTaskEvent(TEXT("FLoadImageRomsTaskRom"));
+#else
+				FGraphEvent::CreateGraphEvent();
+#endif
 
-			op->m_romIndex = RomIndex;
-			op->ConstantType = DT_IMAGE;
+			RomLoadOp->Event = EventTyped;
+			ReadCompleteEvents.Add(EventTyped);
 
-			if (uint32(op->m_streamBuffer.Num()) < RomSize)
+			TFunction<void(bool)> Callback = [EventTyped](bool bSuccess) mutable // Mutable due Trigger not being const
 			{
-				op->m_streamBuffer.SetNum(RomSize);
-				Runner->m_pSystem->WorkingMemoryManager.TrackedBudgetBytes_Stream += op->m_streamBuffer.GetAllocatedSize();
-			}
-
-			uint32 RomId = program.m_roms[RomIndex].Id;
-			op->m_streamID = Runner->m_pSystem->StreamInterface->BeginReadBlock(Runner->m_pModel.Get(), RomId, op->m_streamBuffer.GetData(), RomSize);
-			if (op->m_streamID < 0)
+#ifdef MUTABLE_USE_NEW_TASKGRAPH
+				EventTyped.Trigger();
+#else
+				EventTyped->DispatchSubsequents();
+#endif				
+			};
+			
+			const uint32 RomId = program.m_roms[RomIndex].Id;
+			RomLoadOp->m_streamID = Runner->m_pSystem->StreamInterface->BeginReadBlock(Runner->m_pModel.Get(), RomId, RomLoadOp->m_streamBuffer.GetData(), RomSize, &Callback);
+			if (RomLoadOp->m_streamID < 0)
 			{
 				bOutFailed = true;
 				return false;
 			}
 		}
 
-		//UE_LOG(LogMutableCore, Log, TEXT("FLoadImageRomsTask::Prepare romindex %d pending %d."), RomAt.m_rom, Runner->m_romPendingOps[RomAt.m_rom]);
-
-		// No worker thread work
-		return false;
+		// Wait for all read operations to end
+#ifdef MUTABLE_USE_NEW_TASKGRAPH
+		UE::Tasks::FTaskEvent EventTyped(TEXT("FLoadImageRomsTask"));
+		EventTyped.AddPrerequisites(ReadCompleteEvents);
+		EventTyped.Trigger();
+		Event = EventTyped;
+#else
+		Event = FGraphEvent::CreateGraphEvent();
+		for (const FGraphEventRef& ReadCompleteEvent : ReadCompleteEvents)
+		{
+			Event->DontCompleteUntil(ReadCompleteEvent);
+		}
+		Event->DispatchSubsequents();
+#endif
+			
+		return false; // No worker thread work
 	}
-
-
+	
+	
 	//---------------------------------------------------------------------------------------------
 	void CodeRunner::FLoadImageRomsTask::Complete(CodeRunner* Runner)
 	{
+		// This runs in the Runner thread
+		MUTABLE_CPUPROFILER_SCOPE(FLoadImageRomsTask_Complete);
+
 		if (!Runner || !Runner->m_pSystem)
 		{
 			return;
 		}
 
-		MUTABLE_CPUPROFILER_SCOPE(FLoadImageRomsTask_Complete);
-
-		// This runs in the Runner thread
-
+		FProgram& program = Runner->m_pModel->GetPrivate()->m_program;
+		
 		FWorkingMemoryManager::FModelCacheEntry* ModelCache = Runner->m_pSystem->WorkingMemoryManager.FindModelCache(Runner->m_pModel.Get());
 
-		// Process the constant op normally, now that the rom is loaded.
-		Runner->RunCode(Op, Runner->m_pParams, Runner->m_pModel.Get(), Runner->m_lodMask);
-
-		FProgram& program = Runner->m_pModel->GetPrivate()->m_program;
-		for (int32 i = 0; i < LODIndexCount; ++i)
+		for (const int32 RomIndex : RomIndices)
 		{
-			int32 CurrentIndexIndex = LODIndexIndex + i;
-			int32 CurrentIndex = program.m_constantImageLODIndices[CurrentIndexIndex];
-
-			if (program.m_constantImageLODs[CurrentIndex].Key < 0)
+			// Since task could be reordered, we need to make sure we end the rom read before continuing
+			if (const TSharedPtr<FRomLoadOp>* Result = Runner->RomLoadOps.Find(RomIndex))
 			{
-				// This data is always resident.
-				continue;
-			}
+				const TSharedPtr<FRomLoadOp>& RomLoadOp = *Result;
+				
+				Runner->m_pSystem->StreamInterface->EndRead(RomLoadOp->m_streamID);
 
-			int32 RomIndex = program.m_constantImageLODs[CurrentIndex].Key;
-			check(RomIndex < program.m_roms.Num());
+				MUTABLE_CPUPROFILER_SCOPE(Unserialise);
 
-			Runner->m_pSystem->WorkingMemoryManager.MarkRomUsed(RomIndex, Runner->m_pModel);
-			--ModelCache->PendingOpsPerRom[RomIndex];
+				InputMemoryStream stream(RomLoadOp->m_streamBuffer.GetData(), RomLoadOp->m_streamBuffer.Num());
+				InputArchive arch(&stream);
 
-			if (DebugRom && (DebugRomAll || RomIndex == DebugRomIndex))
-			{
-				UE_LOG(LogMutableCore, Log, TEXT("FLoadImageRomsTask::Complete rom %d, now peding ops is %d."), RomIndex, ModelCache->PendingOpsPerRom[RomIndex]);
+				const int32 ResIndex = program.m_roms[RomIndex].ResourceIndex;
+
+				// TODO: Try to reuse buffer from PooledImages.
+				check(!program.m_constantImageLODs[ResIndex].Value);
+				Ptr<Image> Value = Image::StaticUnserialise(arch);
+				program.m_constantImageLODs[ResIndex].Value = Value;
+				check(program.m_constantImageLODs[ResIndex].Value);
+				Runner->m_pSystem->WorkingMemoryManager.TrackedBudgetBytes_Rom += Value->GetDataSize();
+				
+				Runner->RomLoadOps.Remove(RomIndex);
 			}
 		}
-	}
-
+		
+		// Process the constant op normally, now that the rom is loaded.	
+		Runner->RunCode(Op, Runner->m_pParams, Runner->m_pModel.Get(), Runner->m_lodMask);
 
-	//---------------------------------------------------------------------------------------------
-	bool CodeRunner::FLoadImageRomsTask::IsComplete(CodeRunner* Runner)
-	{
-		FProgram& program = Runner->m_pModel->GetPrivate()->m_program;
-		for (int32 i = 0; i < LODIndexCount; ++i)
+		for (const int32 RomIndex : RomIndices)
 		{
-			int32 CurrentIndexIndex = LODIndexIndex + i;
-			int32 CurrentIndex = program.m_constantImageLODIndices[CurrentIndexIndex];
-
-			if (program.m_constantImageLODs[CurrentIndex].Key < 0)
-			{
-				// This data is always resident.
-				continue;
-			}
+			Runner->m_pSystem->WorkingMemoryManager.MarkRomUsed(RomIndex, Runner->m_pModel);
+			--ModelCache->PendingOpsPerRom[RomIndex];
 
-			int32 RomIndex = program.m_constantImageLODs[CurrentIndex].Key;
-			check(RomIndex < program.m_roms.Num());
-			if (!program.IsRomLoaded(RomIndex))
+			if (DebugRom && (DebugRomAll || RomIndex == DebugRomIndex))
 			{
-				return false;
+				UE_LOG(LogMutableCore, Log, TEXT("FLoadImageRomsTask::Complete rom %d, now peding ops is %d."), RomIndex, ModelCache->PendingOpsPerRom[RomIndex]);
 			}
 		}
-
-		return true;
 	}
 
 
diff --git a/Engine/Plugins/Experimental/Mutable/Source/MutableRuntime/Private/MuR/System.cpp b/Engine/Plugins/Experimental/Mutable/Source/MutableRuntime/Private/MuR/System.cpp
--- a/Engine/Plugins/Experimental/Mutable/Source/MutableRuntime/Private/MuR/System.cpp
+++ b/Engine/Plugins/Experimental/Mutable/Source/MutableRuntime/Private/MuR/System.cpp
@@ -915,12 +915,7 @@ namespace mu
 
 		if (CurrentRunner)
 		{
-			// Data reserved for ongoing streaming operations
-			for ( const CodeRunner::FRomLoadOp& Op : CurrentRunner->m_romLoadOps )
-			{
-				StreamingBufferBytes += Op.m_streamBuffer.GetAllocatedSize();
-			}
-
+			StreamingBufferBytes += CurrentRunner->GetStreamingMemoryBytes();
 			InternalBytes += CurrentRunner->GetInternalMemoryBytes();
 		}
 
@@ -1034,7 +1029,8 @@ namespace mu
 		FModelCacheEntry n;
         n.Model = TWeakPtr<const Model>(InModel);
 		n.PendingOpsPerRom.SetNum(InModel->GetPrivate()->m_program.m_roms.Num());
-        CachePerModel.Add(n);
+    	n.RomWeights.SetNum(InModel->GetPrivate()->m_program.m_roms.Num());
+    	CachePerModel.Add(n);
         return CachePerModel.Last();
     }
 
@@ -1465,8 +1461,6 @@ namespace mu
     {
         check(pModel);
 
-        FProgram& Program = pModel->GetPrivate()->m_program;
-
         // If budget is zero, we don't unload anything here, and we assume it is managed somewhere else.
         if (!BudgetBytes)
         {
@@ -1478,12 +1472,7 @@ namespace mu
         // Update current cache
         {
 			FModelCacheEntry* ModelCache = FindModelCache(pModel.Get());
-
-            if (ModelCache->RomWeights.Num()<Program.m_roms.Num())
-            {
-				ModelCache->RomWeights.SetNumZeroed(Program.m_roms.Num());
-            }
-
+        	
             ModelCache->RomWeights[romIndex].Key++;
             ModelCache->RomWeights[romIndex].Value = RomTick;
         }
diff --git a/Engine/Plugins/Experimental/Mutable/Source/MutableRuntime/Public/MuR/Serialisation.h b/Engine/Plugins/Experimental/Mutable/Source/MutableRuntime/Public/MuR/Serialisation.h
--- a/Engine/Plugins/Experimental/Mutable/Source/MutableRuntime/Public/MuR/Serialisation.h
+++ b/Engine/Plugins/Experimental/Mutable/Source/MutableRuntime/Public/MuR/Serialisation.h
@@ -99,10 +99,11 @@ namespace mu
 		//! \param pBuffer is an already-allocated buffer big enough to receive the expected data.
 		//! \param size is the size of the pBuffer buffer, which must match the size of the data
 		//! requested with the key identifiers.
+		//! \param CompletionCallback Optional callback. Copied inside the called function. Will always be called.
 		//! \return a previously unused identifier, now used for this operation, that can be used in
 		//! calls to the other methods of this interface. If the return value is negative it indicates
 		//! an unrecoverable error.
-		virtual OPERATION_ID BeginReadBlock(const mu::Model*, uint64 key0, void* pBuffer, uint64 size) = 0;
+		virtual OPERATION_ID BeginReadBlock(const mu::Model*, uint64 key0, void* pBuffer, uint64 size, TFunction<void(bool bSuccess)>* CompletionCallback = nullptr) = 0;
 
         //! Check if a data request operation has been completed.
         //! This is a weak check than *may* return true if the given operation has completed, but
