[Mutable] Fix extra UV channels being ignored depending on connection order.

#rnx

diff --git a/Engine/Plugins/Experimental/Mutable/Source/MutableRuntime/Private/MuR/OpMeshMerge.h b/Engine/Plugins/Experimental/Mutable/Source/MutableRuntime/Private/MuR/OpMeshMerge.h
--- a/Engine/Plugins/Experimental/Mutable/Source/MutableRuntime/Private/MuR/OpMeshMerge.h
+++ b/Engine/Plugins/Experimental/Mutable/Source/MutableRuntime/Private/MuR/OpMeshMerge.h
@@ -868,35 +868,63 @@ namespace mu
 					}
 				}
 
-                // See if we need to add additional buffers from the second mesh (like vertex colours)
+                // See if we need to add additional buffers from the second mesh (like vertex colours or additional UV Channels)
                 // This is a bit ad-hoc: we only add buffers containing all new channels
                 for ( const auto& buf : pSecond->GetVertexBuffers().m_buffers )
                 {
                     bool someChannel = false;
                     bool allNewChannels = true;
-                    for ( const auto& chan : buf.m_channels )
-                    {
-                        // Not a system buffer
-                        if ( chan.m_semantic!=MBS_VERTEXINDEX
-                             &&
-                             chan.m_semantic!=MBS_LAYOUTBLOCK
-                             &&
-                             chan.m_semantic!=MBS_CHART
-                             )
-                        {
-                            someChannel = true;
-
-                            int foundBuffer=-1;
-                            int foundChannel=-1;
-                            pFirst->GetVertexBuffers().FindChannel(chan.m_semantic, chan.m_semanticIndex,&foundBuffer,&foundChannel);
-                            if (foundBuffer>=0)
-                            {
-                                // There is a meaningful channel that we don't have in the first
-                                // mesh, we'll need to add the buffer.
-                                allNewChannels = false;
-                            }
-                        }
-                    }
+					for (const auto& chan : buf.m_channels)
+					{
+						// Skip system buffers
+						if (chan.m_semantic == MBS_VERTEXINDEX
+							&&
+							chan.m_semantic == MBS_LAYOUTBLOCK
+							&&
+							chan.m_semantic == MBS_CHART)
+						{
+							continue;
+						}
+
+						someChannel = true;
+
+						int foundBuffer = -1;
+						int foundChannel = -1;
+						pFirst->GetVertexBuffers().FindChannel(chan.m_semantic, chan.m_semanticIndex, &foundBuffer, &foundChannel);
+						if (foundBuffer >= 0)
+						{
+							// There's at least one channel that already exists in the first mesh. Don't add the buffer.
+							allNewChannels = false;
+							continue;
+						}
+						
+						// If there are additional UV channels try to add them
+						if(!allNewChannels && chan.m_semantic == MBS_TEXCOORDS
+							&&
+							chan.m_semanticIndex > 0)
+						{
+							// Add additional UV channels if the previous one is found.
+							FMeshBufferSet& VertexBuffers = Result->GetVertexBuffers();
+							VertexBuffers.FindChannel(MBS_TEXCOORDS, chan.m_semanticIndex - 1, &foundBuffer, &foundChannel);
+
+							if (foundBuffer >= 0)
+							{
+								MESH_BUFFER& Buffer = VertexBuffers.m_buffers[foundBuffer];
+								Buffer.m_channels.Insert(chan, foundChannel + 1);
+
+								// Update offsets
+								int32 Offset = Buffer.m_channels[foundChannel].m_offset;
+								for (int32 c = foundChannel; c < Buffer.m_channels.Num(); ++c)
+								{
+									Buffer.m_channels[c].m_offset = (uint8)Offset;
+									Offset += Buffer.m_channels[c].m_componentCount
+										*
+										GetMeshFormatData(Buffer.m_channels[c].m_format).m_size;
+								}
+								Buffer.m_elementSize = Offset;
+							}
+						}
+					}
 
                     if (someChannel && allNewChannels)
                     {
