[Mutable] Fix typedef collision on some platforms.

diff --git a/Engine/Plugins/Experimental/Mutable/Source/MutableRuntime/Private/MuR/CodeRunner.h b/Engine/Plugins/Experimental/Mutable/Source/MutableRuntime/Private/MuR/CodeRunner.h
--- a/Engine/Plugins/Experimental/Mutable/Source/MutableRuntime/Private/MuR/CodeRunner.h
+++ b/Engine/Plugins/Experimental/Mutable/Source/MutableRuntime/Private/MuR/CodeRunner.h
@@ -30,28 +30,26 @@
 	#include "Async/TaskGraphFwd.h"
 #endif
 
+namespace  mu
+{
+	class Model;
+	class Parameters;
+	class RangeIndex;
 
-using EventType = 
+	using EventType = 
 #ifdef MUTABLE_USE_NEW_TASKGRAPH
-	UE::Tasks::FTaskEvent;
+		UE::Tasks::FTaskEvent;
 #else
-	FGraphEventRef;
+		FGraphEventRef;
 #endif
 
-using TaskType = 
+	using TaskType = 
 #ifdef MUTABLE_USE_NEW_TASKGRAPH
-	UE::Tasks::FTask;
+		UE::Tasks::FTask;
 #else
-	FGraphEventRef;
+		FGraphEventRef;
 #endif
 
-namespace  mu
-{
-	class Model;
-	class Parameters;
-	class RangeIndex;
-
-
     /** Code execution of the mutable virtual machine. */
     class CodeRunner
     {
