[Mutable] Fix incorrect implicit type conversion which resulted in released textures being kept in the generated textures cache.

#rnx

diff --git a/Engine/Plugins/Experimental/Mutable/Source/CustomizableObject/Private/MuCO/CustomizableInstancePrivateData.h b/Engine/Plugins/Experimental/Mutable/Source/CustomizableObject/Private/MuCO/CustomizableInstancePrivateData.h
--- a/Engine/Plugins/Experimental/Mutable/Source/CustomizableObject/Private/MuCO/CustomizableInstancePrivateData.h
+++ b/Engine/Plugins/Experimental/Mutable/Source/CustomizableObject/Private/MuCO/CustomizableInstancePrivateData.h
@@ -15,6 +15,7 @@
 namespace mu { class PhysicsBody; }
 struct FMutableModelImageProperties;
 struct FMutableRefSkeletalMeshData;
+struct FMutableImageCacheKey;
 struct FStreamableHandle;
 
 /** CustomizableObject Instance flags for internal use  */
@@ -54,13 +55,12 @@ struct FGeneratedTexture
 {
 	GENERATED_USTRUCT_BODY();
 
-	UPROPERTY(Category = CustomizableObjectInstance, EditAnywhere)
-	int32 Id = 0;
+	FMutableImageCacheKey Key;
 
-	UPROPERTY(Category = CustomizableObjectInstance, EditAnywhere) 
+	UPROPERTY(Category = CustomizableObjectInstance, VisibleAnywhere)
 	FString Name;
 
-	UPROPERTY(Category = CustomizableObjectInstance, EditAnywhere) 
+	UPROPERTY(Category = CustomizableObjectInstance, VisibleAnywhere)
 	TObjectPtr<UTexture> Texture = nullptr;
 };
 
@@ -233,7 +233,7 @@ public:
 	// Returns true if success (?)
 	bool UpdateSkeletalMesh_PostBeginUpdate0(UCustomizableObjectInstance* Public, const TSharedPtr<FMutableOperationData>& OperationData);
 
-	static void ReleaseMutableTexture(int32 MutableTextureId, UTexture2D* Texture, struct FMutableResourceCache& Cache);
+	static void ReleaseMutableTexture(const FMutableImageCacheKey& MutableTextureKey, UTexture2D* Texture, struct FMutableResourceCache& Cache);
 
 	// Parameter decoration generation: game thread step
 	void UpdateParameterDecorationsEngineResources(const TSharedPtr<FMutableOperationData>& OperationData);
@@ -322,7 +322,7 @@ public:
 	// Converts a ReferencedMaterials index from the CustomizableObject to an index in the ReferencedMaterials in the Instance
 	TMap<uint32, uint32> ObjectToInstanceIndexMap;
 
-	TMap<uint32, FGeneratedTexture> TexturesToRelease;
+	TArray<FGeneratedTexture> TexturesToRelease;
 
 	UPROPERTY(Transient)
 	TArray< TObjectPtr<UPhysicsAsset> > ClothingPhysicsAssets;
diff --git a/Engine/Plugins/Experimental/Mutable/Source/CustomizableObject/Private/MuCO/CustomizableObjectInstance.cpp b/Engine/Plugins/Experimental/Mutable/Source/CustomizableObject/Private/MuCO/CustomizableObjectInstance.cpp
--- a/Engine/Plugins/Experimental/Mutable/Source/CustomizableObject/Private/MuCO/CustomizableObjectInstance.cpp
+++ b/Engine/Plugins/Experimental/Mutable/Source/CustomizableObject/Private/MuCO/CustomizableObjectInstance.cpp
@@ -351,13 +351,13 @@ void UCustomizableObjectInstance::ReleaseMutableResources(bool bCalledFromBeginD
 
 		for (FGeneratedTexture& Texture : PrivateData->GeneratedTextures)
 		{
-			if (CustomizableObjectSystem->RemoveTextureReference(Texture.Id))
+			if (CustomizableObjectSystem->RemoveTextureReference(Texture.Key))
 			{
 				// Do not release textures when called from BeginDestroy, it would produce a texture artifact in the
 				// instance's remaining sk meshes and GC is being performed anyway so it will free the textures if needed
 				if (!bCalledFromBeginDestroy && CustomizableObjectSystem->bReleaseTexturesImmediately)
 				{
-					UCustomizableInstancePrivateData::ReleaseMutableTexture(Texture.Id, Cast<UTexture2D>(Texture.Texture), Cache);
+					UCustomizableInstancePrivateData::ReleaseMutableTexture(Texture.Key, Cast<UTexture2D>(Texture.Texture), Cache);
 				}
 			}
 		}
@@ -508,7 +508,7 @@ bool UCustomizableObjectInstance::IsParamMultidimensional(const FString& ParamNa
 
 
 // Only safe to call if the Mutable texture ref count system returns 0 and absolutely sure nobody holds a reference to the texture
-void UCustomizableInstancePrivateData::ReleaseMutableTexture(int32 MutableTextureId, UTexture2D* Texture, FMutableResourceCache& Cache)
+void UCustomizableInstancePrivateData::ReleaseMutableTexture(const FMutableImageCacheKey& MutableTextureKey, UTexture2D* Texture, FMutableResourceCache& Cache)
 {
 	if (ensure(Texture) && Texture->IsValidLowLevel())
 	{
@@ -521,14 +521,7 @@ void UCustomizableInstancePrivateData::ReleaseMutableTexture(int32 MutableTextur
 	}
 
 	// Must remove texture from cache since it has been released
-	for (TPair<FMutableImageCacheKey, TWeakObjectPtr<UTexture2D> >& CachedTexture : Cache.Images)
-	{
-		if (CachedTexture.Key.Resource == MutableTextureId)
-		{
-			Cache.Images.Remove(CachedTexture.Key);
-			break;
-		}
-	}
+	Cache.Images.Remove(MutableTextureKey);
 }
 
 
@@ -5885,7 +5878,7 @@ void UCustomizableInstancePrivateData::BuildMaterials(const TSharedPtr<FMutableO
 							}
 
 							FGeneratedTexture TextureData;
-							TextureData.Id = Image.ImageID;
+							TextureData.Key = ImageCacheKey;
 							TextureData.Name = Props.TextureParameterName;
 							TextureData.Texture = MutableTexture ? MutableTexture : PassThroughTexture;
 							
@@ -6076,24 +6069,19 @@ void UCustomizableInstancePrivateData::BuildMaterials(const TSharedPtr<FMutableO
 		FCustomizableObjectSystemPrivate* CustomizableObjectSystem = UCustomizableObjectSystem::GetInstance()->GetPrivate();
 		TexturesToRelease.Empty();
 
-		for (FGeneratedTexture& Texture : GeneratedTextures)
+		for (const FGeneratedTexture& Texture : NewGeneratedTextures)
 		{
-			if (CustomizableObjectSystem->RemoveTextureReference(Texture.Id))
-			{
-				if (CustomizableObjectSystem->bReleaseTexturesImmediately)
-				{
-					TexturesToRelease.Add(Texture.Id, Texture); // Texture count is zero, so prepare to release it
-				}
-			}
+			CustomizableObjectSystem->AddTextureReference(Texture.Key);
 		}
 
-		for (FGeneratedTexture& Texture : NewGeneratedTextures)
+		for (const FGeneratedTexture& Texture : GeneratedTextures)
 		{
-			CustomizableObjectSystem->AddTextureReference(Texture.Id);
-
-			if (CustomizableObjectSystem->bReleaseTexturesImmediately)
+			if (CustomizableObjectSystem->RemoveTextureReference(Texture.Key))
 			{
-				TexturesToRelease.Remove(Texture.Id); // Texture count in the end is not zero, so do not release it
+				if (CustomizableObjectSystem->bReleaseTexturesImmediately)
+				{
+					TexturesToRelease.Add(Texture); // Texture count is zero, so prepare to release it
+				}
 			}
 		}
 
diff --git a/Engine/Plugins/Experimental/Mutable/Source/CustomizableObject/Private/MuCO/CustomizableObjectSystem.cpp b/Engine/Plugins/Experimental/Mutable/Source/CustomizableObject/Private/MuCO/CustomizableObjectSystem.cpp
--- a/Engine/Plugins/Experimental/Mutable/Source/CustomizableObject/Private/MuCO/CustomizableObjectSystem.cpp
+++ b/Engine/Plugins/Experimental/Mutable/Source/CustomizableObject/Private/MuCO/CustomizableObjectSystem.cpp
@@ -795,6 +795,8 @@ void UpdateSkeletalMesh(UCustomizableObjectInstance& CustomizableObjectInstance,
 {
 	MUTABLE_CPUPROFILER_SCOPE(UpdateSkeletalMesh);
 
+	check(IsInGameThread());
+
 	for (int32 ComponentIndex = 0; ComponentIndex < CustomizableObjectInstance.SkeletalMeshes.Num(); ++ComponentIndex)
 	{
 		if (TObjectPtr<USkeletalMesh> SkeletalMesh = CustomizableObjectInstance.SkeletalMeshes[ComponentIndex])
@@ -837,21 +839,6 @@ void UpdateSkeletalMesh(UCustomizableObjectInstance& CustomizableObjectInstance,
 
 	CustomizableObjectInstance.bEditorPropertyChanged = false;
 
-	check(UCustomizableObjectSystem::GetInstance() != nullptr);
-	FCustomizableObjectSystemPrivate* CustomizableObjectSystem = UCustomizableObjectSystem::GetInstance()->GetPrivate();
-	check(CustomizableObjectSystem != nullptr);
-	if (CustomizableObjectSystem->bReleaseTexturesImmediately)
-	{
-		FMutableResourceCache& Cache = CustomizableObjectSystem->GetObjectCache(CustomizableObjectInstance.GetCustomizableObject());
-
-		for (TPair<uint32, FGeneratedTexture>& Item : CustomizableObjectInstancePrivateData->TexturesToRelease)
-		{
-			UCustomizableInstancePrivateData::ReleaseMutableTexture(Item.Value.Id, Cast<UTexture2D>(Item.Value.Texture), Cache);
-		}
-
-		CustomizableObjectInstancePrivateData->TexturesToRelease.Empty();
-	}
-
 	FinishUpdateGlobal(&CustomizableObjectInstance, UpdateResult, UpdateCallback, Parameters, UpdatedDescriptorRuntimeHash);
 }
 
@@ -1891,6 +1878,20 @@ namespace impl
 		// TODO MTBL-391: Review This hotfix
 		UpdateSkeletalMesh(*CustomizableObjectInstance, CustomizableObjectSystemPrivateData->CurrentMutableOperation->InstanceDescriptorRuntimeHash, OperationData->UpdateResult, &OperationData->UpdateCallback, OperationData->MutableParameters);
 
+		// All work is done, release unused textures.
+		if (CustomizableObjectSystemPrivateData->bReleaseTexturesImmediately)
+		{
+			FMutableResourceCache& Cache = CustomizableObjectSystemPrivateData->GetObjectCache(CustomizableObjectInstance->GetCustomizableObject());
+
+			UCustomizableInstancePrivateData* CustomizableObjectInstancePrivateData = CustomizableObjectInstance->GetPrivate();
+			for (FGeneratedTexture& GeneratedTexture : CustomizableObjectInstancePrivateData->TexturesToRelease)
+			{
+				UCustomizableInstancePrivateData::ReleaseMutableTexture(GeneratedTexture.Key, Cast<UTexture2D>(GeneratedTexture.Texture), Cache);
+			}
+
+			CustomizableObjectInstancePrivateData->TexturesToRelease.Empty();
+		}
+
 		// TODO: T2927
 		if (LogBenchmarkUtil::isLoggingActive())
 		{
@@ -2158,7 +2159,7 @@ namespace impl
 			FMutableImageCacheKey Key(Image.ImageID, OperationData->MipsToSkip);
 			TWeakObjectPtr<UTexture2D>* TexturePtr = Cache.Images.Find(Key);
 
-			if (TexturePtr && TexturePtr->Get() && SystemPrivateData->TextureHasReferences(Image.ImageID))
+			if (TexturePtr && TexturePtr->Get() && SystemPrivateData->TextureHasReferences(Key))
 			{
 				System->ProtectedCachedTextures.Add(TexturePtr->Get());
 				SystemPrivateData->ProtectedObjectCachedImages.Add(Image.ImageID);
diff --git a/Engine/Plugins/Experimental/Mutable/Source/CustomizableObject/Private/MuCO/CustomizableObjectSystemPrivate.h b/Engine/Plugins/Experimental/Mutable/Source/CustomizableObject/Private/MuCO/CustomizableObjectSystemPrivate.h
--- a/Engine/Plugins/Experimental/Mutable/Source/CustomizableObject/Private/MuCO/CustomizableObjectSystemPrivate.h
+++ b/Engine/Plugins/Experimental/Mutable/Source/CustomizableObject/Private/MuCO/CustomizableObjectSystemPrivate.h
@@ -272,6 +272,8 @@ struct FMutableImageCacheKey
 	mu::FResourceID Resource = 0;
 	int32 SkippedMips = 0;
 
+	FMutableImageCacheKey() {};
+
 	FMutableImageCacheKey(mu::FResourceID InResource, int32 InSkippedMips)
 		: Resource(InResource), SkippedMips(InSkippedMips) {}
 
@@ -708,7 +710,7 @@ public:
 	}
 
 
-	void AddTextureReference(uint32 TextureId)
+	void AddTextureReference(const FMutableImageCacheKey& TextureId)
 	{
 		uint32& CountRef = TextureReferenceCount.FindOrAdd(TextureId);
 
@@ -717,7 +719,7 @@ public:
 
 	
 	// Returns true if the texture's references become zero
-	bool RemoveTextureReference(uint32 TextureId)
+	bool RemoveTextureReference(const FMutableImageCacheKey& TextureId)
 	{
 		uint32* CountPtr = TextureReferenceCount.Find(TextureId);
 
@@ -742,7 +744,7 @@ public:
 	}
 
 
-	bool TextureHasReferences(uint32 TextureId) const
+	bool TextureHasReferences(const FMutableImageCacheKey& TextureId) const
 	{
 		const uint32* CountPtr = TextureReferenceCount.Find(TextureId);
 
@@ -802,7 +804,7 @@ public:
 
 	// \TODO: Remove this array if we are not gathering stats!
 	TArray<TWeakObjectPtr<class UTexture2D>> TextureTrackerArray;
-	TMap<uint32, uint32> TextureReferenceCount; // Keeps a count of texture usage to decide if they have to be blocked from GC during an update
+	TMap<FMutableImageCacheKey, uint32> TextureReferenceCount; // Keeps a count of texture usage to decide if they have to be blocked from GC during an update
 
 	// This is protected from GC by AddReferencedObjects
 	TObjectPtr<UCustomizableObjectInstance> CurrentInstanceBeingUpdated = nullptr;
diff --git a/Engine/Plugins/Experimental/Mutable/Source/CustomizableObject/Private/MuCO/LogInformationUtil.cpp b/Engine/Plugins/Experimental/Mutable/Source/CustomizableObject/Private/MuCO/LogInformationUtil.cpp
--- a/Engine/Plugins/Experimental/Mutable/Source/CustomizableObject/Private/MuCO/LogInformationUtil.cpp
+++ b/Engine/Plugins/Experimental/Mutable/Source/CustomizableObject/Private/MuCO/LogInformationUtil.cpp
@@ -43,7 +43,7 @@ void LogInformationUtil::PrintGeneratedTextures(const TArray<struct FGeneratedTe
 	for (i = 0; i < Max; ++i)
 	{
 		Log += FString::Printf(TEXT("Texture name=%s, "), *GeneratedTextures[i].Name);
-		Log += FString::Printf(TEXT("id=%d, "), GeneratedTextures[i].Id);
+		Log += FString::Printf(TEXT("id=%d, "), GeneratedTextures[i].Key.Resource);
 		Log += "\n\t\t\t";
 		PrintTextureData(Cast<UTexture2D>(GeneratedTextures[i].Texture), Log, false);
 		if (i < (Max - 1))
