[Mutable] Double click to add a reroute node.

#rnx

diff --git a/Engine/Plugins/Experimental/Mutable/Source/CustomizableObjectEditor/Classes/MuCOE/EdGraphSchema_CustomizableObject.h b/Engine/Plugins/Experimental/Mutable/Source/CustomizableObjectEditor/Classes/MuCOE/EdGraphSchema_CustomizableObject.h
--- a/Engine/Plugins/Experimental/Mutable/Source/CustomizableObjectEditor/Classes/MuCOE/EdGraphSchema_CustomizableObject.h
+++ b/Engine/Plugins/Experimental/Mutable/Source/CustomizableObjectEditor/Classes/MuCOE/EdGraphSchema_CustomizableObject.h
@@ -119,6 +119,7 @@ public:
 	virtual FLinearColor GetPinTypeColor(const FEdGraphPinType& PinType) const override;
 	virtual bool ShouldHidePinDefaultValue(UEdGraphPin* Pin) const override;
 	virtual void DroppedAssetsOnGraph(const TArray<FAssetData>& Assets, const FVector2D& GraphPosition, UEdGraph* Graph) const override;
+	virtual void OnPinConnectionDoubleCicked(UEdGraphPin* PinA, UEdGraphPin* PinB, const FVector2D& GraphPosition) const override;
 
 	FLinearColor GetPinTypeColor(const FName& PinType) const;
 
diff --git a/Engine/Plugins/Experimental/Mutable/Source/CustomizableObjectEditor/Private/MuCOE/EdGraphSchema_CustomizableObject.cpp b/Engine/Plugins/Experimental/Mutable/Source/CustomizableObjectEditor/Private/MuCOE/EdGraphSchema_CustomizableObject.cpp
--- a/Engine/Plugins/Experimental/Mutable/Source/CustomizableObjectEditor/Private/MuCOE/EdGraphSchema_CustomizableObject.cpp
+++ b/Engine/Plugins/Experimental/Mutable/Source/CustomizableObjectEditor/Private/MuCOE/EdGraphSchema_CustomizableObject.cpp
@@ -1100,11 +1100,30 @@ void UEdGraphSchema_CustomizableObject::DroppedAssetsOnGraph(const TArray<FAsset
 		{
 			UE_LOG(LogTemp,Error,TEXT("Unable to add null node to graph. "))
 		}
-
 	}
 }
 
 
+void UEdGraphSchema_CustomizableObject::OnPinConnectionDoubleCicked(UEdGraphPin* PinA, UEdGraphPin* PinB, const FVector2D& GraphPosition) const
+{
+	const FScopedTransaction Transaction(LOCTEXT("CreateRerouteNodeOnWire", "Create Reroute Node"));
+
+	// This constant is duplicated from inside of SGraphNodeKnot
+	const FVector2D NodeSpacerSize(42.0f, 24.0f);
+	const FVector2D KnotTopLeft = GraphPosition - (NodeSpacerSize * 0.5f);
+
+	// Create a new knot
+	UEdGraph* ParentGraph = PinA->GetOwningNode()->GetGraph();
+	UCustomizableObjectNodeReroute* DefaultNodeReroute = CastChecked<UCustomizableObjectNodeReroute>(UCustomizableObjectNodeReroute::StaticClass()->GetDefaultObject());
+	UCustomizableObjectNodeReroute* NodeReroute = CastChecked<UCustomizableObjectNodeReroute>(FCustomizableObjectSchemaAction_NewNode::CreateNode(ParentGraph, nullptr, KnotTopLeft, DefaultNodeReroute));
+
+	PinA->BreakLinkTo(PinB);
+	PinA->MakeLinkTo((PinA->Direction == EGPD_Output) ? NodeReroute->GetInputPin() : NodeReroute->GetOutputPin());
+	PinB->MakeLinkTo((PinB->Direction == EGPD_Output) ? NodeReroute->GetInputPin() : NodeReroute->GetOutputPin());
+	NodeReroute->UCustomizableObjectNode::ReconstructNode();
+}
+
+
 bool UEdGraphSchema_CustomizableObject::IsSpawnableAsset(const FAssetData& InAsset, ESpawnableObjectType& OutObjectType) const
 {
 	UObject* Object = InAsset.GetAsset();
@@ -1139,7 +1158,6 @@ bool UEdGraphSchema_CustomizableObject::IsSpawnableAsset(const FAssetData& InAss
 		// Non spawnable object
 		return false;
 	}
-	
 }
 
 
