[Mutable] Fix wrong condition being checked in OpMeshMerge detected by [REDACTED_PERSON]

#rnx

diff --git a/Engine/Plugins/Experimental/Mutable/Source/MutableRuntime/Private/MuR/OpMeshMerge.h b/Engine/Plugins/Experimental/Mutable/Source/MutableRuntime/Private/MuR/OpMeshMerge.h
--- a/Engine/Plugins/Experimental/Mutable/Source/MutableRuntime/Private/MuR/OpMeshMerge.h
+++ b/Engine/Plugins/Experimental/Mutable/Source/MutableRuntime/Private/MuR/OpMeshMerge.h
@@ -878,9 +878,9 @@ namespace mu
 					{
 						// Skip system buffers
 						if (chan.m_semantic == MBS_VERTEXINDEX
-							&&
+							||
 							chan.m_semantic == MBS_LAYOUTBLOCK
-							&&
+							||
 							chan.m_semantic == MBS_CHART)
 						{
 							continue;
