[Backout] - CL26427762
Original CL Desc
-----------------------------------------------------------------
[Mutable] Fix deadlock when loading roms.

- Removed the locking wait in the ModelDiskStreamer.
- Fixed a out of range check with RomWeights.
- ModelDiskStreamer replaced potentially unsafe raw pointers with a shared pointers.

#rnx

diff --git a/Engine/Plugins/Experimental/Mutable/Source/CustomizableObject/Classes/MuCO/CustomizableObject.h b/Engine/Plugins/Experimental/Mutable/Source/CustomizableObject/Classes/MuCO/CustomizableObject.h
--- a/Engine/Plugins/Experimental/Mutable/Source/CustomizableObject/Classes/MuCO/CustomizableObject.h
+++ b/Engine/Plugins/Experimental/Mutable/Source/CustomizableObject/Classes/MuCO/CustomizableObject.h
@@ -1045,7 +1045,7 @@ public:
 
 	/** Creates and returns an array with IAsyncReadFileHandles for each BulkData file.
 	 * Used by Mutable to stream in resources when generating instances. Must be deleted by the caller. */
-	TArray<TSharedPtr<IAsyncReadFileHandle>> GetAsyncReadFileHandles() const;
+	TArray<IAsyncReadFileHandle*> GetAsyncReadFileHandles() const;
 
 #if WITH_EDITOR
 
diff --git a/Engine/Plugins/Experimental/Mutable/Source/CustomizableObject/Private/MuCO/CustomizableObject.cpp b/Engine/Plugins/Experimental/Mutable/Source/CustomizableObject/Private/MuCO/CustomizableObject.cpp
--- a/Engine/Plugins/Experimental/Mutable/Source/CustomizableObject/Private/MuCO/CustomizableObject.cpp
+++ b/Engine/Plugins/Experimental/Mutable/Source/CustomizableObject/Private/MuCO/CustomizableObject.cpp
@@ -1834,14 +1834,14 @@ void UCustomizableObjectBulk::PostLoad()
 }
 
 
-TArray<TSharedPtr<IAsyncReadFileHandle>> UCustomizableObjectBulk::GetAsyncReadFileHandles() const
+TArray<IAsyncReadFileHandle*> UCustomizableObjectBulk::GetAsyncReadFileHandles() const
 {
-	TArray<TSharedPtr<IAsyncReadFileHandle>> ReadFileHandles;
+	TArray<IAsyncReadFileHandle*> ReadFileHandles;
 	ReadFileHandles.Reserve(BulkDataFileNames.Num());
 
 	for (const FString& FilePath : BulkDataFileNames)
 	{
-		TSharedPtr<IAsyncReadFileHandle> ReadFileHandle = MakeShareable(FPlatformFileManager::Get().GetPlatformFile().OpenAsyncRead(*FilePath));
+		IAsyncReadFileHandle* ReadFileHandle = FPlatformFileManager::Get().GetPlatformFile().OpenAsyncRead(*FilePath);
 		
 		if (!ReadFileHandle)
 		{
@@ -1854,6 +1854,11 @@ TArray<TSharedPtr<IAsyncReadFileHandle>> UCustomizableObjectBulk::GetAsyncReadFi
 
 	if (ReadFileHandles.Num() != BulkDataFileNames.Num())
 	{
+		for (IAsyncReadFileHandle* ReadFileHandle : ReadFileHandles)
+		{
+			delete ReadFileHandle;
+		}
+
 		ReadFileHandles.Empty();
 	}
 
diff --git a/Engine/Plugins/Experimental/Mutable/Source/CustomizableObject/Private/MuCO/UnrealMutableModelDiskStreamer.cpp b/Engine/Plugins/Experimental/Mutable/Source/CustomizableObject/Private/MuCO/UnrealMutableModelDiskStreamer.cpp
--- a/Engine/Plugins/Experimental/Mutable/Source/CustomizableObject/Private/MuCO/UnrealMutableModelDiskStreamer.cpp
+++ b/Engine/Plugins/Experimental/Mutable/Source/CustomizableObject/Private/MuCO/UnrealMutableModelDiskStreamer.cpp
@@ -67,9 +67,15 @@ bool FUnrealMutableModelBulkReader::PrepareStreamingForObject(UCustomizableObjec
 		if (!Objects[ObjectIndex].Model.Pin())
 		{
 			// The CustomizableObject is gone, so we won't be streaming for it anymore.
-			for (TPair<OPERATION_ID, FReadRequest>& it : Objects[ObjectIndex].CurrentReadRequests)
+			for (TPair<OPERATION_ID, IAsyncReadRequest*>& it : Objects[ObjectIndex].CurrentReadRequests)
 			{
-				it.Value.ReadRequest->WaitCompletion();
+				it.Value->WaitCompletion();
+				delete it.Value;
+			}
+
+			for (IAsyncReadFileHandle* ReadFileHandle : Objects[ObjectIndex].ReadFileHandles)
+			{
+				delete ReadFileHandle;
 			}
 
 			Objects.RemoveAtSwap(ObjectIndex);
@@ -96,7 +102,7 @@ bool FUnrealMutableModelBulkReader::PrepareStreamingForObject(UCustomizableObjec
 		FString FolderPath = CustomizableObject->GetCompiledDataFolderPath(true);
 		FString FullFileName = FolderPath + CustomizableObject->GetCompiledDataFileName(false, nullptr, true);
 
-		const TSharedPtr<IAsyncReadFileHandle> ReadFileHandle = MakeShareable(FPlatformFileManager::Get().GetPlatformFile().OpenAsyncRead(*FullFileName));
+		IAsyncReadFileHandle* ReadFileHandle = FPlatformFileManager::Get().GetPlatformFile().OpenAsyncRead(*FullFileName);
 		if (!ReadFileHandle)
 		{
 			UE_LOG(LogMutable, Warning, TEXT("Streaming: Customizable Object %s is missing the Editor BulkData."), *CustomizableObject->GetName());
@@ -158,9 +164,15 @@ void FUnrealMutableModelBulkReader::CancelStreamingForObject(const UCustomizable
 	{
 		if (Objects[ObjectIndex].Model.Pin() == CustomizableObject->GetModel())
 		{
-			for (TPair<OPERATION_ID, FReadRequest>& it : Objects[ObjectIndex].CurrentReadRequests)
+			for (TPair<OPERATION_ID, IAsyncReadRequest*>& it : Objects[ObjectIndex].CurrentReadRequests)
 			{
-				it.Value.ReadRequest->WaitCompletion();
+				it.Value->WaitCompletion();
+				delete it.Value;
+			}
+
+			for (IAsyncReadFileHandle* ReadFileHandle : Objects[ObjectIndex].ReadFileHandles)
+			{
+				delete ReadFileHandle;
 			}
 
 			Objects.RemoveAtSwap(ObjectIndex);
@@ -202,9 +214,15 @@ void FUnrealMutableModelBulkReader::EndStreaming()
 {
 	for (FObjectData& o : Objects)
 	{
-		for (TPair<OPERATION_ID, FReadRequest>& it : o.CurrentReadRequests)
+		for (TPair<OPERATION_ID, IAsyncReadRequest*>& it : o.CurrentReadRequests)
+		{
+			it.Value->WaitCompletion();
+			delete it.Value;
+		}
+
+		for (IAsyncReadFileHandle* ReadFileHandle : o.ReadFileHandles)
 		{
-			it.Value.ReadRequest->WaitCompletion();
+			delete ReadFileHandle;
 		}
 	}
 	Objects.Empty();
@@ -218,11 +236,11 @@ FAutoConsoleVariableRef CVarStreamPriority(
 	TEXT(""));
 
 
-mu::ModelReader::OPERATION_ID FUnrealMutableModelBulkReader::BeginReadBlock(const mu::Model* Model, uint64 Key, void* pBuffer, uint64 size, TFunction<void(bool bSuccess)>* CompletionCallback)
+mu::ModelReader::OPERATION_ID FUnrealMutableModelBulkReader::BeginReadBlock(const mu::Model* Model, uint64 Key, void* pBuffer, uint64 size)
 {
 	MUTABLE_CPUPROFILER_SCOPE(FUnrealMutableModelBulkStreamer::OpenReadFile);
 
-	UE_LOG(LogMutable, VeryVerbose, TEXT("Streaming: reading data %08llu."), Key);
+	UE_LOG(LogMutable, VeryVerbose, TEXT("Streaming: reading data %08lx."), Key);
 
 	// Find the object we are streaming for
 	FObjectData* ObjectData = Objects.FindByPredicate(
@@ -233,55 +251,34 @@ mu::ModelReader::OPERATION_ID FUnrealMutableModelBulkReader::BeginReadBlock(cons
 		// The object has been unloaded. Streaming is not possible. 
 		// This may happen in the editor if we are recompiling an object but we still have instances from the old
 		// object that have progressive mip generation.
-		if (CompletionCallback)
-		{
-			(*CompletionCallback)(false);    		
-		}
 		return -1;
 	}
 
-	// this generally cannot fail because it is async
-	if (!ObjectData->StreamableBlocks.Contains(Key))
-	{
-		// File Handle not found! This shouldn't really happen.
-		UE_LOG(LogMutable, Error, TEXT("Streaming Block not found!"));
-		check(false);
+	mu::ModelReader::OPERATION_ID Result = 0;
 
-		if (CompletionCallback)
-		{
-			(*CompletionCallback)(false);    		
-		}
-		return -1;
-	}
+	check(!ObjectData->ReadFileHandles.IsEmpty());
 
-	OPERATION_ID Result = ++LastOperationID;
+	// this generally cannot fail because it is async
+	if (ObjectData->StreamableBlocks.Contains(Key))
+	{
+		Result = ++LastOperationID;
 
-	const FMutableStreamableBlock& Block = ObjectData->StreamableBlocks[Key];
+		const FMutableStreamableBlock& Block = ObjectData->StreamableBlocks[Key];
 
-	int32 BulkDataOffsetInFile = 0;
+		int32 BulkDataOffsetInFile = 0;
 #if WITH_EDITOR
-	BulkDataOffsetInFile = sizeof(MutableCompiledDataStreamHeader);
+		BulkDataOffsetInFile = sizeof(MutableCompiledDataStreamHeader);
 #endif
 
-	FReadRequest ReadRequest;
-
-	if (CompletionCallback)
+		IAsyncReadRequest* ReadRequest = ObjectData->ReadFileHandles[Block.FileIndex]->ReadRequest(BulkDataOffsetInFile + Block.Offset, size, (EAsyncIOPriorityAndFlags)StreamPriority, nullptr, reinterpret_cast<uint8*>(pBuffer));
+		ObjectData->CurrentReadRequests.Add(Result, ReadRequest);
+	}
+	else
 	{
-		ReadRequest.FileCallback = MakeShared<TFunction<void(bool, IAsyncReadRequest*)>>([CompletionCallbackCapture = *CompletionCallback](bool bWasCancelled, IAsyncReadRequest*) -> void
-		{
-			CompletionCallbackCapture(!bWasCancelled);			
-		});		
+		// File Handle not found! This shouldn't really happen.
+		UE_LOG(LogMutable, Error, TEXT("Streaming Block not found!"));
+		check(false);
 	}
-	
-	check(!ObjectData->ReadFileHandles.IsEmpty());
-	ReadRequest.ReadRequest = MakeShareable(ObjectData->ReadFileHandles[Block.FileIndex]->ReadRequest(
-		BulkDataOffsetInFile + Block.Offset,
-		size,
-		(EAsyncIOPriorityAndFlags)StreamPriority,
-		ReadRequest.FileCallback.Get(),
-		reinterpret_cast<uint8*>(pBuffer)));
-
-	ObjectData->CurrentReadRequests.Add(Result, ReadRequest);
 
 	INC_DWORD_STAT(STAT_MutableStreamingOps);
 
@@ -295,9 +292,9 @@ bool FUnrealMutableModelBulkReader::IsReadCompleted(mu::ModelReader::OPERATION_I
 
 	for (FObjectData& o : Objects)
 	{
-		if (FReadRequest* ReadRequest = o.CurrentReadRequests.Find(OperationId))
+		if (IAsyncReadRequest** ReadRequest = o.CurrentReadRequests.Find(OperationId))
 		{
-			return ReadRequest->ReadRequest->PollCompletion();
+			return (*ReadRequest)->PollCompletion();
 		}
 	}
 
@@ -314,17 +311,19 @@ void FUnrealMutableModelBulkReader::EndRead(mu::ModelReader::OPERATION_ID Operat
 	bool bFound = false;
 	for (FObjectData& o : Objects)
 	{
-		FReadRequest* ReadRequest = o.CurrentReadRequests.Find(OperationId);
+		IAsyncReadRequest** ReadRequest = o.CurrentReadRequests.Find(OperationId);
 		if (ReadRequest)
 		{
-			if (ReadRequest->ReadRequest)
+			if (*ReadRequest)
 			{
-				bool bCompleted = ReadRequest->ReadRequest->WaitCompletion();
+				bool bCompleted = (*ReadRequest)->WaitCompletion();
 				if (!bCompleted)
 				{
 					UE_LOG(LogMutable, Error, TEXT("Operation failed to complete in EndRead."));
 					check(false);
 				}
+
+				delete* ReadRequest;
 			}
 			o.CurrentReadRequests.Remove(OperationId);
 			bFound = true;
diff --git a/Engine/Plugins/Experimental/Mutable/Source/CustomizableObject/Public/MuCO/UnrealMutableModelDiskStreamer.h b/Engine/Plugins/Experimental/Mutable/Source/CustomizableObject/Public/MuCO/UnrealMutableModelDiskStreamer.h
--- a/Engine/Plugins/Experimental/Mutable/Source/CustomizableObject/Public/MuCO/UnrealMutableModelDiskStreamer.h
+++ b/Engine/Plugins/Experimental/Mutable/Source/CustomizableObject/Public/MuCO/UnrealMutableModelDiskStreamer.h
@@ -57,24 +57,18 @@ public:
 	void EndStreaming();
 
 	// mu::ModelReader interface
-	OPERATION_ID BeginReadBlock(const mu::Model*, uint64 key0, void* pBuffer, uint64 size, TFunction<void(bool bSuccess)>* CompletionCallback) override;
+	OPERATION_ID BeginReadBlock(const mu::Model*, uint64 key0, void* pBuffer, uint64 size) override;
 	bool IsReadCompleted(OPERATION_ID) override;
 	void EndRead(OPERATION_ID) override;
 
 protected:
 
-	struct FReadRequest
-	{
-		TSharedPtr<IAsyncReadRequest> ReadRequest = nullptr;
-		TSharedPtr<FAsyncFileCallBack> FileCallback;
-	};
-	
 	/** Streaming data for one object. */
 	struct FObjectData
 	{
 		TWeakPtr<const mu::Model> Model;
-		TArray<TSharedPtr<IAsyncReadFileHandle>> ReadFileHandles;
-		TMap<OPERATION_ID, FReadRequest> CurrentReadRequests;
+		TArray<IAsyncReadFileHandle*> ReadFileHandles;
+		TMap<OPERATION_ID, IAsyncReadRequest*> CurrentReadRequests;
 		TMap<uint64, FMutableStreamableBlock> StreamableBlocks;
 	};
 
diff --git a/Engine/Plugins/Experimental/Mutable/Source/MutableRuntime/Private/MuR/CodeRunner.h b/Engine/Plugins/Experimental/Mutable/Source/MutableRuntime/Private/MuR/CodeRunner.h
--- a/Engine/Plugins/Experimental/Mutable/Source/MutableRuntime/Private/MuR/CodeRunner.h
+++ b/Engine/Plugins/Experimental/Mutable/Source/MutableRuntime/Private/MuR/CodeRunner.h
@@ -30,25 +30,13 @@
 	#include "Async/TaskGraphFwd.h"
 #endif
 
+
 namespace  mu
 {
 	class Model;
 	class Parameters;
 	class RangeIndex;
 
-	using EventType = 
-#ifdef MUTABLE_USE_NEW_TASKGRAPH
-		UE::Tasks::FTaskEvent;
-#else
-		FGraphEventRef;
-#endif
-
-	using TaskType = 
-#ifdef MUTABLE_USE_NEW_TASKGRAPH
-		UE::Tasks::FTask;
-#else
-		FGraphEventRef;
-#endif
 
     /** Code execution of the mutable virtual machine. */
     class CodeRunner
@@ -218,13 +206,6 @@ namespace  mu
 			}
 		};
 
-		struct FRomLoadOp
-		{
-			ModelReader::OPERATION_ID m_streamID = -1;
-			TArray<uint8> m_streamBuffer;
-			TaskType Event;
-		};
-
 		class FLoadMeshRomTask : public CodeRunner::FIssuedTask
 		{
 		public:
@@ -237,6 +218,7 @@ namespace  mu
 			// FIssuedTask interface
 			bool Prepare(CodeRunner*, bool& bOutFailed) override;
 			void Complete(CodeRunner*) override;
+			bool IsComplete(CodeRunner*) override;
 
 		private:
 			int32 RomIndex = -1;
@@ -280,12 +262,11 @@ namespace  mu
 			// FIssuedTask interface
 			bool Prepare(CodeRunner*, bool& bOutFailed) override;
 			void Complete(CodeRunner*) override;
+			bool IsComplete(CodeRunner*) override;
 
 		private:
  			int32 LODIndexIndex = -1;
 			int32 LODIndexCount = -1;
-
-			TArray<int32> RomIndices;
 		};
 
 		void AddOp(const FScheduledOp& op)
@@ -375,19 +356,6 @@ namespace  mu
 			}
 		}
 
-    	/** Calculate an approximation of memory used by streaming buffers in this class. */
-		int32 GetStreamingMemoryBytes() const
-    	{
-			int32 Result = RomLoadOps.GetAllocatedSize();
-			
-			for (const TTuple<int32, TSharedPtr<FRomLoadOp>>& RomLoadOp : RomLoadOps)
-			{
-				Result += RomLoadOp.Value->m_streamBuffer.GetAllocatedSize();
-			}
-			
-			return Result;
-    	}
-    	
 		/** Calculate an approximation of memory used by manging structures in this class. */
 		int32 GetInternalMemoryBytes() const
 		{
@@ -600,9 +568,16 @@ namespace  mu
 		const Parameters* m_pParams = nullptr;
 		uint32 m_lodMask = 0;
 
-    	/** Rom read operations already in progress. Key is the rom index. */
-    	TMap<int32, TSharedPtr<FRomLoadOp>> RomLoadOps;
-    	
+		// Async rom loading control
+		struct FRomLoadOp
+		{
+			int32 m_romIndex = 0;
+			DATATYPE ConstantType = DT_NONE;
+			ModelReader::OPERATION_ID m_streamID;
+			TArray<uint8> m_streamBuffer;
+		};
+		TArray<FRomLoadOp> m_romLoadOps;
+
 	private:
 
 		inline void AddChildren(const FScheduledOp& dep)
@@ -629,6 +604,9 @@ namespace  mu
 		/** Calculate an heuristic to select op execution based on memory usage. */
 		int32 GetOpEstimatedMemoryDelta(const FScheduledOp& Candidate, const FProgram& Program);
 
+		/** */
+		void CompleteRomLoadOp(FRomLoadOp& o);
+		
 		/** Update debug stats. */
 		void UpdateTraces();
 
diff --git a/Engine/Plugins/Experimental/Mutable/Source/MutableRuntime/Private/MuR/CodeRunnerTasked.cpp b/Engine/Plugins/Experimental/Mutable/Source/MutableRuntime/Private/MuR/CodeRunnerTasked.cpp
--- a/Engine/Plugins/Experimental/Mutable/Source/MutableRuntime/Private/MuR/CodeRunnerTasked.cpp
+++ b/Engine/Plugins/Experimental/Mutable/Source/MutableRuntime/Private/MuR/CodeRunnerTasked.cpp
@@ -57,6 +57,54 @@ static FAutoConsoleVariableRef CVarCoreRunnerTaskPriority(
 namespace mu
 {
 
+	void CodeRunner::CompleteRomLoadOp( FRomLoadOp& o )
+	{
+		if (DebugRom && (DebugRomAll || o.m_romIndex == DebugRomIndex))
+			UE_LOG(LogMutableCore, Log, TEXT("CodeRunner::CompleteRomLoadOp for rom %d."), o.m_romIndex);
+
+		m_pSystem->StreamInterface->EndRead(o.m_streamID);
+
+		FProgram& program = m_pModel->GetPrivate()->m_program;
+		{
+			MUTABLE_CPUPROFILER_SCOPE(Unserialise);
+
+			InputMemoryStream stream(o.m_streamBuffer.GetData(), o.m_streamBuffer.Num());
+			InputArchive arch(&stream);
+
+			int32 ResIndex = program.m_roms[o.m_romIndex].ResourceIndex;
+			switch (o.ConstantType)
+			{
+			case DATATYPE::DT_MESH:
+			{
+				check(!program.m_constantMeshes[ResIndex].Value);
+				Ptr<Mesh> Value = Mesh::StaticUnserialise(arch);
+				program.m_constantMeshes[ResIndex].Value = Value;
+				check(program.m_constantMeshes[ResIndex].Value);
+				m_pSystem->WorkingMemoryManager.TrackedBudgetBytes_Rom += Value->GetDataSize();
+				break;
+			}
+			case DATATYPE::DT_IMAGE:
+			{
+				check(!program.m_constantImageLODs[ResIndex].Value);
+				// TODO: Try to reuse buffer from PooledImages.
+				Ptr<Image> Value = Image::StaticUnserialise(arch);
+				program.m_constantImageLODs[ResIndex].Value = Value;
+				check(program.m_constantImageLODs[ResIndex].Value);
+				m_pSystem->WorkingMemoryManager.TrackedBudgetBytes_Rom += Value->GetDataSize();
+				break;
+			}
+			default:
+				check(false);
+				break;
+			}
+
+			m_pSystem->WorkingMemoryManager.TrackedBudgetBytes_Stream -= o.m_streamBuffer.GetAllocatedSize();
+
+			o.m_romIndex = -1;
+			o.m_streamBuffer.Empty();
+		}
+	}
+
 	//---------------------------------------------------------------------------------------------
 	TRACE_DECLARE_INT_COUNTER(MutableRuntime_OpenTask, TEXT("MutableRuntime/OpenTask"));
 	TRACE_DECLARE_INT_COUNTER(MutableRuntime_ClosedTasks, TEXT("MutableRuntime/ClosedTasks"));
@@ -481,6 +529,15 @@ namespace mu
 				}
 			}
 
+			// Look for completed streaming ops and complete the rom loading
+			for (FRomLoadOp& o : m_romLoadOps)
+			{
+				if (o.m_romIndex>=0 && m_pSystem->StreamInterface->IsReadCompleted(o.m_streamID))
+				{
+					CompleteRomLoadOp(o);
+				}
+			}
+
 			// Look for a closed task with dependencies satisfied and move them to the open task list.
 			bool bSomeWasReady = false;
 			for (int Index = 0; Index<ClosedTasks.Num(); )
@@ -556,6 +613,19 @@ namespace mu
 					}
 				}
 
+				// If we reached here it means we didn't find an op to wait for. Try to wait for a loading op.
+				// \todo: unify loading ops with normal ones?
+				for (FRomLoadOp& o : m_romLoadOps)
+				{
+					if (o.m_romIndex >= 0)
+					{
+						CompleteRomLoadOp(o);
+						break;
+					}
+
+					// We should never reach this, since it would mean we deadlocked.
+					//check(false);
+				}
 			}
 		}
 
@@ -592,7 +662,6 @@ namespace mu
 		}
     }
 
-	
 	//---------------------------------------------------------------------------------------------
 	void CodeRunner::GetImageDescResult(FImageDesc& OutDesc)
 	{
@@ -1969,120 +2038,91 @@ namespace mu
 	//---------------------------------------------------------------------------------------------
 	//---------------------------------------------------------------------------------------------
 	//---------------------------------------------------------------------------------------------
-	bool CodeRunner::FLoadMeshRomTask::Prepare(CodeRunner* Runner, bool& bOutFailed)
+	bool CodeRunner::FLoadMeshRomTask::Prepare(CodeRunner* Runner, bool& bOutFailed )
 	{
-		// This runs in the mutable Runner thread
-		MUTABLE_CPUPROFILER_SCOPE(FLoadMeshRomTask_Prepare);
-
 		if (!Runner || !Runner->m_pSystem)
 		{
 			return false;
 		}
 
+		// This runs in the mutable Runner thread
+		MUTABLE_CPUPROFILER_SCOPE(FLoadMeshRomTask_Prepare);
 		bOutFailed = false;
 
-		const FProgram& Program = Runner->m_pModel->GetPrivate()->m_program;
+		FProgram& program = Runner->m_pModel->GetPrivate()->m_program;
 
-		check(RomIndex < Program.m_roms.Num());
-		
-		FWorkingMemoryManager::FModelCacheEntry* ModelCache = Runner->m_pSystem->WorkingMemoryManager.FindModelCache(Runner->m_pModel.Get());
-		++ModelCache->PendingOpsPerRom[RomIndex];
+		check(RomIndex < program.m_roms.Num());
+		bool bRomIsLoaded = program.IsRomLoaded(RomIndex);
 
-		if (Program.IsRomLoaded(RomIndex))
+		FWorkingMemoryManager::FModelCacheEntry* ModelCache = Runner->m_pSystem->WorkingMemoryManager.FindModelCache(Runner->m_pModel.Get());
+		bool bRomHasPendingOps = ModelCache->PendingOpsPerRom[RomIndex]!=0;
+		if (!bRomIsLoaded && !bRomHasPendingOps)
 		{
-			return false;
-		}
+			check(Runner->m_pSystem->StreamInterface);
 
-		if (const TSharedPtr<FRomLoadOp>* Result = Runner->RomLoadOps.Find(RomIndex))
-		{
-			Event = (*Result)->Event;  // Wait for the read operation started by other task
-			return false;
-		}
+			uint32 RomSize = program.m_roms[RomIndex].Size;
+			check(RomSize > 0);
 
-		const TSharedPtr<FRomLoadOp> RomLoadOp = MakeShared<FRomLoadOp>();
-		Runner->RomLoadOps.Add(RomIndex, RomLoadOp);
-		
-		check(Runner->m_pSystem->StreamInterface);
+			// Free roms if necessary
+			{
+				MUTABLE_CPUPROFILER_SCOPE(FreeingRoms);
 
-		const uint32 RomSize = Program.m_roms[RomIndex].Size;
-		check(RomSize > 0);
+				Runner->m_pSystem->WorkingMemoryManager.MarkRomUsed(RomIndex, Runner->m_pModel);
+				Runner->m_pSystem->WorkingMemoryManager.EnsureBudgetBelow(RomSize);
+			}
 
-		// Free roms if necessary
-		{
-			MUTABLE_CPUPROFILER_SCOPE(FreeingRoms);
+			CodeRunner::FRomLoadOp* op = nullptr;
+			for (FRomLoadOp& o : Runner->m_romLoadOps)
+			{
+				if (o.m_romIndex < 0)
+				{
+					op = &o;
+				}
+			}
 
-			Runner->m_pSystem->WorkingMemoryManager.MarkRomUsed(RomIndex, Runner->m_pModel);
-			Runner->m_pSystem->WorkingMemoryManager.EnsureBudgetBelow(RomSize);
-		}
+			if (!op)
+			{
+				Runner->m_romLoadOps.Add(CodeRunner::FRomLoadOp());
+				op = &Runner->m_romLoadOps.Last();
+			}
 
-		RomLoadOp->m_streamBuffer.SetNum(RomSize);
+			op->m_romIndex = RomIndex;
+			op->ConstantType = DT_MESH;
 
-		EventType EventTyped =
-#ifdef MUTABLE_USE_NEW_TASKGRAPH
-			UE::Tasks::FTaskEvent(TEXT("FLoadMeshRomTask"));
-#else
-			FGraphEvent::CreateGraphEvent();		
-#endif
-		
-		RomLoadOp->Event = EventTyped;
-		
-		TFunction<void(bool)> Callback = [EventTyped](bool bSuccess) mutable
-		{
-#ifdef MUTABLE_USE_NEW_TASKGRAPH
-			EventTyped.Trigger();
-#else
-			EventTyped->DispatchSubsequents();
-#endif				
-		};
-		
-		const uint32 RomId = Program.m_roms[RomIndex].Id;
-		RomLoadOp->m_streamID = Runner->m_pSystem->StreamInterface->BeginReadBlock(Runner->m_pModel.Get(), RomId, RomLoadOp->m_streamBuffer.GetData(), RomSize, &Callback);
-		if (RomLoadOp->m_streamID < 0)
-		{
-			bOutFailed = true;
-			return false;
+			if (uint32(op->m_streamBuffer.Num()) < RomSize)
+			{
+				op->m_streamBuffer.SetNum((size_t)RomSize);
+				Runner->m_pSystem->WorkingMemoryManager.TrackedBudgetBytes_Stream += op->m_streamBuffer.GetAllocatedSize();
+			}
+
+			uint32 RomId = program.m_roms[RomIndex].Id;
+			op->m_streamID = Runner->m_pSystem->StreamInterface->BeginReadBlock(Runner->m_pModel.Get(), RomId, op->m_streamBuffer.GetData(), RomSize);
+			if (op->m_streamID < 0)
+			{
+				bOutFailed = true;
+				return false;
+			}
 		}
+		++ModelCache->PendingOpsPerRom[RomIndex];
 
-		Event = EventTyped; // Wait for read operation to end
-		return false; // No worker thread work
+		//UE_LOG(LogMutableCore, Log, TEXT("FLoadMeshRomTask::Prepare romindex %d pending %d."), RomAt.m_rom, Runner->m_romPendingOps[RomAt.m_rom]);
+
+		// No worker thread work
+		return false;
 	}
-	
+
 
 	//---------------------------------------------------------------------------------------------
 	void CodeRunner::FLoadMeshRomTask::Complete(CodeRunner* Runner)
 	{
-		// This runs in the Runner thread
-		MUTABLE_CPUPROFILER_SCOPE(FLoadMeshRomTask_Complete);
-
 		if (!Runner || !Runner->m_pSystem)
 		{
 			return;
 		}
 
-		FProgram& Program = Runner->m_pModel->GetPrivate()->m_program;
-
-		// Since task could be reordered, we need to make sure we end the rom read before continuing
-		if (const TSharedPtr<FRomLoadOp>* Result = Runner->RomLoadOps.Find(RomIndex))
-		{
-			const TSharedPtr<FRomLoadOp>& RomLoadOp = *Result;
-
-			Runner->m_pSystem->StreamInterface->EndRead(RomLoadOp->m_streamID);									
-
-			const int32 ResIndex = Program.m_roms[RomIndex].ResourceIndex;
-			check(!Program.m_constantMeshes[ResIndex].Value)
-			MUTABLE_CPUPROFILER_SCOPE(Unserialise);
-
-			InputMemoryStream stream(RomLoadOp->m_streamBuffer.GetData(), RomLoadOp->m_streamBuffer.Num());
-			InputArchive arch(&stream);
-
-			check(!Program.m_constantMeshes[ResIndex].Value);
-			Ptr<Mesh> Value = Mesh::StaticUnserialise(arch);
-			Program.m_constantMeshes[ResIndex].Value = Value;
-			check(Program.m_constantMeshes[ResIndex].Value);
-			Runner->m_pSystem->WorkingMemoryManager.TrackedBudgetBytes_Rom += Value->GetDataSize();
+		MUTABLE_CPUPROFILER_SCOPE(FLoadMeshRomTask_Complete);
 
-			Runner->RomLoadOps.Remove(RomIndex);
-		}
+		// This runs in the Runner thread
 
 		// Process the constant op normally, now that the rom is loaded.
 		Runner->RunCode(Op, Runner->m_pParams, Runner->m_pModel.Get(), Runner->m_lodMask);
@@ -2094,6 +2134,14 @@ namespace mu
 	}
 
 
+	//---------------------------------------------------------------------------------------------
+	bool CodeRunner::FLoadMeshRomTask::IsComplete(CodeRunner* Runner)
+	{ 
+		FProgram& program = Runner->m_pModel->GetPrivate()->m_program;
+		return program.IsRomLoaded(RomIndex);
+	}
+
+
 	//---------------------------------------------------------------------------------------------
 	//---------------------------------------------------------------------------------------------
 	//---------------------------------------------------------------------------------------------
@@ -2194,18 +2242,10 @@ namespace mu
 
 		FWorkingMemoryManager::FModelCacheEntry* ModelCache = Runner->m_pSystem->WorkingMemoryManager.FindModelCache(Runner->m_pModel.Get());
 
-#ifdef MUTABLE_USE_NEW_TASKGRAPH
-		TArray<UE::Tasks::FTask> ReadCompleteEvents;
-#else
-		FGraphEventArray ReadCompleteEvents;
-#endif
-
-		ReadCompleteEvents.Reserve(LODIndexCount); 
-
-		for (int32 LODIndex = 0; LODIndex < LODIndexCount; ++LODIndex)
+		for (int32 i = 0; i<LODIndexCount; ++i )
 		{
-			const int32 CurrentIndexIndex = LODIndexIndex + LODIndex;
-			const int32 CurrentIndex = program.m_constantImageLODIndices[CurrentIndexIndex];
+			int32 CurrentIndexIndex = LODIndexIndex + i;
+			int32 CurrentIndex = program.m_constantImageLODIndices[CurrentIndexIndex];
 
 			if (program.m_constantImageLODs[CurrentIndex].Key<0)
 			{
@@ -2216,30 +2256,21 @@ namespace mu
 			int32 RomIndex = program.m_constantImageLODs[CurrentIndex].Key;
 			check(RomIndex < program.m_roms.Num());
 
+			bool bRomIsLoaded = program.IsRomLoaded(RomIndex);
+			bool bRomHasAlreadyBeenRequested = ModelCache->PendingOpsPerRom[RomIndex] != 0;
 			++ModelCache->PendingOpsPerRom[RomIndex];
 
 			if (DebugRom && (DebugRomAll || RomIndex == DebugRomIndex))
 				UE_LOG(LogMutableCore, Log, TEXT("Preparing rom %d, now peding ops is %d."), RomIndex, ModelCache->PendingOpsPerRom[RomIndex]);
 
-			if (program.IsRomLoaded(RomIndex))
+			if (bRomIsLoaded || bRomHasAlreadyBeenRequested)
 			{
 				continue;
 			}
 
-			RomIndices.Add(RomIndex);
-
-			if (const TSharedPtr<FRomLoadOp>* Result = Runner->RomLoadOps.Find(RomIndex))
-			{
-				ReadCompleteEvents.Add((*Result)->Event); // Wait for the read operation started by other task
-				continue;
-			}
-
-			TSharedPtr<FRomLoadOp> RomLoadOp = MakeShared<FRomLoadOp>();
-			Runner->RomLoadOps.Add(RomIndex, RomLoadOp);
-			
 			check(Runner->m_pSystem->StreamInterface);
 
-			const uint32 RomSize = program.m_roms[RomIndex].Size;
+			uint32 RomSize = program.m_roms[RomIndex].Size;
 			check(RomSize > 0);
 
 			// Free roms if necessary
@@ -2248,102 +2279,78 @@ namespace mu
 				Runner->m_pSystem->WorkingMemoryManager.EnsureBudgetBelow(RomSize);
 			}
 
-			RomLoadOp->m_streamBuffer.SetNum(RomSize);
+			CodeRunner::FRomLoadOp* op = nullptr;
+			for (FRomLoadOp& o : Runner->m_romLoadOps)
+			{
+				if (o.m_romIndex < 0)
+				{
+					op = &o;
+				}
+			}
 
-			EventType EventTyped = 
-#ifdef MUTABLE_USE_NEW_TASKGRAPH
-				UE::Tasks::FTaskEvent(TEXT("FLoadImageRomsTaskRom"));
-#else
-				FGraphEvent::CreateGraphEvent();
-#endif
+			if (!op)
+			{
+				Runner->m_romLoadOps.Add(CodeRunner::FRomLoadOp());
+				op = &Runner->m_romLoadOps.Last();
+			}
 
-			RomLoadOp->Event = EventTyped;
-			ReadCompleteEvents.Add(EventTyped);
+			op->m_romIndex = RomIndex;
+			op->ConstantType = DT_IMAGE;
 
-			TFunction<void(bool)> Callback = [EventTyped](bool bSuccess) mutable // Mutable due Trigger not being const
+			if (uint32(op->m_streamBuffer.Num()) < RomSize)
 			{
-#ifdef MUTABLE_USE_NEW_TASKGRAPH
-				EventTyped.Trigger();
-#else
-				EventTyped->DispatchSubsequents();
-#endif				
-			};
-			
-			const uint32 RomId = program.m_roms[RomIndex].Id;
-			RomLoadOp->m_streamID = Runner->m_pSystem->StreamInterface->BeginReadBlock(Runner->m_pModel.Get(), RomId, RomLoadOp->m_streamBuffer.GetData(), RomSize, &Callback);
-			if (RomLoadOp->m_streamID < 0)
+				op->m_streamBuffer.SetNum(RomSize);
+				Runner->m_pSystem->WorkingMemoryManager.TrackedBudgetBytes_Stream += op->m_streamBuffer.GetAllocatedSize();
+			}
+
+			uint32 RomId = program.m_roms[RomIndex].Id;
+			op->m_streamID = Runner->m_pSystem->StreamInterface->BeginReadBlock(Runner->m_pModel.Get(), RomId, op->m_streamBuffer.GetData(), RomSize);
+			if (op->m_streamID < 0)
 			{
 				bOutFailed = true;
 				return false;
 			}
 		}
 
-		// Wait for all read operations to end
-#ifdef MUTABLE_USE_NEW_TASKGRAPH
-		UE::Tasks::FTaskEvent EventTyped(TEXT("FLoadImageRomsTask"));
-		EventTyped.AddPrerequisites(ReadCompleteEvents);
-		EventTyped.Trigger();
-		Event = EventTyped;
-#else
-		Event = FGraphEvent::CreateGraphEvent();
-		for (const FGraphEventRef& ReadCompleteEvent : ReadCompleteEvents)
-		{
-			Event->DontCompleteUntil(ReadCompleteEvent);
-		}
-		Event->DispatchSubsequents();
-#endif
-			
-		return false; // No worker thread work
+		//UE_LOG(LogMutableCore, Log, TEXT("FLoadImageRomsTask::Prepare romindex %d pending %d."), RomAt.m_rom, Runner->m_romPendingOps[RomAt.m_rom]);
+
+		// No worker thread work
+		return false;
 	}
-	
-	
+
+
 	//---------------------------------------------------------------------------------------------
 	void CodeRunner::FLoadImageRomsTask::Complete(CodeRunner* Runner)
 	{
-		// This runs in the Runner thread
-		MUTABLE_CPUPROFILER_SCOPE(FLoadImageRomsTask_Complete);
-
 		if (!Runner || !Runner->m_pSystem)
 		{
 			return;
 		}
 
-		FProgram& program = Runner->m_pModel->GetPrivate()->m_program;
-		
-		FWorkingMemoryManager::FModelCacheEntry* ModelCache = Runner->m_pSystem->WorkingMemoryManager.FindModelCache(Runner->m_pModel.Get());
+		MUTABLE_CPUPROFILER_SCOPE(FLoadImageRomsTask_Complete);
 
-		for (const int32 RomIndex : RomIndices)
-		{
-			// Since task could be reordered, we need to make sure we end the rom read before continuing
-			if (const TSharedPtr<FRomLoadOp>* Result = Runner->RomLoadOps.Find(RomIndex))
-			{
-				const TSharedPtr<FRomLoadOp>& RomLoadOp = *Result;
-				
-				Runner->m_pSystem->StreamInterface->EndRead(RomLoadOp->m_streamID);
+		// This runs in the Runner thread
 
-				MUTABLE_CPUPROFILER_SCOPE(Unserialise);
+		FWorkingMemoryManager::FModelCacheEntry* ModelCache = Runner->m_pSystem->WorkingMemoryManager.FindModelCache(Runner->m_pModel.Get());
 
-				InputMemoryStream stream(RomLoadOp->m_streamBuffer.GetData(), RomLoadOp->m_streamBuffer.Num());
-				InputArchive arch(&stream);
+		// Process the constant op normally, now that the rom is loaded.
+		Runner->RunCode(Op, Runner->m_pParams, Runner->m_pModel.Get(), Runner->m_lodMask);
 
-				const int32 ResIndex = program.m_roms[RomIndex].ResourceIndex;
+		FProgram& program = Runner->m_pModel->GetPrivate()->m_program;
+		for (int32 i = 0; i < LODIndexCount; ++i)
+		{
+			int32 CurrentIndexIndex = LODIndexIndex + i;
+			int32 CurrentIndex = program.m_constantImageLODIndices[CurrentIndexIndex];
 
-				// TODO: Try to reuse buffer from PooledImages.
-				check(!program.m_constantImageLODs[ResIndex].Value);
-				Ptr<Image> Value = Image::StaticUnserialise(arch);
-				program.m_constantImageLODs[ResIndex].Value = Value;
-				check(program.m_constantImageLODs[ResIndex].Value);
-				Runner->m_pSystem->WorkingMemoryManager.TrackedBudgetBytes_Rom += Value->GetDataSize();
-				
-				Runner->RomLoadOps.Remove(RomIndex);
+			if (program.m_constantImageLODs[CurrentIndex].Key < 0)
+			{
+				// This data is always resident.
+				continue;
 			}
-		}
-		
-		// Process the constant op normally, now that the rom is loaded.	
-		Runner->RunCode(Op, Runner->m_pParams, Runner->m_pModel.Get(), Runner->m_lodMask);
 
-		for (const int32 RomIndex : RomIndices)
-		{
+			int32 RomIndex = program.m_constantImageLODs[CurrentIndex].Key;
+			check(RomIndex < program.m_roms.Num());
+
 			Runner->m_pSystem->WorkingMemoryManager.MarkRomUsed(RomIndex, Runner->m_pModel);
 			--ModelCache->PendingOpsPerRom[RomIndex];
 
@@ -2355,6 +2362,33 @@ namespace mu
 	}
 
 
+	//---------------------------------------------------------------------------------------------
+	bool CodeRunner::FLoadImageRomsTask::IsComplete(CodeRunner* Runner)
+	{
+		FProgram& program = Runner->m_pModel->GetPrivate()->m_program;
+		for (int32 i = 0; i < LODIndexCount; ++i)
+		{
+			int32 CurrentIndexIndex = LODIndexIndex + i;
+			int32 CurrentIndex = program.m_constantImageLODIndices[CurrentIndexIndex];
+
+			if (program.m_constantImageLODs[CurrentIndex].Key < 0)
+			{
+				// This data is always resident.
+				continue;
+			}
+
+			int32 RomIndex = program.m_constantImageLODs[CurrentIndex].Key;
+			check(RomIndex < program.m_roms.Num());
+			if (!program.IsRomLoaded(RomIndex))
+			{
+				return false;
+			}
+		}
+
+		return true;
+	}
+
+
 	//---------------------------------------------------------------------------------------------
 	//---------------------------------------------------------------------------------------------
 	//---------------------------------------------------------------------------------------------
diff --git a/Engine/Plugins/Experimental/Mutable/Source/MutableRuntime/Private/MuR/System.cpp b/Engine/Plugins/Experimental/Mutable/Source/MutableRuntime/Private/MuR/System.cpp
--- a/Engine/Plugins/Experimental/Mutable/Source/MutableRuntime/Private/MuR/System.cpp
+++ b/Engine/Plugins/Experimental/Mutable/Source/MutableRuntime/Private/MuR/System.cpp
@@ -915,7 +915,12 @@ namespace mu
 
 		if (CurrentRunner)
 		{
-			StreamingBufferBytes += CurrentRunner->GetStreamingMemoryBytes();
+			// Data reserved for ongoing streaming operations
+			for ( const CodeRunner::FRomLoadOp& Op : CurrentRunner->m_romLoadOps )
+			{
+				StreamingBufferBytes += Op.m_streamBuffer.GetAllocatedSize();
+			}
+
 			InternalBytes += CurrentRunner->GetInternalMemoryBytes();
 		}
 
@@ -1029,8 +1034,7 @@ namespace mu
 		FModelCacheEntry n;
         n.Model = TWeakPtr<const Model>(InModel);
 		n.PendingOpsPerRom.SetNum(InModel->GetPrivate()->m_program.m_roms.Num());
-    	n.RomWeights.SetNum(InModel->GetPrivate()->m_program.m_roms.Num());
-    	CachePerModel.Add(n);
+        CachePerModel.Add(n);
         return CachePerModel.Last();
     }
 
@@ -1461,6 +1465,8 @@ namespace mu
     {
         check(pModel);
 
+        FProgram& Program = pModel->GetPrivate()->m_program;
+
         // If budget is zero, we don't unload anything here, and we assume it is managed somewhere else.
         if (!BudgetBytes)
         {
@@ -1472,7 +1478,12 @@ namespace mu
         // Update current cache
         {
 			FModelCacheEntry* ModelCache = FindModelCache(pModel.Get());
-        	
+
+            if (ModelCache->RomWeights.Num()<Program.m_roms.Num())
+            {
+				ModelCache->RomWeights.SetNumZeroed(Program.m_roms.Num());
+            }
+
             ModelCache->RomWeights[romIndex].Key++;
             ModelCache->RomWeights[romIndex].Value = RomTick;
         }
diff --git a/Engine/Plugins/Experimental/Mutable/Source/MutableRuntime/Public/MuR/Serialisation.h b/Engine/Plugins/Experimental/Mutable/Source/MutableRuntime/Public/MuR/Serialisation.h
--- a/Engine/Plugins/Experimental/Mutable/Source/MutableRuntime/Public/MuR/Serialisation.h
+++ b/Engine/Plugins/Experimental/Mutable/Source/MutableRuntime/Public/MuR/Serialisation.h
@@ -99,11 +99,10 @@ namespace mu
 		//! \param pBuffer is an already-allocated buffer big enough to receive the expected data.
 		//! \param size is the size of the pBuffer buffer, which must match the size of the data
 		//! requested with the key identifiers.
-		//! \param CompletionCallback Optional callback. Copied inside the called function. Will always be called.
 		//! \return a previously unused identifier, now used for this operation, that can be used in
 		//! calls to the other methods of this interface. If the return value is negative it indicates
 		//! an unrecoverable error.
-		virtual OPERATION_ID BeginReadBlock(const mu::Model*, uint64 key0, void* pBuffer, uint64 size, TFunction<void(bool bSuccess)>* CompletionCallback = nullptr) = 0;
+		virtual OPERATION_ID BeginReadBlock(const mu::Model*, uint64 key0, void* pBuffer, uint64 size) = 0;
 
         //! Check if a data request operation has been completed.
         //! This is a weak check than *may* return true if the given operation has completed, but
