[Mutable] Fix "Failed to cache external image" warning.

- Some update paths where not adding reference the external images. Added an additional check to detect this kind of errors.
- Task_Game_Lock_Cache was removing old references before adding the new ones.
- Fix LOD updates completely skipping DoUpdateSkeletalMesh.

#rnx

diff --git a/Engine/Plugins/Experimental/Mutable/Source/CustomizableObject/Classes/MuCO/CustomizableObjectInstance.h b/Engine/Plugins/Experimental/Mutable/Source/CustomizableObject/Classes/MuCO/CustomizableObjectInstance.h
--- a/Engine/Plugins/Experimental/Mutable/Source/CustomizableObject/Classes/MuCO/CustomizableObjectInstance.h
+++ b/Engine/Plugins/Experimental/Mutable/Source/CustomizableObject/Classes/MuCO/CustomizableObjectInstance.h
@@ -288,8 +288,11 @@ public:
 private:
 	/** Perform all the checks required to see if the update should begin. */
 	EUpdateRequired IsUpdateRequired(bool bIsCloseDistTick, bool bOnlyUpdateIfNotGenerated, bool bIgnoreCloseDist) const;
-	
-	/** Update Skeletal Mesh asynchronously. Immersive function. Not exposed to the API.
+
+public:
+	/** Private API.
+	 *
+	 * Update Skeletal Mesh asynchronously. Immersive function.
 	 * Once the update reaches this function, the update has been considered started and must complete all the update flow.
 	 * Starting at this function, all Update code paths must end up in FinishUpdateGlobal!
 	 *
@@ -297,7 +300,6 @@ private:
 	 */
 	void DoUpdateSkeletalMesh(bool bIsCloseDistTick, bool bOnlyUpdateIfNotGenerated, bool bIgnoreCloseDist, bool bForceHighPriority, const EUpdateRequired* OptionalUpdateRequired, FInstanceUpdateDelegate* UpdateCallback);
 
-public:	
 	// Clones the instance creating a new identical transient instance.
 	UFUNCTION(BlueprintCallable, Category = CustomizableObjectInstance)
 	UCustomizableObjectInstance* Clone();
diff --git a/Engine/Plugins/Experimental/Mutable/Source/CustomizableObject/Classes/MuCO/CustomizableObjectSystem.h b/Engine/Plugins/Experimental/Mutable/Source/CustomizableObject/Classes/MuCO/CustomizableObjectSystem.h
--- a/Engine/Plugins/Experimental/Mutable/Source/CustomizableObject/Classes/MuCO/CustomizableObjectSystem.h
+++ b/Engine/Plugins/Experimental/Mutable/Source/CustomizableObject/Classes/MuCO/CustomizableObjectSystem.h
@@ -184,8 +184,8 @@ public:
 };
 
 
+// TODO GMT Remove hotfix. External images are uncached images if Parameters is valid. Before the Mutable Queue rework this made sense, but this is no longer the case. Remove this when doing MTBL-1409.
 /** End a Customizable Object Instance Update. All code paths of an update have to end here. */
-void FinishUpdateGlobal(UCustomizableObjectInstance* Instance, EUpdateResult UpdateResult, FInstanceUpdateDelegate* UpdateCallback, const FCustomizableObjectInstanceDescriptor* InstanceDescriptor = nullptr, const FDescriptorRuntimeHash InUpdatedHash = FDescriptorRuntimeHash());
 void FinishUpdateGlobal(UCustomizableObjectInstance* Instance, EUpdateResult UpdateResult, FInstanceUpdateDelegate* UpdateCallback, mu::ParametersPtrConst Parameters, const FDescriptorRuntimeHash InUpdatedHash = FDescriptorRuntimeHash());
 
 
diff --git a/Engine/Plugins/Experimental/Mutable/Source/CustomizableObject/Private/MuCO/CustomizableInstancePrivateData.h b/Engine/Plugins/Experimental/Mutable/Source/CustomizableObject/Private/MuCO/CustomizableInstancePrivateData.h
--- a/Engine/Plugins/Experimental/Mutable/Source/CustomizableObject/Private/MuCO/CustomizableInstancePrivateData.h
+++ b/Engine/Plugins/Experimental/Mutable/Source/CustomizableObject/Private/MuCO/CustomizableInstancePrivateData.h
@@ -220,14 +220,10 @@ public:
 	/** Temporal function. Used by GetParameters and ReloadParameters but its body should be moved somewhere else (probably on some other Instance update function). */
 	void InstanceUpdateFlags(const UCustomizableObjectInstance& Public);
 
-	/** See FCustomizableObjectInstanceDescriptor::GetParameters(...). */
-	mu::ParametersPtr GetParameters(UCustomizableObjectInstance* Public);
-
 	/** See FCustomizableObjectInstanceDescriptor::ReloadParameters(...). */
 	void ReloadParameters(UCustomizableObjectInstance* Public, bool bInvalidatePreviousData = false);
 
-	/** If bOnlyUpdatePrioForUpdates is true, Instances that need an update won't be modified. Just their entry in InOutRequestedUpdates will get its priority updated  */
-	void TickUpdateCloseCustomizableObjects(UCustomizableObjectInstance& Public, const bool bOnlyUpdatePrioForUpdates, FMutableInstanceUpdateMap& InOutRequestedUpdates);
+	void TickUpdateCloseCustomizableObjects(UCustomizableObjectInstance& Publics, FMutableInstanceUpdateMap& InOutRequestedUpdates);
 	void UpdateInstanceIfNotGenerated(UCustomizableObjectInstance& Public, FMutableInstanceUpdateMap& InOutRequestedUpdates);
 
 	// Returns true if success (?)
@@ -235,9 +231,6 @@ public:
 
 	static void ReleaseMutableTexture(const FMutableImageCacheKey& MutableTextureKey, UTexture2D* Texture, struct FMutableResourceCache& Cache);
 
-	// Parameter decoration generation: game thread step
-	void UpdateParameterDecorationsEngineResources(const TSharedPtr<FMutableOperationData>& OperationData);
-
 	// Prepare the data for the unreal textures, but don't create them because
 	// it runs in the mutable thread
 	// \TODO: CustomizableObject shouldn't be here
diff --git a/Engine/Plugins/Experimental/Mutable/Source/CustomizableObject/Private/MuCO/CustomizableObjectInstance.cpp b/Engine/Plugins/Experimental/Mutable/Source/CustomizableObject/Private/MuCO/CustomizableObjectInstance.cpp
--- a/Engine/Plugins/Experimental/Mutable/Source/CustomizableObject/Private/MuCO/CustomizableObjectInstance.cpp
+++ b/Engine/Plugins/Experimental/Mutable/Source/CustomizableObject/Private/MuCO/CustomizableObjectInstance.cpp
@@ -547,24 +547,6 @@ void UCustomizableInstancePrivateData::InstanceUpdateFlags(const UCustomizableOb
 }
 
 
-mu::ParametersPtr UCustomizableInstancePrivateData::GetParameters(UCustomizableObjectInstance* Public) 
-{
-	const UCustomizableObject* CustomizableObject = Public->GetCustomizableObject();
-	if (!CustomizableObject)
-	{
-		// May happen when deleting assets in editor, while they are still open.
-		return nullptr;
-	}
-
-	if (!CustomizableObject->IsCompiled())
-	{	
-		return nullptr;
-	}
-		
-	return Public->GetDescriptor().GetParameters();
-}
-
-
 void UCustomizableInstancePrivateData::ReloadParameters(UCustomizableObjectInstance* Public, bool bInvalidatePreviousData)
 {
 	const UCustomizableObject* CustomizableObject = Public->GetCustomizableObject();
@@ -869,32 +851,35 @@ void UCustomizableObjectInstance::DoUpdateSkeletalMesh(bool bIsCloseDistTick, bo
 	check(IsInGameThread());
 	check(!OptionalUpdateRequired || *OptionalUpdateRequired != EUpdateRequired::NoUpdate); // If no update is required this functions must not be called.
 
+	UCustomizableObjectSystem* System = UCustomizableObjectSystem::GetInstance();
+
+	const mu::ParametersPtr Parameters = Descriptor.GetParameters();
+	
+	// Cache Texture Parameters being used during the update:
+	System->GetPrivate()->GetImageProviderChecked()->CacheImages(*Parameters);
+
 	if (!CanUpdateInstance())
 	{
+		FinishUpdateGlobal(this, EUpdateResult::ErrorDiscarded, UpdateCallback, Parameters);
 		return;
 	}
-
-	UCustomizableObjectSystem* System = UCustomizableObjectSystem::GetInstance();
-
+	
 	const EUpdateRequired UpdateRequired = OptionalUpdateRequired ? *OptionalUpdateRequired : IsUpdateRequired(bIsCloseDistTick, bOnlyUpdateIfNotGenerated, bIgnoreCloseDist);
 	switch (UpdateRequired)
 	{
 	case EUpdateRequired::NoUpdate:
 	{	
-		FinishUpdateGlobal(this, EUpdateResult::ErrorDiscarded, UpdateCallback);
+		FinishUpdateGlobal(this, EUpdateResult::ErrorDiscarded, UpdateCallback, Parameters);
 		break;
 	}		
 	case EUpdateRequired::Update:
 	{
 		EQueuePriorityType Priority = GetUpdatePriority(bIsCloseDistTick, bOnlyUpdateIfNotGenerated, bIgnoreCloseDist, bForceHighPriority);
-		const bool bShouldUpdateLODs = GetPrivate()->HasCOInstanceFlags(PendingLODsUpdate);
 
-		{
-			const uint32 InstanceId = GetUniqueID();
-			const float Distance = FMath::Sqrt(GetPrivate()->LastMinSquareDistFromComponentToPlayer);
-			const bool bIsPlayerOrNearIt = GetPrivate()->HasCOInstanceFlags(UsedByPlayerOrNearIt);
-			UE_LOG(LogMutable, Verbose, TEXT("Started UpdateSkeletalMesh Async. of Instance %d with priority %d at dist %f bIsPlayerOrNearIt=%d, frame=%d"), InstanceId, static_cast<int32>(Priority), Distance, bIsPlayerOrNearIt, GFrameNumber);				
-		}
+		const uint32 InstanceId = GetUniqueID();
+		const float Distance = FMath::Sqrt(GetPrivate()->LastMinSquareDistFromComponentToPlayer);
+		const bool bIsPlayerOrNearIt = GetPrivate()->HasCOInstanceFlags(UsedByPlayerOrNearIt);
+		UE_LOG(LogMutable, Verbose, TEXT("Started UpdateSkeletalMesh Async. of Instance %d with priority %d at dist %f bIsPlayerOrNearIt=%d, frame=%d"), InstanceId, static_cast<int32>(Priority), Distance, bIsPlayerOrNearIt, GFrameNumber);				
 
 		UpdateDescriptorRuntimeHash = FDescriptorRuntimeHash(Descriptor);
 
@@ -907,8 +892,8 @@ void UCustomizableObjectInstance::DoUpdateSkeletalMesh(bool bIsCloseDistTick, bo
 			}
 		}
 #endif
-		
-		if (bShouldUpdateLODs)
+
+		if (GetPrivate()->HasCOInstanceFlags(PendingLODsUpdate))
 		{
 			UE_LOG(LogMutable, Verbose, TEXT("LOD change: %d, %d -> %d, %d"), GetCurrentMinLOD(), GetCurrentMaxLOD(), GetMinLODToLoad(), GetMinLODToLoad());
 		}
@@ -924,7 +909,7 @@ void UCustomizableObjectInstance::DoUpdateSkeletalMesh(bool bIsCloseDistTick, bo
 	{
 		System->GetPrivate()->InitDiscardResourcesSkeletalMesh(this);
 		
-		FinishUpdateGlobal(this, EUpdateResult::ErrorDiscarded, UpdateCallback);
+		FinishUpdateGlobal(this, EUpdateResult::ErrorDiscarded, UpdateCallback, Parameters);
 		break;
 	}
 
@@ -934,7 +919,7 @@ void UCustomizableObjectInstance::DoUpdateSkeletalMesh(bool bIsCloseDistTick, bo
 }
 
 
-void UCustomizableInstancePrivateData::TickUpdateCloseCustomizableObjects(UCustomizableObjectInstance& Public, const bool bOnlyUpdatePrioForUpdates, FMutableInstanceUpdateMap& InOutRequestedUpdates)
+void UCustomizableInstancePrivateData::TickUpdateCloseCustomizableObjects(UCustomizableObjectInstance& Public, FMutableInstanceUpdateMap& InOutRequestedUpdates)
 {
 	if (!Public.CanUpdateInstance())
 	{
@@ -944,7 +929,7 @@ void UCustomizableInstancePrivateData::TickUpdateCloseCustomizableObjects(UCusto
 	const EUpdateRequired UpdateRequired = Public.IsUpdateRequired(true, true, false);
 	if (UpdateRequired != EUpdateRequired::NoUpdate) // Since this is done in the tick, avoid starting an update that we know for sure that would not be performed. Once started it has some performance implications that we want to avoid.
 	{
-		if (!bOnlyUpdatePrioForUpdates || UpdateRequired == EUpdateRequired::Discard)
+		if (UpdateRequired == EUpdateRequired::Discard)
 		{
 			Public.DoUpdateSkeletalMesh(true, true, false, false, &UpdateRequired, nullptr);
 			InOutRequestedUpdates.Remove(&Public);
diff --git a/Engine/Plugins/Experimental/Mutable/Source/CustomizableObject/Private/MuCO/CustomizableObjectSystem.cpp b/Engine/Plugins/Experimental/Mutable/Source/CustomizableObject/Private/MuCO/CustomizableObjectSystem.cpp
--- a/Engine/Plugins/Experimental/Mutable/Source/CustomizableObject/Private/MuCO/CustomizableObjectSystem.cpp
+++ b/Engine/Plugins/Experimental/Mutable/Source/CustomizableObject/Private/MuCO/CustomizableObjectSystem.cpp
@@ -133,7 +133,7 @@ void FMutablePendingInstanceWork::AddUpdate(const FMutablePendingInstanceUpdate&
 		FMutablePendingInstanceUpdate TaskToEnqueue = UpdateToAdd;
 
 		FInstanceUpdateDelegate* Callback = ExistingUpdate->Callback;
-		FinishUpdateGlobal(ExistingUpdate->CustomizableObjectInstance.Get(), EUpdateResult::ErrorReplaced, Callback, &ExistingUpdate->InstanceDescriptor);
+		FinishUpdateGlobal(ExistingUpdate->CustomizableObjectInstance.Get(), EUpdateResult::ErrorReplaced, Callback,ExistingUpdate->CustomizableObjectInstance.Get()->GetDescriptor().GetParameters());
 
 		TaskToEnqueue.PriorityType = FMath::Min(ExistingUpdate->PriorityType, UpdateToAdd.PriorityType);
 		TaskToEnqueue.SecondsAtUpdate = FMath::Min(ExistingUpdate->SecondsAtUpdate, UpdateToAdd.SecondsAtUpdate);
@@ -148,7 +148,7 @@ void FMutablePendingInstanceWork::AddUpdate(const FMutablePendingInstanceUpdate&
 
 	if (FMutablePendingInstanceDiscard* ExistingDiscard = PendingInstanceDiscards.Find(UpdateToAdd.CustomizableObjectInstance))
 	{
-		FinishUpdateGlobal(ExistingDiscard->CustomizableObjectInstance.Get(), EUpdateResult::ErrorReplaced, nullptr, nullptr);
+		FinishUpdateGlobal(ExistingDiscard->CustomizableObjectInstance.Get(), EUpdateResult::ErrorReplaced, nullptr, ExistingDiscard->CustomizableObjectInstance.Get()->GetDescriptor().GetParameters());
 
 		PendingInstanceDiscards.Remove(ExistingDiscard->CustomizableObjectInstance);
 	}
@@ -171,7 +171,7 @@ void FMutablePendingInstanceWork::AddDiscard(const FMutablePendingInstanceDiscar
 {
 	if (FMutablePendingInstanceUpdate* ExistingUpdate = PendingInstanceUpdates.Find(TaskToEnqueue.CustomizableObjectInstance.Get()))
 	{
-		FinishUpdateGlobal(ExistingUpdate->CustomizableObjectInstance.Get(), EUpdateResult::ErrorReplaced, ExistingUpdate->Callback, &ExistingUpdate->InstanceDescriptor);
+		FinishUpdateGlobal(ExistingUpdate->CustomizableObjectInstance.Get(), EUpdateResult::ErrorReplaced, ExistingUpdate->Callback, ExistingUpdate->CustomizableObjectInstance.Get()->GetDescriptor().GetParameters());
 		PendingInstanceUpdates.Remove(ExistingUpdate->CustomizableObjectInstance);
 	}
 
@@ -783,13 +783,6 @@ void FinishUpdateGlobal(UCustomizableObjectInstance* Instance, EUpdateResult Upd
 }
 
 
-void FinishUpdateGlobal(UCustomizableObjectInstance* Instance, EUpdateResult UpdateResult, FInstanceUpdateDelegate* UpdateCallback, const FCustomizableObjectInstanceDescriptor* InstanceDescriptor, const FDescriptorRuntimeHash InUpdatedHash)
-{
-	FinishUpdateGlobal(Instance, UpdateResult, UpdateCallback, InstanceDescriptor ? InstanceDescriptor->GetParameters() : nullptr, 
-		InUpdatedHash);
-}
-
-
 /** Update the given Instance Skeletal Meshes and call its callbacks. */
 void UpdateSkeletalMesh(UCustomizableObjectInstance& CustomizableObjectInstance, const FDescriptorRuntimeHash& UpdatedDescriptorRuntimeHash, EUpdateResult UpdateResult, FInstanceUpdateDelegate* UpdateCallback, mu::ParametersPtrConst Parameters = nullptr)
 {
@@ -886,15 +879,16 @@ void FCustomizableObjectSystemPrivate::InitUpdateSkeletalMesh(UCustomizableObjec
 	{
 		if (UpdateDescriptorHash.IsSubset(FDescriptorRuntimeHash(QueueElem->InstanceDescriptor)))
 		{
-			FinishUpdateGlobal(&Instance, EUpdateResult::ErrorOptimized, UpdateCallback);			
+			FinishUpdateGlobal(&Instance, EUpdateResult::ErrorOptimized, UpdateCallback, Instance.GetDescriptor().GetParameters());			
 			return; // The the requested update is equal to the last enqueued update.
 		}
-	}
+	}	
 
 	if (CurrentMutableOperation &&
+		&Instance == CurrentMutableOperation->CustomizableObjectInstance &&
 		UpdateDescriptorHash.IsSubset(CurrentMutableOperation->InstanceDescriptorRuntimeHash))
 	{
-		FinishUpdateGlobal(&Instance, EUpdateResult::ErrorOptimized, UpdateCallback);			
+		FinishUpdateGlobal(&Instance, EUpdateResult::ErrorOptimized, UpdateCallback, Instance.GetDescriptor().GetParameters());			
 		return; // The requested update is equal to the running update.
 	}
 	
@@ -911,9 +905,6 @@ void FCustomizableObjectSystemPrivate::InitUpdateSkeletalMesh(UCustomizableObjec
 	{	
 		Instance.SkeletalMeshStatus = ESkeletalMeshState::AsyncUpdatePending;
 
-		// Cache Texture Parameters being used during the update:
-		GetImageProviderChecked()->CacheImages(*Instance.GetPrivate()->GetParameters(&Instance));
-
 		if (!bIsCloseDistTick) // When called from bIsCloseDistTick, the update operation is directly processed without going to a queue first
 		{
 			const FMutablePendingInstanceUpdate InstanceUpdate(&Instance, Priority, UpdateCallback, bNeverStream, MipsToSkip);
@@ -2112,23 +2103,32 @@ namespace impl
 
 		// Cache Texture Parameters
 		FUnrealMutableImageProvider* ImageProvider = UCustomizableObjectSystem::GetInstance()->GetPrivateChecked()->GetImageProviderChecked();
+
+		// Cache new Texture Parameters
+		for (const FCustomizableObjectTextureParameterValue& TextureParameters : ObjectInstance->GetDescriptor().GetTextureParameters())
+		{
+			ImageProvider->CacheImage(TextureParameters.ParameterValue, false);				
+		
+			for (const FString& TextureParameter: TextureParameters.ParameterRangeValues)
+			{
+				ImageProvider->CacheImage(TextureParameter, false);
+			}
+		}
 		
 		// Uncache old Texture Parameters
 		for (const FString& TextureParameter : ObjectInstancePrivateData->UpdateTextureParameters)
 		{
 			ImageProvider->UnCacheImage(TextureParameter, false);
 		}
-		ObjectInstancePrivateData->UpdateTextureParameters.Reset();
 
-		// Cache new Texture Parameters and update which ones are currently are being used
+		// Update which ones are currently are being used
+		ObjectInstancePrivateData->UpdateTextureParameters.Reset();
 		for (const FCustomizableObjectTextureParameterValue& TextureParameters : ObjectInstance->GetDescriptor().GetTextureParameters())
 		{
-			ImageProvider->CacheImage(TextureParameters.ParameterValue, false);				
 			ObjectInstancePrivateData->UpdateTextureParameters.Add(TextureParameters.ParameterValue);
 		
 			for (const FString& TextureParameter: TextureParameters.ParameterRangeValues)
 			{
-				ImageProvider->CacheImage(TextureParameter, false);
 				ObjectInstancePrivateData->UpdateTextureParameters.Add(TextureParameter);
 			}
 		}
@@ -2572,7 +2572,7 @@ bool UCustomizableObjectSystem::Tick(float DeltaTime)
 
 				if (ObjectInstancePrivateData->HasCOInstanceFlags(UsedByComponentInPlay))
 				{
-					ObjectInstancePrivateData->TickUpdateCloseCustomizableObjects(**CustomizableObjectInstance, true, RequestedLODUpdates);
+					ObjectInstancePrivateData->TickUpdateCloseCustomizableObjects(**CustomizableObjectInstance, RequestedLODUpdates);
 				}
 				else if (ObjectInstancePrivateData->HasCOInstanceFlags(UsedByComponent))
 				{
@@ -2701,9 +2701,9 @@ bool UCustomizableObjectSystem::Tick(float DeltaTime)
 				// previous if statements
 
 				// LOD updates are detected automatically and have never had a DoUpdateSkeletalMesh, so call it without enqueing a new update
-				UCustomizableInstancePrivateData* ObjectInstancePrivateData = LODUpdateCandidateFound->CustomizableObjectInstance->GetPrivate();
-				ObjectInstancePrivateData->TickUpdateCloseCustomizableObjects(*LODUpdateCandidateFound->CustomizableObjectInstance, false, RequestedLODUpdates);
-
+				const EUpdateRequired Required = EUpdateRequired::Update;
+				LODUpdateCandidateFound->CustomizableObjectInstance->DoUpdateSkeletalMesh(true, false, false, false, &Required, nullptr);
+				
 				bool bNeverStream = false;
 				int32 MipsToSkip = 0;
 
@@ -2921,7 +2921,7 @@ FMutableOperation FMutableOperation::CreateInstanceUpdate(UCustomizableObjectIns
 	Op.InstanceDescriptorRuntimeHash = InCustomizableObjectInstance->GetUpdateDescriptorRuntimeHash();
 	Op.bStarted = false;
 	Op.bBuildParameterRelevancy = InCustomizableObjectInstance->GetBuildParameterRelevancy();
-	Op.Parameters = InCustomizableObjectInstance->GetPrivate()->GetParameters(InCustomizableObjectInstance);
+	Op.Parameters = InCustomizableObjectInstance->GetDescriptor().GetParameters();
 
 	if (UpdateCallback)
 	{
diff --git a/Engine/Plugins/Experimental/Mutable/Source/CustomizableObject/Private/MuCO/UnrealMutableImageProvider.cpp b/Engine/Plugins/Experimental/Mutable/Source/CustomizableObject/Private/MuCO/UnrealMutableImageProvider.cpp
--- a/Engine/Plugins/Experimental/Mutable/Source/CustomizableObject/Private/MuCO/UnrealMutableImageProvider.cpp
+++ b/Engine/Plugins/Experimental/Mutable/Source/CustomizableObject/Private/MuCO/UnrealMutableImageProvider.cpp
@@ -570,7 +570,8 @@ void FUnrealMutableImageProvider::UnCacheImage(const FExternalImageID& Id, bool
 		}
 		else
 		{
-			--Result->ReferencesSystem;							
+			--Result->ReferencesSystem;
+			check(Result->ReferencesSystem >= 0); // Something went wrong. The image has been uncached more times than it has been cached.
 		}
 		
 		if (Result->ReferencesUser + Result->ReferencesSystem == 0)
