[Mutable] Fix missing material parameters when reusing surfaces between LODs.

#rnx

diff --git a/Engine/Plugins/Experimental/Mutable/Source/CustomizableObject/Private/MuCO/CustomizableObjectSystem.cpp b/Engine/Plugins/Experimental/Mutable/Source/CustomizableObject/Private/MuCO/CustomizableObjectSystem.cpp
--- a/Engine/Plugins/Experimental/Mutable/Source/CustomizableObject/Private/MuCO/CustomizableObjectSystem.cpp
+++ b/Engine/Plugins/Experimental/Mutable/Source/CustomizableObject/Private/MuCO/CustomizableObjectSystem.cpp
@@ -1424,14 +1424,14 @@ namespace impl
 
 						// Images
 						Surface.FirstImage = OperationData->InstanceUpdateData.Images.Num();
-						Surface.ImageCount = Instance->GetImageCount(MutableLODIndex, ComponentIndex, InstanceSurfaceIndex);
+						Surface.ImageCount = Instance->GetImageCount(BaseLODIndex, ComponentIndex, InstanceSurfaceIndex);
 						for (int32 ImageIndex = 0; ImageIndex < Surface.ImageCount; ++ImageIndex)
 						{
 							MUTABLE_CPUPROFILER_SCOPE(GetImageId);
 
 							OperationData->InstanceUpdateData.Images.Push({});
 							FInstanceUpdateData::FImage& Image = OperationData->InstanceUpdateData.Images.Last();
-							Image.Name = Instance->GetImageName(MutableLODIndex, ComponentIndex, InstanceSurfaceIndex, ImageIndex);
+							Image.Name = Instance->GetImageName(BaseLODIndex, ComponentIndex, InstanceSurfaceIndex, ImageIndex);
 							Image.ImageID = Instance->GetImageId(BaseLODIndex, ComponentIndex, BaseSurfaceIndex, ImageIndex);
 							Image.FullImageSizeX = 0;
 							Image.FullImageSizeY = 0;
@@ -1440,26 +1440,26 @@ namespace impl
 
 						// Vectors
 						Surface.FirstVector = OperationData->InstanceUpdateData.Vectors.Num();
-						Surface.VectorCount = Instance->GetVectorCount(MutableLODIndex, ComponentIndex, InstanceSurfaceIndex);
+						Surface.VectorCount = Instance->GetVectorCount(BaseLODIndex, ComponentIndex, InstanceSurfaceIndex);
 						for (int32 VectorIndex = 0; VectorIndex < Surface.VectorCount; ++VectorIndex)
 						{
 							MUTABLE_CPUPROFILER_SCOPE(GetVector);
 							OperationData->InstanceUpdateData.Vectors.Push({});
 							FInstanceUpdateData::FVector& Vector = OperationData->InstanceUpdateData.Vectors.Last();
-							Vector.Name = Instance->GetVectorName(MutableLODIndex, ComponentIndex, InstanceSurfaceIndex, VectorIndex);
-							Vector.Vector = Instance->GetVector(MutableLODIndex, ComponentIndex, InstanceSurfaceIndex, VectorIndex);
+							Vector.Name = Instance->GetVectorName(BaseLODIndex, ComponentIndex, InstanceSurfaceIndex, VectorIndex);
+							Vector.Vector = Instance->GetVector(BaseLODIndex, ComponentIndex, InstanceSurfaceIndex, VectorIndex);
 						}
 
 						// Scalars
 						Surface.FirstScalar = OperationData->InstanceUpdateData.Scalars.Num();
-						Surface.ScalarCount = Instance->GetScalarCount(MutableLODIndex, ComponentIndex, InstanceSurfaceIndex);
+						Surface.ScalarCount = Instance->GetScalarCount(BaseLODIndex, ComponentIndex, InstanceSurfaceIndex);
 						for (int32 ScalarIndex = 0; ScalarIndex < Surface.ScalarCount; ++ScalarIndex)
 						{
 							MUTABLE_CPUPROFILER_SCOPE(GetScalar);
 							OperationData->InstanceUpdateData.Scalars.Push({});
 							FInstanceUpdateData::FScalar& Scalar = OperationData->InstanceUpdateData.Scalars.Last();
-							Scalar.Name = Instance->GetScalarName(MutableLODIndex, ComponentIndex, InstanceSurfaceIndex, ScalarIndex);
-							Scalar.Scalar = Instance->GetScalar(MutableLODIndex, ComponentIndex, InstanceSurfaceIndex, ScalarIndex);
+							Scalar.Name = Instance->GetScalarName(BaseLODIndex, ComponentIndex, InstanceSurfaceIndex, ScalarIndex);
+							Scalar.Scalar = Instance->GetScalar(BaseLODIndex, ComponentIndex, InstanceSurfaceIndex, ScalarIndex);
 						}
 					}
 				}
