[Mutable] Mutable Code Viewer search by name.
- Added string based search option to locate operations with certain labels.
- Added option to interpret the string provided as a regular expression.

diff --git a/Engine/Plugins/Experimental/Mutable/Source/CustomizableObjectEditor/Private/MuCOE/SMutableCodeViewer.cpp b/Engine/Plugins/Experimental/Mutable/Source/CustomizableObjectEditor/Private/MuCOE/SMutableCodeViewer.cpp
--- a/Engine/Plugins/Experimental/Mutable/Source/CustomizableObjectEditor/Private/MuCOE/SMutableCodeViewer.cpp
+++ b/Engine/Plugins/Experimental/Mutable/Source/CustomizableObjectEditor/Private/MuCOE/SMutableCodeViewer.cpp
@@ -35,6 +35,10 @@
 #include "Widgets/Layout/SScrollBox.h"
 #include "Widgets/Views/STreeView.h"
 
+
+#include "Widgets/Input/SSearchBox.h"
+#include "Internationalization/Regex.h"
+
 class FExtender;
 class FReferenceCollector;
 class FUICommandList;
@@ -132,149 +136,6 @@ public:
 		// Primary column showing the name of the operation and tye type
 		if (ColumnName == MutableCodeTreeViewColumns::OperationsColumnID)
 		{
-			const mu::FProgram& Program = RowItem->MutableModel->GetPrivate()->m_program;
-			const mu::OP_TYPE Type = Program.GetOpType(RowItem->MutableOperation);
-			FString OpName = mu::s_opNames[int32(Type)];
-			OpName.TrimEndInline();
-
-			// See if the operation type accepts additional information in the label
-			switch ( Type )
-			{
-			case mu::OP_TYPE::BO_PARAMETER:
-			case mu::OP_TYPE::NU_PARAMETER:
-			case mu::OP_TYPE::SC_PARAMETER:
-			case mu::OP_TYPE::CO_PARAMETER:
-			case mu::OP_TYPE::PR_PARAMETER:
-			case mu::OP_TYPE::IM_PARAMETER:
-			case mu::OP_TYPE::ST_PARAMETER:
-			{
-				mu::OP::ParameterArgs Args = Program.GetOpArgs<mu::OP::ParameterArgs>(RowItem->MutableOperation);
-				OpName += TEXT(" ");
-				OpName += StringCast<TCHAR>(Program.m_parameters[int32(Args.variable)].m_name.c_str()).Get();
-				break;
-			}
-
-			case mu::OP_TYPE::IM_SWIZZLE:
-			{
-				mu::OP::ImageSwizzleArgs Args = Program.GetOpArgs<mu::OP::ImageSwizzleArgs>(RowItem->MutableOperation);
-				OpName += TEXT(" ");
-				OpName += StringCast<TCHAR>(mu::TypeInfo::s_imageFormatName[int32(Args.format)]).Get();
-				break;
-			}
-
-			case mu::OP_TYPE::IM_PIXELFORMAT:
-			{
-				mu::OP::ImagePixelFormatArgs Args = Program.GetOpArgs<mu::OP::ImagePixelFormatArgs>(RowItem->MutableOperation);
-				OpName += TEXT(" ");
-				OpName += StringCast<TCHAR>(mu::TypeInfo::s_imageFormatName[int32(Args.format)]).Get();
-				OpName += TEXT(" or ");
-				OpName += StringCast<TCHAR>(mu::TypeInfo::s_imageFormatName[int32(Args.formatIfAlpha)]).Get();
-				break;
-			}
-
-			case mu::OP_TYPE::IM_MIPMAP:
-			{
-				mu::OP::ImageMipmapArgs Args = Program.GetOpArgs<mu::OP::ImageMipmapArgs>(RowItem->MutableOperation);
-				OpName += FString::Printf(TEXT(" levels: %d-%d tail: %d"), Args.levels, Args.blockLevels, int32(Args.onlyTail));
-				break;
-			}
-
-			case mu::OP_TYPE::IM_RESIZE:
-			{
-				mu::OP::ImageResizeArgs Args = Program.GetOpArgs<mu::OP::ImageResizeArgs>(RowItem->MutableOperation);
-				OpName += FString::Printf(TEXT(" %d x %d"), int32(Args.size[0]), int32(Args.size[1]));
-				break;
-			}
-
-			case mu::OP_TYPE::IM_RESIZEREL:
-			{
-				mu::OP::ImageResizeRelArgs Args = Program.GetOpArgs<mu::OP::ImageResizeRelArgs>(RowItem->MutableOperation);
-				OpName += FString::Printf(TEXT(" %.3f x %.3f"), Args.factor[0], Args.factor[1]);
-				break;
-			}
-
-			case mu::OP_TYPE::IM_MULTILAYER:
-			{
-				mu::OP::ImageMultiLayerArgs Args = Program.GetOpArgs<mu::OP::ImageMultiLayerArgs>(RowItem->MutableOperation);
-				OpName += TEXT(" rgb: ");
-				OpName += mu::TypeInfo::s_blendModeName[int32(Args.blendType)];
-				OpName += TEXT(", a: ");
-				OpName += mu::TypeInfo::s_blendModeName[int32(Args.blendTypeAlpha)];
-				OpName += FString::Printf(TEXT(" a from %d "), Args.BlendAlphaSourceChannel);
-				OpName += FString::Printf(TEXT(" range-id: %d"), Args.rangeId);
-				OpName += FString::Printf(TEXT(" mask-from-alpha: %d"), int32(Args.bUseMaskFromBlended));
-				break;
-			}
-
-			case mu::OP_TYPE::IM_LAYER:
-			{
-				mu::OP::ImageLayerArgs Args = Program.GetOpArgs<mu::OP::ImageLayerArgs>(RowItem->MutableOperation);
-				OpName += TEXT(" rgb: ");
-				OpName += mu::TypeInfo::s_blendModeName[int32(Args.blendType)];
-				OpName += TEXT(", a: ");
-				OpName += mu::TypeInfo::s_blendModeName[int32(Args.blendTypeAlpha)];
-				OpName += FString::Printf(TEXT(" a from %d "), Args.BlendAlphaSourceChannel);
-				OpName += FString::Printf(TEXT(" flags %d"),Args.flags);
-				break;
-			}
-
-			case mu::OP_TYPE::IM_LAYERCOLOUR:
-			{
-				mu::OP::ImageLayerColourArgs Args = Program.GetOpArgs<mu::OP::ImageLayerColourArgs>(RowItem->MutableOperation);
-				OpName += TEXT(" rgb: ");
-				OpName += mu::TypeInfo::s_blendModeName[int32(Args.blendType)];
-				OpName += TEXT(" a: ");
-				OpName += mu::TypeInfo::s_blendModeName[int32(Args.blendTypeAlpha)];
-				OpName += TEXT(" a from ");
-				OpName += FString::Printf(TEXT(" a from %d "), Args.BlendAlphaSourceChannel);
-				OpName += FString::Printf(TEXT(" flags %d"), Args.flags);
-				break;
-			}
-
-			case mu::OP_TYPE::IM_PLAINCOLOUR:
-			{
-				mu::OP::ImagePlainColourArgs Args = Program.GetOpArgs<mu::OP::ImagePlainColourArgs>(RowItem->MutableOperation);
-				OpName += TEXT(" format: ");
-				OpName += StringCast<TCHAR>(mu::TypeInfo::s_imageFormatName[int32(Args.format)]).Get();
-				OpName += FString::Printf(TEXT(" size %d x %d"), Args.size[0], Args.size[1]);
-				OpName += FString::Printf(TEXT(" mips %d"), Args.LODs);
-				break;
-			}
-
-			case mu::OP_TYPE::IN_ADDIMAGE:
-			{
-				mu::OP::InstanceAddArgs Args = Program.GetOpArgs<mu::OP::InstanceAddArgs>(RowItem->MutableOperation);
-				if (Program.m_constantStrings.IsValidIndex(Args.name))
-				{
-					OpName += TEXT(" name: ");
-					OpName += FString(Program.m_constantStrings[Args.name].c_str());
-				}
- 				break;
-			}
-
-			default:
-				break;
-			}
-
-			// Prepare the text shown on the UI side of the operation tree
-			FString MainLabel = FString::Printf(TEXT("%s %d : %s"), *RowItem->Caption, int32(RowItem->MutableOperation), *OpName);
-
-			// DEBUG : 
-			// FString IndexOnTree = FString::FromInt(RowItem->IndexOnTree);
-			// IndexOnTree.Append(TEXT("- "));
-			// MainLabel.InsertAt(0,IndexOnTree);
-
-			// DEBUG : 
-			// FString RowStateIndex = FString::FromInt(RowItem->GetStateIndex());
-			// RowStateIndex.Append(TEXT(" st "));
-			// MainLabel.InsertAt(0,RowStateIndex);
-
-			
-			if (RowItem->DuplicatedOf)
-			{
-				MainLabel.Append(TEXT(" (duplicated)"));
-			}
-
 			// TODO:
 			const FSlateBrush* IconBrush = nullptr;
 			
@@ -320,7 +181,7 @@ public:
 					+ SHorizontalBox::Slot()
 					[
 						SNew(STextBlock)
-						.Text(FText::FromString(MainLabel))
+						.Text(FText::FromString(RowItem->MainLabel))
 						.ColorAndOpacity(RowItem->LabelColor)
 					]
 				]
@@ -462,7 +323,14 @@ void SMutableCodeViewer::SetCurrentModel(const TSharedPtr<mu::Model, ESPMode::Th
 	FoundModelOperationTypeElements.Empty();
 	ModelOperationTypes.Empty();
 	ModelOperationTypeNames.Empty();
+
+	// Reset navigation by type / constant resource
 	NavigationElements.Empty();
+	NavigationIndex = -1;
+
+	// Reset navigation by string
+	NameBasedNavigationElements.Empty();
+	StringNavigationIndex = -1;
 
 	// Generate all elements before starting the tree UI so we have a deterministic set of unique and duplicated elements
 	GenerateAllTreeElements();
@@ -580,7 +448,63 @@ void SMutableCodeViewer::Construct(const FArguments& InArgs, const TSharedPtr<mu
 			.Value(0.35f)
 			[
 				SNew(SVerticalBox)
-		
+
+				// Search box for tree operations
+				+ SVerticalBox::Slot()
+				.AutoHeight()
+				.HAlign(HAlign_Right)
+				[
+					SNew(SHorizontalBox)
+
+					// Search by name
+					+ SHorizontalBox::Slot()
+					.AutoWidth()
+					[
+						SNew(SHorizontalBox)
+
+						+ SHorizontalBox::Slot()
+						.AutoWidth()
+						[
+							SNew(STextBlock)
+							.Text(LOCTEXT("SelectedOperationByStringLabel","Search Operation by String :"))
+						]
+					
+						+ SHorizontalBox::Slot()
+						.AutoWidth()
+						[
+							SNew(SSearchBox)
+							.HintText(LOCTEXT("OperationToSearchHintText","Search OP"))
+							.SearchResultData(this,&SMutableCodeViewer::SearchResultsData)
+							.OnSearch(this, &SMutableCodeViewer::OnTreeStringSearch)
+							.OnTextChanged(this, &SMutableCodeViewer::OnTreeSearchTextChanged)
+							.OnTextCommitted(this, &SMutableCodeViewer::OnTreeSearchTextCommitted)
+						]
+					]
+					
+					
+					// Regex control for search by name
+					+ SHorizontalBox::Slot()
+					.AutoWidth()
+					.Padding(4,2)
+					[
+						SNew(SHorizontalBox)
+
+						+ SHorizontalBox::Slot()
+						.AutoWidth()
+						[
+							SNew(STextBlock)
+							.Text(LOCTEXT("OperationToSearchRegexLabel","Is RegEx?"))
+						]
+
+						+ SHorizontalBox::Slot()
+						.AutoWidth()
+						[
+							SNew(SCheckBox)
+							.OnCheckStateChanged(this,&SMutableCodeViewer::OnRegexToggleChanged)
+						]
+					]	
+				]
+				
 				// Operation type filtering slot
 				+ SVerticalBox::Slot()
 				.AutoHeight()
@@ -816,6 +740,234 @@ void SMutableCodeViewer::OnCurrentMipSkipChanged(int32 NewValue)
 	bIsPreviewPendingUpdate = true;
 }
 
+#pragma region CodeTree operation name search
+
+
+void SMutableCodeViewer::OnRegexToggleChanged(ECheckBoxState CheckBoxState)
+{
+	const bool bPreChangeValue = bIsSearchStringRegularExpression;
+	bIsSearchStringRegularExpression = CheckBoxState == ECheckBoxState::Checked ? true : false;
+	
+	if (bPreChangeValue != bIsSearchStringRegularExpression)
+	{
+		CacheOperationsMatchingStringPattern();
+		GoToNextOperation();
+	}
+}
+
+void SMutableCodeViewer::OnTreeStringSearch(SSearchBox::SearchDirection SearchDirection)
+{
+	if (SearchDirection==SSearchBox::SearchDirection::Next)
+	{
+		GoToNextOperation();
+	}
+	else
+	{
+		GoToPreviousOperation();
+	}
+}
+
+void SMutableCodeViewer::GoToNextOperation()
+{
+	// Contingency : Prevent a second scroll operation from being performed if still we do not have the first target in view
+	if (bWasScrollToTargetRequested)
+	{
+		return;
+	}
+	
+	if (NameBasedNavigationElements.Num())
+	{
+		const int32 PreviousIndex = StringNavigationIndex;
+
+		// Focus on next target
+		StringNavigationIndex = StringNavigationIndex >= NameBasedNavigationElements.Num() - 1
+										   ? 0
+										   : StringNavigationIndex + 1;
+
+		if (StringNavigationIndex != PreviousIndex)
+		{
+			FocusViewOnNavigationTarget(NameBasedNavigationElements[StringNavigationIndex]);
+		}
+	}
+}
+
+
+void SMutableCodeViewer::GoToPreviousOperation()
+{
+	// Contingency : Prevent a second scroll operation from being performed if still we do not have the first target in view
+	if (bWasScrollToTargetRequested)
+	{
+		return;
+	}
+	
+	if (NameBasedNavigationElements.Num())
+	{
+		const int32 PreviousIndex = StringNavigationIndex;
+		
+		// Focus on previous target
+		StringNavigationIndex = StringNavigationIndex <= 0 ?  NameBasedNavigationElements.Num() -1 : StringNavigationIndex -1;
+
+		if (PreviousIndex != StringNavigationIndex)
+		{
+			FocusViewOnNavigationTarget(NameBasedNavigationElements[StringNavigationIndex]);
+		}
+	}
+}
+
+void SMutableCodeViewer::GoToTargetOperation(const int32& InTargetIndex)
+{
+	if (InTargetIndex == StringNavigationIndex)
+	{
+		return;
+	}
+	
+	if (NameBasedNavigationElements.Num() && InTargetIndex > 0 && InTargetIndex <= NameBasedNavigationElements.Num()-1)
+	{
+		// Focus on the target index
+		StringNavigationIndex = InTargetIndex;
+		FocusViewOnNavigationTarget(NameBasedNavigationElements[StringNavigationIndex]);
+	}
+}
+
+
+void SMutableCodeViewer::OnTreeSearchTextChanged(const FText& InUpdatedText)
+{
+	SearchString = InUpdatedText.ToString();
+}
+
+
+TOptional<SSearchBox::FSearchResultData> SMutableCodeViewer::SearchResultsData() const
+{
+	if (NameBasedNavigationElements.Num() == 0)
+	{
+		return TOptional<SSearchBox::FSearchResultData>();
+	}
+	return TOptional<SSearchBox::FSearchResultData>({ NameBasedNavigationElements.Num(), StringNavigationIndex + 1});
+}
+
+
+void SMutableCodeViewer::OnTreeSearchTextCommitted(const FText& InUpdatedText, ETextCommit::Type TextCommitType)
+{
+	if (TextCommitType == ETextCommit::OnEnter)
+	{
+		check (InUpdatedText.ToString() == SearchString);
+		CacheOperationsMatchingStringPattern();	
+		GoToNextOperation();
+	}
+}
+
+
+void SMutableCodeViewer::CacheOperationsMatchingStringPattern()
+{
+	check(MutableModel);
+	check(RootNodeAddresses.Num());
+	
+	if ( LastSearchedString == SearchString &&
+		bWasLastSearchRegEx == bIsSearchStringRegularExpression && 
+		LastSearchedModel == MutableModel )
+	{
+		// Do not perform a search again since the context has not changed
+		return;
+	} 
+	
+	if (!SearchString.IsEmpty())
+	{
+		UE_LOG(LogMutable,Display,TEXT("Starting string search with target string ""\"%s""\" "),*SearchString);
+	
+		// Object containing all data required by the search operation to be able to be called recursively
+		FElementsSearchCache SearchPayload;
+		// Initialize the Search Payload with the root node addresses. This way the search will use them as the root nodes where
+		// to start searching
+		SearchPayload.SetupRootBatch(RootNodeAddresses);
+
+		const mu::FProgram& Program = MutableModel->GetPrivate()->m_program;
+		GetOperationsMatchingStringPattern(SearchString,bIsSearchStringRegularExpression,SearchPayload, Program);
+	
+		// Dump the located resources array onto the navigation array
+		NameBasedNavigationElements = MoveTemp(SearchPayload.FoundElements);
+		SortElementsByTreeIndex(NameBasedNavigationElements);
+		
+		UE_LOG(LogMutable, Display, TEXT("Operations found with matching pattern ""\"%s""\" is  %i"), *SearchString, NameBasedNavigationElements.Num());
+	}
+	else
+	{
+		NameBasedNavigationElements.Reset();
+	}
+
+	// Reset the search index
+	StringNavigationIndex = -1;
+
+	// Keep track of what context was used to perform the search to avoid doing it again if the context has not changed
+	LastSearchedString = SearchString;
+	bWasLastSearchRegEx = bIsSearchStringRegularExpression;
+	LastSearchedModel = MutableModel;
+}
+
+
+void SMutableCodeViewer::GetOperationsMatchingStringPattern(const FString& InStringPattern,const bool bIsRegularExpression ,FElementsSearchCache& SearchPayload,const mu::FProgram& InProgram)
+{
+	// next batch of addresses to be explored 
+	TArray<FItemCacheKey> NextBatchAddressesData;
+	
+	for (int32 ParentIndex = 0; ParentIndex < SearchPayload.BatchData.Num(); ParentIndex++)
+	{	
+		const FItemCacheKey CacheKey = SearchPayload.BatchData[ParentIndex];
+		const FString OperationDescriptiveText = GetOperationDescriptiveText(CacheKey);
+		
+		bool bMatchesPattern = false;
+		if (!bIsRegularExpression)
+		{
+			// Check if the provided text is contained over the element identification text
+			bMatchesPattern = OperationDescriptiveText.Contains(InStringPattern);
+		}
+		else
+		{
+			FRegexPattern Pattern{InStringPattern};
+			FRegexMatcher RegexMatcher{Pattern,OperationDescriptiveText};
+			bMatchesPattern = RegexMatcher.FindNext();
+		}
+		
+		// Get one of the previous run "children" and treat as a parent to get it's children and process them
+		const mu::OP::ADDRESS& ParentAddress = SearchPayload.BatchData[ParentIndex].Child;
+		
+		if (bMatchesPattern)
+		{
+			SearchPayload.AddToFoundElements(ParentAddress,ParentIndex,ItemCache);
+		}
+		
+		// Get all NON PROCESSED the children of this operation to later be able to process them (on next recursive call)
+		SearchPayload.CacheChildrenOfAddressIfNotProcessed(ParentAddress, InProgram, NextBatchAddressesData);
+	}
+
+	// At this point all the addresses to be computed on the next batch have already been set and will be computed on
+	// the next recursive call
+	
+	// Explore children if found 
+	if (NextBatchAddressesData.Num())
+	{
+		// Cache next batch data so the next invocations is able to locate the provided addresses on the itemsCache
+		SearchPayload.BatchData = MoveTemp(NextBatchAddressesData);
+		
+		GetOperationsMatchingStringPattern(InStringPattern,bIsRegularExpression, SearchPayload, InProgram);
+	}
+}
+
+
+FString SMutableCodeViewer::GetOperationDescriptiveText(const FItemCacheKey& InItemCacheKey)
+{
+	FString OperationDescriptiveText;
+		
+	if (const TSharedPtr<FMutableCodeTreeElement>* Element = ItemCache.Find(InItemCacheKey))
+	{
+		OperationDescriptiveText = Element->Get()->MainLabel;
+		check (!OperationDescriptiveText.IsEmpty());
+	}
+	
+	return OperationDescriptiveText;
+}
+
+#pragma endregion 
+
 
 #pragma region CodeTree operation search
 
@@ -968,10 +1120,10 @@ void SMutableCodeViewer::OnSelectedOperationTypeFromTree()
 }
 
 
-void SMutableCodeViewer::SortNavigationElements()
+void SMutableCodeViewer::SortElementsByTreeIndex(TArray<TSharedPtr<FMutableCodeTreeElement>>& InElementsArrayToSort)
 {
 	// Sort the array from lower index to bigger index (0 , 1 , 2 ...)
-	NavigationElements.Sort([](const TSharedPtr<FMutableCodeTreeElement> A , const TSharedPtr<FMutableCodeTreeElement> B)
+	InElementsArrayToSort.Sort([](const TSharedPtr<FMutableCodeTreeElement> A , const TSharedPtr<FMutableCodeTreeElement> B)
 	{
 		return A->IndexOnTree < B->IndexOnTree;
 	});
@@ -998,7 +1150,7 @@ void SMutableCodeViewer::CacheAddressesOfOperationsOfType()
 	{
 		// Cache the navigation addresses so we are able to navigate over them
 		NavigationElements = MoveTemp(SearchPayload.FoundElements);
-		SortNavigationElements();
+		SortElementsByTreeIndex(NavigationElements);
 		
 		// Reset the navigation index
 		NavigationIndex = -1;
@@ -1012,7 +1164,6 @@ void SMutableCodeViewer::GetOperationsOfType(const mu::OP_TYPE& TargetOperationT
 	// next batch of addresses to be explored 
 	TArray<FItemCacheKey> NextBatchAddressesData;
 	
-	// uint32 CurrentParentBeingScanned = 0;
 	for	(int32 ParentIndex = 0 ; ParentIndex < InSearchPayload.BatchData.Num(); ParentIndex++)
 	{
 		// Get one of the previous run "children" and treat as a parent to get it's children and process them
@@ -1155,7 +1306,7 @@ FReply SMutableCodeViewer::OnGoToPreviousOperationButtonPressed()
 {
 	// Focus on previous target
 	NavigationIndex = NavigationIndex<=0 ? 0 : NavigationIndex - 1;
-	FocusViewOnNavigationTarget();
+	FocusViewOnNavigationTarget(NavigationElements[NavigationIndex]);
 	
 	return FReply::Handled();
 }
@@ -1164,37 +1315,39 @@ FReply SMutableCodeViewer::OnGoToNextOperationButtonPressed()
 {
 	// Focus on next target
 	NavigationIndex = NavigationIndex>=NavigationElements.Num() -1 ? NavigationElements.Num() -1 : NavigationIndex + 1;
-	FocusViewOnNavigationTarget();
+	FocusViewOnNavigationTarget(NavigationElements[NavigationIndex]);
 	
 	return FReply::Handled();
 }
 
-void SMutableCodeViewer::FocusViewOnNavigationTarget()
+void SMutableCodeViewer::FocusViewOnNavigationTarget(TSharedPtr<FMutableCodeTreeElement> InTargetElement)
 {
 	// Stage 1 : Expand all tree so all navigable elements get to be reachable
 	if (!bWasUniqueExpansionInvokedForNavigation && !bWasScrollToTargetRequested)
 	{
 		TreeExpandUnique();
 		bWasUniqueExpansionInvokedForNavigation = true;
+		
+		// Cache the current navigation target so after the update we can focus it 
+		ToFocusElement = InTargetElement;
+		
 		// Early exit, this method will get called again later after tree update
 		return;		
 	}
 	
 	// Stage 2 : Try to get to the targeted element. if not visible scroll into view
-	
-	// Locate the element on our map of elements
-	const TSharedPtr<FMutableCodeTreeElement> TargetElement = NavigationElements[NavigationIndex];
-	check (TargetElement.IsValid());
+	check (InTargetElement.IsValid());
 	
 	// If required scroll to the area where we know the element is going to be in view
 	// a way to ensure this happens is by calling 
-	if (TreeView->IsItemVisible(TargetElement))
+	if (TreeView->IsItemVisible(InTargetElement))
 	{
 		// Stage 3-b : Select the element we have provided since now is sure to be in view
 		
 		// This line selects the element with at the same time updates the UI to show the row representing this element selected
-		TreeView->SetSelection(TargetElement);
-	
+		TreeView->SetSelection(InTargetElement);
+		ToFocusElement.Reset();							// We have focused the target so we no longer need to keep a reference to it
+		
 		// Done!
 		// We have the element in view and we have selected it!
 	}
@@ -1205,8 +1358,8 @@ void SMutableCodeViewer::FocusViewOnNavigationTarget()
 		// Failing this check would mean we have performed a scroll but we are still not able to view the element
 		check (!bWasScrollToTargetRequested);
 		
-		// Amount of elements that can be shown by the UI at any given time (search for better way of knowing this)
-		TreeView->RequestScrollIntoView(TargetElement);
+		// Request the tree to show us the target element we want to get focused
+		TreeView->RequestScrollIntoView(InTargetElement);
 		
 		// Read this variable after the update and then select the object (easy at this point)
 		// You may want to just call again this method after refresh since the element will be on view
@@ -2058,7 +2211,7 @@ void SMutableCodeViewer::CacheRootNodeAddresses()
 
 		// Dump the located resources array onto the navigation array since we have content to navigate over
 		NavigationElements = MoveTemp(SearchPayload.FoundElements);
-		SortNavigationElements();
+		SortElementsByTreeIndex(NavigationElements);
 		
 		// Reset the navigation index
 		NavigationIndex = -1;
@@ -2389,7 +2542,7 @@ void SMutableCodeViewer::Tick(const FGeometry& AllottedGeometry, const double In
 		/** If we have expanded the tree elements in order to reach one of them then continue the operation */
 		if (bWasUniqueExpansionInvokedForNavigation || bWasScrollToTargetRequested)
 		{
-			FocusViewOnNavigationTarget();
+			FocusViewOnNavigationTarget(ToFocusElement);
 		}
 	}
 	
diff --git a/Engine/Plugins/Experimental/Mutable/Source/CustomizableObjectEditor/Private/MuCOE/SMutableCodeViewer.h b/Engine/Plugins/Experimental/Mutable/Source/CustomizableObjectEditor/Private/MuCOE/SMutableCodeViewer.h
--- a/Engine/Plugins/Experimental/Mutable/Source/CustomizableObjectEditor/Private/MuCOE/SMutableCodeViewer.h
+++ b/Engine/Plugins/Experimental/Mutable/Source/CustomizableObjectEditor/Private/MuCOE/SMutableCodeViewer.h
@@ -3,8 +3,11 @@
 #pragma once
 
 #include "MuR/ModelPrivate.h"
+#include "MuT/ErrorLogPrivate.h"
+#include "MuT/TypeInfo.h"
 #include "UObject/GCObject.h"
 #include "Widgets/SCompoundWidget.h"
+#include "Widgets/Input/SSearchBox.h"
 
 class STableViewBase;
 namespace ESelectInfo { enum Type : int; }
@@ -360,12 +363,16 @@ private:
 	int64 NavigationIndex = -1;
 	
 	/** Sort the contents of NavigationElements to follow a sequential pattern using the indices of the elements from 0 to +n*/
-	void SortNavigationElements();
-
+	void SortElementsByTreeIndex(TArray<TSharedPtr<FMutableCodeTreeElement>>& InElementsArrayToSort);
+	
+	/** If the focus could not be performed directly then cache the target until it is in view and can therefore be focused.
+	 * It will be valid if we are trying to focus something (it) and invalid if the focus operation has already been resolved. */
+	TSharedPtr<FMutableCodeTreeElement> ToFocusElement;
+	
 	/** Focus the current NavigationElement and places it into view. It also selects it so the previewer for the
 	 * element gets invoked
 	 */
-	void FocusViewOnNavigationTarget();
+	void FocusViewOnNavigationTarget(TSharedPtr<FMutableCodeTreeElement> InTargetElement);
 	
 	/** Wrapper struct designed to be used as cache for all elements found during the navigation system search for elements
 	 * of X Type or relation with a targeted constant resource. It is designed to be used and then destroyed once the
@@ -439,9 +446,9 @@ private:
 			const TSharedPtr<FMutableCodeTreeElement>* PtrFoundElement = InItemCache.Find(Key);
 			check(PtrFoundElement);
 			TSharedPtr<FMutableCodeTreeElement> FoundElement = *PtrFoundElement;
-			check(FoundElement != nullptr);
+			check(FoundElement);
 			
-			// Store this element on our temp Map of elements
+			// Store this element on our temp array of elements
 			FoundElements.Add(FoundElement);
 		}
 
@@ -556,6 +563,84 @@ private:
 	                                       const mu::DATATYPE ConstantDataType, const mu::OP::ADDRESS OperationAddress,
 	                                       const mu::FProgram& InProgram) const;
 	
+
+	/*
+	* Navigation : Text based navigation
+	*/
+
+	/** Current text being used to perform a search by name. It gets updated each time the user changes the text in the SearchBox */
+	FString SearchString;
+	
+	/** Flag that determines how to interpret the SearchString provided.
+	 * If true it means it is a regular expression.
+	 * If not it is just a fragment of text to seek for in the tree of operations. */
+	bool bIsSearchStringRegularExpression = false;
+	
+	/** Last search configuration -> Prevents recurrent searches with the same search parameters */
+	FString LastSearchedString;
+	bool bWasLastSearchRegEx = false;
+	TSharedPtr<mu::Model> LastSearchedModel;
+
+	/** A collection with all the elements we have found to be related with the SearchString. */
+	TArray<TSharedPtr<FMutableCodeTreeElement>> NameBasedNavigationElements;
+
+	/**  Index in the above array. Current navigation index for the navigation over elements related with the SearchString */
+	int32 StringNavigationIndex = -1;
+	
+	/** UI Callback that sets the backend flag to mirror the UI toggle.
+	* @param CheckBoxState The new state of the checkBox object being updated
+	*/
+	void OnRegexToggleChanged(ECheckBoxState CheckBoxState);
+
+	/** Exposes arrows to allow navigation over the found elements and handles the navigation using those.
+	 * @param SearchDirection The direction (up or down) being used to navigate the collection of tree elements
+	 */
+	void OnTreeStringSearch(SSearchBox::SearchDirection SearchDirection);
+
+	/** Invoked each time the user changes the search text in the UI. NewText cached in SearchString variable 
+	* @param InUpdatedText The new value set by the user for the TextBox used to provide the search string.
+	*/
+	void OnTreeSearchTextChanged(const FText& InUpdatedText);
+
+	/** Get the search results for the current search. Used by the Ui to generate some context related elements.
+	* @return Data for external search results to be shown in the search box.
+	*/
+	TOptional<SSearchBox::FSearchResultData> SearchResultsData() const;
+
+	/** Invoked each time the user commits to the text provided. It also caches the text provided by the user 
+	* @param InUpdatedText The current text being exposed in the UI textBox to be used for search. It is not being used
+	* since we rely in the value cached in the global variable "SearchString"
+	* @param TextCommitType Determines how the user did commit the search string.
+	*/
+	void OnTreeSearchTextCommitted(const FText& InUpdatedText, ETextCommit::Type TextCommitType);
+	
+	/** Move the focus to the next element in the collection of related elements to the cached SearchString. */
+	void GoToNextOperation();
+
+	/** Move the focus to the previous element in the collection of related elements to the cached SearchString. */
+	void GoToPreviousOperation();
+
+	/** Provided an index inside the NameBasedNavigationElements array change the focus to it. */
+	void GoToTargetOperation(const int32& InTargetIndex);
+
+	/** Provided a search string and a search type a series of elements to be navigated over */
+	void CacheOperationsMatchingStringPattern();
+
+	/** Searches the item cache in search of a string pattern that can be treated as a regular string or as a Regular expression.
+	 * @param InStringPattern The string to be searched over. It will get compared with the "Main Label" of the comparing operations.
+	 * @param bIsRegularExpression True if the InStringPattern is ought to be treated as a regular expression, false otherwise.
+	 * @param SearchPayload Structure containing the operations to be used as the origin of the search. This object will also contain the results of the search.
+	 * @param InProgram MutableProgram used to allow the search of the children of the operations being currently processed.
+	 */
+	void GetOperationsMatchingStringPattern(const FString& InStringPattern,const bool bIsRegularExpression, FElementsSearchCache& SearchPayload, const mu::FProgram& InProgram);
+
+	/** Provided a ItemCacheKey this method returns the string that represents it. It will not take in consideration duplicated entries since those are not part
+	 * of said ItemCache.
+	 * @param InItemCacheKey The key representing the current operation being searched.
+	 * @return The MainLabel used to display the provided operation.
+	 */
+	FString GetOperationDescriptiveText(const FItemCacheKey& InItemCacheKey);
+
 	
 	/*
 	* Main callbacks from the tree widget standard operation. 
@@ -762,6 +847,9 @@ public:
 		{
 			DuplicatedOf = *InDuplicatedOf;
 		}
+
+		// Generate the label to be used to display this operation in the operation tree
+		GenerateLabelText();
 		
 		// Process the data that can be extracted from the current state
 		SetElementCurrentState(InMutableStateIndex);
@@ -803,6 +891,151 @@ public:
 		bIsStateConstant = CurrentState.m_updateCache.Contains(MutableOperation);
 	}
 
+	void GenerateLabelText()
+	{
+		const mu::FProgram& Program = MutableModel->GetPrivate()->m_program;
+		const mu::OP_TYPE Type = Program.GetOpType(MutableOperation);
+		FString OpName = mu::s_opNames[static_cast<int32>(Type)];
+		OpName.TrimEndInline();
+
+		// See if the operation type accepts additional information in the label
+		switch ( Type )
+		{
+		case mu::OP_TYPE::BO_PARAMETER:
+		case mu::OP_TYPE::NU_PARAMETER:
+		case mu::OP_TYPE::SC_PARAMETER:
+		case mu::OP_TYPE::CO_PARAMETER:
+		case mu::OP_TYPE::PR_PARAMETER:
+		case mu::OP_TYPE::IM_PARAMETER:
+		case mu::OP_TYPE::ST_PARAMETER:
+		{
+			mu::OP::ParameterArgs Args = Program.GetOpArgs<mu::OP::ParameterArgs>(MutableOperation);
+			OpName += TEXT(" ");
+			OpName += StringCast<TCHAR>(Program.m_parameters[int32(Args.variable)].m_name.c_str()).Get();
+			break;
+		}
+
+		case mu::OP_TYPE::IM_SWIZZLE:
+		{
+			mu::OP::ImageSwizzleArgs Args = Program.GetOpArgs<mu::OP::ImageSwizzleArgs>(MutableOperation);
+			OpName += TEXT(" ");
+			OpName += StringCast<TCHAR>(mu::TypeInfo::s_imageFormatName[int32(Args.format)]).Get();
+			break;
+		}
+
+		case mu::OP_TYPE::IM_PIXELFORMAT:
+		{
+			mu::OP::ImagePixelFormatArgs Args = Program.GetOpArgs<mu::OP::ImagePixelFormatArgs>(MutableOperation);
+			OpName += TEXT(" ");
+			OpName += StringCast<TCHAR>(mu::TypeInfo::s_imageFormatName[int32(Args.format)]).Get();
+			OpName += TEXT(" or ");
+			OpName += StringCast<TCHAR>(mu::TypeInfo::s_imageFormatName[int32(Args.formatIfAlpha)]).Get();
+			break;
+		}
+
+		case mu::OP_TYPE::IM_MIPMAP:
+		{
+			mu::OP::ImageMipmapArgs Args = Program.GetOpArgs<mu::OP::ImageMipmapArgs>(MutableOperation);
+			OpName += FString::Printf(TEXT(" levels: %d-%d tail: %d"), Args.levels, Args.blockLevels, int32(Args.onlyTail));
+			break;
+		}
+
+		case mu::OP_TYPE::IM_RESIZE:
+		{
+			mu::OP::ImageResizeArgs Args = Program.GetOpArgs<mu::OP::ImageResizeArgs>(MutableOperation);
+			OpName += FString::Printf(TEXT(" %d x %d"), int32(Args.size[0]), int32(Args.size[1]));
+			break;
+		}
+
+		case mu::OP_TYPE::IM_RESIZEREL:
+		{
+			mu::OP::ImageResizeRelArgs Args = Program.GetOpArgs<mu::OP::ImageResizeRelArgs>(MutableOperation);
+			OpName += FString::Printf(TEXT(" %.3f x %.3f"), Args.factor[0], Args.factor[1]);
+			break;
+		}
+
+		case mu::OP_TYPE::IM_MULTILAYER:
+		{
+			mu::OP::ImageMultiLayerArgs Args = Program.GetOpArgs<mu::OP::ImageMultiLayerArgs>(MutableOperation);
+			OpName += TEXT(" rgb: ");
+			OpName += mu::TypeInfo::s_blendModeName[int32(Args.blendType)];
+			OpName += TEXT(", a: ");
+			OpName += mu::TypeInfo::s_blendModeName[int32(Args.blendTypeAlpha)];
+			OpName += FString::Printf(TEXT(" a from %d "), Args.BlendAlphaSourceChannel);
+			OpName += FString::Printf(TEXT(" range-id: %d"), Args.rangeId);
+			OpName += FString::Printf(TEXT(" mask-from-alpha: %d"), int32(Args.bUseMaskFromBlended));
+			break;
+		}
+
+		case mu::OP_TYPE::IM_LAYER:
+		{
+			mu::OP::ImageLayerArgs Args = Program.GetOpArgs<mu::OP::ImageLayerArgs>(MutableOperation);
+			OpName += TEXT(" rgb: ");
+			OpName += mu::TypeInfo::s_blendModeName[int32(Args.blendType)];
+			OpName += TEXT(", a: ");
+			OpName += mu::TypeInfo::s_blendModeName[int32(Args.blendTypeAlpha)];
+			OpName += FString::Printf(TEXT(" a from %d "), Args.BlendAlphaSourceChannel);
+			OpName += FString::Printf(TEXT(" flags %d"),Args.flags);
+			break;
+		}
+
+		case mu::OP_TYPE::IM_LAYERCOLOUR:
+		{
+			mu::OP::ImageLayerColourArgs Args = Program.GetOpArgs<mu::OP::ImageLayerColourArgs>(MutableOperation);
+			OpName += TEXT(" rgb: ");
+			OpName += mu::TypeInfo::s_blendModeName[int32(Args.blendType)];
+			OpName += TEXT(" a: ");
+			OpName += mu::TypeInfo::s_blendModeName[int32(Args.blendTypeAlpha)];
+			OpName += TEXT(" a from ");
+			OpName += FString::Printf(TEXT(" a from %d "), Args.BlendAlphaSourceChannel);
+			OpName += FString::Printf(TEXT(" flags %d"), Args.flags);
+			break;
+		}
+
+		case mu::OP_TYPE::IM_PLAINCOLOUR:
+		{
+			mu::OP::ImagePlainColourArgs Args = Program.GetOpArgs<mu::OP::ImagePlainColourArgs>(MutableOperation);
+			OpName += TEXT(" format: ");
+			OpName += StringCast<TCHAR>(mu::TypeInfo::s_imageFormatName[int32(Args.format)]).Get();
+			OpName += FString::Printf(TEXT(" size %d x %d"), Args.size[0], Args.size[1]);
+			OpName += FString::Printf(TEXT(" mips %d"), Args.LODs);
+			break;
+		}
+
+		case mu::OP_TYPE::IN_ADDIMAGE:
+		{
+			mu::OP::InstanceAddArgs Args = Program.GetOpArgs<mu::OP::InstanceAddArgs>(MutableOperation);
+			if (Program.m_constantStrings.IsValidIndex(Args.name))
+			{
+				OpName += TEXT(" name: ");
+				OpName += FString(Program.m_constantStrings[Args.name].c_str());
+			}
+ 			break;
+		}
+
+		default:
+			break;
+		}
+
+		// Prepare the text shown on the UI side of the operation tree
+		MainLabel = FString::Printf(TEXT("%s %d : %s"), *Caption, int32(MutableOperation), *OpName);
+
+		// DEBUG : 
+		// FString IndexOnTree = FString::FromInt(IndexOnTree);
+		// IndexOnTree.Append(TEXT("- "));
+		// MainLabel.InsertAt(0,IndexOnTree);
+
+		// DEBUG : 
+		// FString RowStateIndex = FString::FromInt(GetStateIndex());
+		// RowStateIndex.Append(TEXT(" st "));
+		// MainLabel.InsertAt(0,RowStateIndex);
+		
+		if (DuplicatedOf)
+		{
+			MainLabel.Append(TEXT(" (duplicated)"));
+		}
+	}
+	
 	int32 GetStateIndex() const
 	{
 		return CurrentMutableStateIndex;
@@ -810,21 +1043,21 @@ public:
 	
 public:
 
-	/** */
+	/** Model this operation is part of */
 	TSharedPtr<mu::Model, ESPMode::ThreadSafe> MutableModel;
 	
 	/** Mutable Graph Node represented in this tree row. */
 	mu::OP::ADDRESS MutableOperation;
 
-	/** Label representing this operation. */
-	FString Caption;
-
 	/** If this tree element is a duplicated of another op, this is the op. */
 	TSharedPtr<FMutableCodeTreeElement> DuplicatedOf;
 
 	/** The color to be used by the row representing this object */
 	FSlateColor LabelColor;
 
+	/** Label to be used to represent this operation in the tree */
+	FString MainLabel;
+
 	/*
 	 * Navigation metadata
 	 */
@@ -854,6 +1087,8 @@ private:
 	 * of duplicated elements (duplicated) since children of duplicated elements are still unique and therefore they can
 	 * be present in more than one state (but can just appear in one state at any given time)
 	 */
-	int32 CurrentMutableStateIndex = -1;	
+	int32 CurrentMutableStateIndex = -1;
 
+	/** String representing additional data like the state this operation is from. It may be empty. */
+	FString Caption;
 };
