[Mutable] Disable editor compilations and the CO system tick when the cook commandlet is running.

#rnx

diff --git a/Engine/Plugins/Experimental/Mutable/Source/CustomizableObject/Private/MuCO/CustomizableObject.cpp b/Engine/Plugins/Experimental/Mutable/Source/CustomizableObject/Private/MuCO/CustomizableObject.cpp
--- a/Engine/Plugins/Experimental/Mutable/Source/CustomizableObject/Private/MuCO/CustomizableObject.cpp
+++ b/Engine/Plugins/Experimental/Mutable/Source/CustomizableObject/Private/MuCO/CustomizableObject.cpp
@@ -745,6 +745,13 @@ bool UCustomizableObject::ConditionalAutoCompile()
 		return false;
 	}
 
+	// Don't compile if the cook commandlet is running. Do not add a warning/error. Otherwise, we could end up
+	// invalidating the cook for no reason.
+	if (IsRunningCookCommandlet())
+	{
+		return false;
+	}
+
 	// Don't compile if we're running game, Compile and/or if AutoCompile is disabled
 	if (IsRunningGame() || System->IsCompilationDisabled() || !System->IsAutoCompileEnabled())
 	{
diff --git a/Engine/Plugins/Experimental/Mutable/Source/CustomizableObject/Private/MuCO/CustomizableObjectSystem.cpp b/Engine/Plugins/Experimental/Mutable/Source/CustomizableObject/Private/MuCO/CustomizableObjectSystem.cpp
--- a/Engine/Plugins/Experimental/Mutable/Source/CustomizableObject/Private/MuCO/CustomizableObjectSystem.cpp
+++ b/Engine/Plugins/Experimental/Mutable/Source/CustomizableObject/Private/MuCO/CustomizableObjectSystem.cpp
@@ -1433,14 +1433,14 @@ namespace impl
 
 						// Images
 						Surface.FirstImage = OperationData->InstanceUpdateData.Images.Num();
-						Surface.ImageCount = Instance->GetImageCount(MutableLODIndex, ComponentIndex, InstanceSurfaceIndex);
+						Surface.ImageCount = Instance->GetImageCount(BaseLODIndex, ComponentIndex, InstanceSurfaceIndex);
 						for (int32 ImageIndex = 0; ImageIndex < Surface.ImageCount; ++ImageIndex)
 						{
 							MUTABLE_CPUPROFILER_SCOPE(GetImageId);
 
 							OperationData->InstanceUpdateData.Images.Push({});
 							FInstanceUpdateData::FImage& Image = OperationData->InstanceUpdateData.Images.Last();
-							Image.Name = Instance->GetImageName(MutableLODIndex, ComponentIndex, InstanceSurfaceIndex, ImageIndex);
+							Image.Name = Instance->GetImageName(BaseLODIndex, ComponentIndex, InstanceSurfaceIndex, ImageIndex);
 							Image.ImageID = Instance->GetImageId(BaseLODIndex, ComponentIndex, BaseSurfaceIndex, ImageIndex);
 							Image.FullImageSizeX = 0;
 							Image.FullImageSizeY = 0;
@@ -1449,26 +1449,26 @@ namespace impl
 
 						// Vectors
 						Surface.FirstVector = OperationData->InstanceUpdateData.Vectors.Num();
-						Surface.VectorCount = Instance->GetVectorCount(MutableLODIndex, ComponentIndex, InstanceSurfaceIndex);
+						Surface.VectorCount = Instance->GetVectorCount(BaseLODIndex, ComponentIndex, InstanceSurfaceIndex);
 						for (int32 VectorIndex = 0; VectorIndex < Surface.VectorCount; ++VectorIndex)
 						{
 							MUTABLE_CPUPROFILER_SCOPE(GetVector);
 							OperationData->InstanceUpdateData.Vectors.Push({});
 							FInstanceUpdateData::FVector& Vector = OperationData->InstanceUpdateData.Vectors.Last();
-							Vector.Name = Instance->GetVectorName(MutableLODIndex, ComponentIndex, InstanceSurfaceIndex, VectorIndex);
-							Vector.Vector = Instance->GetVector(MutableLODIndex, ComponentIndex, InstanceSurfaceIndex, VectorIndex);
+							Vector.Name = Instance->GetVectorName(BaseLODIndex, ComponentIndex, InstanceSurfaceIndex, VectorIndex);
+							Vector.Vector = Instance->GetVector(BaseLODIndex, ComponentIndex, InstanceSurfaceIndex, VectorIndex);
 						}
 
 						// Scalars
 						Surface.FirstScalar = OperationData->InstanceUpdateData.Scalars.Num();
-						Surface.ScalarCount = Instance->GetScalarCount(MutableLODIndex, ComponentIndex, InstanceSurfaceIndex);
+						Surface.ScalarCount = Instance->GetScalarCount(BaseLODIndex, ComponentIndex, InstanceSurfaceIndex);
 						for (int32 ScalarIndex = 0; ScalarIndex < Surface.ScalarCount; ++ScalarIndex)
 						{
 							MUTABLE_CPUPROFILER_SCOPE(GetScalar);
 							OperationData->InstanceUpdateData.Scalars.Push({});
 							FInstanceUpdateData::FScalar& Scalar = OperationData->InstanceUpdateData.Scalars.Last();
-							Scalar.Name = Instance->GetScalarName(MutableLODIndex, ComponentIndex, InstanceSurfaceIndex, ScalarIndex);
-							Scalar.Scalar = Instance->GetScalar(MutableLODIndex, ComponentIndex, InstanceSurfaceIndex, ScalarIndex);
+							Scalar.Name = Instance->GetScalarName(BaseLODIndex, ComponentIndex, InstanceSurfaceIndex, ScalarIndex);
+							Scalar.Scalar = Instance->GetScalar(BaseLODIndex, ComponentIndex, InstanceSurfaceIndex, ScalarIndex);
 						}
 					}
 				}
@@ -2553,6 +2553,12 @@ bool UCustomizableObjectSystem::Tick(float DeltaTime)
 	{
 		return true; // Assets are still being loaded, so subobjects won't be found, compiled objects incomplete and thus updates wrong
 	}
+
+	// Do not tick if the CookCommandlet is running.
+	if (IsRunningCookCommandlet())
+	{
+		return true;
+	}
 #endif
 	
 	// Get a new operation if we aren't working on one
diff --git a/Engine/Plugins/Experimental/Mutable/Source/CustomizableObjectEditor/Private/MuCOE/CustomizableObjectCompiler.cpp b/Engine/Plugins/Experimental/Mutable/Source/CustomizableObjectEditor/Private/MuCOE/CustomizableObjectCompiler.cpp
--- a/Engine/Plugins/Experimental/Mutable/Source/CustomizableObjectEditor/Private/MuCOE/CustomizableObjectCompiler.cpp
+++ b/Engine/Plugins/Experimental/Mutable/Source/CustomizableObjectEditor/Private/MuCOE/CustomizableObjectCompiler.cpp
@@ -938,6 +938,12 @@ void FCustomizableObjectCompiler::CompileInternal(UCustomizableObject* Object, c
 	Options.bAnimBpPhysicsManipulationEnabled = Object->bEnableAnimBpPhysicsAssetsManipualtion;
 	Options.ImageTiling = Object->CompileOptions.ImageTiling;
 
+	if (!InOptions.bIsCooking && IsRunningCookCommandlet())
+	{
+		UE_LOG(LogMutable, Display, TEXT("Editor compilation suspended for Customizable Object [%s]. Can not compile COs when the cook commandlet is running. "), *Object->GetName());
+		return;
+	}
+
 	if (Object->IsLocked() || !UCustomizableObjectSystem::GetInstance()->LockObject(Object))
 	{
 		UE_LOG(LogMutable, Display, TEXT("Customizable Object is already being compiled or updated %s. Please wait a few seconds and try again."), *Object->GetName());
@@ -956,8 +962,6 @@ void FCustomizableObjectCompiler::CompileInternal(UCustomizableObject* Object, c
 		UE_LOG(LogMutable, Display, TEXT("Compiling Customizable Object with %d LODBias."), InOptions.DebugBias);
 	}
 
-	CompilationLogsContainer.ClearMessageCounters();
-
 	FMutableGraphGenerationContext GenerationContext(Object, this, Options);
 	GenerationContext.ParamNamesToSelectedOptions = ParamNamesToSelectedOptions;
 	if (ParamNamesToSelectedOptions.Num() > 0)
@@ -978,6 +982,7 @@ void FCustomizableObjectCompiler::CompileInternal(UCustomizableObject* Object, c
 	}
 
 	// Clear Messages from previous Compilations
+	CompilationLogsContainer.ClearMessageCounters();
 	CompilationLogsContainer.ClearMessagesArray();
 
 	// Generate the mutable node expression
