[mutable] Fix customizable object instance placement

#rnx

diff --git a/Engine/Plugins/Experimental/Mutable/Source/CustomizableObjectEditor/Classes/MuCOE/CustomizableObjectInstanceFactory.h b/Engine/Plugins/Experimental/Mutable/Source/CustomizableObjectEditor/Classes/MuCOE/CustomizableObjectInstanceFactory.h
--- a/Engine/Plugins/Experimental/Mutable/Source/CustomizableObjectEditor/Classes/MuCOE/CustomizableObjectInstanceFactory.h
+++ b/Engine/Plugins/Experimental/Mutable/Source/CustomizableObjectEditor/Classes/MuCOE/CustomizableObjectInstanceFactory.h
@@ -3,7 +3,6 @@
 #pragma once
 
 #include "ActorFactories/ActorFactory.h"
-
 #include "CustomizableObjectInstanceFactory.generated.h"
 
 class AActor;
@@ -21,10 +20,11 @@ protected:
 	//~ Begin UActorFactory Interface
 	void PostSpawnActor(UObject* Asset, AActor* NewActor) override;
 	void PostCreateBlueprint(UObject* Asset, AActor* CDO) override;
-	FQuat AlignObjectToSurfaceNormal(const FVector& InSurfaceNormal, const FQuat& ActorRotation) const;
+	UObject* GetAssetFromActorInstance(AActor* ActorInstance) override;
 	bool CanCreateActorFrom(const FAssetData& AssetData, FText& OutErrorMsg) override;
 	//~ End UActorFactory Interface
 
-	int32 GetNumberOfComponents(class UCustomizableObjectInstance* COInstance);
 	USkeletalMesh* GetSkeletalMeshFromAsset(UObject* Asset, int32 ComponentIndex = 0) const;
+	int32 GetNumberOfComponents(class UCustomizableObjectInstance* COInstance);
+	FQuat AlignObjectToSurfaceNormal(const FVector& InSurfaceNormal, const FQuat& ActorRotation) const;
 };
diff --git a/Engine/Plugins/Experimental/Mutable/Source/CustomizableObjectEditor/Private/MuCOE/CustomizableObjectEditorModule.cpp b/Engine/Plugins/Experimental/Mutable/Source/CustomizableObjectEditor/Private/MuCOE/CustomizableObjectEditorModule.cpp
--- a/Engine/Plugins/Experimental/Mutable/Source/CustomizableObjectEditor/Private/MuCOE/CustomizableObjectEditorModule.cpp
+++ b/Engine/Plugins/Experimental/Mutable/Source/CustomizableObjectEditor/Private/MuCOE/CustomizableObjectEditorModule.cpp
@@ -61,6 +61,7 @@
 #include "MuCO/CustomizableObjectInstance.h"
 #include "MuCOE/Nodes/CustomizableObjectNodeGroupProjectorParameter.h"
 #include "UObject/UObjectIterator.h"
+#include "Subsystems/PlacementSubsystem.h"
 
 
 class AActor;
@@ -190,10 +191,11 @@ void FCustomizableObjectEditorModule::StartupModule()
 	PropertyModule.NotifyCustomizationModuleChanged();
 
 	// Register factory
-	UCustomizableObjectInstanceFactory* CustomizableObjectInstanceFactory = NewObject<UCustomizableObjectInstanceFactory>();
-	CustomizableObjectInstanceFactory->NewActorClass = ACustomizableSkeletalMeshActor::StaticClass();
-	UEditorEngine* engine = GEditor;
-	engine->ActorFactories.Add(CustomizableObjectInstanceFactory);
+	GEditor->ActorFactories.Add(NewObject<UCustomizableObjectInstanceFactory>());
+	if (UPlacementSubsystem* PlacementSubsystem = GEditor->GetEditorSubsystem<UPlacementSubsystem>())
+	{
+		PlacementSubsystem->RegisterAssetFactory(NewObject<UCustomizableObjectInstanceFactory>());
+	}
 	
 
 	// Additional UI style
diff --git a/Engine/Plugins/Experimental/Mutable/Source/CustomizableObjectEditor/Private/MuCOE/CustomizableObjectInstanceFactory.cpp b/Engine/Plugins/Experimental/Mutable/Source/CustomizableObjectEditor/Private/MuCOE/CustomizableObjectInstanceFactory.cpp
--- a/Engine/Plugins/Experimental/Mutable/Source/CustomizableObjectEditor/Private/MuCOE/CustomizableObjectInstanceFactory.cpp
+++ b/Engine/Plugins/Experimental/Mutable/Source/CustomizableObjectEditor/Private/MuCOE/CustomizableObjectInstanceFactory.cpp
@@ -17,13 +17,87 @@
 
 
 UCustomizableObjectInstanceFactory::UCustomizableObjectInstanceFactory(const FObjectInitializer& ObjectInitializer)
-    : Super(ObjectInitializer)
+	: Super(ObjectInitializer)
 {
     DisplayName = LOCTEXT("CustomizableObjectInstanceDisplayName", "Customizable Object Instance");
     NewActorClass = ACustomizableSkeletalMeshActor::StaticClass();
     bUseSurfaceOrientation = true;
 }
 
+
+void UCustomizableObjectInstanceFactory::PostSpawnActor(UObject* Asset, AActor* NewActor)
+{
+	Super::PostSpawnActor(Asset, NewActor);
+	UCustomizableObjectInstance* COInstance = Cast<UCustomizableObjectInstance>(Asset);
+	ACustomizableSkeletalMeshActor* NewCSMActor = CastChecked<ACustomizableSkeletalMeshActor>(NewActor);
+	if (NewCSMActor && COInstance)
+	{
+		int32 NumComponents = GetNumberOfComponents(COInstance);
+
+		for (int32 ComponentIndex = 0; ComponentIndex < NumComponents; ++ComponentIndex)
+		{
+			USkeletalMesh* SkeletalMesh = GetSkeletalMeshFromAsset(Asset, ComponentIndex);
+
+			if (ComponentIndex > 0 && ComponentIndex > NewCSMActor->GetNumComponents() - 1)
+			{
+				NewCSMActor->AttachNewComponent();
+			}
+
+			if (USkeletalMeshComponent* SkeletalMeshComp = NewCSMActor->GetSkeletalMeshComponentAt(ComponentIndex))
+			{
+				SkeletalMeshComp->UnregisterComponent();
+				UE_MUTABLE_SETSKINNEDASSET(SkeletalMeshComp, SkeletalMesh);
+
+				if (ComponentIndex == 0 && NewCSMActor->GetWorld()->IsGameWorld())
+				{
+					NewCSMActor->ReplicatedMesh = SkeletalMesh;
+				}
+
+				if (UCustomizableSkeletalComponent* CustomSkeletalComp = NewCSMActor->GetCustomizableSkeletalComponent(ComponentIndex))
+				{
+					CustomSkeletalComp->UnregisterComponent();
+					CustomSkeletalComp->CustomizableObjectInstance = COInstance;
+					CustomSkeletalComp->ComponentIndex = ComponentIndex;
+					CustomSkeletalComp->SetSkeletalMesh(SkeletalMesh, false);
+					CustomSkeletalComp->UpdateSkeletalMeshAsync();
+					CustomSkeletalComp->RegisterComponent();
+				}
+
+				SkeletalMeshComp->RegisterComponent();
+			}
+		}
+	}
+}
+
+
+void UCustomizableObjectInstanceFactory::PostCreateBlueprint(UObject* Asset, AActor* CDO)
+{
+	//if (Asset != NULL && CDO != NULL)
+	//{
+	//	USkeletalMesh* SkeletalMesh = GetSkeletalMeshFromAsset(Asset);
+	//	UAnimBlueprint* AnimBlueprint = Cast<UAnimBlueprint>(Asset);
+
+	//	ASkeletalMeshActor* SkeletalMeshActor = CastChecked<ASkeletalMeshActor>(CDO);
+	//	SkeletalMeshActor->GetSkeletalMeshComponent()->SkeletalMesh = SkeletalMesh;
+	//	SkeletalMeshActor->GetSkeletalMeshComponent()->AnimClass = AnimBlueprint ? Cast<UAnimBlueprintGeneratedClass>(AnimBlueprint->GeneratedClass) : NULL;
+	//}
+}
+
+UObject* UCustomizableObjectInstanceFactory::GetAssetFromActorInstance(AActor* ActorInstance)
+{
+	if (ACustomizableSkeletalMeshActor* CSMActor = CastChecked<ACustomizableSkeletalMeshActor>(ActorInstance))
+	{
+		if (CSMActor->GetNumComponents() > 0)
+		{
+			if (UCustomizableSkeletalComponent* CustomSkeletalComp = CSMActor->GetCustomizableSkeletalComponent(0))
+			{
+				return CustomSkeletalComp->CustomizableObjectInstance;
+			}
+		}
+	}
+	return nullptr;
+}
+
 bool UCustomizableObjectInstanceFactory::CanCreateActorFrom(const FAssetData& AssetData, FText& OutErrorMsg)
 {
     if (!AssetData.IsValid() ||
@@ -63,12 +137,6 @@ bool UCustomizableObjectInstanceFactory::CanCreateActorFrom(const FAssetData& As
         }
     }
 
-    if (!CustomizableObjectInstanceData.IsValid())
-    {
-        OutErrorMsg = LOCTEXT("NoSkeletalMeshAss", "No valid skeletal mesh was found associated with the animation sequence.");
-        return false;
-    }
-
     if (USkeletalMesh* SkeletalMeshCDO = Cast<USkeletalMesh>(AssetData.GetClass()->GetDefaultObject()))
     {
         if (SkeletalMeshCDO->HasCustomActorFactory())
@@ -102,67 +170,6 @@ USkeletalMesh* UCustomizableObjectInstanceFactory::GetSkeletalMeshFromAsset(UObj
 }
 
 
-void UCustomizableObjectInstanceFactory::PostSpawnActor(UObject* Asset, AActor* NewActor)
-{
-    UCustomizableObjectInstance* COInstance = Cast<UCustomizableObjectInstance>(Asset);
-    ACustomizableSkeletalMeshActor* NewCSMActor = CastChecked<ACustomizableSkeletalMeshActor>(NewActor);
-
-	if (NewCSMActor && COInstance)
-	{
-		int32 NumComponents = GetNumberOfComponents(COInstance);
-
-		for (int32 ComponentIndex = 0; ComponentIndex < NumComponents; ++ComponentIndex)
-		{
-			USkeletalMesh* SkeletalMesh = GetSkeletalMeshFromAsset(Asset, ComponentIndex);
-
-			if (ComponentIndex > 0 && ComponentIndex > NewCSMActor->GetNumComponents() - 1)
-			{
-				NewCSMActor->AttachNewComponent();
-			}
-
-			if (USkeletalMeshComponent* SkeletalMeshComp = NewCSMActor->GetSkeletalMeshComponentAt(ComponentIndex))
-			{
-				SkeletalMeshComp->UnregisterComponent();
-				UE_MUTABLE_SETSKINNEDASSET(SkeletalMeshComp,SkeletalMesh);
-
-				if (ComponentIndex == 0 && NewCSMActor->GetWorld()->IsGameWorld())
-				{
-					NewCSMActor->ReplicatedMesh = SkeletalMesh;
-				}
-
-				if (UCustomizableSkeletalComponent* CustomSkeletalComp = NewCSMActor->GetCustomizableSkeletalComponent(ComponentIndex))
-				{
-					CustomSkeletalComp->UnregisterComponent();
-					CustomSkeletalComp->CustomizableObjectInstance = COInstance;
-					CustomSkeletalComp->ComponentIndex = ComponentIndex;
-					CustomSkeletalComp->SetSkeletalMesh(SkeletalMesh, false);
-					CustomSkeletalComp->UpdateSkeletalMeshAsync();
-					CustomSkeletalComp->RegisterComponent();
-				}
-
-				SkeletalMeshComp->RegisterComponent();
-			}			
-		}
-
-		Super::PostSpawnActor(COInstance, NewActor);
-	}
-}
-
-
-void UCustomizableObjectInstanceFactory::PostCreateBlueprint(UObject* Asset, AActor* CDO)
-{
-    //if (Asset != NULL && CDO != NULL)
-    //{
-    //	USkeletalMesh* SkeletalMesh = GetSkeletalMeshFromAsset(Asset);
-    //	UAnimBlueprint* AnimBlueprint = Cast<UAnimBlueprint>(Asset);
-
-    //	ASkeletalMeshActor* SkeletalMeshActor = CastChecked<ASkeletalMeshActor>(CDO);
-    //	SkeletalMeshActor->GetSkeletalMeshComponent()->SkeletalMesh = SkeletalMesh;
-    //	SkeletalMeshActor->GetSkeletalMeshComponent()->AnimClass = AnimBlueprint ? Cast<UAnimBlueprintGeneratedClass>(AnimBlueprint->GeneratedClass) : NULL;
-    //}
-}
-
-
 FQuat UCustomizableObjectInstanceFactory::AlignObjectToSurfaceNormal(const FVector& InSurfaceNormal, const FQuat& ActorRotation) const
 {
     // Meshes align the Z (up) axis with the surface normal
