[Mutable] Add custom allocator policy for mutable resource memory tracking.

#rnx

diff --git a/Engine/Plugins/Experimental/Mutable/Source/MutableRuntime/Private/MuR/CodeRunner.h b/Engine/Plugins/Experimental/Mutable/Source/MutableRuntime/Private/MuR/CodeRunner.h
--- a/Engine/Plugins/Experimental/Mutable/Source/MutableRuntime/Private/MuR/CodeRunner.h
+++ b/Engine/Plugins/Experimental/Mutable/Source/MutableRuntime/Private/MuR/CodeRunner.h
@@ -30,6 +30,11 @@
 	#include "Async/TaskGraphFwd.h"
 #endif
 
+namespace mu::MemoryCounters
+{
+	struct FStreamingMemoryCounterTag {};
+	using  FStreamingMemoryCounter = TMemoryCounter<FStreamingMemoryCounterTag>;
+}
 
 namespace  mu
 {
@@ -571,10 +576,12 @@ namespace  mu
 		// Async rom loading control
 		struct FRomLoadOp
 		{
+			using StreamingDataContainerType = TArray<uint8, FDefaultMemoryTrackingAllocator<MemoryCounters::FStreamingMemoryCounter>>;
+
 			int32 m_romIndex = 0;
 			DATATYPE ConstantType = DT_NONE;
 			ModelReader::OPERATION_ID m_streamID;
-			TArray<uint8> m_streamBuffer;
+			StreamingDataContainerType m_streamBuffer;
 		};
 		TArray<FRomLoadOp> m_romLoadOps;
 
diff --git a/Engine/Plugins/Experimental/Mutable/Source/MutableRuntime/Private/MuR/CodeRunnerTasked.cpp b/Engine/Plugins/Experimental/Mutable/Source/MutableRuntime/Private/MuR/CodeRunnerTasked.cpp
--- a/Engine/Plugins/Experimental/Mutable/Source/MutableRuntime/Private/MuR/CodeRunnerTasked.cpp
+++ b/Engine/Plugins/Experimental/Mutable/Source/MutableRuntime/Private/MuR/CodeRunnerTasked.cpp
@@ -80,7 +80,6 @@ namespace mu
 				Ptr<Mesh> Value = Mesh::StaticUnserialise(arch);
 				program.m_constantMeshes[ResIndex].Value = Value;
 				check(program.m_constantMeshes[ResIndex].Value);
-				m_pSystem->WorkingMemoryManager.TrackedBudgetBytes_Rom += Value->GetDataSize();
 				break;
 			}
 			case DATATYPE::DT_IMAGE:
@@ -90,7 +89,6 @@ namespace mu
 				Ptr<Image> Value = Image::StaticUnserialise(arch);
 				program.m_constantImageLODs[ResIndex].Value = Value;
 				check(program.m_constantImageLODs[ResIndex].Value);
-				m_pSystem->WorkingMemoryManager.TrackedBudgetBytes_Rom += Value->GetDataSize();
 				break;
 			}
 			default:
@@ -98,8 +96,6 @@ namespace mu
 				break;
 			}
 
-			m_pSystem->WorkingMemoryManager.TrackedBudgetBytes_Stream -= o.m_streamBuffer.GetAllocatedSize();
-
 			o.m_romIndex = -1;
 			o.m_streamBuffer.Empty();
 		}
@@ -2092,7 +2088,6 @@ namespace mu
 			if (uint32(op->m_streamBuffer.Num()) < RomSize)
 			{
 				op->m_streamBuffer.SetNum((size_t)RomSize);
-				Runner->m_pSystem->WorkingMemoryManager.TrackedBudgetBytes_Stream += op->m_streamBuffer.GetAllocatedSize();
 			}
 
 			uint32 RomId = program.m_roms[RomIndex].Id;
@@ -2300,7 +2295,6 @@ namespace mu
 			if (uint32(op->m_streamBuffer.Num()) < RomSize)
 			{
 				op->m_streamBuffer.SetNum(RomSize);
-				Runner->m_pSystem->WorkingMemoryManager.TrackedBudgetBytes_Stream += op->m_streamBuffer.GetAllocatedSize();
 			}
 
 			uint32 RomId = program.m_roms[RomIndex].Id;
diff --git a/Engine/Plugins/Experimental/Mutable/Source/MutableRuntime/Private/MuR/ImageRLE.cpp b/Engine/Plugins/Experimental/Mutable/Source/MutableRuntime/Private/MuR/ImageRLE.cpp
--- a/Engine/Plugins/Experimental/Mutable/Source/MutableRuntime/Private/MuR/ImageRLE.cpp
+++ b/Engine/Plugins/Experimental/Mutable/Source/MutableRuntime/Private/MuR/ImageRLE.cpp
@@ -280,7 +280,7 @@ namespace mu
     //---------------------------------------------------------------------------------------------
     void CompressRLE_RGBA( int32 width, int32 rows,
                            const uint8* pBaseDataByte,
-                           TArray<uint8>& destData )
+                           Image::ImageDataContainerType& destData )
     {
         // TODO: Support for compression from compressed data size, like L_RLE formats.
         TArray<int8_t> rle;
@@ -432,7 +432,7 @@ namespace mu
     //---------------------------------------------------------------------------------------------
     void CompressRLE_RGB( int32 width, int32 rows,
                           const uint8* pBaseDataByte,
-                          TArray<uint8>& destData )
+                          Image::ImageDataContainerType& destData )
     {
         TArray<int8_t> rle;
         rle.Reserve(  (width*rows) );
diff --git a/Engine/Plugins/Experimental/Mutable/Source/MutableRuntime/Private/MuR/ImageRLE.h b/Engine/Plugins/Experimental/Mutable/Source/MutableRuntime/Private/MuR/ImageRLE.h
--- a/Engine/Plugins/Experimental/Mutable/Source/MutableRuntime/Private/MuR/ImageRLE.h
+++ b/Engine/Plugins/Experimental/Mutable/Source/MutableRuntime/Private/MuR/ImageRLE.h
@@ -66,7 +66,7 @@ namespace mu
     //---------------------------------------------------------------------------------------------
 	MUTABLERUNTIME_API extern void CompressRLE_RGBA(int32 width, int32 rows,
 									const uint8* pBaseDataByte,
-									TArray<uint8>& destData );
+									Image::ImageDataContainerType& destData );
 
     //---------------------------------------------------------------------------------------------
 	MUTABLERUNTIME_API extern void UncompressRLE_RGBA(int32 width, int32 rows,
@@ -85,7 +85,7 @@ namespace mu
     //---------------------------------------------------------------------------------------------
 	MUTABLERUNTIME_API extern void CompressRLE_RGB(int32 width, int32 rows,
 									const uint8* pBaseDataByte,
-									TArray<uint8>& destData );
+									Image::ImageDataContainerType& destData );
 
     //---------------------------------------------------------------------------------------------
 	MUTABLERUNTIME_API extern void UncompressRLE_RGB(int32 width, int32 rows,
diff --git a/Engine/Plugins/Experimental/Mutable/Source/MutableRuntime/Private/MuR/Serialisation.cpp b/Engine/Plugins/Experimental/Mutable/Source/MutableRuntime/Private/MuR/Serialisation.cpp
--- a/Engine/Plugins/Experimental/Mutable/Source/MutableRuntime/Private/MuR/Serialisation.cpp
+++ b/Engine/Plugins/Experimental/Mutable/Source/MutableRuntime/Private/MuR/Serialisation.cpp
@@ -27,12 +27,23 @@ namespace mu
     MUTABLE_IMPLEMENT_POD_SERIALISABLE( uint16 ); 
     MUTABLE_IMPLEMENT_POD_SERIALISABLE( uint32 ); 
     MUTABLE_IMPLEMENT_POD_SERIALISABLE( uint64 )
-
+	
+                                                                                      
+	MUTABLE_IMPLEMENT_POD_SERIALISABLE(vec2f);                                        
+	MUTABLE_IMPLEMENT_POD_SERIALISABLE(vec3f);                                        
+	MUTABLE_IMPLEMENT_POD_SERIALISABLE(mat3f);                                        
+	MUTABLE_IMPLEMENT_POD_SERIALISABLE(mat4f);                                        
+	MUTABLE_IMPLEMENT_POD_SERIALISABLE(vec2<int>);                                    
+	
+	// Unreal POD Serializables                                                       
+	MUTABLE_IMPLEMENT_POD_SERIALISABLE(FUintVector2);                                 
+	MUTABLE_IMPLEMENT_POD_SERIALISABLE(UE::Math::TIntVector2<uint16>);                
+	MUTABLE_IMPLEMENT_POD_SERIALISABLE(UE::Math::TIntVector2<int16>); 
 
 	template <>
     void operator<<<FString>(OutputArchive& arch, const FString& t)
     {
-	    const TArray<TCHAR> Data = t.GetCharArray();
+	    const TArray<TCHAR>& Data = t.GetCharArray();
     	arch << Data;
     }
 
@@ -45,39 +56,6 @@ namespace mu
     	
     	t = FString(Data);
     }
-                                 
-
-	MUTABLE_IMPLEMENT_POD_VECTOR_SERIALISABLE(float);                                 
-	MUTABLE_IMPLEMENT_POD_VECTOR_SERIALISABLE(double);                                
-	MUTABLE_IMPLEMENT_POD_VECTOR_SERIALISABLE(uint8);                                 
-	MUTABLE_IMPLEMENT_POD_VECTOR_SERIALISABLE(uint16);                                
-	MUTABLE_IMPLEMENT_POD_VECTOR_SERIALISABLE(uint32);                                
-	MUTABLE_IMPLEMENT_POD_VECTOR_SERIALISABLE(uint64);                                
-	MUTABLE_IMPLEMENT_POD_VECTOR_SERIALISABLE(int8);                                  
-	MUTABLE_IMPLEMENT_POD_VECTOR_SERIALISABLE(int16);                                 
-	MUTABLE_IMPLEMENT_POD_VECTOR_SERIALISABLE(int32);                                 
-	MUTABLE_IMPLEMENT_POD_VECTOR_SERIALISABLE(int64);                                 
-                                                                                      
-	MUTABLE_IMPLEMENT_POD_SERIALISABLE(vec2f);                                        
-	MUTABLE_IMPLEMENT_POD_SERIALISABLE(vec3f);                                        
-	MUTABLE_IMPLEMENT_POD_SERIALISABLE(mat3f);                                        
-	MUTABLE_IMPLEMENT_POD_SERIALISABLE(mat4f);                                        
-	MUTABLE_IMPLEMENT_POD_SERIALISABLE(vec2<int>);                                    
-	MUTABLE_IMPLEMENT_POD_VECTOR_SERIALISABLE(vec2f);                                 
-	MUTABLE_IMPLEMENT_POD_VECTOR_SERIALISABLE(vec3f);                                 
-	MUTABLE_IMPLEMENT_POD_VECTOR_SERIALISABLE(mat3f);                                 
-	MUTABLE_IMPLEMENT_POD_VECTOR_SERIALISABLE(mat4f);                                 
-	MUTABLE_IMPLEMENT_POD_VECTOR_SERIALISABLE(vec2<int>);
-	MUTABLE_IMPLEMENT_POD_VECTOR_SERIALISABLE(TCHAR);
-	
-	// Unreal POD Serializables                                                       
-	MUTABLE_IMPLEMENT_POD_SERIALISABLE(FUintVector2);                                 
-	MUTABLE_IMPLEMENT_POD_SERIALISABLE(UE::Math::TIntVector2<uint16>);                
-	MUTABLE_IMPLEMENT_POD_SERIALISABLE(UE::Math::TIntVector2<int16>);                 
-	MUTABLE_IMPLEMENT_POD_VECTOR_SERIALISABLE(FUintVector2);                          
-	MUTABLE_IMPLEMENT_POD_VECTOR_SERIALISABLE(UE::Math::TIntVector2<uint16>);         
-	MUTABLE_IMPLEMENT_POD_VECTOR_SERIALISABLE(UE::Math::TIntVector2<int16>);          
-                                                                                      
 
     //---------------------------------------------------------------------------------------------
     //---------------------------------------------------------------------------------------------
diff --git a/Engine/Plugins/Experimental/Mutable/Source/MutableRuntime/Private/MuR/SerialisationPrivate.h b/Engine/Plugins/Experimental/Mutable/Source/MutableRuntime/Private/MuR/SerialisationPrivate.h
--- a/Engine/Plugins/Experimental/Mutable/Source/MutableRuntime/Private/MuR/SerialisationPrivate.h
+++ b/Engine/Plugins/Experimental/Mutable/Source/MutableRuntime/Private/MuR/SerialisationPrivate.h
@@ -158,6 +158,30 @@ namespace mu
 		}																\
 		
 
+#define MUTABLE_IMPLEMENT_POD_VECTOR_SERIALISABLE(T)									     \
+	template<typename Alloc>                                                             \
+	void operator <<(OutputArchive& arch, const TArray<T, Alloc>& v)					 \
+	{                                                                                    \
+		uint32 Num = uint32(v.Num());													 \
+		arch << Num;																	 \
+		if (Num)																		 \
+		{                                                                                \
+			arch.GetPrivate()->m_pStream->Write(&v[0], Num * sizeof(T));				 \
+		}                                                                                \
+	}                                                                                    \
+	                                                                                     \
+	template<typename Alloc>                                                             \
+	void operator >>(InputArchive& arch, TArray<T, Alloc>& v)                            \
+	{                                                                                    \
+		uint32 Num;																		 \
+		arch >> Num;																	 \
+		v.SetNum(Num);																	 \
+		if (Num)																		 \
+		{                                                                                \
+			arch.GetPrivate()->m_pStream->Read(&v[0], Num * sizeof(T));					 \
+		}                                                                                \
+	}                                                                                    \
+
 	//---------------------------------------------------------------------------------------------
 	// TVariant custom serialize. Based on the default serialization.
 	//---------------------------------------------------------------------------------------------
@@ -311,68 +335,41 @@ namespace mu
 		}
 	}
 
-
-	//---------------------------------------------------------------------------------------------
-#define MUTABLE_IMPLEMENT_POD_VECTOR_SERIALISABLE(T)									\
-	template<>																			\
-	void DLLEXPORT operator<< <T>(OutputArchive& arch, const TArray<T>& v)						\
-	{                                                                                   \
-		uint32 Num = uint32(v.Num());													\
-		arch << Num;																	\
-		if (Num)																		\
-		{                                                                               \
-			arch.GetPrivate()->m_pStream->Write(&v[0], Num * sizeof(T));				\
-		}                                                                               \
-	}                                                                                   \
-																						\
-	template<>																			\
-	void DLLEXPORT operator>> <T>(InputArchive& arch, TArray<T>& v)								\
-	{                                                                                   \
-		uint32 Num;																		\
-		arch >> Num;																	\
-		v.SetNum(Num);																	\
-		if (Num)																		\
-		{                                                                               \
-			arch.GetPrivate()->m_pStream->Read(&v[0], Num * sizeof(T));					\
-		}                                                                               \
-	}																					\
-
-	MUTABLE_DEFINE_POD_VECTOR_SERIALISABLE(float);
-	MUTABLE_DEFINE_POD_VECTOR_SERIALISABLE(double);
-	MUTABLE_DEFINE_POD_VECTOR_SERIALISABLE(uint8);
-	MUTABLE_DEFINE_POD_VECTOR_SERIALISABLE(uint16);
-	MUTABLE_DEFINE_POD_VECTOR_SERIALISABLE(uint32);
-	MUTABLE_DEFINE_POD_VECTOR_SERIALISABLE(uint64);
-	MUTABLE_DEFINE_POD_VECTOR_SERIALISABLE(int8);
-	MUTABLE_DEFINE_POD_VECTOR_SERIALISABLE(int16);
-	MUTABLE_DEFINE_POD_VECTOR_SERIALISABLE(int32);
-	MUTABLE_DEFINE_POD_VECTOR_SERIALISABLE(int64);
-
 	MUTABLE_DEFINE_POD_SERIALISABLE(vec2f);
 	MUTABLE_DEFINE_POD_SERIALISABLE(vec3f);
 	MUTABLE_DEFINE_POD_SERIALISABLE(mat3f);
 	MUTABLE_DEFINE_POD_SERIALISABLE(mat4f);
 	MUTABLE_DEFINE_POD_SERIALISABLE(vec2<int>);
-	MUTABLE_DEFINE_POD_VECTOR_SERIALISABLE(vec2f);
-	MUTABLE_DEFINE_POD_VECTOR_SERIALISABLE(vec3f);
-	MUTABLE_DEFINE_POD_VECTOR_SERIALISABLE(mat3f);
-	MUTABLE_DEFINE_POD_VECTOR_SERIALISABLE(mat4f);
-	MUTABLE_DEFINE_POD_VECTOR_SERIALISABLE(vec2<int>);
-	MUTABLE_DEFINE_POD_VECTOR_SERIALISABLE(TCHAR);
 
 	// Unreal POD Serializables
 	MUTABLE_DEFINE_POD_SERIALISABLE(FUintVector2);
 	MUTABLE_DEFINE_POD_SERIALISABLE(UE::Math::TIntVector2<uint16>);
 	MUTABLE_DEFINE_POD_SERIALISABLE(UE::Math::TIntVector2<int16>);
-	MUTABLE_DEFINE_POD_VECTOR_SERIALISABLE(FUintVector2);
-	MUTABLE_DEFINE_POD_VECTOR_SERIALISABLE(UE::Math::TIntVector2<uint16>);
-	MUTABLE_DEFINE_POD_VECTOR_SERIALISABLE(UE::Math::TIntVector2<int16>);
-
 	MUTABLE_DEFINE_POD_SERIALISABLE(FVector4f);
 	MUTABLE_DEFINE_POD_SERIALISABLE(UE::Math::TVector4<float>);
-	MUTABLE_DEFINE_POD_VECTOR_SERIALISABLE(FVector4f);
-	MUTABLE_DEFINE_POD_VECTOR_SERIALISABLE(UE::Math::TVector4<float>);
 
+	MUTABLE_IMPLEMENT_POD_VECTOR_SERIALISABLE(float);
+	MUTABLE_IMPLEMENT_POD_VECTOR_SERIALISABLE(double);
+	MUTABLE_IMPLEMENT_POD_VECTOR_SERIALISABLE(uint8);
+	MUTABLE_IMPLEMENT_POD_VECTOR_SERIALISABLE(uint16);
+	MUTABLE_IMPLEMENT_POD_VECTOR_SERIALISABLE(uint32);
+	MUTABLE_IMPLEMENT_POD_VECTOR_SERIALISABLE(uint64);
+	MUTABLE_IMPLEMENT_POD_VECTOR_SERIALISABLE(int8);
+	MUTABLE_IMPLEMENT_POD_VECTOR_SERIALISABLE(int16);
+	MUTABLE_IMPLEMENT_POD_VECTOR_SERIALISABLE(int32);
+	MUTABLE_IMPLEMENT_POD_VECTOR_SERIALISABLE(int64);
+
+	MUTABLE_IMPLEMENT_POD_VECTOR_SERIALISABLE(vec2f);
+	MUTABLE_IMPLEMENT_POD_VECTOR_SERIALISABLE(vec3f);
+	MUTABLE_IMPLEMENT_POD_VECTOR_SERIALISABLE(mat3f);
+	MUTABLE_IMPLEMENT_POD_VECTOR_SERIALISABLE(mat4f);
+	MUTABLE_IMPLEMENT_POD_VECTOR_SERIALISABLE(vec2<int>);
+	MUTABLE_IMPLEMENT_POD_VECTOR_SERIALISABLE(TCHAR);
+
+	MUTABLE_IMPLEMENT_POD_VECTOR_SERIALISABLE(FUintVector2);
+	MUTABLE_IMPLEMENT_POD_VECTOR_SERIALISABLE(UE::Math::TIntVector2<uint16>);
+	MUTABLE_IMPLEMENT_POD_VECTOR_SERIALISABLE(UE::Math::TIntVector2<int16>);
+	MUTABLE_IMPLEMENT_POD_VECTOR_SERIALISABLE(FVector4f);
 
 	//---------------------------------------------------------------------------------------------
 	template<>
diff --git a/Engine/Plugins/Experimental/Mutable/Source/MutableRuntime/Private/MuR/System.cpp b/Engine/Plugins/Experimental/Mutable/Source/MutableRuntime/Private/MuR/System.cpp
--- a/Engine/Plugins/Experimental/Mutable/Source/MutableRuntime/Private/MuR/System.cpp
+++ b/Engine/Plugins/Experimental/Mutable/Source/MutableRuntime/Private/MuR/System.cpp
@@ -49,6 +49,8 @@ namespace mu
 	TRACE_DECLARE_INT_COUNTER(MutableRuntime_MemPool,		TEXT("MutableRuntime/MemPool"));
 	TRACE_DECLARE_INT_COUNTER(MutableRuntime_MemCache,		TEXT("MutableRuntime/MemCache"));
 	TRACE_DECLARE_INT_COUNTER(MutableRuntime_MemRom,		TEXT("MutableRuntime/MemRom"));
+	TRACE_DECLARE_INT_COUNTER(MutableRuntime_MemImage,		TEXT("MutableRuntime/MemImage"));
+	TRACE_DECLARE_INT_COUNTER(MutableRuntime_MemMesh,		TEXT("MutableRuntime/MemMesh"));
 	TRACE_DECLARE_INT_COUNTER(MutableRuntime_MemStream,		TEXT("MutableRuntime/MemStream"));
 	TRACE_DECLARE_INT_COUNTER(MutableRuntime_MemTotal,		TEXT("MutableRuntime/MemTotal"));
 	TRACE_DECLARE_INT_COUNTER(MutableRuntime_MemBudget,		TEXT("MutableRuntime/MemBudget"));
@@ -176,12 +178,6 @@ namespace mu
 		m_pD->WorkingMemoryManager.CacheResources.Empty();
 		check(m_pD->WorkingMemoryManager.TempImages.IsEmpty());
 		check(m_pD->WorkingMemoryManager.TempMeshes.IsEmpty());
-
-		m_pD->WorkingMemoryManager.TrackedBudgetBytes_Pooled = 0;
-		//m_pD->WorkingMemoryManager.TrackedBudgetBytes_Temp = 0;
-		m_pD->WorkingMemoryManager.TrackedBudgetBytes_Cached = 0;
-		m_pD->WorkingMemoryManager.TrackedBudgetBytes_Rom = 0;
-		m_pD->WorkingMemoryManager.TrackedBudgetBytes_Stream = 0;
 	}
 
 	
@@ -441,7 +437,6 @@ namespace mu
 		if (m_pD->WorkingMemoryManager.BudgetBytes == 0)
 		{
 			m_pD->WorkingMemoryManager.PooledImages.Empty();
-			m_pD->WorkingMemoryManager.TrackedBudgetBytes_Pooled = 0;
 		}
 
 	}
@@ -463,12 +458,7 @@ namespace mu
 				{
 					if (Data)
 					{
-						int32 DataSize = Data->GetDataSize();
-						int32 Removed = m_pD->WorkingMemoryManager.CacheResources.Remove(Data);
-						if (Removed)
-						{
-							m_pD->WorkingMemoryManager.TrackedBudgetBytes_Cached -= DataSize;
-						}
+						m_pD->WorkingMemoryManager.CacheResources.Remove(Data);
 					}
 				}
 
@@ -476,12 +466,7 @@ namespace mu
 				{
 					if (Data)
 					{
-						int32 DataSize = Data->GetDataSize();
-						int32 Removed = m_pD->WorkingMemoryManager.CacheResources.Remove(Data);
-						if (Removed)
-						{
-							m_pD->WorkingMemoryManager.TrackedBudgetBytes_Cached -= DataSize;
-						}
+						m_pD->WorkingMemoryManager.CacheResources.Remove(Data);
 					}
 				}
 
@@ -897,89 +882,41 @@ namespace mu
 	{
 #if !(UE_BUILD_SHIPPING || UE_BUILD_TEST)
 
+		MUTABLE_CPUPROFILER_SCOPE(LogWorkingMemory);
+
 		// For now, we calculate these for every log. We will later track on resource creation, destruction or state change.
+		// All resource memory is tracked by the memory allocator, but that does not give information about where the memory is
+		// located. Keep the localized memory computation for now.   
 		const uint32 RomBytes = GetRomBytes();
 		const uint32 CacheBytes = GetCacheBytes();
 		const uint32 TrackedCacheBytes = GetTrackedCacheBytes();
 		const uint32 PoolBytes = GetPooledBytes();
 		const uint32 TempBytes = GetTempBytes();
-		uint32 StreamingBufferBytes = 0;
-		uint32 InternalBytes = 0;
-
-		InternalBytes += PooledImages.GetAllocatedSize();
-		InternalBytes += TempImages.GetAllocatedSize();
-		InternalBytes += TempMeshes.GetAllocatedSize();
-		InternalBytes += CacheResources.GetAllocatedSize();
-		InternalBytes += CachePerModel.GetAllocatedSize();
-		InternalBytes += LiveInstances.GetAllocatedSize();
-
-		if (CurrentRunner)
-		{
-			// Data reserved for ongoing streaming operations
-			for ( const CodeRunner::FRomLoadOp& Op : CurrentRunner->m_romLoadOps )
-			{
-				StreamingBufferBytes += Op.m_streamBuffer.GetAllocatedSize();
-			}
-
-			InternalBytes += CurrentRunner->GetInternalMemoryBytes();
-		}
-
-		TArray<const Model*> Models;
-		for (const FLiveInstance& Instance: LiveInstances)
-		{
-			Models.AddUnique(Instance.Model.Get());
-
-			InternalBytes += Instance.Cache->OpExecutionData.GetAllocatedSize()
-				+ Instance.Cache->ImageResults.GetAllocatedSize()
-				+ Instance.Cache->MeshResults.GetAllocatedSize()
-				+ Instance.Cache->LayoutResults.GetAllocatedSize()
-				+ Instance.Cache->InstanceResults.GetAllocatedSize()
-				+ Instance.Cache->ColorResults.GetAllocatedSize()
-				+ Instance.Cache->ProjectorResults.GetAllocatedSize()
-				+ Instance.Cache->m_usedRangeIndices.GetAllocatedSize();
-
-			for (const FWorkingMemoryManager::FModelCacheEntry& c : CachePerModel)
-			{
-				InternalBytes += c.RomWeights.GetAllocatedSize();
-			}			
-		}
-
-		// Data stored per-model, but related to instance construction
-		for (const FModelCacheEntry& CacheEntry:CachePerModel)
-		{
-			// Keys for resources generated so far
-			for (const FModelCacheEntry::FGeneratedResourceData& Data : CacheEntry.GeneratedResources)
-			{
-				InternalBytes += Data.ParameterValuesBlob.GetAllocatedSize();
-			}
-		}
 
-		uint32 TotalBytes = InternalBytes 
-			+ TempBytes 
-			+ TrackedCacheBytes
-			+ PoolBytes
-			+ StreamingBufferBytes
-			+ RomBytes;
-
-		UE_LOG(LogMutableCore, Log, TEXT("Mem KB: Internal %7d   Temp %7d    Pool %7d (%7d)   Cache0+1 %7d (%7d) (%7d)   Rom %7d (%7d)   Stream %7d (%7d)   Total %7d / %7d"),
-			InternalBytes / 1024,
-			TempBytes / 1024, //TrackedBudgetBytes_Temp / 1024,
-			PoolBytes / 1024, TrackedBudgetBytes_Pooled / 1024,
-			CacheBytes / 1024, TrackedCacheBytes / 1024, TrackedBudgetBytes_Cached / 1024,
-			RomBytes / 1024, TrackedBudgetBytes_Rom / 1024,
-			StreamingBufferBytes / 1024, TrackedBudgetBytes_Stream / 1024,
-			TotalBytes / 1024,
-			BudgetBytes / 1024
-			);
-
-		TRACE_COUNTER_SET(MutableRuntime_MemInternal, InternalBytes);
-		TRACE_COUNTER_SET(MutableRuntime_MemTemp, TempBytes);
-		TRACE_COUNTER_SET(MutableRuntime_MemPool, PoolBytes);
-		TRACE_COUNTER_SET(MutableRuntime_MemCache, TrackedCacheBytes);
-		TRACE_COUNTER_SET(MutableRuntime_MemRom, RomBytes);
-		TRACE_COUNTER_SET(MutableRuntime_MemStream, StreamingBufferBytes);
-		TRACE_COUNTER_SET(MutableRuntime_MemTotal, TotalBytes);
-		TRACE_COUNTER_SET(MutableRuntime_MemBudget, BudgetBytes);
+		// Get allocator counters.
+		const SSIZE_T ImageAllocBytes	 = MemoryCounters::FImageMemoryCounter::Counter.load(std::memory_order_relaxed);
+		const SSIZE_T MeshAllocBytes     = MemoryCounters::FMeshMemoryCounter::Counter.load(std::memory_order_relaxed);
+		const SSIZE_T StreamAllocBytes   = MemoryCounters::FStreamingMemoryCounter::Counter.load(std::memory_order_relaxed);
+		const SSIZE_T InternalAllocBytes = MemoryCounters::FMemoryTrackerInternalMemoryCounter::Counter.load(std::memory_order_relaxed);
+
+		SSIZE_T TotalBytes = ImageAllocBytes + MeshAllocBytes + StreamAllocBytes + InternalAllocBytes;
+
+		UE_LOG(LogMutableCore, Log, 
+			TEXT("Mem KB: ImageAlloc %7d, MeshAlloc %7d, StreamAlloc %7d, InternalAlloc %7d,  AllocTotal %7d / %7d. \
+				  Resources MemLoc: Temp %7d, Pool %7d, Cache0+1 %7d (%7d), Rom %7d."),
+			ImageAllocBytes/1024, MeshAllocBytes/1024, StreamAllocBytes/1024, InternalAllocBytes/1024, TotalBytes/1024, BudgetBytes/1024,  
+			TempBytes/1024, PoolBytes/1024, CacheBytes/1024, TrackedCacheBytes/1024, RomBytes/1024);
+
+		TRACE_COUNTER_SET(MutableRuntime_MemTemp,	  TempBytes);
+		TRACE_COUNTER_SET(MutableRuntime_MemPool,     PoolBytes);
+		TRACE_COUNTER_SET(MutableRuntime_MemCache,    TrackedCacheBytes);
+		TRACE_COUNTER_SET(MutableRuntime_MemRom,      RomBytes);
+		TRACE_COUNTER_SET(MutableRuntime_MemInternal, InternalAllocBytes);
+		TRACE_COUNTER_SET(MutableRuntime_MemMesh,     MeshAllocBytes);
+		TRACE_COUNTER_SET(MutableRuntime_MemImage,    ImageAllocBytes);
+		TRACE_COUNTER_SET(MutableRuntime_MemStream,   StreamAllocBytes);
+		TRACE_COUNTER_SET(MutableRuntime_MemTotal,    TotalBytes);
+		TRACE_COUNTER_SET(MutableRuntime_MemBudget,   BudgetBytes);
 #endif
 	}
 
@@ -1034,7 +971,8 @@ namespace mu
 		FModelCacheEntry n;
         n.Model = TWeakPtr<const Model>(InModel);
 		n.PendingOpsPerRom.SetNum(InModel->GetPrivate()->m_program.m_roms.Num());
-        CachePerModel.Add(n);
+    	n.RomWeights.SetNum(InModel->GetPrivate()->m_program.m_roms.Num());
+    	CachePerModel.Add(n);
         return CachePerModel.Last();
     }
 
@@ -1044,26 +982,10 @@ namespace mu
 	{
 		MUTABLE_CPUPROFILER_SCOPE(GetCurrentMemoryBytes);
 
-		int32 RomBytes = TrackedBudgetBytes_Rom; // GetRomBytes();
-		int32 PoolBytes = TrackedBudgetBytes_Pooled; // GetPooledBytes();
-		int32 TempBytes = GetTempBytes(); // TrackedBudgetBytes_Temp + GetTempMeshesBytes(); 
-		int32 TrackedCacheBytes = TrackedBudgetBytes_Cached; // GetTrackedCacheBytes();
-		int32 TrackedStreamBytes = TrackedBudgetBytes_Stream;
-
-		int64 TotalBytes = 0;
-		TotalBytes += TempBytes;
-		TotalBytes += RomBytes;
-		TotalBytes += PoolBytes;
-		TotalBytes += TrackedCacheBytes;
-		TotalBytes += TrackedStreamBytes;
-
-		TRACE_COUNTER_SET(MutableRuntime_MemTemp, TempBytes);
-		TRACE_COUNTER_SET(MutableRuntime_MemPool, PoolBytes);
-		TRACE_COUNTER_SET(MutableRuntime_MemCache, TrackedCacheBytes);
-		TRACE_COUNTER_SET(MutableRuntime_MemRom, RomBytes);
-		TRACE_COUNTER_SET(MutableRuntime_MemStream, TrackedStreamBytes);
-		TRACE_COUNTER_SET(MutableRuntime_MemTotal, TotalBytes);
-		TRACE_COUNTER_SET(MutableRuntime_MemBudget, BudgetBytes);
+		SSIZE_T TotalBytes = MemoryCounters::FImageMemoryCounter::Counter.load(std::memory_order_relaxed) + 
+						     MemoryCounters::FMeshMemoryCounter::Counter.load(std::memory_order_relaxed) +
+							 MemoryCounters::FStreamingMemoryCounter::Counter.load(std::memory_order_relaxed) +
+							 MemoryCounters::FMemoryTrackerInternalMemoryCounter::Counter.load(std::memory_order_relaxed);
 
 		return TotalBytes;
 	}
@@ -1114,7 +1036,6 @@ namespace mu
 				int32 PooledResourceSize = PooledImages[0]->GetDataSize();
 				TotalBytes -= PooledResourceSize;				
 				PooledImages.RemoveAtSwap(0);
-				TrackedBudgetBytes_Pooled -= PooledResourceSize;
 				bFinished = TotalBytes <= BudgetBytes;
 			}
 		}
@@ -1176,8 +1097,6 @@ namespace mu
 				// UE_LOG(LogMutableCore,Log, "Unloading rom because of memory budget: %d.", lowestPriorityRom);
 				int32 UnloadedSize = Candidate.Value.Model->GetPrivate()->m_program.UnloadRom(Candidate.Value.RomIndex);
 				TotalBytes -= UnloadedSize;
-				TrackedBudgetBytes_Rom -= UnloadedSize;
-
 				bFinished = TotalBytes <= BudgetBytes;
 			}
 		}
@@ -1275,12 +1194,7 @@ namespace mu
 
 							if (Value == Removed)
 							{
-								int32 RemoveCount = CacheResources.Remove(Removed);
-								if (RemoveCount > 0)
-								{
-									TrackedBudgetBytes_Cached -= RemovedDataSize;
-								}
-
+								CacheResources.Remove(Removed);
 								Instance.Cache->SetUnused(*RemIt);
 							}
 						}
@@ -1422,11 +1336,7 @@ namespace mu
 							{
 								MUTABLE_CPUPROFILER_SCOPE(EnsureBudgetBelow_FreeCached_Free_Current_ActualFree);
 
-								int32 RemoveCount = CacheResources.Remove(Removed);
-								if (RemoveCount > 0)
-								{
-									TrackedBudgetBytes_Cached -= RemovedDataSize;
-								}
+								CacheResources.Remove(Removed);
 								CurrentInstanceCache->SetUnused(*RemIt);
 							}
 						}
@@ -1465,8 +1375,6 @@ namespace mu
     {
         check(pModel);
 
-        FProgram& Program = pModel->GetPrivate()->m_program;
-
         // If budget is zero, we don't unload anything here, and we assume it is managed somewhere else.
         if (!BudgetBytes)
         {
@@ -1478,12 +1386,7 @@ namespace mu
         // Update current cache
         {
 			FModelCacheEntry* ModelCache = FindModelCache(pModel.Get());
-
-            if (ModelCache->RomWeights.Num()<Program.m_roms.Num())
-            {
-				ModelCache->RomWeights.SetNumZeroed(Program.m_roms.Num());
-            }
-
+        	
             ModelCache->RomWeights[romIndex].Key++;
             ModelCache->RomWeights[romIndex].Value = RomTick;
         }
@@ -1491,7 +1394,7 @@ namespace mu
 
 
 	//---------------------------------------------------------------------------------------------
-	static int32 AddMultiValueKeys(TArray<uint8_t>& Blob, int32 pos, const TMap< TArray<int>, PARAMETER_VALUE >& multi)
+	static int32 AddMultiValueKeys(FWorkingMemoryManager::TMemoryTrackedArray<uint8>& Blob, int32 pos, const TMap< TArray<int>, PARAMETER_VALUE >& multi)
 	{
 		Blob.SetNum(pos + 4);
 		uint32 s = uint32(multi.Num());
@@ -1544,7 +1447,7 @@ namespace mu
 		}
 
 		// Generate the relevant parameters blob
-		TArray<uint8> Blob;
+		TMemoryTrackedArray<uint8> Blob;
 		Blob.Reserve(2048);
 		for (int32 ParamIndex : *RelevantParams)
 		{
diff --git a/Engine/Plugins/Experimental/Mutable/Source/MutableRuntime/Private/MuR/SystemPrivate.h b/Engine/Plugins/Experimental/Mutable/Source/MutableRuntime/Private/MuR/SystemPrivate.h
--- a/Engine/Plugins/Experimental/Mutable/Source/MutableRuntime/Private/MuR/SystemPrivate.h
+++ b/Engine/Plugins/Experimental/Mutable/Source/MutableRuntime/Private/MuR/SystemPrivate.h
@@ -19,6 +19,11 @@
 #include "HAL/PlatformTLS.h"
 #endif
 
+namespace mu::MemoryCounters
+{
+	struct FMemoryTrackerInternalMemoryCounterTag {};
+	using FMemoryTrackerInternalMemoryCounter = TMemoryCounter<FMemoryTrackerInternalMemoryCounterTag>;
+}
 
 namespace mu
 {
@@ -205,8 +210,11 @@ namespace mu
 	template<class DATA>
 	class CodeContainer
 	{
-	public:
+		using MemoryCounter = MemoryCounters::FMemoryTrackerInternalMemoryCounter;
+		using ArrayDataContainerType = TArray<DATA, FDefaultMemoryTrackingAllocator<MemoryCounter>>;
+		using MapDataContainerType = TMap<FCacheAddress, DATA, FDefaultMemoryTrackingSetAllocator<MemoryCounter>>;
 
+	public:
 		void resize(size_t s)
 		{
 			m_index0.SetNumZeroed(s);
@@ -335,9 +343,10 @@ namespace mu
 		{
 		private:
 			friend class CodeContainer<DATA>;
+
 			const CodeContainer<DATA>* container;
-			typename TArray<DATA>::TIterator it0;
-			typename TMap<FCacheAddress, DATA>::TIterator it1;
+			typename CodeContainer<DATA>::ArrayDataContainerType::TIterator it0;
+			typename CodeContainer<DATA>::MapDataContainerType::TIterator it1;
 
 			iterator(CodeContainer<DATA>* InContainer)
 				: container(InContainer)
@@ -419,13 +428,12 @@ namespace mu
 			return m_index0.GetAllocatedSize()
 				+ m_otherIndex.GetAllocatedSize();
 		}
-
 	private:
 		// For index 0
-		TArray<DATA> m_index0;
+		ArrayDataContainerType m_index0;
 
 		// For index>0
-		TMap<FCacheAddress, DATA> m_otherIndex;
+		MapDataContainerType m_otherIndex;
 	};
 
 
@@ -434,7 +442,12 @@ namespace mu
 	{
 	public:
 
-		TArray< ExecutionIndex, TInlineAllocator<4> > m_usedRangeIndices;
+		using AllocType = FDefaultMemoryTrackingAllocator<MemoryCounters::FMemoryTrackerInternalMemoryCounter>;
+
+		template<class Type, class Alloc = AllocType >
+		using TMemoryTrackedArray = TArray<Type, Alloc>;
+		
+		TMemoryTrackedArray<ExecutionIndex, TInlineAllocator<4, AllocType>> m_usedRangeIndices;
 
 		/** Runtime data for each program op. */
 		struct FOpExecutionData
@@ -469,14 +482,15 @@ namespace mu
 		CodeContainer<FOpExecutionData> OpExecutionData;
 
 		/** */
-		TArray<FVector4f> ColorResults;
-		TArray<Ptr<const Image>> ImageResults;
-		TArray<Ptr<const Layout>> LayoutResults;
-		TArray<Ptr<const Mesh>> MeshResults;
-		TArray<Ptr<const Instance>> InstanceResults;
-		TArray<FProjector> ProjectorResults;
-		TArray<Ptr<const String>> StringResults;
-		TArray<Ptr<const ExtensionData>> ExtensionDataResults;
+
+		TMemoryTrackedArray<FVector4f> ColorResults;
+		TMemoryTrackedArray<Ptr<const Image>> ImageResults;
+		TMemoryTrackedArray<Ptr<const Layout>> LayoutResults;
+		TMemoryTrackedArray<Ptr<const Mesh>> MeshResults;
+		TMemoryTrackedArray<Ptr<const Instance>> InstanceResults;
+		TMemoryTrackedArray<FProjector> ProjectorResults;
+		TMemoryTrackedArray<Ptr<const String>> StringResults;
+		TMemoryTrackedArray<Ptr<const ExtensionData>> ExtensionDataResults;
 
 		/** */
 		inline const ExecutionIndex& GetRangeIndex(uint32_t i)
@@ -1093,6 +1107,14 @@ namespace mu
     /** Struct to manage all the memory allocated for resources used during mutable operation. */
     struct FWorkingMemoryManager
     {
+		using MemoryCounter = MemoryCounters::FMemoryTrackerInternalMemoryCounter;
+
+		template<class Type>
+		using TMemoryTrackedArray = TArray<Type, FDefaultMemoryTrackingAllocator<MemoryCounter>>;
+
+		template<class KeyType, class ValueType>
+		using TMemoryTrackedMap = TMap<KeyType, ValueType, FDefaultMemoryTrackingSetAllocator<MemoryCounter>>;
+
 		/** Cached traking for streamed model data for one model. */
         struct FModelCacheEntry
         {
@@ -1100,10 +1122,10 @@ namespace mu
             TWeakPtr<const Model> Model;
 
             /** For each model rom, the last time its streamed data was used. */
-            TArray< TPair<uint64,uint64> > RomWeights;
+            TMemoryTrackedArray<TPair<uint64, uint64>> RomWeights;
 
 			/** Count of pending operations for every rom index. */
-			TArray<uint16> PendingOpsPerRom;
+			TMemoryTrackedArray<uint16> PendingOpsPerRom;
 
 			//! Management of generated resources
 			//! @{
@@ -1118,7 +1140,7 @@ namespace mu
 				uint32 LastRequestId;
 
 				//! An opaque blob with the values of the relevant parameters
-				TArray<uint8> ParameterValuesBlob;
+				TMemoryTrackedArray<uint8> ParameterValuesBlob;
 			};
 
 			//! The last id generated for a resource
@@ -1130,7 +1152,7 @@ namespace mu
 
 			//! Cached ids for returned assets
 			//! This is non-persistent runtime data
-			TArray<FGeneratedResourceData> GeneratedResources;
+			TMemoryTrackedArray<FGeneratedResourceData> GeneratedResources;
 
 			//! Get a resource key for a given resource with given parameter values.
 			FResourceID GetResourceKey(uint32 ParamListIndex, OP::ADDRESS RootAt, const Parameters*, int32 InMaxResourceKeys);
@@ -1147,34 +1169,28 @@ namespace mu
 		/** Maximum number of resource keys that will be stored for resource reusal. */
 		int32 MaxGeneratedResourceCacheSize = 1024;
 
-		/** Signed ints to account for imprecisions especially with temp.*/
-		int32 TrackedBudgetBytes_Rom = 0;
-		// int32 TrackedBudgetBytes_Temp = 0;
-		int32 TrackedBudgetBytes_Pooled = 0;
-		int32 TrackedBudgetBytes_Cached = 0;
-		int32 TrackedBudgetBytes_Stream = 0;
 
 		/** This value is used to track the order of loading of roms. */
         uint64 RomTick = 0;
 
 		/** Control info for the per-model cache of streamed data. */
-        TArray< FModelCacheEntry > CachePerModel;
+        TMemoryTrackedArray<FModelCacheEntry> CachePerModel;
 
 		/** Data for each mutable instance that is being updated. */
-		TArray<FLiveInstance> LiveInstances;
+		TMemoryTrackedArray<FLiveInstance> LiveInstances;
 
 		/** Temporary reference to the memory of the current instance being updated. Only valid during a mutable "atomic" operation, like a BeginUpdate or a GetImage. */
 		TSharedPtr<FProgramCache> CurrentInstanceCache;
 
 		/** Resources that have been used in the past, but haven't been deallocated because they still fitted the memory budget and they could be reused. */
-		TArray<Ptr<Image>> PooledImages;
+		TMemoryTrackedArray<Ptr<Image>> PooledImages;
 
 		/** List of intermediate resources that are not soterd anywhere yet. They are still locally referenced by code. */
-		TArray<Ptr<const Image>> TempImages;
-		TArray<Ptr<const Mesh>> TempMeshes;
+		TMemoryTrackedArray<Ptr<const Image>> TempImages;
+		TMemoryTrackedArray<Ptr<const Mesh>> TempMeshes;
 
 		/** List of resources that are currently in any cache position, and the number of positions they are in. */
-		TMap<Ptr<const Resource>, int32> CacheResources;
+		TMemoryTrackedMap<Ptr<const Resource>, int32> CacheResources;
 
 		/** Given a mutable model, find or create its rom cache. */
 		FModelCacheEntry* FindModelCache(const Model*);
@@ -1215,7 +1231,7 @@ namespace mu
 				{
 					Ptr<Image> Result = Candidate;
 					PooledImages.RemoveAtSwap(Index);
-					TrackedBudgetBytes_Pooled -= Result->GetDataSize();
+					
 					if (Init == EInitializationType::Black)
 					{
 						Result->InitToBlack();
@@ -1237,7 +1253,6 @@ namespace mu
 			Ptr<Image> Result = new Image(SizeX, SizeY, Lods, Format, Init);
 
 			TempImages.Add(Result);
-			//TrackedBudgetBytes_Temp += Result->GetDataSize();
 			return Result;
 		}
 
@@ -1246,11 +1261,7 @@ namespace mu
 		{
 			CheckRunnerThread();
 
-			int32 Removed = TempImages.RemoveSingle(Resource);
-			if (Removed)
-			{
-				//TrackedBudgetBytes_Temp -= Resource->GetDataSize();
-			}
+			TempImages.RemoveSingle(Resource);
 			
 			check(!TempImages.Contains(Resource));
 			check(!PooledImages.Contains(Resource));
@@ -1285,12 +1296,8 @@ namespace mu
 				return;
 			}
 
-			int32 ResourceDataSize = Resource->GetDataSize();
-			int32 Removed = TempImages.RemoveSingle(Resource);
-			if (Removed)
-			{
-				//TrackedBudgetBytes_Temp -= ResourceDataSize;
-			}
+			const int32 ResourceDataSize = Resource->GetDataSize();
+			TempImages.RemoveSingle(Resource);
 
 			check(!TempImages.Contains(Resource));
 			check(!PooledImages.Contains(Resource));
@@ -1302,7 +1309,6 @@ namespace mu
 				if (bInBudget)
 				{
 					PooledImages.Add(const_cast<Image*>(Resource.get()));
-					TrackedBudgetBytes_Pooled += ResourceDataSize;
 				}
 			}
 			else
@@ -1324,12 +1330,8 @@ namespace mu
 				return;
 			}
 
-			int32 ResourceDataSize = Resource->GetDataSize();
-			int32 Removed = TempImages.RemoveSingle(Resource);
-			if (Removed)
-			{
-				//TrackedBudgetBytes_Temp -= ResourceDataSize;
-			}
+			const int32 ResourceDataSize = Resource->GetDataSize();
+			TempImages.RemoveSingle(Resource);
 
 			check(!TempImages.Contains(Resource));
 			check(!PooledImages.Contains(Resource));
@@ -1341,7 +1343,6 @@ namespace mu
 				if (bInBudget)
 				{
 					PooledImages.Add(Resource.get());
-					TrackedBudgetBytes_Pooled += ResourceDataSize;
 				}
 			}
 			else
@@ -1362,7 +1363,6 @@ namespace mu
 			Ptr<Mesh> Result = new Mesh();
 
 			TempMeshes.Add(Result);
-			//TrackedBudgetBytes_Temp += Result->GetDataSize();
 
 			return Result;
 		}
@@ -1372,11 +1372,7 @@ namespace mu
 			CheckRunnerThread();
 
 			const int32 ResourceDataSize = Resource->GetDataSize();
-			const int32 Removed = TempMeshes.RemoveSingle(Resource);
-			if (Removed)
-			{
-				//TrackedBudgetBytes_Temp -= ResourceDataSize;
-			}
+			TempMeshes.RemoveSingle(Resource);
 
 			Ptr<Mesh> Result;
 			if (!Resource->IsUnique())
@@ -1403,12 +1399,8 @@ namespace mu
 				return;
 			}
 
-			int32 ResourceDataSize = Resource->GetDataSize();
-			int32 Removed = TempMeshes.RemoveSingle(Resource);
-			if (Removed)
-			{
-				//TrackedBudgetBytes_Temp -= ResourceDataSize;
-			}
+			TempMeshes.RemoveSingle(Resource);
+			check(!TempMeshes.Contains(Resource));
 
 			EnsureBudgetBelow(0);
 
@@ -1434,8 +1426,6 @@ namespace mu
 				return nullptr;
 			}
 
-			int32 ResultDataSize = Result->GetDataSize();
-
 			// If we retrieved the last reference to this resource in "From" cache position (it could still be in other cache positions as well)
 			if (bIsLastReference)
 			{
@@ -1445,14 +1435,12 @@ namespace mu
 				if (!*CountPtr)
 				{
 					CacheResources.FindAndRemoveChecked(Result);
-					TrackedBudgetBytes_Cached -= ResultDataSize;
 				}
 			}
 
 			if (!bTakeOwnership && Result->IsUnique())
 			{
 				TempMeshes.Add(Result);
-				//TrackedBudgetBytes_Temp += ResultDataSize;
 			}
 
 			return Result;
@@ -1468,8 +1456,6 @@ namespace mu
 				return nullptr;
 			}
 
-			int32 ResultDataSize = Result->GetDataSize();
-
 			// If we retrieved the last reference to this resource in "From" cache position (it could still be in other cache positions as well)
 			if (bIsLastReference)
 			{
@@ -1479,7 +1465,6 @@ namespace mu
 				if (!*CountPtr)
 				{
 					CacheResources.FindAndRemoveChecked(Result);
-					TrackedBudgetBytes_Cached -= ResultDataSize;
 				}
 			}
 
@@ -1487,7 +1472,6 @@ namespace mu
 			if (!bTakeOwnership && Result->IsUnique())
 			{
 				TempImages.Add(Result);
-				//TrackedBudgetBytes_Temp += ResultDataSize;
 			}
 
 			return Result;
@@ -1499,20 +1483,11 @@ namespace mu
 			if (Resource)
 			{
 				int32 ResourceDataSize = Resource->GetDataSize();
-				int32 Removed = TempImages.RemoveSingle(Resource);
-				if (Removed)
-				{
-					//TrackedBudgetBytes_Temp -= ResourceDataSize;
-				}
+				TempImages.RemoveSingle(Resource);
 
 				check(!TempImages.Contains(Resource));
 
 				int32& Count = CacheResources.FindOrAdd(Resource, 0);
-				if (!Count)
-				{
-					TrackedBudgetBytes_Cached += ResourceDataSize;
-				}
-
 				++Count;
 			}
 
@@ -1523,19 +1498,9 @@ namespace mu
 		{
 			if (Resource)
 			{
-				int32 ResourceDataSize = Resource->GetDataSize();
-				int32 Removed = TempMeshes.RemoveSingle(Resource);
-				if (Removed)
-				{
-					//TrackedBudgetBytes_Temp -= ResourceDataSize;
-				}
+				TempMeshes.RemoveSingle(Resource);
 
 				int32& Count = CacheResources.FindOrAdd(Resource, 0);
-				if (!Count)
-				{
-					TrackedBudgetBytes_Cached += ResourceDataSize;
-				}
-
 				++Count;
 			}
 
@@ -1562,31 +1527,23 @@ namespace mu
 				Result += Value->GetDataSize();
 			}
 
-			ensure(Result == TrackedBudgetBytes_Pooled);
-
 			return Result;
 		}
 
-		int32 GetTempMeshesBytes() const
+		int32 GetTempBytes() const
 		{
 			int32 Result = 0;
-			for (const Ptr<const Mesh>& Value : TempMeshes)
+			for (const Ptr<const Image>& Value : TempImages)
 			{
 				Result += Value->GetDataSize();
 			}
 
-			return Result;
-		}
-
-		int32 GetTempBytes() const
-		{
-			int32 Result = 0;
-			for (const Ptr<const Image>& Value : TempImages)
+			for (const Ptr<const Mesh>& Value : TempMeshes)
 			{
 				Result += Value->GetDataSize();
 			}
 
-			return Result + GetTempMeshesBytes();
+			return Result;
 		}
 
 
@@ -1717,8 +1674,6 @@ namespace mu
 				}
 
 				const Resource* Value = nullptr;
-				int32 ValueDataSize = 0;
-				int32 RemoveCount = 0;
 
 				switch (Data.DataType)
 				{
@@ -1726,12 +1681,7 @@ namespace mu
 					Value = CurrentInstanceCache->ImageResults[Data.DataTypeIndex].get();
 					if (Value)
 					{
-						ValueDataSize = Value->GetDataSize();
-						RemoveCount = CacheResources.Remove(Value);
-						if (RemoveCount)
-						{
-							TrackedBudgetBytes_Cached -= ValueDataSize;
-						}
+						CacheResources.Remove(Value);
 						CurrentInstanceCache->ImageResults[Data.DataTypeIndex] = nullptr;
 					}
 					break;
@@ -1740,12 +1690,7 @@ namespace mu
 					Value = CurrentInstanceCache->MeshResults[Data.DataTypeIndex].get();
 					if (Value)
 					{
-						ValueDataSize = Value->GetDataSize();
-						RemoveCount = CacheResources.Remove(Value);
-						if (RemoveCount)
-						{
-							TrackedBudgetBytes_Cached -= ValueDataSize;
-						}
+						CacheResources.Remove(Value);
 						CurrentInstanceCache->MeshResults[Data.DataTypeIndex] = nullptr;
 					}
 					break;
diff --git a/Engine/Plugins/Experimental/Mutable/Source/MutableRuntime/Public/MuR/Image.h b/Engine/Plugins/Experimental/Mutable/Source/MutableRuntime/Public/MuR/Image.h
--- a/Engine/Plugins/Experimental/Mutable/Source/MutableRuntime/Public/MuR/Image.h
+++ b/Engine/Plugins/Experimental/Mutable/Source/MutableRuntime/Public/MuR/Image.h
@@ -12,6 +12,14 @@
 #include "HAL/Platform.h"
 #include "Misc/AssertionMacros.h"
 
+#include "MuR/MemoryTrackingAllocationPolicy.h"
+
+namespace mu::MemoryCounters
+{
+	struct FImageMemoryCounterTag {};
+	using FImageMemoryCounter = TMemoryCounter<FImageMemoryCounterTag>;
+}
+
 namespace mu
 {
 
@@ -292,7 +300,9 @@ namespace mu
 		uint32 ReferenceID = 0;
 
 		/** Pixel data for all lods. */
-		TArray<uint8> m_data;
+
+		using ImageDataContainerType = TArray<uint8, FDefaultMemoryTrackingAllocator<MemoryCounters::FImageMemoryCounter>>;
+		ImageDataContainerType m_data;
 
 
 		// This used to be the methods in the private implementation of the image interface
diff --git a/Engine/Plugins/Experimental/Mutable/Source/MutableRuntime/Public/MuR/MemoryTrackingAllocationPolicy.h b/Engine/Plugins/Experimental/Mutable/Source/MutableRuntime/Public/MuR/MemoryTrackingAllocationPolicy.h
new file mode 100644
--- /dev/null
+++ b/Engine/Plugins/Experimental/Mutable/Source/MutableRuntime/Public/MuR/MemoryTrackingAllocationPolicy.h
@@ -0,0 +1,330 @@
+// Copyright Epic Games, Inc. All Rights Reserved.
+
+#pragma once
+
+#include "Containers/ContainerAllocationPolicies.h"
+#include "Containers/ContainersFwd.h"
+
+#include <atomic>
+#include <type_traits>
+
+namespace mu
+{
+	template<typename TagType>
+	struct TMemoryCounter
+	{
+		alignas(8) static inline std::atomic<SSIZE_T> Counter{ 0 };
+	};
+
+	template<typename Type>
+	struct TIsAMemoryCounter : std::false_type {};
+
+	template<typename TagType>
+	struct TIsAMemoryCounter<TMemoryCounter<TagType>> : std::true_type {};
+
+	template<typename BaseAlloc, typename CounterType>
+	class FMemoryTrackingAllocatorWrapper
+	{
+		static_assert(TIsAMemoryCounter<CounterType>::value);
+
+	public:
+		using SizeType = typename BaseAlloc::SizeType;
+
+		enum { NeedsElementType = BaseAlloc::NeedsElementType };
+		enum { RequireRangeCheck = BaseAlloc::RequireRangeCheck };
+
+		class ForAnyElementType : public BaseAlloc::ForAnyElementType
+		{
+			SSIZE_T AllocSize = 0;
+
+		public:
+			ForAnyElementType() 
+				: BaseAlloc::ForAnyElementType()
+			{
+			}
+
+			/** Destructor. */
+			FORCEINLINE ~ForAnyElementType()
+			{
+				CounterType::Counter.fetch_sub(AllocSize, std::memory_order::memory_order_relaxed);
+
+				AllocSize = 0;
+			}
+
+			template <typename OtherAllocator>
+			FORCEINLINE void MoveToEmptyFromOtherAllocator(typename OtherAllocator::ForAnyElementType& Other) = delete;
+
+			FORCEINLINE void MoveToEmpty(ForAnyElementType& Other)
+			{
+				BaseAlloc::ForAnyElementType::MoveToEmpty(Other);
+				
+				CounterType::Counter.fetch_sub(AllocSize, std::memory_order::memory_order_relaxed);
+
+				AllocSize = Other.AllocSize;
+				Other.AllocSize = 0;
+			}
+
+			FORCEINLINE void ResizeAllocation(
+				SizeType PreviousNumElements,
+				SizeType NumElements,
+				SIZE_T NumBytesPerElement
+				)
+			{
+				BaseAlloc::ForAnyElementType::ResizeAllocation(
+					PreviousNumElements, NumElements, NumBytesPerElement);
+
+				const SSIZE_T Differential = SSIZE_T(NumElements * NumBytesPerElement) - AllocSize;
+				const SSIZE_T PrevCounterValue = CounterType::Counter.fetch_add(Differential, std::memory_order::memory_order_relaxed);
+				check(PrevCounterValue >= AllocSize);
+
+				AllocSize = NumElements * NumBytesPerElement;
+			}
+
+			FORCEINLINE typename TEnableIf<TAllocatorTraits<BaseAlloc>::SupportsElementAlignment, void>::Type ResizeAllocation(
+				SizeType PreviousNumElements,
+				SizeType NumElements,
+				SIZE_T NumBytesPerElement,
+				uint32 AlignmentOfElement
+				)
+			{
+				BaseAlloc::ForAnyElementType::ResizeAllocation(
+					PreviousNumElements, NumElements, NumBytesPerElement, AlignmentOfElement);
+
+				const SSIZE_T Differential = SSIZE_T(NumElements * NumBytesPerElement) - AllocSize;
+				const SSIZE_T PrevCounterValue = CounterType::Counter.fetch_add(Differential, std::memory_order::memory_order_relaxed);
+				check(PrevCounterValue >= AllocSize);
+
+				AllocSize = NumElements * NumBytesPerElement;
+			}
+
+		};
+
+		template<typename ElementType>
+		class ForElementType : public ForAnyElementType
+		{
+			// Some Allocators, e.g, TAlignedHeapAllocator do some static checks. Instanciate the BaseAlloc::ForElementType
+			// at compile time so those warnings are emitted. 
+			using BaseTypeInstantaitionType = decltype(DeclVal<typename BaseAlloc::template ForElementType<ElementType>>());
+
+		public:
+			ForElementType()
+			{
+			}
+
+			FORCEINLINE ElementType* GetAllocation() const
+			{
+				return (ElementType*)ForAnyElementType::GetAllocation();
+			}
+		};
+	};
+
+
+	template<typename BaseAlloc, typename CounterType>
+	class FMemoryTrackingAllocatorWrapper2
+	{
+		static_assert(TIsAMemoryCounter<CounterType>::value);
+
+	public:
+		using SizeType = typename BaseAlloc::SizeType;
+
+		enum { NeedsElementType = BaseAlloc::NeedsElementType };
+		enum { RequireRangeCheck = BaseAlloc::RequireRangeCheck };
+
+		class ForAnyElementType 
+		{
+		public:
+			ForAnyElementType() 
+			{
+			}
+
+			/** Destructor. */
+			FORCEINLINE ~ForAnyElementType()
+			{
+				CounterType::Counter.fetch_sub(AllocSize, std::memory_order::memory_order_relaxed);
+				AllocSize = 0;
+			}
+
+			ForAnyElementType(const ForAnyElementType&) = delete;
+			ForAnyElementType& operator=(const ForAnyElementType&) = delete;
+
+			FORCEINLINE decltype(auto) GetAllocation() const
+			{
+				return Base.GetAllocation();
+			}
+
+			FORCEINLINE void MoveToEmpty(ForAnyElementType& Other)
+			{
+				Base.MoveToEmpty(Other.Base);
+
+				CounterType::Counter.fetch_sub(AllocSize, std::memory_order::memory_order_relaxed);
+
+				AllocSize = Other.AllocSize;
+				Other.AllocSize = 0;
+			}
+
+			FORCEINLINE void ResizeAllocation(
+				SizeType PreviousNumElements,
+				SizeType NumElements,
+				SIZE_T NumBytesPerElement
+			)
+			{
+				Base.ResizeAllocation(PreviousNumElements, NumElements, NumBytesPerElement);
+
+				const SSIZE_T Differential = SSIZE_T(NumElements * NumBytesPerElement) - AllocSize;
+				const SSIZE_T PrevCounterValue = CounterType::Counter.fetch_add(Differential, std::memory_order::memory_order_relaxed);
+				check(PrevCounterValue >= AllocSize);
+
+				AllocSize = NumElements * NumBytesPerElement;
+			}
+
+			FORCEINLINE typename TEnableIf<TAllocatorTraits<BaseAlloc>::SupportsElementAlignment, void>::Type ResizeAllocation(
+				SizeType PreviousNumElements,
+				SizeType NumElements,
+				SIZE_T NumBytesPerElement,
+				uint32 AlignmentOfElement
+			)
+			{
+				Base.ResizeAllocation(PreviousNumElements, NumElements, NumBytesPerElement, AlignmentOfElement);
+
+				const SSIZE_T Differential = SSIZE_T(NumElements * NumBytesPerElement) - AllocSize;
+				const SSIZE_T PrevCounterValue = CounterType::Counter.fetch_add(Differential, std::memory_order::memory_order_relaxed);
+				check(PrevCounterValue >= AllocSize);
+
+				AllocSize = NumElements * NumBytesPerElement;
+			}
+
+			FORCEINLINE SizeType CalculateSlackReserve(
+				SizeType NumElements,
+				SIZE_T NumBytesPerElement
+			) const
+			{
+				return Base.CalculateSlackReserve(NumElements, NumBytesPerElement);
+			}
+
+			FORCEINLINE typename TEnableIf<TAllocatorTraits<BaseAlloc>::SupportsElementAlignment, SizeType>::Type CalculateSlackReserve(
+				SizeType NumElements,
+				SIZE_T NumBytesPerElement,
+				uint32 AlignmentOfElement
+			) const
+			{
+				return Base.CalculateSlackReserve(NumElements, NumBytesPerElement, AlignmentOfElement);
+			}
+
+			FORCEINLINE SizeType CalculateSlackShrink(
+				SizeType NumElements,
+				SizeType CurrentNumSlackElements,
+				SIZE_T NumBytesPerElement
+			) const
+			{
+				return Base.CalculateSlackShrink(NumElements, CurrentNumSlackElements, NumBytesPerElement);
+			}
+
+			FORCEINLINE typename TEnableIf<TAllocatorTraits<BaseAlloc>::SupportsElementAlignment, SizeType>::Type CalculateSlackShrink(
+				SizeType NumElements,
+				SizeType CurrentNumSlackElements,
+				SIZE_T NumBytesPerElement,
+				uint32 AlignmentOfElement
+			) const
+			{
+				return Base.CalculateSlackShrink(NumElements, CurrentNumSlackElements, NumBytesPerElement);
+			}
+
+			SizeType CalculateSlackGrow(
+				SizeType NumElements,
+				SizeType CurrentNumSlackElements,
+				SIZE_T NumBytesPerElement
+			) const
+			{
+				return Base.CalculateSlackGrow(NumElements, CurrentNumSlackElements, NumBytesPerElement);
+			}
+
+			FORCEINLINE typename TEnableIf<TAllocatorTraits<BaseAlloc>::SupportsElementAlignment, SizeType>::Type CalculateSlackGrow(
+				SizeType NumElements,
+				SizeType CurrentNumSlackElements,
+				SIZE_T NumBytesPerElement,
+				uint32 AlignmentOfElement
+			) const
+			{
+				return Base.CalculateSlackGrow(NumElements, CurrentNumSlackElements, NumBytesPerElement, AlignmentOfElement);
+			}
+
+			FORCEINLINE SIZE_T GetAllocatedSize(
+				SizeType NumAllocatedElements, 
+				SIZE_T NumBytesPerElement
+			) const
+			{
+				return Base.GetAllocatedSize(NumAllocatedElements, NumBytesPerElement);
+			}
+
+			FORCEINLINE bool HasAllocation() const
+			{
+				return Base.HasAllocation();
+			}
+
+			FORCEINLINE SizeType GetInitialCapacity() const
+			{
+				return Base.GetInitialCapacity();
+			}
+		private:
+			typename BaseAlloc::ForAnyElementType Base;
+			SSIZE_T AllocSize = 0;
+		};
+
+		template<typename ElementType>
+		class ForElementType : public ForAnyElementType
+		{
+			// Some Allocators, e.g, TAlignedHeapAllocator do some static checks. Instanciate the BaseAlloc::ForElementType
+			// at compile time so those warnings are emitted. 
+			using BaseTypeInstantaitionType = decltype(DeclVal<typename BaseAlloc::template ForElementType<ElementType>>());
+
+		public:
+			ForElementType()
+			{
+			}
+
+			FORCEINLINE ElementType* GetAllocation() const
+			{
+				return (ElementType*)ForAnyElementType::GetAllocation();
+			}
+		};
+	};
+
+
+	// Default memory tarcking allocator needed for TArray and TMap.
+
+	template<typename CounterType>
+	using FDefaultMemoryTrackingAllocator = FMemoryTrackingAllocatorWrapper2<FDefaultAllocator, CounterType>;
+
+	template<typename CounterType>
+	using FDefaultMemoryTrackingBitArrayAllocator = TInlineAllocator<4, FDefaultMemoryTrackingAllocator<CounterType>>;
+
+	template<typename CounterType>
+	using FDefaultMemoryTrackingSparceArrayAllocator = TSparseArrayAllocator<
+		FDefaultMemoryTrackingAllocator<CounterType>,
+		FDefaultMemoryTrackingBitArrayAllocator<CounterType>>;
+		
+	template<typename CounterType>
+	using FDefaultMemoryTrackingSetAllocator = TSetAllocator<
+		FDefaultMemoryTrackingSparceArrayAllocator<CounterType>,
+		TInlineAllocator<1, FDefaultMemoryTrackingAllocator<CounterType>>,
+		DEFAULT_NUMBER_OF_ELEMENTS_PER_HASH_BUCKET,
+		DEFAULT_BASE_NUMBER_OF_HASH_BUCKETS,
+		DEFAULT_MIN_NUMBER_OF_HASHED_ELEMENTS>;
+}
+
+
+template<typename BaseAlloc, typename Counter>
+struct TAllocatorTraits<mu::FMemoryTrackingAllocatorWrapper<BaseAlloc, Counter>> : public TAllocatorTraitsBase<mu::FMemoryTrackingAllocatorWrapper<BaseAlloc, Counter>>
+{
+	enum { SupportsElementAlignment = TAllocatorTraits<BaseAlloc>::SupportsElementAlignment };
+	enum { SupportsMoveFromOtherAllocator = false };
+};
+
+template<typename BaseAlloc, typename Counter>
+struct TAllocatorTraits<mu::FMemoryTrackingAllocatorWrapper2<BaseAlloc, Counter>> : public TAllocatorTraitsBase<mu::FMemoryTrackingAllocatorWrapper2<BaseAlloc, Counter>>
+{
+	enum { SupportsElementAlignment = TAllocatorTraits<BaseAlloc>::SupportsElementAlignment };
+	enum { SupportsMoveFromOtherAllocator = false };
+};
+
+
diff --git a/Engine/Plugins/Experimental/Mutable/Source/MutableRuntime/Public/MuR/Mesh.h b/Engine/Plugins/Experimental/Mutable/Source/MutableRuntime/Public/MuR/Mesh.h
--- a/Engine/Plugins/Experimental/Mutable/Source/MutableRuntime/Public/MuR/Mesh.h
+++ b/Engine/Plugins/Experimental/Mutable/Source/MutableRuntime/Public/MuR/Mesh.h
@@ -346,6 +346,9 @@ namespace mu
     
 	public:
 
+		template<typename Type>
+		using TMemoryTrackedArray = TArray<Type, FDefaultMemoryTrackingAllocator<MemoryCounters::FMeshMemoryCounter>>;
+
 		//! Non-persistent internal id unique for a mesh generated for a specific state and
 		//! parameter values.
 		mutable uint32 m_internalId = 0;
@@ -430,7 +433,7 @@ namespace mu
 		};
 		// This is the pose used by this mesh fragment, used to update the transforms of the final skeleton
 		// taking into consideration the meshes being used.
-		TArray<FBonePose> BonePoses;
+		TMemoryTrackedArray<FBonePose> BonePoses;
 
 		// Array containing the bonemaps of all surfaces in the mesh.
 		TArray<uint16> BoneMap;
diff --git a/Engine/Plugins/Experimental/Mutable/Source/MutableRuntime/Public/MuR/MeshBufferSet.h b/Engine/Plugins/Experimental/Mutable/Source/MutableRuntime/Public/MuR/MeshBufferSet.h
--- a/Engine/Plugins/Experimental/Mutable/Source/MutableRuntime/Public/MuR/MeshBufferSet.h
+++ b/Engine/Plugins/Experimental/Mutable/Source/MutableRuntime/Public/MuR/MeshBufferSet.h
@@ -7,6 +7,13 @@
 #include "HAL/PlatformMath.h"
 #include "MuR/MutableMemory.h"
 
+#include "MuR/MemoryTrackingAllocationPolicy.h"
+
+namespace mu::MemoryCounters
+{
+	struct FMeshMemoryCounterTag {};
+	using FMeshMemoryCounter = TMemoryCounter<FMeshMemoryCounterTag>;
+}
 
 namespace mu
 {
@@ -164,6 +171,9 @@ namespace mu
 
 	struct MESH_BUFFER
 	{
+		template<typename Type>
+		using TMemoryTrackedArray = TArray<Type, FDefaultMemoryTrackingAllocator<MemoryCounters::FMeshMemoryCounter>>;
+
 		//!
 		MESH_BUFFER()
 		{
@@ -174,7 +184,7 @@ namespace mu
 		TArray<mu::MESH_BUFFER_CHANNEL> m_channels;
 
 		//!
-		TArray<uint8> m_data;
+		TMemoryTrackedArray<uint8> m_data;
 
 		//!
 		uint32 m_elementSize;
diff --git a/Engine/Plugins/Experimental/Mutable/Source/MutableRuntime/Public/MuR/Serialisation.h b/Engine/Plugins/Experimental/Mutable/Source/MutableRuntime/Public/MuR/Serialisation.h
--- a/Engine/Plugins/Experimental/Mutable/Source/MutableRuntime/Public/MuR/Serialisation.h
+++ b/Engine/Plugins/Experimental/Mutable/Source/MutableRuntime/Public/MuR/Serialisation.h
@@ -7,13 +7,16 @@
 #include "MuR/Ptr.h"
 #include "MuR/RefCounted.h"
 
+#include "Math/Vector.h"
+#include "Math/IntVector.h"
+#include "Math/Vector4.h"
+#include "MuR/MutableMath.h"
 
 namespace mu
 {    
     class Image;
 	class Model;
 
-
 #define MUTABLE_DEFINE_POD_SERIALISABLE(T)							\
 	template<>														\
 	void DLLEXPORT operator<< <T>(OutputArchive& arch, const T& t);	\
@@ -21,14 +24,12 @@ namespace mu
 	template<>														\
 	void DLLEXPORT operator>> <T>(InputArchive& arch, T& t);		\
 
-
-#define MUTABLE_DEFINE_POD_VECTOR_SERIALISABLE(T)							\
-	template<>																\
-	void DLLEXPORT operator<< <T>(OutputArchive& arch, const TArray<T>& v);	\
-																			\
-	template<>																\
-	void DLLEXPORT operator>> <T>(InputArchive& arch, TArray<T>& v);		\
-
+#define MUTABLE_DEFINE_POD_VECTOR_SERIALISABLE(T)									     \
+	template<typename Alloc>                                                             \
+	void operator <<(OutputArchive& arch, const TArray<T, Alloc>& v);					 \
+	                                                                                     \
+	template<typename Alloc>                                                             \
+	void operator >>(InputArchive& arch, TArray<T, Alloc>& v);                           \
 
 #define MUTABLE_DEFINE_ENUM_SERIALISABLE(T)							\
 	template<>														\
@@ -334,20 +335,41 @@ namespace mu
         t.Unserialise( arch );
 	}
 
-
 	MUTABLE_DEFINE_POD_SERIALISABLE(float);
 	MUTABLE_DEFINE_POD_SERIALISABLE(double);
 
-    MUTABLE_DEFINE_POD_SERIALISABLE( int8 );
-    MUTABLE_DEFINE_POD_SERIALISABLE( int16 );
-    MUTABLE_DEFINE_POD_SERIALISABLE( int32 );
-    MUTABLE_DEFINE_POD_SERIALISABLE( int64 );
-
-    MUTABLE_DEFINE_POD_SERIALISABLE( uint8 );
-    MUTABLE_DEFINE_POD_SERIALISABLE( uint16 );
-    MUTABLE_DEFINE_POD_SERIALISABLE( uint32 );
-    MUTABLE_DEFINE_POD_SERIALISABLE( uint64 );
-
+    MUTABLE_DEFINE_POD_SERIALISABLE(int8);
+    MUTABLE_DEFINE_POD_SERIALISABLE(int16);
+    MUTABLE_DEFINE_POD_SERIALISABLE(int32);
+    MUTABLE_DEFINE_POD_SERIALISABLE(int64);
+
+    MUTABLE_DEFINE_POD_SERIALISABLE(uint8);
+    MUTABLE_DEFINE_POD_SERIALISABLE(uint16);
+    MUTABLE_DEFINE_POD_SERIALISABLE(uint32);
+    MUTABLE_DEFINE_POD_SERIALISABLE(uint64);
+
+	MUTABLE_DEFINE_POD_VECTOR_SERIALISABLE(float);
+	MUTABLE_DEFINE_POD_VECTOR_SERIALISABLE(double);
+	MUTABLE_DEFINE_POD_VECTOR_SERIALISABLE(uint8);
+	MUTABLE_DEFINE_POD_VECTOR_SERIALISABLE(uint16);
+	MUTABLE_DEFINE_POD_VECTOR_SERIALISABLE(uint32);
+	MUTABLE_DEFINE_POD_VECTOR_SERIALISABLE(uint64);
+	MUTABLE_DEFINE_POD_VECTOR_SERIALISABLE(int8);
+	MUTABLE_DEFINE_POD_VECTOR_SERIALISABLE(int16);
+	MUTABLE_DEFINE_POD_VECTOR_SERIALISABLE(int32);
+	MUTABLE_DEFINE_POD_VECTOR_SERIALISABLE(int64);
+
+	MUTABLE_DEFINE_POD_VECTOR_SERIALISABLE(vec2f);
+	MUTABLE_DEFINE_POD_VECTOR_SERIALISABLE(vec3f);
+	MUTABLE_DEFINE_POD_VECTOR_SERIALISABLE(mat3f);
+	MUTABLE_DEFINE_POD_VECTOR_SERIALISABLE(mat4f);
+	MUTABLE_DEFINE_POD_VECTOR_SERIALISABLE(vec2<int>);
+	MUTABLE_DEFINE_POD_VECTOR_SERIALISABLE(TCHAR);
+
+	MUTABLE_DEFINE_POD_VECTOR_SERIALISABLE(FUintVector2);
+	MUTABLE_DEFINE_POD_VECTOR_SERIALISABLE(UE::Math::TIntVector2<uint16>);
+	MUTABLE_DEFINE_POD_VECTOR_SERIALISABLE(UE::Math::TIntVector2<int16>);
+	MUTABLE_DEFINE_POD_VECTOR_SERIALISABLE(FVector4f);
 	
 	//---------------------------------------------------------------------------------------------
 	template<>
@@ -358,23 +380,26 @@ namespace mu
 	
 	
 	//---------------------------------------------------------------------------------------------
-	template<typename T> void operator<<(OutputArchive& arch, const TArray<T>& v)
+	template<typename T, typename Alloc> 
+	void operator<<(OutputArchive& arch, const TArray<T, Alloc>& v)
 	{
-		// TODO: Optimise for vectors of PODs
-		arch << (uint32)v.Num();
-		for (std::size_t i = 0; i < v.Num(); ++i)
+		const uint32 Num = (uint32)v.Num();
+		arch << Num;
+		
+		for (SIZE_T i = 0; i < Num; ++i)
 		{
 			arch << v[i];
 		}
 	}
 
-	template<typename T> void operator>>(InputArchive& arch, TArray<T>& v)
+	template<typename T, typename Alloc> 
+	void operator>>(InputArchive& arch, TArray<T, Alloc>& v)
 	{
-		// TODO: Optimise for vectors of PODs
-		uint32 size;
-		arch >> size;
-		v.SetNum(size);
-		for (std::size_t i = 0; i < size; ++i)
+		uint32 Num;
+		arch >> Num;
+		v.SetNum(Num);
+
+		for (SIZE_T i = 0; i < Num; ++i)
 		{
 			arch >> v[i];
 		}
