[Mutable] Fix CIS compile errors.

#rnx

diff --git a/Engine/Plugins/Experimental/Mutable/Source/MutableRuntime/Public/MuR/MemoryTrackingAllocationPolicy.h b/Engine/Plugins/Experimental/Mutable/Source/MutableRuntime/Public/MuR/MemoryTrackingAllocationPolicy.h
--- a/Engine/Plugins/Experimental/Mutable/Source/MutableRuntime/Public/MuR/MemoryTrackingAllocationPolicy.h
+++ b/Engine/Plugins/Experimental/Mutable/Source/MutableRuntime/Public/MuR/MemoryTrackingAllocationPolicy.h
@@ -46,7 +46,7 @@ namespace mu
 			/** Destructor. */
 			FORCEINLINE ~ForAnyElementType()
 			{
-				CounterType::Counter.fetch_sub(AllocSize, std::memory_order::memory_order_relaxed);
+				CounterType::Counter.fetch_sub(AllocSize, std::memory_order_relaxed);
 
 				AllocSize = 0;
 			}
@@ -58,7 +58,7 @@ namespace mu
 			{
 				BaseAlloc::ForAnyElementType::MoveToEmpty(Other);
 				
-				CounterType::Counter.fetch_sub(AllocSize, std::memory_order::memory_order_relaxed);
+				CounterType::Counter.fetch_sub(AllocSize, std::memory_order_relaxed);
 
 				AllocSize = Other.AllocSize;
 				Other.AllocSize = 0;
@@ -74,7 +74,7 @@ namespace mu
 					PreviousNumElements, NumElements, NumBytesPerElement);
 
 				const SSIZE_T Differential = SSIZE_T(NumElements * NumBytesPerElement) - AllocSize;
-				const SSIZE_T PrevCounterValue = CounterType::Counter.fetch_add(Differential, std::memory_order::memory_order_relaxed);
+				const SSIZE_T PrevCounterValue = CounterType::Counter.fetch_add(Differential, std::memory_order_relaxed);
 				check(PrevCounterValue >= AllocSize);
 
 				AllocSize = NumElements * NumBytesPerElement;
@@ -91,7 +91,7 @@ namespace mu
 					PreviousNumElements, NumElements, NumBytesPerElement, AlignmentOfElement);
 
 				const SSIZE_T Differential = SSIZE_T(NumElements * NumBytesPerElement) - AllocSize;
-				const SSIZE_T PrevCounterValue = CounterType::Counter.fetch_add(Differential, std::memory_order::memory_order_relaxed);
+				const SSIZE_T PrevCounterValue = CounterType::Counter.fetch_add(Differential, std::memory_order_relaxed);
 				check(PrevCounterValue >= AllocSize);
 
 				AllocSize = NumElements * NumBytesPerElement;
@@ -140,7 +140,7 @@ namespace mu
 			/** Destructor. */
 			FORCEINLINE ~ForAnyElementType()
 			{
-				CounterType::Counter.fetch_sub(AllocSize, std::memory_order::memory_order_relaxed);
+				CounterType::Counter.fetch_sub(AllocSize, std::memory_order_relaxed);
 				AllocSize = 0;
 			}
 
@@ -156,7 +156,7 @@ namespace mu
 			{
 				Base.MoveToEmpty(Other.Base);
 
-				CounterType::Counter.fetch_sub(AllocSize, std::memory_order::memory_order_relaxed);
+				CounterType::Counter.fetch_sub(AllocSize, std::memory_order_relaxed);
 
 				AllocSize = Other.AllocSize;
 				Other.AllocSize = 0;
@@ -171,7 +171,7 @@ namespace mu
 				Base.ResizeAllocation(PreviousNumElements, NumElements, NumBytesPerElement);
 
 				const SSIZE_T Differential = SSIZE_T(NumElements * NumBytesPerElement) - AllocSize;
-				const SSIZE_T PrevCounterValue = CounterType::Counter.fetch_add(Differential, std::memory_order::memory_order_relaxed);
+				const SSIZE_T PrevCounterValue = CounterType::Counter.fetch_add(Differential, std::memory_order_relaxed);
 				check(PrevCounterValue >= AllocSize);
 
 				AllocSize = NumElements * NumBytesPerElement;
@@ -187,7 +187,7 @@ namespace mu
 				Base.ResizeAllocation(PreviousNumElements, NumElements, NumBytesPerElement, AlignmentOfElement);
 
 				const SSIZE_T Differential = SSIZE_T(NumElements * NumBytesPerElement) - AllocSize;
-				const SSIZE_T PrevCounterValue = CounterType::Counter.fetch_add(Differential, std::memory_order::memory_order_relaxed);
+				const SSIZE_T PrevCounterValue = CounterType::Counter.fetch_add(Differential, std::memory_order_relaxed);
 				check(PrevCounterValue >= AllocSize);
 
 				AllocSize = NumElements * NumBytesPerElement;
