[Mutable] Fix allocator wrapper aligned CalculateSlackShrink.

#rnx

diff --git a/Engine/Plugins/Experimental/Mutable/Source/MutableRuntime/Public/MuR/MemoryTrackingAllocationPolicy.h b/Engine/Plugins/Experimental/Mutable/Source/MutableRuntime/Public/MuR/MemoryTrackingAllocationPolicy.h
--- a/Engine/Plugins/Experimental/Mutable/Source/MutableRuntime/Public/MuR/MemoryTrackingAllocationPolicy.h
+++ b/Engine/Plugins/Experimental/Mutable/Source/MutableRuntime/Public/MuR/MemoryTrackingAllocationPolicy.h
@@ -226,7 +226,7 @@ namespace mu
 				uint32 AlignmentOfElement
 			) const
 			{
-				return Base.CalculateSlackShrink(NumElements, CurrentNumSlackElements, NumBytesPerElement);
+				return Base.CalculateSlackShrink(NumElements, CurrentNumSlackElements, NumBytesPerElement, AlignmentOfElement);
 			}
 
 			SizeType CalculateSlackGrow(
