[Mutable] Addapt ClipMeshMorphPlane to the new BoneId type (string to uint16).
Fix wrong legacy deserialization.

#rnx

diff --git a/Engine/Plugins/Experimental/Mutable/Source/CustomizableObject/Classes/MuCO/CustomizableObject.h b/Engine/Plugins/Experimental/Mutable/Source/CustomizableObject/Classes/MuCO/CustomizableObject.h
--- a/Engine/Plugins/Experimental/Mutable/Source/CustomizableObject/Classes/MuCO/CustomizableObject.h
+++ b/Engine/Plugins/Experimental/Mutable/Source/CustomizableObject/Classes/MuCO/CustomizableObject.h
@@ -1439,7 +1439,7 @@ private:
 	// This is a manual version number for the binary blobs in this asset.
 	// Increasing it invalidates all the previously compiled models.
 	// Warning: If while merging code both versions have changed, take the highest+1.
-	static const int32 CurrentSupportedVersion = 394;
+	static const int32 CurrentSupportedVersion = 395;
 
 public:
 
diff --git a/Engine/Plugins/Experimental/Mutable/Source/CustomizableObjectEditor/Private/MuCOE/GenerateMutableSource/GenerateMutableSourceModifier.cpp b/Engine/Plugins/Experimental/Mutable/Source/CustomizableObjectEditor/Private/MuCOE/GenerateMutableSource/GenerateMutableSourceModifier.cpp
--- a/Engine/Plugins/Experimental/Mutable/Source/CustomizableObjectEditor/Private/MuCOE/GenerateMutableSource/GenerateMutableSourceModifier.cpp
+++ b/Engine/Plugins/Experimental/Mutable/Source/CustomizableObjectEditor/Private/MuCOE/GenerateMutableSource/GenerateMutableSourceModifier.cpp
@@ -58,7 +58,7 @@ mu::NodeModifierPtr GenerateMutableSourceModifier(const UEdGraphPin * Pin, FMuta
 		ClipNode->SetParams(TypedNodeClip->B, TypedNodeClip->Exponent);
 		ClipNode->SetMorphEllipse(TypedNodeClip->Radius, TypedNodeClip->Radius2, TypedNodeClip->RotationAngle);
 
-		ClipNode->SetVertexSelectionBone(StringCast<ANSICHAR>(*TypedNodeClip->BoneName.ToString()).Get(), TypedNodeClip->MaxEffectRadius);
+		ClipNode->SetVertexSelectionBone(GenerationContext.BoneNames.AddUnique(TypedNodeClip->BoneName), TypedNodeClip->MaxEffectRadius);
 
 		for (const FString& Tag : TypedNodeClip->Tags)
 		{
diff --git a/Engine/Plugins/Experimental/Mutable/Source/CustomizableObjectEditor/Private/MuCOE/SMutableCodeViewer.cpp b/Engine/Plugins/Experimental/Mutable/Source/CustomizableObjectEditor/Private/MuCOE/SMutableCodeViewer.cpp
--- a/Engine/Plugins/Experimental/Mutable/Source/CustomizableObjectEditor/Private/MuCOE/SMutableCodeViewer.cpp
+++ b/Engine/Plugins/Experimental/Mutable/Source/CustomizableObjectEditor/Private/MuCOE/SMutableCodeViewer.cpp
@@ -2317,16 +2317,6 @@ bool SMutableCodeViewer::IsConstantResourceUsedByOperation(const int32 IndexOnCo
 				{
 					bResourceLocated = IndexOnConstantsArray == InProgram.GetOpArgs<mu::OP::InstanceAddArgs>(OperationAddress).name;
 				}
-				else if (OperationType == mu::OP_TYPE::ME_CLIPMORPHPLANE)
-				{
-					const mu::OP::MeshClipMorphPlaneArgs Arguments = InProgram.GetOpArgs<mu::OP::MeshClipMorphPlaneArgs>(OperationAddress);
-
-					// treat the data as if it was a bone name
-					if (Arguments.vertexSelectionType == mu::OP::MeshClipMorphPlaneArgs::VS_BONE_HIERARCHY)
-					{
-						bResourceLocated = IndexOnConstantsArray == InProgram.GetOpArgs<mu::OP::MeshClipMorphPlaneArgs>(OperationAddress).vertexSelectionShapeOrBone;
-					}
-				}
 				else if (OperationType == mu::OP_TYPE::ME_BINDSHAPE)
 				{
 					mu::OP::MeshBindShapeArgs Arguments = InProgram.GetOpArgs<mu::OP::MeshBindShapeArgs>(OperationAddress);
diff --git a/Engine/Plugins/Experimental/Mutable/Source/MutableRuntime/Private/MuR/Mesh.cpp b/Engine/Plugins/Experimental/Mutable/Source/MutableRuntime/Private/MuR/Mesh.cpp
--- a/Engine/Plugins/Experimental/Mutable/Source/MutableRuntime/Private/MuR/Mesh.cpp
+++ b/Engine/Plugins/Experimental/Mutable/Source/MutableRuntime/Private/MuR/Mesh.cpp
@@ -1147,16 +1147,15 @@ void UnserialiseLegacySurfaces(InputArchive& arch, TArray<MESH_SURFACE>& OutMesh
 	const int32 NumSurfaces = LegacyMeshSurfaces.Num();
 	OutMeshSurfaces.SetNumZeroed(NumSurfaces);
 
-	MESH_SURFACE* Surfaces = OutMeshSurfaces.GetData();
-	FMeshSurfaceLegacy* LegacySurfaces = LegacyMeshSurfaces.GetData();
-
 	for (int32 SurfaceIndex = 0; SurfaceIndex < NumSurfaces; ++SurfaceIndex)
 	{
-		Surfaces->m_firstVertex = LegacySurfaces->m_firstVertex;
-		Surfaces->m_vertexCount = LegacySurfaces->m_vertexCount;
-		Surfaces->m_firstIndex = LegacySurfaces->m_firstIndex;
-		Surfaces->m_indexCount = LegacySurfaces->m_indexCount;
-		Surfaces->m_id = LegacySurfaces->m_id;
+		FMeshSurfaceLegacy& LegacySurface = LegacyMeshSurfaces[SurfaceIndex];
+		MESH_SURFACE& Surface = OutMeshSurfaces[SurfaceIndex];
+		Surface.m_firstVertex = LegacySurface.m_firstVertex;
+		Surface.m_vertexCount = LegacySurface.m_vertexCount;
+		Surface.m_firstIndex = LegacySurface.m_firstIndex;
+		Surface.m_indexCount = LegacySurface.m_indexCount;
+		Surface.m_id = LegacySurface.m_id;
 	}
 }
 
diff --git a/Engine/Plugins/Experimental/Mutable/Source/MutableTools/Private/MuT/ASTOpMeshClipMorphPlane.cpp b/Engine/Plugins/Experimental/Mutable/Source/MutableTools/Private/MuT/ASTOpMeshClipMorphPlane.cpp
--- a/Engine/Plugins/Experimental/Mutable/Source/MutableTools/Private/MuT/ASTOpMeshClipMorphPlane.cpp
+++ b/Engine/Plugins/Experimental/Mutable/Source/MutableTools/Private/MuT/ASTOpMeshClipMorphPlane.cpp
@@ -93,7 +93,7 @@ namespace mu
 			args.vertexSelectionType = (uint8_t)vertexSelectionType;
 			if (vertexSelectionType == OP::MeshClipMorphPlaneArgs::VS_BONE_HIERARCHY)
 			{
-				args.vertexSelectionShapeOrBone = program.AddConstant(vertexSelectionBone);
+				args.vertexSelectionShapeOrBone = vertexSelectionBone;
 			}
 			else if (vertexSelectionType == OP::MeshClipMorphPlaneArgs::VS_SHAPE)
 			{
diff --git a/Engine/Plugins/Experimental/Mutable/Source/MutableTools/Private/MuT/ASTOpMeshClipMorphPlane.h b/Engine/Plugins/Experimental/Mutable/Source/MutableTools/Private/MuT/ASTOpMeshClipMorphPlane.h
--- a/Engine/Plugins/Experimental/Mutable/Source/MutableTools/Private/MuT/ASTOpMeshClipMorphPlane.h
+++ b/Engine/Plugins/Experimental/Mutable/Source/MutableTools/Private/MuT/ASTOpMeshClipMorphPlane.h
@@ -24,7 +24,7 @@ struct FProgram;
 
 		FShape morphShape;
 		FShape selectionShape;
-		string vertexSelectionBone;
+		uint16 vertexSelectionBone;
 
 		OP::MeshClipMorphPlaneArgs::VERTEX_SELECTION_TYPE vertexSelectionType = OP::MeshClipMorphPlaneArgs::VS_NONE;
 
diff --git a/Engine/Plugins/Experimental/Mutable/Source/MutableTools/Private/MuT/NodeMeshClipMorphPlane.cpp b/Engine/Plugins/Experimental/Mutable/Source/MutableTools/Private/MuT/NodeMeshClipMorphPlane.cpp
--- a/Engine/Plugins/Experimental/Mutable/Source/MutableTools/Private/MuT/NodeMeshClipMorphPlane.cpp
+++ b/Engine/Plugins/Experimental/Mutable/Source/MutableTools/Private/MuT/NodeMeshClipMorphPlane.cpp
@@ -123,10 +123,10 @@ namespace mu
 	}
 
 	//---------------------------------------------------------------------------------------------
-	void NodeMeshClipMorphPlane::SetVertexSelectionBone(const char* strBoneName, float maxEffectRadius)
+	void NodeMeshClipMorphPlane::SetVertexSelectionBone(uint16 BoneId, float maxEffectRadius)
 	{
 		m_pD->m_vertexSelectionType = Private::VS_BONE_HIERARCHY;
-		m_pD->m_vertexSelectionBone = strBoneName;
+		m_pD->m_vertexSelectionBone = BoneId;
 		m_pD->m_maxEffectRadius = maxEffectRadius;
 	}
 
diff --git a/Engine/Plugins/Experimental/Mutable/Source/MutableTools/Private/MuT/NodeMeshClipMorphPlanePrivate.h b/Engine/Plugins/Experimental/Mutable/Source/MutableTools/Private/MuT/NodeMeshClipMorphPlanePrivate.h
--- a/Engine/Plugins/Experimental/Mutable/Source/MutableTools/Private/MuT/NodeMeshClipMorphPlanePrivate.h
+++ b/Engine/Plugins/Experimental/Mutable/Source/MutableTools/Private/MuT/NodeMeshClipMorphPlanePrivate.h
@@ -64,7 +64,7 @@ namespace mu
         uint8_t m_vertexSelectionType;
 		vec3f m_selectionBoxOrigin;
 		vec3f m_selectionBoxRadius;
-		string m_vertexSelectionBone;
+		uint16 m_vertexSelectionBone;
 
 		TArray<mu::string> m_tags;
 
@@ -75,7 +75,7 @@ namespace mu
 		//!
 		void Serialise( OutputArchive& arch ) const
 		{
-            uint32_t ver = 3;
+            uint32_t ver = 4;
 			arch << ver;
 
 			arch << m_pSource;
@@ -113,7 +113,17 @@ namespace mu
 			arch >> m_selectionBoxOrigin;
 			arch >> m_selectionBoxRadius;
 			arch >> m_tags;
-			arch >> m_vertexSelectionBone;
+
+			if (ver >= 4)
+			{
+				arch >> m_vertexSelectionBone;
+			}
+			else
+			{
+				string OldVertexSelectionBone;
+				arch >> OldVertexSelectionBone;
+				m_vertexSelectionBone = 0;
+			}
 			arch >> m_maxEffectRadius;
 		}
 
diff --git a/Engine/Plugins/Experimental/Mutable/Source/MutableTools/Private/MuT/NodeModifierMeshClipMorphPlane.cpp b/Engine/Plugins/Experimental/Mutable/Source/MutableTools/Private/MuT/NodeModifierMeshClipMorphPlane.cpp
--- a/Engine/Plugins/Experimental/Mutable/Source/MutableTools/Private/MuT/NodeModifierMeshClipMorphPlane.cpp
+++ b/Engine/Plugins/Experimental/Mutable/Source/MutableTools/Private/MuT/NodeModifierMeshClipMorphPlane.cpp
@@ -87,10 +87,10 @@ namespace mu
 	}
 
 	//---------------------------------------------------------------------------------------------
-	void NodeModifierMeshClipMorphPlane::SetVertexSelectionBone(const char* strBoneName, float maxEffectRadius)
+	void NodeModifierMeshClipMorphPlane::SetVertexSelectionBone(uint16 BoneId, float maxEffectRadius)
 	{
 		m_pD->m_vertexSelectionType = Private::VS_BONE_HIERARCHY;
-		m_pD->m_vertexSelectionBone = strBoneName;
+		m_pD->m_vertexSelectionBone = BoneId;
 		m_pD->m_maxEffectRadius = maxEffectRadius;
 	}
 
diff --git a/Engine/Plugins/Experimental/Mutable/Source/MutableTools/Private/MuT/NodeModifierMeshClipMorphPlanePrivate.h b/Engine/Plugins/Experimental/Mutable/Source/MutableTools/Private/MuT/NodeModifierMeshClipMorphPlanePrivate.h
--- a/Engine/Plugins/Experimental/Mutable/Source/MutableTools/Private/MuT/NodeModifierMeshClipMorphPlanePrivate.h
+++ b/Engine/Plugins/Experimental/Mutable/Source/MutableTools/Private/MuT/NodeModifierMeshClipMorphPlanePrivate.h
@@ -59,7 +59,7 @@ namespace mu
         uint8_t m_vertexSelectionType;
         vec3f m_selectionBoxOrigin;
         vec3f m_selectionBoxRadius;
-        string m_vertexSelectionBone;
+        uint16 m_vertexSelectionBone;
 
 		// Max distance a vertex can have to the bone in order to be affected. A negative value
 		// means no limit.
@@ -70,7 +70,7 @@ namespace mu
         {
             NodeModifier::Private::Serialise(arch);
 
-			uint32_t ver = 2;
+			uint32_t ver = 3;
             arch << ver;
 
             arch << m_origin;
@@ -94,7 +94,7 @@ namespace mu
 
             uint32_t ver;
             arch >> ver;
-            check(ver==2);
+            check(ver>=2);
 
             arch >> m_origin;
             arch >> m_normal;
@@ -106,7 +106,16 @@ namespace mu
             arch >> m_vertexSelectionType;
             arch >> m_selectionBoxOrigin;
             arch >> m_selectionBoxRadius;
-			arch >> m_vertexSelectionBone;
+			if (ver >= 3)
+			{
+				arch >> m_vertexSelectionBone;
+			}
+			else
+			{
+				string OldVertexSelectionBone;
+				arch >> OldVertexSelectionBone;
+				m_vertexSelectionBone = 0;
+			}
 			arch >> m_maxEffectRadius;
         }
 
diff --git a/Engine/Plugins/Experimental/Mutable/Source/MutableTools/Public/MuT/NodeMeshClipMorphPlane.h b/Engine/Plugins/Experimental/Mutable/Source/MutableTools/Public/MuT/NodeMeshClipMorphPlane.h
--- a/Engine/Plugins/Experimental/Mutable/Source/MutableTools/Public/MuT/NodeMeshClipMorphPlane.h
+++ b/Engine/Plugins/Experimental/Mutable/Source/MutableTools/Public/MuT/NodeMeshClipMorphPlane.h
@@ -66,7 +66,7 @@ namespace mu
 
 		//! Define the root bone of the subhierarchy of the mesh that will be affected.
 		//! Only one of Box or Bone Hierarchy can be used (the last one set)
-		void SetVertexSelectionBone(const char* strBoneName, float maxEffectRadius);
+		void SetVertexSelectionBone(uint16 BoneId, float maxEffectRadius);
 
 		//! Add a tag to the clip morph operation, which will only affect surfaces with the same tag
 		void AddTag(const char* tagName);
diff --git a/Engine/Plugins/Experimental/Mutable/Source/MutableTools/Public/MuT/NodeModifierMeshClipMorphPlane.h b/Engine/Plugins/Experimental/Mutable/Source/MutableTools/Public/MuT/NodeModifierMeshClipMorphPlane.h
--- a/Engine/Plugins/Experimental/Mutable/Source/MutableTools/Public/MuT/NodeModifierMeshClipMorphPlane.h
+++ b/Engine/Plugins/Experimental/Mutable/Source/MutableTools/Public/MuT/NodeModifierMeshClipMorphPlane.h
@@ -61,7 +61,7 @@ namespace mu
 
 		//! Define the root bone of the subhierarchy of the mesh that will be affected.
 		//! Only one of Box or Bone Hierarchy can be used (the last one set)
-		void SetVertexSelectionBone(const char* strBoneName, float maxEffectRadius);
+		void SetVertexSelectionBone(uint16 BoneId, float maxEffectRadius);
 
 		//-----------------------------------------------------------------------------------------
 		// Interface pattern
