[Mutable] Add CVar to change the priority of the streaming request of Additional Assets.

#rnx

diff --git a/Engine/Plugins/Experimental/Mutable/Source/CustomizableObject/Private/MuCO/CustomizableObjectInstance.cpp b/Engine/Plugins/Experimental/Mutable/Source/CustomizableObject/Private/MuCO/CustomizableObjectInstance.cpp
--- a/Engine/Plugins/Experimental/Mutable/Source/CustomizableObject/Private/MuCO/CustomizableObjectInstance.cpp
+++ b/Engine/Plugins/Experimental/Mutable/Source/CustomizableObject/Private/MuCO/CustomizableObjectInstance.cpp
@@ -5027,6 +5027,11 @@ bool UCustomizableInstancePrivateData::BuildOrCopyRenderData(const TSharedPtr<FM
 	return true;
 }
 
+static bool bEnableHighPriorityLoading = true;
+FAutoConsoleVariableRef CVarMutableHighPriorityLoading(
+	TEXT("Mutable.EnableLoadingAssetsWithHighPriority"),
+	bEnableHighPriorityLoading,
+	TEXT("If enabled, the request to load additional assets will have high priority."));
 
 FGraphEventRef UCustomizableInstancePrivateData::LoadAdditionalAssetsAsync(const TSharedPtr<FMutableOperationData>& OperationData, UCustomizableObjectInstance* Public, FStreamableManager& StreamableManager)
 {
@@ -5275,7 +5280,8 @@ FGraphEventRef UCustomizableInstancePrivateData::LoadAdditionalAssetsAsync(const
 		check(!StreamingHandle);
 
 		Result = FGraphEvent::CreateGraphEvent();
-		StreamingHandle = StreamableManager.RequestAsyncLoad(AssetsToStream, FStreamableDelegate::CreateUObject(Public, &UCustomizableObjectInstance::AdditionalAssetsAsyncLoaded, Result));
+		StreamingHandle = StreamableManager.RequestAsyncLoad(AssetsToStream, FStreamableDelegate::CreateUObject(Public, &UCustomizableObjectInstance::AdditionalAssetsAsyncLoaded, Result),
+			bEnableHighPriorityLoading ? FStreamableManager::AsyncLoadHighPriority : FStreamableManager::DefaultAsyncLoadPriority);
 	}
 
 	return Result;
