[Mutable] Reduce MutableRuntime memory peak.

Forbidden to do an UpdateImage before releasing the resources from a generation update.

#rnx

diff --git a/Engine/Plugins/Experimental/Mutable/Source/CustomizableObject/Private/MuCO/FMutableTaskGraph.cpp b/Engine/Plugins/Experimental/Mutable/Source/CustomizableObject/Private/MuCO/FMutableTaskGraph.cpp
--- a/Engine/Plugins/Experimental/Mutable/Source/CustomizableObject/Private/MuCO/FMutableTaskGraph.cpp
+++ b/Engine/Plugins/Experimental/Mutable/Source/CustomizableObject/Private/MuCO/FMutableTaskGraph.cpp
@@ -155,7 +155,7 @@ void FMutableTaskGraph::TryLaunchMutableTaskLowPriority(bool bFromMutableTask)
 			return;		
 		}
 		
-		if (!IsTaskCompleted(LastMutableTaskLowPriority)) // Also check inside the lock since we will be writing it
+		if (!IsTaskCompleted(LastMutableTaskLowPriority)) // Check #1 // Also check inside the lock since we will be writing it
 		{
 			return;
 		}
@@ -166,9 +166,8 @@ void FMutableTaskGraph::TryLaunchMutableTaskLowPriority(bool bFromMutableTask)
 		TimeElapsed = FPlatformTime::Seconds() - NextTask.CreationTime;
 
 		bTimeLimit = TimeElapsed >= TimeLimit;
-		if (!bTimeLimit && (
-				!bAllowLaunchMutableTaskLowPriority ||
-				(bFromMutableTask && IsTaskCompleted(LastMutableTask))))
+		if (!bAllowLaunchMutableTaskLowPriority || // Check #2
+			bFromMutableTask && IsTaskCompleted(LastMutableTask) && !bTimeLimit) // Check #3
 		{			
 			return;
 		}
diff --git a/Engine/Plugins/Experimental/Mutable/Source/CustomizableObject/Public/MuCO/FMutableTaskGraph.h b/Engine/Plugins/Experimental/Mutable/Source/CustomizableObject/Public/MuCO/FMutableTaskGraph.h
--- a/Engine/Plugins/Experimental/Mutable/Source/CustomizableObject/Public/MuCO/FMutableTaskGraph.h
+++ b/Engine/Plugins/Experimental/Mutable/Source/CustomizableObject/Public/MuCO/FMutableTaskGraph.h
@@ -26,7 +26,7 @@
  *
  * - A Low Priority task will not be launched if one of the follow conditions is true (in order):
  *     1. There is a task Low Priority task running.
- *     2. Flag bLaunchMutableTaskLowPriority is false (unless time limit).
+ *     2. Flag bLaunchMutableTaskLowPriority is false.
  *     3. There is a Normal Priority task running (unless time limit).
  */
 class FMutableTaskGraph
