[mutable] Fix debugger crash when previewing images.
#rnx

diff --git a/Engine/Plugins/Experimental/Mutable/Source/MutableRuntime/Private/MuR/System.cpp b/Engine/Plugins/Experimental/Mutable/Source/MutableRuntime/Private/MuR/System.cpp
--- a/Engine/Plugins/Experimental/Mutable/Source/MutableRuntime/Private/MuR/System.cpp
+++ b/Engine/Plugins/Experimental/Mutable/Source/MutableRuntime/Private/MuR/System.cpp
@@ -621,14 +621,17 @@ namespace mu
 
 
 	//---------------------------------------------------------------------------------------------
-	void System::Private::BeginBuild(const TSharedPtr<const Model>& pModel)
+	void System::Private::BeginBuild(const TSharedPtr<const Model>& InModel)
 	{
 		// We don't have a FLiveInstance, let's create the memory
 		// \TODO: There is no clear moment to remove this... EndBuild?
 		WorkingMemoryManager.CurrentInstanceCache = MakeShared<FProgramCache>();
-		WorkingMemoryManager.CurrentInstanceCache->Init(pModel->GetPrivate()->m_program.m_opAddress.Num());
+		WorkingMemoryManager.CurrentInstanceCache->Init(InModel->GetPrivate()->m_program.m_opAddress.Num());
 
-		PrepareCache(pModel.Get(), -1);
+		// Ensure the model cache has been created
+		WorkingMemoryManager.FindOrAddModelCache(InModel);
+
+		PrepareCache(InModel.Get(), -1);
 	}
 
 
diff --git a/Engine/Plugins/Experimental/Mutable/Source/MutableTools/Private/MuT/AST.h b/Engine/Plugins/Experimental/Mutable/Source/MutableTools/Private/MuT/AST.h
--- a/Engine/Plugins/Experimental/Mutable/Source/MutableTools/Private/MuT/AST.h
+++ b/Engine/Plugins/Experimental/Mutable/Source/MutableTools/Private/MuT/AST.h
@@ -1121,12 +1121,5 @@ namespace mu
         }
     };
 
-
-    //
-    extern void DebugLogAST(const Ptr<ASTOp>& at,
-                            int indent=0,
-                            ASTOpList* done=nullptr,
-                            const char* label=nullptr);
-
 }
 
