#rnx

diff --git a/Engine/Plugins/Experimental/Mutable/Source/CustomizableObject/Classes/MuCO/CustomizableInstanceLODManagement.h b/Engine/Plugins/Experimental/Mutable/Source/CustomizableObject/Classes/MuCO/CustomizableInstanceLODManagement.h
--- a/Engine/Plugins/Experimental/Mutable/Source/CustomizableObject/Classes/MuCO/CustomizableInstanceLODManagement.h
+++ b/Engine/Plugins/Experimental/Mutable/Source/CustomizableObject/Classes/MuCO/CustomizableInstanceLODManagement.h
@@ -7,8 +7,6 @@
 
 typedef TMap<const class UCustomizableObjectInstance*, class FMutableUpdateCandidate> FMutableInstanceUpdateMap;
 
-#define NEW_MUTABLE_PLUGIN_QUEUE_SYSTEM
-
 // This is an abstract base class, override it to create a new Instance LOD management system and register with UCustomizableObjectSystem::SetInstanceLODManagement
 UCLASS(Blueprintable, BlueprintType)
 class CUSTOMIZABLEOBJECT_API UCustomizableInstanceLODManagementBase : public UObject
@@ -20,8 +18,6 @@ public:
 
 	// WARNING! The following methods must be overriden in derived classes
 	virtual void UpdateInstanceDistsAndLODs(FMutableInstanceUpdateMap& InOutRequestedUpdates) { check(0); };
-	// Remove the following method after CL 25973936 is propagated to all streams
-	virtual void UpdateInstanceDistsAndLODs() { check(0); };
 	virtual int32 GetNumGeneratedInstancesLimitFullLODs() const { check(0); return 0; };
 	virtual int32 GetNumGeneratedInstancesLimitLOD1() const { check(0); return 0; };
 	virtual int32 GetNumGeneratedInstancesLimitLOD2() const { check(0); return 0; };
diff --git a/Engine/Plugins/Experimental/Mutable/Source/CustomizableObject/Classes/MuCO/CustomizableObjectInstance.h b/Engine/Plugins/Experimental/Mutable/Source/CustomizableObject/Classes/MuCO/CustomizableObjectInstance.h
--- a/Engine/Plugins/Experimental/Mutable/Source/CustomizableObject/Classes/MuCO/CustomizableObjectInstance.h
+++ b/Engine/Plugins/Experimental/Mutable/Source/CustomizableObject/Classes/MuCO/CustomizableObjectInstance.h
@@ -647,10 +647,6 @@ public:
 	UE_DEPRECATED(5.2, "Use SetRequestedLODs instead.")
 	void SetMinMaxLODToLoad(FMutableInstanceUpdateMap& InOutRequestedUpdates, int32 NewMinLOD = 0, int32 NewMaxLOD = INT32_MAX, bool bLimitLODUpgrades = true);
 
-	// Remove the following method after CL 25973936 is propagated to all streams
-	UE_DEPRECATED(5.2, "Use SetRequestedLODs instead.")
-	void SetMinMaxLODToLoad(int32 NewMinLOD = 0, int32 NewMaxLOD = INT32_MAX, bool bLimitLODUpgrades = true);
-
 	/** Return the Min LOD this Instance is using (from the beginning of an update. If an update fails this value will be incorrect). */
 	int32 GetCurrentMinLOD() const;
 
@@ -668,9 +664,6 @@ public:
 	 * @param InOutRequestedUpdates - Map from Instance to Update data that stores a request for the Instance to be updated, which will be either processed or discarded by priority (to be rerequested the next tick) */
 	void SetRequestedLODs(int32 InMinLOD, int32 InMaxLOD, const TArray<uint16>& InRequestedLODsPerComponent, FMutableInstanceUpdateMap& InOutRequestedUpdates);
 
-	// Remove the following method after CL 25973936 is propagated to all streams
-	void SetRequestedLODs(int32 InMinLOD, int32 InMaxLOD, const TArray<uint16>& InRequestedLODsPerComponent);
-
 	const TArray<uint16>& GetRequestedLODsPerComponent() const;
 
 	/** Common end point of all updates. Even those which failed. */
diff --git a/Engine/Plugins/Experimental/Mutable/Source/CustomizableObject/Private/MuCO/CustomizableObjectInstance.cpp b/Engine/Plugins/Experimental/Mutable/Source/CustomizableObject/Private/MuCO/CustomizableObjectInstance.cpp
--- a/Engine/Plugins/Experimental/Mutable/Source/CustomizableObject/Private/MuCO/CustomizableObjectInstance.cpp
+++ b/Engine/Plugins/Experimental/Mutable/Source/CustomizableObject/Private/MuCO/CustomizableObjectInstance.cpp
@@ -6251,12 +6251,6 @@ void UCustomizableObjectInstance::SetMinMaxLODToLoad(FMutableInstanceUpdateMap&
 }
 
 
-void UCustomizableObjectInstance::SetMinMaxLODToLoad(int32 NewMinLOD, int32 NewMaxLOD, bool bLimitLODUpgrades)
-{
-	SetRequestedLODs(NewMinLOD, NewMaxLOD, Descriptor.RequestedLODLevels);
-}
-
-
 int32 UCustomizableObjectInstance::GetNumComponents() const
 {
 	return GetCustomizableObject() ? GetCustomizableObject()->GetComponentCount() : 0;
@@ -6407,12 +6401,6 @@ void UCustomizableObjectInstance::SetRequestedLODs(int32 InMinLOD, int32 InMaxLO
 }
 
 
-void UCustomizableObjectInstance::SetRequestedLODs(int32 InMinLOD, int32 InMaxLOD, const TArray<uint16>& InRequestedLODsPerComponent)
-{
-	// Remove this method after CL 25973936 is propagated to all streams
-}
-
-
 const TArray<uint16>& UCustomizableObjectInstance::GetRequestedLODsPerComponent() const
 {
 	return Descriptor.RequestedLODLevels;
