[mutable] Changed table material functionality. Now a table material works as a real material swtch instead of creating a swtch for each modified parameter.
+ Now, if a node is linked to a texture parameter, the image pin mode would not be able to be modified.

#rnx

diff --git a/Engine/Plugins/Experimental/Mutable/Source/CustomizableObject/Private/MuCO/CustomizableObjectSystem.cpp b/Engine/Plugins/Experimental/Mutable/Source/CustomizableObject/Private/MuCO/CustomizableObjectSystem.cpp
--- a/Engine/Plugins/Experimental/Mutable/Source/CustomizableObject/Private/MuCO/CustomizableObjectSystem.cpp
+++ b/Engine/Plugins/Experimental/Mutable/Source/CustomizableObject/Private/MuCO/CustomizableObjectSystem.cpp
@@ -1400,7 +1400,7 @@ namespace impl
 						FInstanceUpdateData::FSurface& Surface = OperationData->InstanceUpdateData.Surfaces.Last();
 						++Component.SurfaceCount;
 
-						Surface.MaterialIndex = Instance->GetSurfaceCustomId(MutableLODIndex, ComponentIndex, InstanceSurfaceIndex);
+						// Now Surface.MaterialIndex is decoded from a parameter at the end of this if()
 						Surface.SurfaceId = SurfaceId;
 
 						const int32 SharedSurfaceId = Instance->GetSharedSurfaceId(MutableLODIndex, ComponentIndex, InstanceSurfaceIndex);
@@ -1455,11 +1455,25 @@ namespace impl
 						Surface.ScalarCount = Instance->GetScalarCount(BaseLODIndex, ComponentIndex, InstanceSurfaceIndex);
 						for (int32 ScalarIndex = 0; ScalarIndex < Surface.ScalarCount; ++ScalarIndex)
 						{
-							MUTABLE_CPUPROFILER_SCOPE(GetScalar);
-							OperationData->InstanceUpdateData.Scalars.Push({});
-							FInstanceUpdateData::FScalar& Scalar = OperationData->InstanceUpdateData.Scalars.Last();
-							Scalar.Name = Instance->GetScalarName(BaseLODIndex, ComponentIndex, InstanceSurfaceIndex, ScalarIndex);
-							Scalar.Scalar = Instance->GetScalar(BaseLODIndex, ComponentIndex, InstanceSurfaceIndex, ScalarIndex);
+							MUTABLE_CPUPROFILER_SCOPE(GetScalar)
+
+							FString ScalarName = Instance->GetScalarName(BaseLODIndex, ComponentIndex, InstanceSurfaceIndex, ScalarIndex);
+							float ScalarValue = Instance->GetScalar(BaseLODIndex, ComponentIndex, InstanceSurfaceIndex, ScalarIndex);
+							
+							FString EncodingMaterialIdString = "__MutableMaterialId";
+							
+							// Decoding Material Switch from Mutable parameter name
+							if (ScalarName.Equals(EncodingMaterialIdString))
+							{
+								Surface.MaterialIndex = (uint32)ScalarValue;
+							
+								// This parameter is not needed in the final material instance
+								Surface.ScalarCount -= 1;
+							}
+							else
+							{
+								OperationData->InstanceUpdateData.Scalars.Push({ ScalarName, ScalarValue });
+							}
 						}
 					}
 				}
diff --git a/Engine/Plugins/Experimental/Mutable/Source/CustomizableObject/Public/MuCO/CustomizableObjectCustomVersion.h b/Engine/Plugins/Experimental/Mutable/Source/CustomizableObject/Public/MuCO/CustomizableObjectCustomVersion.h
--- a/Engine/Plugins/Experimental/Mutable/Source/CustomizableObject/Public/MuCO/CustomizableObjectCustomVersion.h
+++ b/Engine/Plugins/Experimental/Mutable/Source/CustomizableObject/Public/MuCO/CustomizableObjectCustomVersion.h
@@ -85,6 +85,8 @@ struct CUSTOMIZABLEOBJECT_API FCustomizableObjectCustomVersion
 		IgnoreDisabledSections,
 
 		SkeletalMeshNodeDefaultPinWithoutPinData,
+		
+		AddedTableMaterialSwitch,
 
 		// -----<new versions can be added above this line>-------------------------------------------------
 		VersionPlusOne,
diff --git a/Engine/Plugins/Experimental/Mutable/Source/CustomizableObjectEditor/Classes/MuCOE/Nodes/CustomizableObjectNodeMaterial.h b/Engine/Plugins/Experimental/Mutable/Source/CustomizableObjectEditor/Classes/MuCOE/Nodes/CustomizableObjectNodeMaterial.h
--- a/Engine/Plugins/Experimental/Mutable/Source/CustomizableObjectEditor/Classes/MuCOE/Nodes/CustomizableObjectNodeMaterial.h
+++ b/Engine/Plugins/Experimental/Mutable/Source/CustomizableObjectEditor/Classes/MuCOE/Nodes/CustomizableObjectNodeMaterial.h
@@ -242,7 +242,8 @@ private:
 	UPROPERTY()
 	TMap<FGuid, FEdGraphPinReference> PinsParameter;
 
-	/** Relates an Image pin (key) to its Image Pin Mode (value). */
+	/** Relates an Image pin (key) to its Image Pin Mode (value). 
+		Represents the real mode of the pin. Required due to some node configurations can force the mode independently of what the user had previously selected */
 	UPROPERTY()
 	TMap<FGuid, EPinMode> PinsImagePinMode;
 	
@@ -270,7 +271,7 @@ private:
 
 	/** Returns the Image Pin Mode the pin should be at. It does not update its mode, to update it call UpdateImagePinMode. */
 	EPinMode GetImagePinMode(const UEdGraphPin& Pin) const;
-	
+
 	// Deprecated properties
 	/** Set all pins to Mutable mode. Even so, each pin can override its behaviour. */
 	UPROPERTY()
diff --git a/Engine/Plugins/Experimental/Mutable/Source/CustomizableObjectEditor/Classes/MuCOE/Nodes/CustomizableObjectNodeTable.h b/Engine/Plugins/Experimental/Mutable/Source/CustomizableObjectEditor/Classes/MuCOE/Nodes/CustomizableObjectNodeTable.h
--- a/Engine/Plugins/Experimental/Mutable/Source/CustomizableObjectEditor/Classes/MuCOE/Nodes/CustomizableObjectNodeTable.h
+++ b/Engine/Plugins/Experimental/Mutable/Source/CustomizableObjectEditor/Classes/MuCOE/Nodes/CustomizableObjectNodeTable.h
@@ -191,10 +191,7 @@ public:
 	/*** Allows to perform work when remapping the pin data. */
 	virtual void RemapPinsData(const TMap<UEdGraphPin*, UEdGraphPin*>& PinsToRemap) override;
 	
-	// Own interface
-	/** Given an Image pin, return the PinMode. */
-	bool ForceImageMutableMode(const UEdGraphPin* Pin, FGuid ParameterId) const;
-	
+	// Own interface	
 	// Returns the reference Texture parameter from a Material
 	UTexture2D* FindReferenceTextureParameter(const UEdGraphPin* Pin, FString ParameterImageName) const;
 
diff --git a/Engine/Plugins/Experimental/Mutable/Source/CustomizableObjectEditor/Private/MuCOE/GenerateMutableSource/GenerateMutableSourceColor.cpp b/Engine/Plugins/Experimental/Mutable/Source/CustomizableObjectEditor/Private/MuCOE/GenerateMutableSource/GenerateMutableSourceColor.cpp
--- a/Engine/Plugins/Experimental/Mutable/Source/CustomizableObjectEditor/Private/MuCOE/GenerateMutableSource/GenerateMutableSourceColor.cpp
+++ b/Engine/Plugins/Experimental/Mutable/Source/CustomizableObjectEditor/Private/MuCOE/GenerateMutableSource/GenerateMutableSourceColor.cpp
@@ -47,8 +47,6 @@ mu::NodeColourPtr GenerateMutableSourceColor(const UEdGraphPin* Pin, FMutableGra
 		return static_cast<mu::NodeColour*>(Generated->Node.get());
 	}
 
-	bool bDoNotAddToGeneratedCache = false;
-
 	mu::NodeColourPtr Result;
 
 	if (const UCustomizableObjectNodeColorConstant* TypedNodeColorConst = Cast<UCustomizableObjectNodeColorConstant>(Node))
@@ -279,12 +277,6 @@ mu::NodeColourPtr GenerateMutableSourceColor(const UEdGraphPin* Pin, FMutableGra
 
 		bool bSuccess = true;
 
-		if (Pin->PinType.PinCategory == Schema->PC_MaterialAsset)
-		{
-			// Material pins have to skip the cache of nodes or they will return always the same column node
-			bDoNotAddToGeneratedCache = true;
-		}
-
 		if (TypedNodeTable->Table)
 		{
 			FString ColumnName = Pin->PinFriendlyName.ToString();
@@ -308,18 +300,12 @@ mu::NodeColourPtr GenerateMutableSourceColor(const UEdGraphPin* Pin, FMutableGra
 				{
 					mu::NodeColourTablePtr ColorTableNode = new mu::NodeColourTable();
 
-					if (Pin->PinType.PinCategory == Schema->PC_MaterialAsset)
-					{
-						// Materials use the parameter id as column names
-						ColumnName = GenerationContext.CurrentMaterialTableParameterId;
-					}
-
 					// Generating a new Color column if not exists
 					if (Table->FindColumn(StringCast<ANSICHAR>(*ColumnName).Get()) == INDEX_NONE)
 					{
 						int32 Dummy = -1; // TODO MTBL-1512
 						bool Dummy2 = false;
-						bSuccess = GenerateTableColumn(TypedNodeTable, Pin, Table, ColumnName, Dummy, Dummy, GenerationContext.CurrentLOD, Dummy, Dummy2, GenerationContext);
+						bSuccess = GenerateTableColumn(TypedNodeTable, Pin, Table, ColumnName, Property, Dummy, Dummy, GenerationContext.CurrentLOD, Dummy, Dummy2, GenerationContext);
 
 						if (!bSuccess)
 						{
@@ -357,12 +343,9 @@ mu::NodeColourPtr GenerateMutableSourceColor(const UEdGraphPin* Pin, FMutableGra
 		GenerationContext.Compiler->CompilerLog(LOCTEXT("UnimplementedNode", "Node type not implemented yet."), Node);
 	}
 
-	if (!bDoNotAddToGeneratedCache)
-	{
-		FGeneratedData CacheData = FGeneratedData(Node, Result);
-		GenerationContext.Generated.Add(Key, CacheData);
-		GenerationContext.GeneratedNodes.Add(Node);
-	}
+	FGeneratedData CacheData = FGeneratedData(Node, Result);
+	GenerationContext.Generated.Add(Key, CacheData);
+	GenerationContext.GeneratedNodes.Add(Node);
 	
 	if (Result)
 	{
diff --git a/Engine/Plugins/Experimental/Mutable/Source/CustomizableObjectEditor/Private/MuCOE/GenerateMutableSource/GenerateMutableSourceFloat.cpp b/Engine/Plugins/Experimental/Mutable/Source/CustomizableObjectEditor/Private/MuCOE/GenerateMutableSource/GenerateMutableSourceFloat.cpp
--- a/Engine/Plugins/Experimental/Mutable/Source/CustomizableObjectEditor/Private/MuCOE/GenerateMutableSource/GenerateMutableSourceFloat.cpp
+++ b/Engine/Plugins/Experimental/Mutable/Source/CustomizableObjectEditor/Private/MuCOE/GenerateMutableSource/GenerateMutableSourceFloat.cpp
@@ -345,7 +345,7 @@ mu::NodeScalarPtr GenerateMutableSourceFloat(const UEdGraphPin* Pin, FMutableGra
 					{
 						int32 Dummy = -1; // TODO MTBL-1512
 						bool Dummy2 = false;
-						bSuccess = GenerateTableColumn(TypedNodeTable, Pin, Table, ColumnName, Dummy, Dummy, GenerationContext.CurrentLOD, Dummy, Dummy2, GenerationContext);
+						bSuccess = GenerateTableColumn(TypedNodeTable, Pin, Table, ColumnName, Property, Dummy, Dummy, GenerationContext.CurrentLOD, Dummy, Dummy2, GenerationContext);
 
 						if (!bSuccess)
 						{
diff --git a/Engine/Plugins/Experimental/Mutable/Source/CustomizableObjectEditor/Private/MuCOE/GenerateMutableSource/GenerateMutableSourceImage.cpp b/Engine/Plugins/Experimental/Mutable/Source/CustomizableObjectEditor/Private/MuCOE/GenerateMutableSource/GenerateMutableSourceImage.cpp
--- a/Engine/Plugins/Experimental/Mutable/Source/CustomizableObjectEditor/Private/MuCOE/GenerateMutableSource/GenerateMutableSourceImage.cpp
+++ b/Engine/Plugins/Experimental/Mutable/Source/CustomizableObjectEditor/Private/MuCOE/GenerateMutableSource/GenerateMutableSourceImage.cpp
@@ -988,8 +988,8 @@ mu::NodeImagePtr GenerateMutableSourceImage(const UEdGraphPin* Pin, FMutableGrap
 
 						if (Pin->PinType.PinCategory == Schema->PC_MaterialAsset)
 						{
-							// Material parameters use the parameter id as column names
-							ColumnName = GenerationContext.CurrentMaterialTableParameterId;
+							// Material parameters use the Data Table Column Name + Parameter id as mutable column Name to aboid duplicated names (i.e. two MI columns with the same parents but different values).
+							ColumnName = Property->GetDisplayNameText().ToString() + GenerationContext.CurrentMaterialTableParameterId;
 						}
 
 						// Generating a new Texture column if not exists
@@ -997,7 +997,7 @@ mu::NodeImagePtr GenerateMutableSourceImage(const UEdGraphPin* Pin, FMutableGrap
 						{
 							int32 Dummy = -1; // TODO MTBL-1512
 							bool Dummy2 = false;
-							bSuccess = GenerateTableColumn(TypedNodeTable, Pin, Table, ColumnName, Dummy, Dummy, GenerationContext.CurrentLOD, Dummy, Dummy2, GenerationContext);
+							bSuccess = GenerateTableColumn(TypedNodeTable, Pin, Table, ColumnName, Property, Dummy, Dummy, GenerationContext.CurrentLOD, Dummy, Dummy2, GenerationContext);
 
 							if (!bSuccess)
 							{
diff --git a/Engine/Plugins/Experimental/Mutable/Source/CustomizableObjectEditor/Private/MuCOE/GenerateMutableSource/GenerateMutableSourceMesh.cpp b/Engine/Plugins/Experimental/Mutable/Source/CustomizableObjectEditor/Private/MuCOE/GenerateMutableSource/GenerateMutableSourceMesh.cpp
--- a/Engine/Plugins/Experimental/Mutable/Source/CustomizableObjectEditor/Private/MuCOE/GenerateMutableSource/GenerateMutableSourceMesh.cpp
+++ b/Engine/Plugins/Experimental/Mutable/Source/CustomizableObjectEditor/Private/MuCOE/GenerateMutableSource/GenerateMutableSourceMesh.cpp
@@ -3577,7 +3577,7 @@ mu::NodeMeshPtr GenerateMutableSourceMesh(const UEdGraphPin* Pin,
 					// Generating a new Mesh column if not exists
 					if (Table->FindColumn(StringCast<ANSICHAR>(*MutableColumnName).Get()) == INDEX_NONE)
 					{
-						bSuccess = GenerateTableColumn(TypedNodeTable, Pin, Table, DataTableColumnName, LODIndexConnected, SectionIndexConnected, LODIndex, SectionIndex, bOnlyConnectedLOD, GenerationContext);
+						bSuccess = GenerateTableColumn(TypedNodeTable, Pin, Table, DataTableColumnName, Property, LODIndexConnected, SectionIndexConnected, LODIndex, SectionIndex, bOnlyConnectedLOD, GenerationContext);
 
 						if (!bSuccess)
 						{
diff --git a/Engine/Plugins/Experimental/Mutable/Source/CustomizableObjectEditor/Private/MuCOE/GenerateMutableSource/GenerateMutableSourceSurface.cpp b/Engine/Plugins/Experimental/Mutable/Source/CustomizableObjectEditor/Private/MuCOE/GenerateMutableSource/GenerateMutableSourceSurface.cpp
--- a/Engine/Plugins/Experimental/Mutable/Source/CustomizableObjectEditor/Private/MuCOE/GenerateMutableSource/GenerateMutableSourceSurface.cpp
+++ b/Engine/Plugins/Experimental/Mutable/Source/CustomizableObjectEditor/Private/MuCOE/GenerateMutableSource/GenerateMutableSourceSurface.cpp
@@ -41,6 +41,7 @@
 #include "MuT/NodeMeshFragment.h"
 #include "MuT/NodePatchImage.h"
 #include "MuT/NodePatchMesh.h"
+#include "MuT/NodeScalarConstant.h"
 #include "MuT/NodeSurfaceEdit.h"
 #include "MuT/NodeSurfaceVariation.h"
 
@@ -333,7 +334,6 @@ mu::NodeSurfacePtr GenerateMutableSourceSurface(const UEdGraphPin * Pin, FMutabl
 			const int32 lastMaterialAmount = GenerationContext.ReferencedMaterials.Num();
 			ReferencedMaterialsIndex = GenerationContext.ReferencedMaterials.AddUnique(TypedNodeMat->Material);
 			// Used ReferencedMaterialsIndex instead of TypedNodeMat->Material->GetName() to prevent material name collisions
-			SurfNode->SetCustomID(ReferencedMaterialsIndex);
 
 			// Set shared surface Id to reuse materials and textures between LODs if automatic LODs from mesh is being used
 			UpdateSharedSurfaceId(GenerationContext, TypedNodeMat, SurfNode);
@@ -365,6 +365,8 @@ mu::NodeSurfacePtr GenerateMutableSourceSurface(const UEdGraphPin * Pin, FMutabl
 									SlotNameFound = true;
 								}
 							}
+
+							//TODO(max): Add support for table nodes. Convert this to a function as it will be also used in the GenerateMutableSourceTable file
 						}
 					}
 				}
@@ -379,7 +381,6 @@ mu::NodeSurfacePtr GenerateMutableSourceSurface(const UEdGraphPin * Pin, FMutabl
 
 		mu::NodeMeshPtr MeshNode;
 
-
 		TypedNodeMat = TypedNodeCopyMaterial; // NodeCopyMaterial. Start reading mesh pin. Set TypedNodeMat pointer to NodeCopyMaterial
 
 		if (const UEdGraphPin* ConnectedPin = FollowInputPin(*TypedNodeMat->GetMeshPin()))
@@ -414,69 +415,32 @@ mu::NodeSurfacePtr GenerateMutableSourceSurface(const UEdGraphPin * Pin, FMutabl
 		TMap<FString, float> TextureNameToProjectionResFactor;
 		FString AlternateResStateName;
 
-		bool bMaterialPinConnected = FollowInputPin(*TypedNodeMat->GetMaterialAssetPin()) != nullptr;
-		TArray<FGuid> NodeTableParametersGenerated;
+		bool bTableMaterialPinLinked = TypedNodeMat->GetMaterialAssetPin() && FollowInputPin(*TypedNodeMat->GetMaterialAssetPin()) != nullptr;
+		FString TableColumnName;
 
 		// Checking if we should not use the material of the table node even if it is linked to the material node
 		if (const UEdGraphPin* ConnectedPin = FollowInputPin(*TypedNodeMat->GetMaterialAssetPin()))
 		{
 			if (const UCustomizableObjectNodeTable* TypedNodeTable = Cast< UCustomizableObjectNodeTable >(ConnectedPin->GetOwningNode()))
 			{
-				FString ColumnName = ConnectedPin->PinFriendlyName.ToString();
+				TableColumnName = ConnectedPin->PinFriendlyName.ToString();
 
-				if (!TypedNodeTable->GetColumnDefaultAssetByType<UMaterialInstance>(ConnectedPin))
+				if (UMaterialInstance * TableMaterial = TypedNodeTable->GetColumnDefaultAssetByType<UMaterialInstance>(ConnectedPin))
 				{
-					FString Msg = FString::Printf(TEXT("Couldn't find a default value in the data table's struct for the column [%s]. The default value is null or not a Material Instance."), *ColumnName);
-					GenerationContext.Compiler->CompilerLog(FText::FromString(Msg), Node);
+					// Checking if the reference material of the Table Node has the same parent as the material of the Material Node 
+					if (!TypedNodeMat->Material || TableMaterial->GetMaterial() != TypedNodeMat->Material->GetMaterial())
+					{
+						bTableMaterialPinLinked = false;
 
-					bMaterialPinConnected = false;
+						GenerationContext.Compiler->CompilerLog(LOCTEXT("DifferentParentMaterial","The Deafult Material Instance of the Data Table must have the same Parent Material."), TypedNodeMat);
+					}
 				}
 				else
 				{
-					// Generating a new data table if not exists
-					mu::TablePtr Table = GenerateMutableSourceTable(TypedNodeTable->Table->GetName(), ConnectedPin, GenerationContext);
-
-					// Generating a column for each modified parameter(texture, color & float) if not exists
-					if (Table && Table->FindColumn(StringCast<ANSICHAR>(*ColumnName).Get()) == INDEX_NONE)
-					{
-						int32 Dummy = -1; // TODO MTBL-1512
-						bool Dummy2 = false;
-						GenerateTableColumn(TypedNodeTable, ConnectedPin, Table, ColumnName, Dummy, Dummy, GenerationContext.CurrentLOD, Dummy, Dummy2, GenerationContext);
-					}
-
-					// Checking if this material has some parameters modified by the table node linked to it
-					if (GenerationContext.GeneratedParametersInTables.Contains(TypedNodeTable))
-					{
-						NodeTableParametersGenerated = GenerationContext.GeneratedParametersInTables[TypedNodeTable];
-
-						UMaterialInstance* TableMaterial = TypedNodeTable->GetColumnDefaultAssetByType<UMaterialInstance>(ConnectedPin);
-
-						if (TableMaterial && TypedNodeMat->Material)
-						{
-							// Checking if the reference material of the Table Node has the same parent as the material of the Material Node 
-							if (TableMaterial->GetMaterial() != TypedNodeMat->Material->GetMaterial())
-							{
-								bMaterialPinConnected = false;
-
-								FString msg = FString::Printf(TEXT("Material from NodeTable is an instance from a different Parent"));
-								GenerationContext.Compiler->CompilerLog(FText::FromString(msg), TypedNodeMat);
-							}
-						}
-						else
-						{
-							bMaterialPinConnected = false;
+					FText Msg = FText::Format(LOCTEXT("DefaultValueNotFound", "Couldn't find a default value in the data table's struct for the column {0}. The default value is null or not a Material Instance."), FText::FromString(TableColumnName));
+					GenerationContext.Compiler->CompilerLog(Msg, Node);
 
-							if (!TypedNodeMat->Material)
-							{
-								FString msg = FString::Printf(TEXT("Need to select a Material to work with Node Table Materials."));
-								GenerationContext.Compiler->CompilerLog(FText::FromString(msg), TypedNodeMat);
-							}
-						}
-					}
-					else
-					{
-						bMaterialPinConnected = false;
-					}
+					bTableMaterialPinLinked = false;
 				}
 			}
 		}
@@ -487,14 +451,14 @@ mu::NodeSurfacePtr GenerateMutableSourceSurface(const UEdGraphPin * Pin, FMutabl
 		for (int32 ImageIndex = 0; ImageIndex < NumImages; ++ImageIndex)
 		{
 			const UEdGraphPin* ImagePin = TypedNodeMat->GetParameterPin(EMaterialParameterType::Texture, ImageIndex);
-			const bool bIsImagePinLinked = ImagePin && !ImagePin->LinkedTo.IsEmpty();
+
+			const bool bIsImagePinLinked = ImagePin && FollowInputPin(*ImagePin);
 
 			if (bIsImagePinLinked && !TypedNodeMaterial->IsImageMutableMode(ImageIndex))
 			{
-				// This is a connected pass-through texture that simply has to be passed to the core
-
 				if (const UEdGraphPin* ConnectedPin = FollowInputPin(*ImagePin))
 				{
+					// This is a connected pass-through texture that simply has to be passed to the core
 					mu::Ptr<mu::NodeImage> PassThroughImagePtr = GenerateMutableSourceImage(ConnectedPin, GenerationContext, 0);
 					SurfNode->SetImage(ImageIndex, PassThroughImagePtr);
 
@@ -502,12 +466,14 @@ mu::NodeSurfacePtr GenerateMutableSourceSurface(const UEdGraphPin * Pin, FMutabl
 					FString SurfNodeImageName = FString::Printf(TEXT("%d"), GenerationContext.ImageProperties.Num());
 					SurfNode->SetImageName(ImageIndex, StringCast<ANSICHAR>(*SurfNodeImageName).Get());
 					int32 UVLayout = TypedNodeMat->GetImageUVLayout(ImageIndex);
+
 					if (UVLayout >= 0)
 					{
 						FString msg = FString::Printf(TEXT("Passthrough texture [%s] will ignore layout despite set to use layout [%d]"), *ImageName, UVLayout);
 						GenerationContext.Compiler->CompilerLog(FText::FromString(msg), Node, EMessageSeverity::Warning);
 
 					}
+
 					SurfNode->SetImageLayoutIndex(ImageIndex, -1);
 					SurfNode->SetImageAdditionalNames(ImageIndex, StringCast<ANSICHAR>(*TypedNodeMat->Material->GetName()).Get(), StringCast<ANSICHAR>(*ImageName).Get());
 
@@ -516,7 +482,6 @@ mu::NodeSurfacePtr GenerateMutableSourceSurface(const UEdGraphPin * Pin, FMutabl
 					Props.TextureParameterName = ImageName;
 					GenerationContext.ImageProperties.Add(Props);
 					SurfaceData.ImageProperties = Props;
-
 				}
 			}
 			else
@@ -539,7 +504,7 @@ mu::NodeSurfacePtr GenerateMutableSourceSurface(const UEdGraphPin * Pin, FMutabl
 						GroupProjectionReferenceTexture, TextureNameToProjectionResFactor, AlternateResStateName,
 						nullptr);
 
-					if (GroupProjectionImg.get() || TypedNodeMaterial->IsImageMutableMode(ImageIndex) || (bMaterialPinConnected && NodeTableParametersGenerated.Contains(ImageId)))
+					if (GroupProjectionImg.get() || TypedNodeMaterial->IsImageMutableMode(ImageIndex))
 					{
 						// Get the reference texture
 						UTexture2D* ReferenceTexture = nullptr;
@@ -564,7 +529,7 @@ mu::NodeSurfacePtr GenerateMutableSourceSurface(const UEdGraphPin * Pin, FMutabl
 										ReferenceTexture = FindReferenceImage(ConnectedPin, GenerationContext);
 									}
 								}
-								if (!ReferenceTexture && bMaterialPinConnected && NodeTableParametersGenerated.Contains(ImageId))
+								if (!ReferenceTexture && bTableMaterialPinLinked)
 								{
 									if (const UEdGraphPin* ConnectedPin = FollowInputPin(*TypedNodeMat->GetMaterialAssetPin()))
 									{
@@ -634,7 +599,7 @@ mu::NodeSurfacePtr GenerateMutableSourceSurface(const UEdGraphPin * Pin, FMutabl
 									}
 								}
 
-								if (bMaterialPinConnected && NodeTableParametersGenerated.Contains(ImageId))
+								if (bTableMaterialPinLinked)
 								{
 									if (const UEdGraphPin* ConnectedPin = FollowInputPin(*TypedNodeMat->GetMaterialAssetPin()))
 									{
@@ -874,23 +839,9 @@ mu::NodeSurfacePtr GenerateMutableSourceSurface(const UEdGraphPin * Pin, FMutabl
 			FString VectorName = TypedNodeMat->GetParameterName(EMaterialParameterType::Vector, VectorIndex).ToString();
 			FGuid VectorId = TypedNodeMat->GetParameterId(EMaterialParameterType::Vector, VectorIndex);
 
-			if (bVectorPinConnected || (bMaterialPinConnected && NodeTableParametersGenerated.Contains(VectorId)))
-			{
-				const UEdGraphPin* SourcePin = nullptr;
-
-				if (bVectorPinConnected)
-				{
-					SourcePin = VectorPin;
-				}
-				else if (bMaterialPinConnected)
-				{
-					GenerationContext.CurrentMaterialTableParameter = VectorName;
-					GenerationContext.CurrentMaterialTableParameterId = VectorId.ToString();
-
-					SourcePin = TypedNodeMat->GetMaterialAssetPin();
-				}
-
-				if (const UEdGraphPin* ConnectedPin = FollowInputPin(*SourcePin))
+			if (bVectorPinConnected)
+			{				
+				if (const UEdGraphPin* ConnectedPin = FollowInputPin(*VectorPin))
 				{
 					mu::NodeColourPtr ColorNode = GenerateMutableSourceColor(ConnectedPin, GenerationContext);
 
@@ -916,23 +867,9 @@ mu::NodeSurfacePtr GenerateMutableSourceSurface(const UEdGraphPin * Pin, FMutabl
 			FString ScalarName = TypedNodeMat->GetParameterName(EMaterialParameterType::Scalar, ScalarIndex).ToString();
 			FGuid ScalarId = TypedNodeMat->GetParameterId(EMaterialParameterType::Scalar, ScalarIndex);
 
-			if (bScalarPinConnected || (bMaterialPinConnected && NodeTableParametersGenerated.Contains(ScalarId)))
+			if (bScalarPinConnected)
 			{
-				const UEdGraphPin* SourcePin = nullptr;
-
-				if (bScalarPinConnected)
-				{
-					SourcePin = ScalarPin;
-				}
-				else if (bMaterialPinConnected)
-				{
-					GenerationContext.CurrentMaterialTableParameter = ScalarName;
-					GenerationContext.CurrentMaterialTableParameterId = ScalarId.ToString();
-
-					SourcePin = TypedNodeMat->GetMaterialAssetPin();
-				}
-
-				if (const UEdGraphPin* ConnectedPin = FollowInputPin(*SourcePin))
+				if (const UEdGraphPin* ConnectedPin = FollowInputPin(*ScalarPin))
 				{
 					mu::NodeScalarPtr ScalarNode = GenerateMutableSourceFloat(ConnectedPin, GenerationContext);
 
@@ -948,6 +885,38 @@ mu::NodeSurfacePtr GenerateMutableSourceSurface(const UEdGraphPin * Pin, FMutabl
 			}
 		}
 
+		// New method to pass the surface id as a scalar parameter
+		{
+			int32 MaterialIndex = NumScalar;
+			SurfNode->SetScalarCount(NumScalar + 1);
+
+			const UEdGraphPin* MaterialPin = TypedNodeMat->GetMaterialAssetPin();
+
+			//Encoding name for material material id parameter
+			FString MaterialName = "__MutableMaterialId";
+
+			if (bTableMaterialPinLinked)
+			{
+				if (const UEdGraphPin* ConnectedPin = FollowInputPin(*MaterialPin))
+				{
+					GenerationContext.CurrentMaterialTableParameterId = MaterialName;
+					mu::NodeScalarPtr ScalarNode = GenerateMutableSourceFloat(ConnectedPin, GenerationContext);
+
+					SurfNode->SetScalar(MaterialIndex, ScalarNode);
+					SurfNode->SetScalarName(MaterialIndex, StringCast<ANSICHAR>(*MaterialName).Get());
+				}
+			}
+			else
+			{
+				mu::NodeScalarConstantPtr ScalarNode = new mu::NodeScalarConstant();
+				ScalarNode->SetValue(ReferencedMaterialsIndex);
+
+				SurfNode->SetScalar(MaterialIndex, ScalarNode);
+				SurfNode->SetScalarName(MaterialIndex, StringCast<ANSICHAR>(*MaterialName).Get());
+			}
+		}
+		
+
 		for (const FString& Tag : TypedNodeMat->Tags)
 		{
 			SurfNode->AddTag(StringCast<ANSICHAR>(*Tag).Get());
@@ -1630,5 +1599,6 @@ mu::NodeSurfacePtr GenerateMutableSourceSurface(const UEdGraphPin * Pin, FMutabl
 	return Result;
 }
 
+
 #undef LOCTEXT_NAMESPACE
 
diff --git a/Engine/Plugins/Experimental/Mutable/Source/CustomizableObjectEditor/Private/MuCOE/GenerateMutableSource/GenerateMutableSourceTable.cpp b/Engine/Plugins/Experimental/Mutable/Source/CustomizableObjectEditor/Private/MuCOE/GenerateMutableSource/GenerateMutableSourceTable.cpp
--- a/Engine/Plugins/Experimental/Mutable/Source/CustomizableObjectEditor/Private/MuCOE/GenerateMutableSource/GenerateMutableSourceTable.cpp
+++ b/Engine/Plugins/Experimental/Mutable/Source/CustomizableObjectEditor/Private/MuCOE/GenerateMutableSource/GenerateMutableSourceTable.cpp
@@ -21,13 +21,13 @@
 #define LOCTEXT_NAMESPACE "CustomizableObjectEditor"
 
 
-bool FillTableColumn(const UCustomizableObjectNodeTable* TableNode,	mu::TablePtr MutableTable,	const FString& ColumnName,	const FString& RowName,	const int32 RowIdx,	uint8* CellData, FProperty* Property,
+bool FillTableColumn(const UCustomizableObjectNodeTable* TableNode,	mu::TablePtr MutableTable,	const FString& ColumnName,	const FString& RowName,	const int32 RowIdx,	uint8* CellData, const FProperty* ColumnProperty,
 	const int LODIndexConnected, const int32 SectionIndexConnected, int32 LODIndex, int32 SectionIndex, const bool bOnlyConnectedLOD, FMutableGraphGenerationContext& GenerationContext)
 {
 	int32 CurrentColumn;
 
 	// Getting property type
-	if (const FSoftObjectProperty* SoftObjectProperty = CastField<FSoftObjectProperty>(Property))
+	if (const FSoftObjectProperty* SoftObjectProperty = CastField<FSoftObjectProperty>(ColumnProperty))
 	{
 		UObject* Object = SoftObjectProperty->GetPropertyValue(CellData).LoadSynchronous();
 
@@ -37,7 +37,20 @@ bool FillTableColumn(const UCustomizableObjectNodeTable* TableNode,	mu::TablePtr
 
 			if(!SkeletalMesh)
 			{
-				return false;
+				// Generating an Empty cell
+				FString MutableColumnName = TableNode->GenerateSkeletalMeshMutableColumName(ColumnName, LODIndex, SectionIndex);
+
+				CurrentColumn = MutableTable.get()->FindColumn(StringCast<ANSICHAR>(*MutableColumnName).Get());
+
+				if (CurrentColumn == -1)
+				{
+					CurrentColumn = MutableTable->AddColumn(StringCast<ANSICHAR>(*MutableColumnName).Get(), mu::TABLE_COLUMN_TYPE::TCT_MESH);
+				}
+
+				mu::MeshPtr EmptySkeletalMesh = nullptr;
+				MutableTable->SetCell(CurrentColumn, RowIdx, EmptySkeletalMesh.get());
+
+				return true;
 			}
 
 			// Getting Animation Blueprint and Animation Slot
@@ -391,256 +404,123 @@ bool FillTableColumn(const UCustomizableObjectNodeTable* TableNode,	mu::TablePtr
 
 		else if (SoftObjectProperty->PropertyClass->IsChildOf(UMaterialInstance::StaticClass()))
 		{
-			UMaterialInstance* Material = Cast<UMaterialInstance>(Object);
-
-			if (!Material)
-			{
-				Material = TableNode->GetColumnDefaultAssetByType<UMaterialInstance>(ColumnName);
+			// Getting the name of material column of the data table
+			FString MaterialColumnName = ColumnProperty->GetDisplayNameText().ToString();
 
-				FString msg = FString::Printf(TEXT("Material Instance from column [%s] row [%s] is null. The default Material Instance will be used instead."), *ColumnName, *RowName);
-				GenerationContext.Compiler->CompilerLog(FText::FromString(msg), TableNode);
-			}
-
-			//Adding an empty column for searching purposes
-			if (MutableTable.get()->FindColumn(StringCast<ANSICHAR>(*ColumnName).Get()) == -1)
-			{
-				CurrentColumn = MutableTable->AddColumn(StringCast<ANSICHAR>(*ColumnName).Get(), mu::TABLE_COLUMN_TYPE::TCT_NONE);
-			}
-
-			UMaterialInstance* ReferenceMaterial = TableNode->GetColumnDefaultAssetByType<UMaterialInstance>(ColumnName);
-			
-			if (!GenerationContext.GeneratedParametersInTables.Contains(TableNode))
-			{
-				GenerationContext.GeneratedParametersInTables.Add(TableNode);
-			}
-
-			TArray<FGuid>& Parameters = GenerationContext.GeneratedParametersInTables[TableNode];
+			UMaterialInstance* Material = Cast<UMaterialInstance>(Object);
+			UMaterialInstance* ReferenceMaterial = TableNode->GetColumnDefaultAssetByType<UMaterialInstance>(MaterialColumnName);
 
 			if (!ReferenceMaterial)
 			{
-				FString msg = FString::Printf(TEXT("Reference Material not found for column [%s]."), *ColumnName);
+				FString msg = FString::Printf(TEXT("Reference Material not found for column [%s]."), *MaterialColumnName);
 				GenerationContext.Compiler->CompilerLog(FText::FromString(msg), TableNode);
 
 				return false;
 			}
 
-			if (ReferenceMaterial->GetMaterial() == Material->GetMaterial())
+			if (!Material || ReferenceMaterial->GetMaterial() != Material->GetMaterial())
 			{
-				// Getting parameter Guids
-				TArray<FMaterialParameterInfo > ParameterInfos;
-				TArray<FGuid> ParameterGuids;
+				FText Warning;
 
-				ReferenceMaterial->GetAllParameterInfoOfType(EMaterialParameterType::Texture, ParameterInfos, ParameterGuids);
-
-				// Number of modified texture parameters in the Material Instance
-				int32 ModifiedTextureParameters = 0;
-				
-				for (FTextureParameterValue ReferenceTexture : ReferenceMaterial->TextureParameterValues)
+				if (!Material)
 				{
-					if (UTexture2D* Texture = Cast<UTexture2D>(ReferenceTexture.ParameterValue))
-					{
-						int32 ColumnIndex;
-
-						FString TextureParameterName = ReferenceTexture.ParameterInfo.Name.ToString();
-						FString ParameterGuid;
-						
-						for (int32 i = 0; i < ParameterInfos.Num(); ++i)
-						{
-							if (ParameterInfos[i].Name == ReferenceTexture.ParameterInfo.Name)
-							{
-								ParameterGuid = ParameterGuids[i].ToString();
-								Parameters.Add(ParameterGuids[i]);
-
-								break;
-							}
-						}
-
-						// Getting column index from parameter name
-						ColumnIndex = MutableTable->FindColumn(StringCast<ANSICHAR>(*ParameterGuid).Get());
-
-						if (ColumnIndex == INDEX_NONE)
-						{
-							// If there is no column with the parameters name, we generate a new one
-							ColumnIndex = MutableTable->AddColumn(StringCast<ANSICHAR>(*ParameterGuid).Get(), mu::TABLE_COLUMN_TYPE::TCT_IMAGE);
-						}
-
-						// Getting the parameter value from the instance if it has been modified
-						for (FTextureParameterValue InstanceTexture : Material->TextureParameterValues)
-						{
-							if (InstanceTexture.ParameterInfo.Name == ReferenceTexture.ParameterInfo.Name)
-							{
-								if (UTexture2D* ParameterTexture = Cast<UTexture2D>(InstanceTexture.ParameterValue))
-								{
-									Texture = ParameterTexture;
-								}
-								else
-								{
-									FString ParamName = ReferenceTexture.ParameterInfo.Name.ToString();
-									FString Message = Cast<UObject>(InstanceTexture.ParameterValue) ? "not a Texture2D" : "null";
+					Warning = FText::Format(LOCTEXT("NullMaterialInstance", "Material Instance from column [{0}] row [{1}] is null. The default Material Instance will be used instead."),
+						FText::FromString(MaterialColumnName), FText::FromString(RowName));
+				}
+				else
+				{
+					Warning = FText::Format(LOCTEXT("MatInstanceFromDifferentParent","Material Instance from column [{0}] row [{1}] has a different Material Parent than the Default Material Instance. The Default Material Instance will be used instead."),
+						FText::FromString(MaterialColumnName), FText::FromString(RowName));
+				}
 
-									FString msg = FString::Printf(TEXT("Parameter [%s] from material instance of column [%s] row [%s] is %s. The parameter texture of the default material will be used instead."), *ParamName,  *ColumnName, *RowName, *Message);
-									GenerationContext.Compiler->CompilerLog(FText::FromString(msg), TableNode);
-								}
+				Material = ReferenceMaterial;
 
-								ModifiedTextureParameters++;
-								
-								break;
-							}
-						}
-
-						check(Texture);
-						GenerationContext.ArrayTextureUnrealToMutableTask.Add(FTextureUnrealToMutableTask(MutableTable, Texture, TableNode, ColumnIndex, RowIdx));
-					}
-					else
-					{
-						FString ParamName = ReferenceTexture.ParameterInfo.Name.ToString();
-						FString Message = Cast<UObject>(ReferenceTexture.ParameterValue) ? "not a Texture2D" : "null";
+				GenerationContext.Compiler->CompilerLog(Warning, TableNode);
+			}
 
-						FString msg = FString::Printf(TEXT("Parameter [%s] from Default Material Instance of column [%s] is %s. This parameter will be ignored."), *ParamName, *ColumnName, *Message);
-						GenerationContext.Compiler->CompilerLog(FText::FromString(msg), TableNode);
-					}
-				}
+			FString EncodedSwitchParameterName = "__MutableMaterialId";
+			if (ColumnName.Contains(EncodedSwitchParameterName))
+			{
+				CurrentColumn = MutableTable.get()->FindColumn(StringCast<ANSICHAR>(*ColumnName).Get());
 
-				// Checking if all modifiable textures of the material instance have been modified
-				if (Material->TextureParameterValues.Num() > ModifiedTextureParameters)
+				if (CurrentColumn == -1)
 				{
-					int32 ParametersDiff = Material->TextureParameterValues.Num() - ModifiedTextureParameters;
-					FString msg =
-						FString::Printf(TEXT("Material Instance [%s] from column [%s] row [%s] has %d modifiable Textures that will not be modified, they are non-modifiable parameters in the Default Material Instance"),
-						*Material->GetName(), *ColumnName, *RowName, ParametersDiff);
-					GenerationContext.Compiler->CompilerLog(FText::FromString(msg), TableNode);
+					CurrentColumn = MutableTable->AddColumn(StringCast<ANSICHAR>(*ColumnName).Get(), mu::TABLE_COLUMN_TYPE::TCT_SCALAR);
 				}
 
-				ParameterInfos.Empty();
-				ParameterGuids.Empty();
+				const int32 lastMaterialAmount = GenerationContext.ReferencedMaterials.Num();
+				int32 ReferenceMaterialId = GenerationContext.ReferencedMaterials.AddUnique(Material);
 
-				ReferenceMaterial->GetAllParameterInfoOfType(EMaterialParameterType::Vector, ParameterInfos, ParameterGuids);
+				// Take slot name from skeletal mesh if one can be found, else leave empty.
+				// Keep Referenced Materials and Materail Slot Names synchronized even if no material name can be found.
+				const bool IsNewSlotName = GenerationContext.ReferencedMaterialSlotNames.Num() == ReferenceMaterialId;
 
-				// Number of modified Vector parameters in the Material Instance
-				int32 ModifiedVectorParameters = 0;
-
-				for (FVectorParameterValue ReferenceVector : ReferenceMaterial->VectorParameterValues)
+				if (IsNewSlotName)
 				{
-					// Getting the parameter value from the default material
-					FLinearColor VectorValue = ReferenceVector.ParameterValue;
+					GenerationContext.ReferencedMaterialSlotNames.Add(FName(NAME_None));
+				}
 
-					int32 ColumnIndex;
-					FString ParameterGuid;
+				MutableTable->SetCell(CurrentColumn, RowIdx, (float)ReferenceMaterialId);
 
-					for (int32 i = 0; i < ParameterInfos.Num(); ++i)
-					{
-						if (ParameterInfos[i].Name == ReferenceVector.ParameterInfo.Name)
-						{
-							ParameterGuid = ParameterGuids[i].ToString();
-							Parameters.Add(ParameterGuids[i]);
+				return true;
+			}
 
-							break;
-						}
-					}
-					
-					// Getting column index from parameter name
-					ColumnIndex = MutableTable->FindColumn(StringCast<ANSICHAR>(*ParameterGuid).Get());
+			int32 ColumnIndex;
 
-					if (ColumnIndex == INDEX_NONE)
-					{
-						// If there is no column with the parameters name, we generate a new one
-						ColumnIndex = MutableTable->AddColumn(StringCast<ANSICHAR>(*ParameterGuid).Get(), mu::TABLE_COLUMN_TYPE::TCT_COLOUR);
-					}
+			// Getting parameter value
+			TArray<FMaterialParameterInfo> ParameterInfos;
+			TArray<FGuid> ParameterGuids;
 
-					// Getting the parameter value from the instance if it has been modified
-					for (FVectorParameterValue InstanceVector : Material->VectorParameterValues)
-					{
-						if (InstanceVector.ParameterInfo.Name == ReferenceVector.ParameterInfo.Name)
-						{
-							VectorValue = InstanceVector.ParameterValue;
-							ModifiedVectorParameters++;
-							
-							break;
-						}
-					}
+			Material->GetMaterial()->GetAllParameterInfoOfType(EMaterialParameterType::Texture, ParameterInfos, ParameterGuids);
+			
+			FGuid ParameterId(GenerationContext.CurrentMaterialTableParameterId);
+			int32 ParameterIndex = ParameterGuids.Find(ParameterId);
 
-					// Setting cell value
-					MutableTable->SetCell(ColumnIndex, RowIdx, VectorValue.R, VectorValue.G, VectorValue.B, VectorValue.A);
-				}
+			if (ParameterIndex != INDEX_NONE && ParameterInfos[ParameterIndex].Name == GenerationContext.CurrentMaterialTableParameter)
+			{
+				// Getting column index from parameter name
+				ColumnIndex = MutableTable->FindColumn(StringCast<ANSICHAR>(*ColumnName).Get());
 
-				// Checking if all modifiable vectors of the material instance have been modified
-				if (Material->VectorParameterValues.Num() > ModifiedVectorParameters)
+				if (ColumnIndex == INDEX_NONE)
 				{
-					int32 ParametersDiff = Material->VectorParameterValues.Num() - ModifiedVectorParameters;
+					// If there is no column with the parameters name, we generate a new one
+					ColumnIndex = MutableTable->AddColumn(StringCast<ANSICHAR>(*ColumnName).Get(), mu::TABLE_COLUMN_TYPE::TCT_IMAGE);
+				}
 
-					FString msg =
-						FString::Printf(TEXT("Material Instance [%s] from column [%s] row [%s] has %d modifiable Vectors that will not be modified, they are non-modifiable parameters in the Default Material Instance"),
-							*Material->GetName(), *ColumnName, *RowName, ParametersDiff);
+				UTexture* ParentTextureValue = nullptr;
+				Material->GetMaterial()->GetTextureParameterValue(ParameterInfos[ParameterIndex], ParentTextureValue);
+				
+				UTexture2D* ParentParameterTexture = Cast<UTexture2D>(ParentTextureValue);
+				if (!ParentParameterTexture)
+				{
+					FString ParamName = ParameterInfos[ParameterIndex].Name.ToString();
+					FString Message = Cast<UObject>(ParentParameterTexture) ? "not a Texture2D" : "null";
+					
+					FString msg = FString::Printf(TEXT("Parameter [%s] from Default Material Instance of column [%s] is %s. This parameter will be ignored."), *ParamName, *MaterialColumnName, *Message);
 					GenerationContext.Compiler->CompilerLog(FText::FromString(msg), TableNode);
+					 
+					 return false;
 				}
 
-				// Number of modified Float parameters in the Material Instance
-				int32 ModifiedFloatParameters = 0;
-
-				ParameterInfos.Empty();
-				ParameterGuids.Empty();
+				UTexture* TextureValue = nullptr;
+				Material->GetTextureParameterValue(ParameterInfos[ParameterIndex], TextureValue);
 
-				ReferenceMaterial->GetAllParameterInfoOfType(EMaterialParameterType::Scalar, ParameterInfos, ParameterGuids);
+				UTexture2D* ParameterTexture = Cast<UTexture2D>(TextureValue);
 
-				for (FScalarParameterValue ReferenceScalar : ReferenceMaterial->ScalarParameterValues)
+				if (!ParameterTexture)
 				{
-					// Getting the parameter value from the parent material
-					float ScalarValue = ReferenceScalar.ParameterValue;
+					ParameterTexture = ParentParameterTexture;
 
-					int32 ColumnIndex;
-					FString ParameterGuid;
+					FString ParamName = GenerationContext.CurrentMaterialTableParameter;
+					FString Message = Cast<UObject>(TextureValue) ? "not a Texture2D" : "null";
 
-					for (int32 i = 0; i < ParameterInfos.Num(); ++i)
-					{
-						if (ParameterInfos[i].Name == ReferenceScalar.ParameterInfo.Name)
-						{
-							ParameterGuid = ParameterGuids[i].ToString();
-							Parameters.Add(ParameterGuids[i]);
-							
-							break;
-						}
-					}
-
-					// Getting column index from parameter name
-					ColumnIndex = MutableTable->FindColumn(StringCast<ANSICHAR>(*ParameterGuid).Get());
-
-					if (ColumnIndex == INDEX_NONE)
-					{
-						// If there is no column with the parameters name, we generate a new one
-						ColumnIndex = MutableTable->AddColumn(StringCast<ANSICHAR>(*ParameterGuid).Get(), mu::TABLE_COLUMN_TYPE::TCT_SCALAR);
-					}
-
-					// Getting the parameter value from the instance if it has been modified
-					for (FScalarParameterValue InstanceScalar : Material->ScalarParameterValues)
-					{
-						if (InstanceScalar.ParameterInfo.Name == ReferenceScalar.ParameterInfo.Name)
-						{
-							ScalarValue = InstanceScalar.ParameterValue;
-							ModifiedFloatParameters++;
-							
-							break;
-						}
-					}
-
-					// Setting cell value
-					MutableTable->SetCell(ColumnIndex, RowIdx, ScalarValue);
+					FString msg = FString::Printf(TEXT("Parameter [%s] from material instance of column [%s] row [%s] is %s. The parameter texture of the default material will be used instead."), *ParamName, *MaterialColumnName, *RowName, *Message);
+					GenerationContext.Compiler->CompilerLog(FText::FromString(msg), TableNode);
 				}
 
-				// Checking if all modifiable floats of the material instance have been modified
-				if (Material->ScalarParameterValues.Num() > ModifiedFloatParameters)
-				{
-					int32 ParametersDiff = Material->ScalarParameterValues.Num() - ModifiedFloatParameters;
+				GenerationContext.ArrayTextureUnrealToMutableTask.Add(FTextureUnrealToMutableTask(MutableTable, ParameterTexture, TableNode, ColumnIndex, RowIdx));
 
-					FString msg =
-						FString::Printf(TEXT("Material Instance [%s] from column [%s] row [%s] has %d modifiable Scalars that will not be modified, they are non-modifiable parameters in the Default Material Instance"),
-							*Material->GetName(), *ColumnName, *RowName, ParametersDiff);
-					GenerationContext.Compiler->CompilerLog(FText::FromString(msg), TableNode);
-				}
-			}
-			else
-			{
-				FString msg = FString::Printf(TEXT("Material from column [%s] row [%s] is a diferent instance than the Reference Material of the table."), *ColumnName, *RowName);
-				GenerationContext.Compiler->CompilerLog(FText::FromString(msg), TableNode);
+				return true;
 			}
 		}
 
@@ -651,7 +531,7 @@ bool FillTableColumn(const UCustomizableObjectNodeTable* TableNode,	mu::TablePtr
 		}
 	}
 
-	else if (const FStructProperty* StructProperty = CastField<FStructProperty>(Property))
+	else if (const FStructProperty* StructProperty = CastField<FStructProperty>(ColumnProperty))
 	{
 		if (StructProperty->Struct == TBaseStructure<FLinearColor>::Get())
 		{
@@ -674,7 +554,7 @@ bool FillTableColumn(const UCustomizableObjectNodeTable* TableNode,	mu::TablePtr
 		}
 	}
 
-	else if (const FNumericProperty* FloatNumProperty = CastField<FFloatProperty>(Property))
+	else if (const FNumericProperty* FloatNumProperty = CastField<FFloatProperty>(ColumnProperty))
 	{
 		CurrentColumn = MutableTable->FindColumn(StringCast<ANSICHAR>(*ColumnName).Get());
 
@@ -688,7 +568,7 @@ bool FillTableColumn(const UCustomizableObjectNodeTable* TableNode,	mu::TablePtr
 		MutableTable->SetCell(CurrentColumn, RowIdx, Value);
 	}
 
-	else if (const FNumericProperty* DoubleNumProperty = CastField<FDoubleProperty>(Property))
+	else if (const FNumericProperty* DoubleNumProperty = CastField<FDoubleProperty>(ColumnProperty))
 	{
 		CurrentColumn = MutableTable->FindColumn(StringCast<ANSICHAR>(*ColumnName).Get());
 	
@@ -712,25 +592,16 @@ bool FillTableColumn(const UCustomizableObjectNodeTable* TableNode,	mu::TablePtr
 }
 
 
-bool GenerateTableColumn(const UCustomizableObjectNodeTable* TableNode, const UEdGraphPin* Pin, mu::TablePtr MutableTable, const FString& DataTableColumnName,
+bool GenerateTableColumn(const UCustomizableObjectNodeTable* TableNode, const UEdGraphPin* Pin, mu::TablePtr MutableTable, const FString& DataTableColumnName, const FProperty* ColumnProperty,
 	const int32 LODIndexConnected, const int32 SectionIndexConnected, const int32 LODIndex, const int32 SectionIndex, const bool bOnlyConnectedLOD, FMutableGraphGenerationContext& GenerationContext)
 {
 	SCOPED_PIN_DATA(GenerationContext, Pin)
 
-	bool bSuccess = false;
-
 	if (!TableNode || !TableNode->Table || !TableNode->Table->GetRowStruct())
 	{
 		return false;
 	}
 
-	FProperty* ColumnProperty = TableNode->Table->FindTableProperty(FName(*DataTableColumnName));
-
-	if (!ColumnProperty)
-	{
-		return false;
-	}
-	
 	// Getting names of the rows to access the information
 	TArray<FName> RowNames = TableNode->GetRowNames();
 
@@ -746,18 +617,18 @@ bool GenerateTableColumn(const UCustomizableObjectNodeTable* TableNode, const UE
 
 			if (CellData)
 			{
-				bool bColumnGenerated = FillTableColumn(TableNode, MutableTable, DataTableColumnName, RowNames[RowIndex].ToString(), RowIndex, CellData, ColumnProperty,
+				bool bCellGenerated = FillTableColumn(TableNode, MutableTable, DataTableColumnName, RowNames[RowIndex].ToString(), RowIndex, CellData, ColumnProperty,
 					LODIndexConnected, SectionIndexConnected, LODIndex, SectionIndex, bOnlyConnectedLOD, GenerationContext);
 				
-				if (!bSuccess)
+				if (!bCellGenerated)
 				{
-					bSuccess = bColumnGenerated;
+					return false;
 				}
 			}
 		}
 	}
 
-	return bSuccess;
+	return true;
 }
 
 
diff --git a/Engine/Plugins/Experimental/Mutable/Source/CustomizableObjectEditor/Private/MuCOE/GenerateMutableSource/GenerateMutableSourceTable.h b/Engine/Plugins/Experimental/Mutable/Source/CustomizableObjectEditor/Private/MuCOE/GenerateMutableSource/GenerateMutableSourceTable.h
--- a/Engine/Plugins/Experimental/Mutable/Source/CustomizableObjectEditor/Private/MuCOE/GenerateMutableSource/GenerateMutableSourceTable.h
+++ b/Engine/Plugins/Experimental/Mutable/Source/CustomizableObjectEditor/Private/MuCOE/GenerateMutableSource/GenerateMutableSourceTable.h
@@ -18,6 +18,7 @@ struct FMutableGraphGenerationContext;
  * @param Pin 
  * @param MutableTable 
  * @param DataTableColumnName 
+ * @param TableProperty 
  * @param LODIndexConnected LOD which the pin is connected to.
  * @param SectionIndexConnected Section which the pin is connected to.
  * @param LODIndex LOD we are generating. Will be different from LODIndexConnected only when using Automatic LOD From Mesh. 
@@ -25,7 +26,7 @@ struct FMutableGraphGenerationContext;
  * @param bOnlyConnectedLOD Corrected LOD and Section will unconditionally always be the connected ones.
  * @param GenerationContext 
  * @return  */
-bool GenerateTableColumn(const UCustomizableObjectNodeTable* TableNode, const UEdGraphPin* Pin, mu::TablePtr MutableTable, const FString& DataTableColumnName,
+bool GenerateTableColumn(const UCustomizableObjectNodeTable* TableNode, const UEdGraphPin* Pin, mu::TablePtr MutableTable, const FString& DataTableColumnName, const FProperty* ColumnProperty,
 	int32 LODIndexConnected, int32 SectionIndexConnected,
 	int32 LODIndex, int32 SectionIndex,
 	bool bOnlyConnectedLOD,
@@ -37,6 +38,7 @@ bool GenerateTableColumn(const UCustomizableObjectNodeTable* TableNode, const UE
  * @param TableNode 
  * @param MutableTable 
  * @param ColumnName 
+ * @param TableProperty 
  * @param RowName 
  * @param RowIdx 
  * @param CellData 
@@ -49,7 +51,7 @@ bool GenerateTableColumn(const UCustomizableObjectNodeTable* TableNode, const UE
  * @param GenerationContext 
  * @return  */
 bool FillTableColumn(const UCustomizableObjectNodeTable* TableNode, mu::TablePtr MutableTable, const FString& ColumnName,
-	const FString& RowName, int32 RowIdx, uint8* CellData, FProperty* Property,
+	const FString& RowName, int32 RowIdx, uint8* CellData, const FProperty* ColumnProperty,
 	int32 LODIndexConnected, int32 SectionIndexConnected, int32 LODIndex, int32 SectionIndex,
 	bool bOnlyConnectedLOD,
 	FMutableGraphGenerationContext& GenerationContext);
diff --git a/Engine/Plugins/Experimental/Mutable/Source/CustomizableObjectEditor/Private/MuCOE/Nodes/CustomizableObjectNodeMaterial.cpp b/Engine/Plugins/Experimental/Mutable/Source/CustomizableObjectEditor/Private/MuCOE/Nodes/CustomizableObjectNodeMaterial.cpp
--- a/Engine/Plugins/Experimental/Mutable/Source/CustomizableObjectEditor/Private/MuCOE/Nodes/CustomizableObjectNodeMaterial.cpp
+++ b/Engine/Plugins/Experimental/Mutable/Source/CustomizableObjectEditor/Private/MuCOE/Nodes/CustomizableObjectNodeMaterial.cpp
@@ -6,6 +6,7 @@
 #include "Materials/Material.h"
 #include "Materials/MaterialExpressionTextureCoordinate.h"
 #include "Materials/MaterialExpressionTextureSample.h"
+#include "Materials/MaterialInstance.h"
 #include "MuCOE/CustomizableObjectEditor_Deprecated.h"
 #include "MuCOE/EdGraphSchema_CustomizableObject.h"
 #include "MuCOE/GraphTraversal.h"
@@ -138,6 +139,7 @@ int32 UCustomizableObjectNodeMaterial::GetExpressionTextureCoordinate(UMaterial*
 	}
 }
 
+
 EPinMode UCustomizableObjectNodeMaterial::NodePinModeToImagePinMode(const ENodePinMode NodePinMode)
 {
 	switch (NodePinMode)
@@ -166,20 +168,6 @@ EPinMode UCustomizableObjectNodeMaterial::GetImagePinMode(const UEdGraphPin& Pin
 
 		return EPinMode::Mutable;
 	}
-	else if (const UEdGraphPin* MaterialAssetPin = GetMaterialAssetPin())
-	{
-		if (const UEdGraphPin* ConnectedPin = FollowInputPin(*MaterialAssetPin))
-		{
-			if (const UCustomizableObjectNodeTable* TableNode = Cast<UCustomizableObjectNodeTable>(ConnectedPin->GetOwningNode()))
-			{
-				const FGuid ParameterGuid = GetPinData<UCustomizableObjectNodeMaterialPinDataImage>(Pin).ParameterId;
-				if (TableNode->ForceImageMutableMode(ConnectedPin, ParameterGuid))
-				{
-					return EPinMode::Mutable;
-				}
-			}
-		}
-	}
 	
 	switch (GetPinData<UCustomizableObjectNodeMaterialPinDataImage>(Pin).GetPinMode())
 	{
@@ -191,7 +179,7 @@ EPinMode UCustomizableObjectNodeMaterial::GetImagePinMode(const UEdGraphPin& Pin
 		return EPinMode::Passthrough;
 	default:
 		check(false); // Missing case.
-		return EPinMode::Mutable;		
+		return EPinMode::Mutable;
 	}
 }
 
@@ -470,8 +458,66 @@ void UCustomizableObjectNodeMaterial::BackwardsCompatibleFixup()
 	{
 		Super::ReconstructNode();
 	}
+
+	if (CustomizableObjectCustomVersion < FCustomizableObjectCustomVersion::AddedTableMaterialSwitch)
+	{
+		UMaterialInstance* DefaultPinValue = nullptr;
+
+		if (const UEdGraphPin* MaterialAssetPin = GetMaterialAssetPin())
+		{
+			if (const UEdGraphPin* ConnectedPin = FollowInputPin(*MaterialAssetPin))
+			{
+				if (const UCustomizableObjectNodeTable* TableNode = Cast<UCustomizableObjectNodeTable>(ConnectedPin->GetOwningNode()))
+				{
+					DefaultPinValue = TableNode->GetColumnDefaultAssetByType<UMaterialInstance>(ConnectedPin);
+				}
+			}
+		}
+
+		if (DefaultPinValue && DefaultPinValue->TextureParameterValues.Num())
+		{
+			const uint32 NumTextureParameters = GetNumParameters(EMaterialParameterType::Texture);
+			for (uint32 ImageIndex = 0; ImageIndex < NumTextureParameters; ++ImageIndex)
+			{
+				if (const UEdGraphPin* ImagePin = GetParameterPin(EMaterialParameterType::Texture, ImageIndex))
+				{
+					FGuid ParameterId = GetParameterId(EMaterialParameterType::Texture, ImageIndex);
+
+					TArray<FMaterialParameterInfo> TextureParameterInfo;
+					TArray<FGuid> TextureGuids;
+
+					// Getting parent's texture infos
+					DefaultPinValue->GetMaterial()->GetAllTextureParameterInfo(TextureParameterInfo, TextureGuids);
+
+					int32 TextureIndex = TextureGuids.Find(ParameterId);
+
+					if (TextureIndex == INDEX_NONE)
+					{
+						continue;
+					}
+
+					FName TextureName = TextureParameterInfo[TextureIndex].Name;
+
+					// Checking if the pin's texture has been modified in the material instance
+					for (const FTextureParameterValue Texture : DefaultPinValue->TextureParameterValues)
+					{
+						if (TextureName == Texture.ParameterInfo.Name)
+						{
+							if (UCustomizableObjectNodeMaterialPinDataImage* PinDataImage = Cast<UCustomizableObjectNodeMaterialPinDataImage>(GetPinData(*ImagePin)))
+							{
+								PinDataImage->PinMode = EPinMode::Mutable;
+
+								PinsImagePinMode.Add(ImagePin->PinId, GetImagePinMode(*ImagePin));
+							}
+						}
+					}
+				}
+			}
+		}
+	}
 }
 
+
 void UCustomizableObjectNodeMaterial::PostBackwardsCompatibleFixup()
 {
 	Super::PostBackwardsCompatibleFixup();
@@ -994,7 +1040,7 @@ TSharedPtr<SWidget> UCustomizableObjectNodeMaterial::CustomizePinDetails(UEdGrap
 {
 	if (UCustomizableObjectNodeMaterialPinDataImage* PinData = Cast<UCustomizableObjectNodeMaterialPinDataImage>(GetPinData(Pin)))
 	{
-		return SNew(SPinViewerNodeMaterialPinImageDetails).PinData(PinData);
+		return SNew(SPinViewerNodeMaterialPinImageDetails).Pin(&Pin).PinData(PinData);
 	}
 	else
 	{
diff --git a/Engine/Plugins/Experimental/Mutable/Source/CustomizableObjectEditor/Private/MuCOE/Nodes/CustomizableObjectNodeTable.cpp b/Engine/Plugins/Experimental/Mutable/Source/CustomizableObjectEditor/Private/MuCOE/Nodes/CustomizableObjectNodeTable.cpp
--- a/Engine/Plugins/Experimental/Mutable/Source/CustomizableObjectEditor/Private/MuCOE/Nodes/CustomizableObjectNodeTable.cpp
+++ b/Engine/Plugins/Experimental/Mutable/Source/CustomizableObjectEditor/Private/MuCOE/Nodes/CustomizableObjectNodeTable.cpp
@@ -747,39 +747,6 @@ bool UCustomizableObjectNodeTable::IsPinRelevant(const UEdGraphPin* Pin) const
 }
 
 
-bool UCustomizableObjectNodeTable::ForceImageMutableMode(const UEdGraphPin* Pin, FGuid ParameterId) const
-{
-	UMaterialInstance* DefaultPinValue = GetColumnDefaultAssetByType<UMaterialInstance>(Pin);
-	
-	if (DefaultPinValue && DefaultPinValue->TextureParameterValues.Num())
-	{
-		TArray<FMaterialParameterInfo> TextureParameterInfo;
-		TArray<FGuid> TextureGuids;
-
-		DefaultPinValue->GetMaterial()->GetAllTextureParameterInfo(TextureParameterInfo, TextureGuids);
-
-		int32 TextureIndex = TextureGuids.Find(ParameterId);
-		
-		if (TextureIndex == INDEX_NONE)
-		{
-			return false;
-		}
-		
-		FName TextureName = TextureParameterInfo[TextureIndex].Name;
-		
-		for (const FTextureParameterValue Texture : DefaultPinValue->TextureParameterValues)
-		{
-			if (TextureName == Texture.ParameterInfo.Name)
-			{
-				return true;
-			}
-		}
-	}
-
-	return false;
-}
-
-
 void UCustomizableObjectNodeTable::SetLayoutInLayoutEditor(UCustomizableObjectLayout* CurrentLayout)
 {
 	TSharedPtr<ICustomizableObjectEditor> Editor = GetGraphEditor();
diff --git a/Engine/Plugins/Experimental/Mutable/Source/CustomizableObjectEditor/Private/MuCOE/Nodes/SPinViewerNodeMaterialPinImageDetails.cpp b/Engine/Plugins/Experimental/Mutable/Source/CustomizableObjectEditor/Private/MuCOE/Nodes/SPinViewerNodeMaterialPinImageDetails.cpp
--- a/Engine/Plugins/Experimental/Mutable/Source/CustomizableObjectEditor/Private/MuCOE/Nodes/SPinViewerNodeMaterialPinImageDetails.cpp
+++ b/Engine/Plugins/Experimental/Mutable/Source/CustomizableObjectEditor/Private/MuCOE/Nodes/SPinViewerNodeMaterialPinImageDetails.cpp
@@ -2,9 +2,11 @@
 
 #include "MuCOE/Nodes/SPinViewerNodeMaterialPinImageDetails.h"
 
+#include "EdGraph/EdGraphPin.h"
 #include "Framework/Views/TableViewMetadata.h"
 #include "ISinglePropertyView.h"
 #include "Modules/ModuleManager.h"
+#include "MuCOE/GraphTraversal.h"
 #include "MuCOE/Nodes/CustomizableObjectNodeMaterial.h"
 #include "MuCOE/SMutableTextComboBox.h"
 #include "ScopedTransaction.h"
@@ -81,6 +83,7 @@ void SPinViewerNodeMaterialPinImageDetails::Construct(const FArguments& InArgs)
 {
 	SPinViewerPinDetails::Construct(SPinViewerPinDetails::FArguments());
 
+	Pin = InArgs._Pin;
 	PinData = InArgs._PinData;
 	
 	{
@@ -92,6 +95,8 @@ void SPinViewerNodeMaterialPinImageDetails::Construct(const FArguments& InArgs)
 		{
 			return *Element == PinData->GetPinMode();
 		});
+
+		bool bIsLinked = FollowInputPin(*Pin) != nullptr;
 		
 		TSharedPtr<SWidget> Widget;
 		SAssignNew(Widget, SMutableTextComboBox<TSharedPtr<EPinMode>>)
@@ -101,6 +106,16 @@ void SPinViewerNodeMaterialPinImageDetails::Construct(const FArguments& InArgs)
 			.Translate_Lambda([this](const TSharedPtr<EPinMode> Option)
 			{
 				return EPinModeToText(*Option);
+			})
+			.IsEnabled(!bIsLinked)
+			.ToolTipText_Lambda([bIsLinked]()
+			{
+				if (bIsLinked)
+				{
+					return LOCTEXT("SPinViewer_PinModeComboboxTooltip", "Value being overridden due to pin linked to a node.");
+				}
+
+				return LOCTEXT("SPinViewer_PinModeComboboxTooltipDisabled", "");
 			});
 
 		const FText Tooltip = LOCTEXT("PinModeTooltip", "Texture Parameter mode.");
diff --git a/Engine/Plugins/Experimental/Mutable/Source/CustomizableObjectEditor/Private/MuCOE/Nodes/SPinViewerNodeMaterialPinImageDetails.h b/Engine/Plugins/Experimental/Mutable/Source/CustomizableObjectEditor/Private/MuCOE/Nodes/SPinViewerNodeMaterialPinImageDetails.h
--- a/Engine/Plugins/Experimental/Mutable/Source/CustomizableObjectEditor/Private/MuCOE/Nodes/SPinViewerNodeMaterialPinImageDetails.h
+++ b/Engine/Plugins/Experimental/Mutable/Source/CustomizableObjectEditor/Private/MuCOE/Nodes/SPinViewerNodeMaterialPinImageDetails.h
@@ -13,6 +13,7 @@ template <typename NumericType> class SSpinBox;
 
 class FString;
 class SWidget;
+class UEdGraphPin;
 
 /** UV Layout Mode UI Enum. */
 UENUM()
@@ -37,6 +38,7 @@ class SPinViewerNodeMaterialPinImageDetails : public SPinViewerPinDetails
 {
 public:
 	SLATE_BEGIN_ARGS(SPinViewerNodeMaterialPinImageDetails) {}
+		SLATE_ARGUMENT(UEdGraphPin*, Pin)
 		SLATE_ARGUMENT(UCustomizableObjectNodeMaterialPinDataImage*, PinData)
 	SLATE_END_ARGS()
 
@@ -50,6 +52,9 @@ private:
 	TSharedPtr<SWidget> UVLayout;
 	TMap<TSharedPtr<FString>, EPinMode> PinModeComboBoxOptionValues;
 	
+	/** Pointer to the Pin to show and modigy. */
+	UEdGraphPin* Pin = nullptr;
+
 	/** Pin Data to show and modify. */
 	UCustomizableObjectNodeMaterialPinDataImage* PinData = nullptr;
 	
