[mutable] Cleanup code generation for images to store state in stack instead of heap vector.
Optimize mesh comparison at code linking stage.
#rnx

diff --git a/Engine/Plugins/Experimental/Mutable/Source/CustomizableObjectEditor/Private/MuCOE/GenerateMutableSource/GenerateMutableSourceColor.cpp b/Engine/Plugins/Experimental/Mutable/Source/CustomizableObjectEditor/Private/MuCOE/GenerateMutableSource/GenerateMutableSourceColor.cpp
--- a/Engine/Plugins/Experimental/Mutable/Source/CustomizableObjectEditor/Private/MuCOE/GenerateMutableSource/GenerateMutableSourceColor.cpp
+++ b/Engine/Plugins/Experimental/Mutable/Source/CustomizableObjectEditor/Private/MuCOE/GenerateMutableSource/GenerateMutableSourceColor.cpp
@@ -54,7 +54,7 @@ mu::NodeColourPtr GenerateMutableSourceColor(const UEdGraphPin* Pin, FMutableGra
 		mu::NodeColourConstantPtr ColorNode = new mu::NodeColourConstant();
 		Result = ColorNode;
 
-		ColorNode->SetValue(TypedNodeColorConst->Value.R, TypedNodeColorConst->Value.G, TypedNodeColorConst->Value.B, TypedNodeColorConst->Value.A);
+		ColorNode->SetValue(TypedNodeColorConst->Value);
 	}
 
 	else if (const UCustomizableObjectNodeColorParameter* TypedNodeColorParam = Cast<UCustomizableObjectNodeColorParameter>(Node))
@@ -271,7 +271,7 @@ mu::NodeColourPtr GenerateMutableSourceColor(const UEdGraphPin* Pin, FMutableGra
 	{
 		//This node will add a white color in case of error
 		mu::NodeColourConstantPtr WhiteColorNode = new mu::NodeColourConstant();
-		WhiteColorNode->SetValue(1.0f, 1.0f, 1.0f, 1.0f);
+		WhiteColorNode->SetValue(FVector4f(1.0f, 1.0f, 1.0f, 1.0f));
 
 		Result = WhiteColorNode;
 
diff --git a/Engine/Plugins/Experimental/Mutable/Source/CustomizableObjectEditor/Private/MuCOE/GenerateMutableSource/GenerateMutableSourceGroupProjector.cpp b/Engine/Plugins/Experimental/Mutable/Source/CustomizableObjectEditor/Private/MuCOE/GenerateMutableSource/GenerateMutableSourceGroupProjector.cpp
--- a/Engine/Plugins/Experimental/Mutable/Source/CustomizableObjectEditor/Private/MuCOE/GenerateMutableSource/GenerateMutableSourceGroupProjector.cpp
+++ b/Engine/Plugins/Experimental/Mutable/Source/CustomizableObjectEditor/Private/MuCOE/GenerateMutableSource/GenerateMutableSourceGroupProjector.cpp
@@ -276,7 +276,7 @@ mu::NodeImagePtr GenerateMutableGroupProjection(const int32 NodeLOD, const int32
 	}
 	
 	mu::NodeColourConstantPtr ZeroColorNode = new mu::NodeColourConstant();
-	ZeroColorNode->SetValue(0.f, 0.f, 0.f, 1.0f);
+	ZeroColorNode->SetValue(FVector4f(0.f, 0.f, 0.f, 1.0f));
 
 	if (TextureSize <= 0 || !FMath::IsPowerOfTwo(TextureSize))
 	{
diff --git a/Engine/Plugins/Experimental/Mutable/Source/MutableRuntime/Private/MuR/CodeRunner.cpp b/Engine/Plugins/Experimental/Mutable/Source/MutableRuntime/Private/MuR/CodeRunner.cpp
--- a/Engine/Plugins/Experimental/Mutable/Source/MutableRuntime/Private/MuR/CodeRunner.cpp
+++ b/Engine/Plugins/Experimental/Mutable/Source/MutableRuntime/Private/MuR/CodeRunner.cpp
@@ -4162,7 +4162,7 @@ namespace mu
             		
                 Ptr<const Image> pA = LoadImage( FCacheAddress(args.source,item) );
 
-                box< vec2<int> > rect;
+                box< UE::Math::TIntVector2<int32> > rect;
                 rect.min[0] = args.minX;
                 rect.min[1] = args.minY;
                 rect.size[0] = args.sizeX;
diff --git a/Engine/Plugins/Experimental/Mutable/Source/MutableRuntime/Private/MuR/ImageRLE.cpp b/Engine/Plugins/Experimental/Mutable/Source/MutableRuntime/Private/MuR/ImageRLE.cpp
--- a/Engine/Plugins/Experimental/Mutable/Source/MutableRuntime/Private/MuR/ImageRLE.cpp
+++ b/Engine/Plugins/Experimental/Mutable/Source/MutableRuntime/Private/MuR/ImageRLE.cpp
@@ -512,6 +512,8 @@ namespace mu
     //---------------------------------------------------------------------------------------------
     void UncompressRLE_RGB( int32 width, int32 rows, const uint8* pBaseData, uint8* pDestDataB )
     {
+		check(pBaseData && pDestDataB);
+
         UINT24* pDestData = reinterpret_cast<UINT24*>( pDestDataB );
 
         pBaseData += rows*sizeof(uint32);
diff --git a/Engine/Plugins/Experimental/Mutable/Source/MutableRuntime/Private/MuR/ModelPrivate.h b/Engine/Plugins/Experimental/Mutable/Source/MutableRuntime/Private/MuR/ModelPrivate.h
--- a/Engine/Plugins/Experimental/Mutable/Source/MutableRuntime/Private/MuR/ModelPrivate.h
+++ b/Engine/Plugins/Experimental/Mutable/Source/MutableRuntime/Private/MuR/ModelPrivate.h
@@ -48,7 +48,7 @@ namespace mu
     inline void AppendCode(TArray<uint8>& code, const DATA& data )
     {
         int32 pos = code.Num();
-        code.SetNum( pos+sizeof(DATA) );
+        code.SetNum( pos+sizeof(DATA), false );
 		FMemory::Memcpy (&code[pos], &data, sizeof(DATA));
     }
 
@@ -375,7 +375,7 @@ namespace mu
 			return RomSize;
 		}
 
-        //! Adds a constant data and returns its constant index.
+        //! Adds a constant image data and returns its constant index.
 		int32 AddConstant(Ptr<const Image> pImage, int32 MinTextureResidentMipCount)
 		{
 			check(pImage->GetSizeX()*pImage->GetSizeY()>0);
@@ -471,16 +471,7 @@ namespace mu
 
 		int32 AddConstant(Ptr<const Mesh> pMesh)
 		{
-			// Ensure unique
-			for (int32 i = 0; i < m_constantMeshes.Num(); ++i)
-			{
-				const Mesh* pCandidate = m_constantMeshes[i].Value.get();
-				if (*pCandidate == *pMesh)
-				{
-					return i;
-				}
-			}
-
+			// Uniques needs to be ensured outside
 			return m_constantMeshes.Add(TPair<int32, Ptr<const Mesh>>( -1, pMesh.get() ));
 		}
 
diff --git a/Engine/Plugins/Experimental/Mutable/Source/MutableRuntime/Private/MuR/OpImageCrop.cpp b/Engine/Plugins/Experimental/Mutable/Source/MutableRuntime/Private/MuR/OpImageCrop.cpp
--- a/Engine/Plugins/Experimental/Mutable/Source/MutableRuntime/Private/MuR/OpImageCrop.cpp
+++ b/Engine/Plugins/Experimental/Mutable/Source/MutableRuntime/Private/MuR/OpImageCrop.cpp
@@ -8,7 +8,7 @@
 namespace mu
 {
 
-	bool FImageOperator::ImageCrop( Image* InCropped, int32 CompressionQuality, const Image* InBase, const box< vec2<int32> >& Rect)
+	bool FImageOperator::ImageCrop( Image* InCropped, int32 CompressionQuality, const Image* InBase, const box< FIntVector2 >& Rect)
 	{
 		Ptr<const Image> Base = InBase;
 		Ptr<Image> BaseReformat;
diff --git a/Engine/Plugins/Experimental/Mutable/Source/MutableRuntime/Public/MuR/Image.h b/Engine/Plugins/Experimental/Mutable/Source/MutableRuntime/Public/MuR/Image.h
--- a/Engine/Plugins/Experimental/Mutable/Source/MutableRuntime/Public/MuR/Image.h
+++ b/Engine/Plugins/Experimental/Mutable/Source/MutableRuntime/Public/MuR/Image.h
@@ -594,7 +594,7 @@ namespace mu
 		void ImageMipmap_ReleaseScratch(FScratchImageMipmap&);
 
 		/** */
-		MUTABLERUNTIME_API bool ImageCrop(Image* Cropped, int32 CompressionQuality, const Image* Base, const box< vec2<int32> >& Rect);
+		MUTABLERUNTIME_API bool ImageCrop(Image* Cropped, int32 CompressionQuality, const Image* Base, const box<FIntVector2>& Rect);
 
 		/** */
 		void ImageCompose(Image* Base, const Image* Block, const box< UE::Math::TIntVector2<uint16> >& Rect);
diff --git a/Engine/Plugins/Experimental/Mutable/Source/MutableTools/Private/MuT/AST.cpp b/Engine/Plugins/Experimental/Mutable/Source/MutableTools/Private/MuT/AST.cpp
--- a/Engine/Plugins/Experimental/Mutable/Source/MutableTools/Private/MuT/AST.cpp
+++ b/Engine/Plugins/Experimental/Mutable/Source/MutableTools/Private/MuT/AST.cpp
@@ -243,9 +243,10 @@ mu::Ptr<ASTOp> ASTOp::DeepClone( const Ptr<ASTOp>& root )
 
 
 //-------------------------------------------------------------------------------------------------
-void ASTOp::FullLink( Ptr<ASTOp>& root, FProgram& program, const FLinkerOptions* Options )
+void ASTOp::FullLink( Ptr<ASTOp>& root, FProgram& program, FLinkerOptions* Options )
 {
     MUTABLE_CPUPROFILER_SCOPE(AST_FullLink);
+
     Traverse_BottomUp_Unique( root,
                               [&](Ptr<ASTOp> n){ n->Link(program, Options); },
                               [&](Ptr<const ASTOp> n){ return n->linkedAddress==0; });
@@ -1022,8 +1023,10 @@ void ASTOpFixed::ForEachChild( const TFunctionRef<void(ASTChild&)> f )
 }
 
 
-void ASTOpFixed::Link( FProgram& program, const FLinkerOptions* )
+void ASTOpFixed::Link( FProgram& program, FLinkerOptions* )
 {
+	MUTABLE_CPUPROFILER_SCOPE(ASTOpFixed_Link);
+
     if (!linkedAddress)
     {
         OP lop = op;
diff --git a/Engine/Plugins/Experimental/Mutable/Source/MutableTools/Private/MuT/AST.h b/Engine/Plugins/Experimental/Mutable/Source/MutableTools/Private/MuT/AST.h
--- a/Engine/Plugins/Experimental/Mutable/Source/MutableTools/Private/MuT/AST.h
+++ b/Engine/Plugins/Experimental/Mutable/Source/MutableTools/Private/MuT/AST.h
@@ -4,6 +4,7 @@
 
 #include "MuT/Platform.h"
 #include "MuR/Image.h"
+#include "MuR/Mesh.h"
 #include "MuR/ModelPrivate.h"
 #include "MuR/MutableMath.h"
 #include "MuR/MutableMemory.h"
@@ -210,7 +211,27 @@ namespace mu
 
 	struct FLinkerOptions
 	{
-		int MinTextureResidentMipCount = 0;
+		int32 MinTextureResidentMipCount = 0;
+
+		/** */
+		struct FDeduplicationMeshFuncs : TDefaultMapHashableKeyFuncs<mu::Ptr<const mu::Mesh>, int32, false>
+		{
+			static FORCEINLINE bool Matches(KeyInitType A, KeyInitType B)
+			{
+				return *A == *B;
+			}
+
+			static FORCEINLINE uint32 GetKeyHash(KeyInitType Key)
+			{
+				const mu::Mesh* Data = Key.get();
+				return HashCombine(
+					::GetTypeHash(Data->m_VertexBuffers.GetElementCount()),
+					::GetTypeHash(Data->m_IndexBuffers.GetElementCount())
+				);
+			}
+		};
+
+		TMap<mu::Ptr<const mu::Mesh>, int32, FDefaultSetAllocator, FDeduplicationMeshFuncs> MeshConstantMap;
 	};
 
 
@@ -462,7 +483,7 @@ namespace mu
         //---------------------------------------------------------------------------------------------
 
         //!
-        static void FullLink( Ptr<ASTOp>& root, FProgram&, const FLinkerOptions*);
+        static void FullLink( Ptr<ASTOp>& root, FProgram&, FLinkerOptions*);
 
         //!
         static void ClearLinkData( Ptr<ASTOp>& root );
@@ -473,7 +494,7 @@ namespace mu
     private:
 
         //!
-        virtual void Link( FProgram& program, const FLinkerOptions* Options ) = 0;
+        virtual void Link( FProgram& program, FLinkerOptions* Options ) = 0;
 
     protected:
 
@@ -980,7 +1001,7 @@ namespace mu
 
         Ptr<ASTOp> Clone( MapChildFuncRef mapChild ) const override;
         void ForEachChild( const TFunctionRef<void(ASTChild&)> f ) override;
-        void Link( FProgram& program, const FLinkerOptions* Options) override;
+        void Link( FProgram& program, FLinkerOptions* Options) override;
         bool IsEqual(const ASTOp& otherUntyped) const override;
 		uint64 Hash() const override;
 		FImageDesc GetImageDesc( bool returnBestOption, FGetImageDescContext* context ) const override;
diff --git a/Engine/Plugins/Experimental/Mutable/Source/MutableTools/Private/MuT/ASTOpAddExtensionData.cpp b/Engine/Plugins/Experimental/Mutable/Source/MutableTools/Private/MuT/ASTOpAddExtensionData.cpp
--- a/Engine/Plugins/Experimental/Mutable/Source/MutableTools/Private/MuT/ASTOpAddExtensionData.cpp
+++ b/Engine/Plugins/Experimental/Mutable/Source/MutableTools/Private/MuT/ASTOpAddExtensionData.cpp
@@ -69,7 +69,7 @@ void ASTOpAddExtensionData::ForEachChild(const TFunctionRef<void(ASTChild&)> F)
 	F(ExtensionData);
 }
 
-void ASTOpAddExtensionData::Link(FProgram& Program, const FLinkerOptions*)
+void ASTOpAddExtensionData::Link(FProgram& Program, FLinkerOptions*)
 {
 	// Already linked?
 	if (linkedAddress)
diff --git a/Engine/Plugins/Experimental/Mutable/Source/MutableTools/Private/MuT/ASTOpAddExtensionData.h b/Engine/Plugins/Experimental/Mutable/Source/MutableTools/Private/MuT/ASTOpAddExtensionData.h
--- a/Engine/Plugins/Experimental/Mutable/Source/MutableTools/Private/MuT/ASTOpAddExtensionData.h
+++ b/Engine/Plugins/Experimental/Mutable/Source/MutableTools/Private/MuT/ASTOpAddExtensionData.h
@@ -36,7 +36,7 @@ public:
 	void ForEachChild(const TFunctionRef<void(ASTChild&)> F) override;
 	bool IsEqual(const ASTOp& OtherUntyped) const override;
 	Ptr<ASTOp> Clone(MapChildFuncRef MapChild) const override;
-	void Link(FProgram& Program, const FLinkerOptions* Options) override;
+	void Link(FProgram& Program, FLinkerOptions* Options) override;
 };
 
 
diff --git a/Engine/Plugins/Experimental/Mutable/Source/MutableTools/Private/MuT/ASTOpAddLOD.cpp b/Engine/Plugins/Experimental/Mutable/Source/MutableTools/Private/MuT/ASTOpAddLOD.cpp
--- a/Engine/Plugins/Experimental/Mutable/Source/MutableTools/Private/MuT/ASTOpAddLOD.cpp
+++ b/Engine/Plugins/Experimental/Mutable/Source/MutableTools/Private/MuT/ASTOpAddLOD.cpp
@@ -71,7 +71,7 @@ namespace mu
 	}
 
 
-	void ASTOpAddLOD::Link(FProgram& program, const FLinkerOptions*)
+	void ASTOpAddLOD::Link(FProgram& program, FLinkerOptions*)
 	{
 		// Already linked?
 		if (!linkedAddress)
diff --git a/Engine/Plugins/Experimental/Mutable/Source/MutableTools/Private/MuT/ASTOpAddLOD.h b/Engine/Plugins/Experimental/Mutable/Source/MutableTools/Private/MuT/ASTOpAddLOD.h
--- a/Engine/Plugins/Experimental/Mutable/Source/MutableTools/Private/MuT/ASTOpAddLOD.h
+++ b/Engine/Plugins/Experimental/Mutable/Source/MutableTools/Private/MuT/ASTOpAddLOD.h
@@ -33,7 +33,7 @@ struct FProgram;
 		void ForEachChild(const TFunctionRef<void(ASTChild&)>) override;
 		bool IsEqual(const ASTOp& otherUntyped) const override;
 		Ptr<ASTOp> Clone(MapChildFuncRef mapChild) const override;
-		void Link(FProgram& program, const FLinkerOptions* Options) override;
+		void Link(FProgram& program, FLinkerOptions* Options) override;
 	};
 
 
diff --git a/Engine/Plugins/Experimental/Mutable/Source/MutableTools/Private/MuT/ASTOpConditional.cpp b/Engine/Plugins/Experimental/Mutable/Source/MutableTools/Private/MuT/ASTOpConditional.cpp
--- a/Engine/Plugins/Experimental/Mutable/Source/MutableTools/Private/MuT/ASTOpConditional.cpp
+++ b/Engine/Plugins/Experimental/Mutable/Source/MutableTools/Private/MuT/ASTOpConditional.cpp
@@ -103,7 +103,7 @@ namespace mu
 
 
 	//-------------------------------------------------------------------------------------------------
-	void ASTOpConditional::Link(FProgram& program, const FLinkerOptions*)
+	void ASTOpConditional::Link(FProgram& program, FLinkerOptions*)
 	{
 		// Already linked?
 		if (!linkedAddress)
diff --git a/Engine/Plugins/Experimental/Mutable/Source/MutableTools/Private/MuT/ASTOpConditional.h b/Engine/Plugins/Experimental/Mutable/Source/MutableTools/Private/MuT/ASTOpConditional.h
--- a/Engine/Plugins/Experimental/Mutable/Source/MutableTools/Private/MuT/ASTOpConditional.h
+++ b/Engine/Plugins/Experimental/Mutable/Source/MutableTools/Private/MuT/ASTOpConditional.h
@@ -43,7 +43,7 @@ struct FProgram;
 		uint64 Hash() const override;
 		void Assert() override;
 		void ForEachChild(const TFunctionRef<void(ASTChild&)> f) override;
-		void Link(FProgram& program, const FLinkerOptions* Options) override;
+		void Link(FProgram& program, FLinkerOptions* Options) override;
 		FImageDesc GetImageDesc(bool returnBestOption, class FGetImageDescContext* context) const override;
 		void GetLayoutBlockSize(int* pBlockX, int* pBlockY) override;
 		void GetBlockLayoutSize(int blockIndex, int* pBlockX, int* pBlockY, FBlockLayoutSizeCache* cache) override;
diff --git a/Engine/Plugins/Experimental/Mutable/Source/MutableTools/Private/MuT/ASTOpConstantBool.cpp b/Engine/Plugins/Experimental/Mutable/Source/MutableTools/Private/MuT/ASTOpConstantBool.cpp
--- a/Engine/Plugins/Experimental/Mutable/Source/MutableTools/Private/MuT/ASTOpConstantBool.cpp
+++ b/Engine/Plugins/Experimental/Mutable/Source/MutableTools/Private/MuT/ASTOpConstantBool.cpp
@@ -53,7 +53,7 @@ namespace mu
 	}
 
 
-	void ASTOpConstantBool::Link(FProgram& program, const FLinkerOptions*)
+	void ASTOpConstantBool::Link(FProgram& program, FLinkerOptions*)
 	{
 		if (!linkedAddress)
 		{
diff --git a/Engine/Plugins/Experimental/Mutable/Source/MutableTools/Private/MuT/ASTOpConstantBool.h b/Engine/Plugins/Experimental/Mutable/Source/MutableTools/Private/MuT/ASTOpConstantBool.h
--- a/Engine/Plugins/Experimental/Mutable/Source/MutableTools/Private/MuT/ASTOpConstantBool.h
+++ b/Engine/Plugins/Experimental/Mutable/Source/MutableTools/Private/MuT/ASTOpConstantBool.h
@@ -32,7 +32,7 @@ struct FProgram;
 		void ForEachChild(const TFunctionRef<void(ASTChild&)>) override;
 		bool IsEqual(const ASTOp& otherUntyped) const override;
 		Ptr<ASTOp> Clone(MapChildFuncRef mapChild) const override;
-		void Link(FProgram& program, const FLinkerOptions* Options) override;
+		void Link(FProgram& program, FLinkerOptions* Options) override;
 		FBoolEvalResult EvaluateBool(ASTOpList& facts, FEvaluateBoolCache* cache) const override;
 	};
 
diff --git a/Engine/Plugins/Experimental/Mutable/Source/MutableTools/Private/MuT/ASTOpConstantExtensionData.cpp b/Engine/Plugins/Experimental/Mutable/Source/MutableTools/Private/MuT/ASTOpConstantExtensionData.cpp
--- a/Engine/Plugins/Experimental/Mutable/Source/MutableTools/Private/MuT/ASTOpConstantExtensionData.cpp
+++ b/Engine/Plugins/Experimental/Mutable/Source/MutableTools/Private/MuT/ASTOpConstantExtensionData.cpp
@@ -44,7 +44,7 @@ uint64 ASTOpConstantExtensionData::Hash() const
 
 
 //-------------------------------------------------------------------------------------------------
-void ASTOpConstantExtensionData::Link(FProgram& Program, const FLinkerOptions* Options)
+void ASTOpConstantExtensionData::Link(FProgram& Program, FLinkerOptions* Options)
 {
 	if (!linkedAddress)
 	{
diff --git a/Engine/Plugins/Experimental/Mutable/Source/MutableTools/Private/MuT/ASTOpConstantExtensionData.h b/Engine/Plugins/Experimental/Mutable/Source/MutableTools/Private/MuT/ASTOpConstantExtensionData.h
--- a/Engine/Plugins/Experimental/Mutable/Source/MutableTools/Private/MuT/ASTOpConstantExtensionData.h
+++ b/Engine/Plugins/Experimental/Mutable/Source/MutableTools/Private/MuT/ASTOpConstantExtensionData.h
@@ -32,7 +32,7 @@ public:
 	bool IsEqual(const ASTOp& OtherUntyped) const override;
 	Ptr<ASTOp> Clone(MapChildFuncRef MapChild) const override;
 	uint64 Hash() const override;
-	void Link(FProgram& Program, const FLinkerOptions*) override;
+	void Link(FProgram& Program, FLinkerOptions*) override;
 };
 
 
diff --git a/Engine/Plugins/Experimental/Mutable/Source/MutableTools/Private/MuT/ASTOpConstantProjector.cpp b/Engine/Plugins/Experimental/Mutable/Source/MutableTools/Private/MuT/ASTOpConstantProjector.cpp
--- a/Engine/Plugins/Experimental/Mutable/Source/MutableTools/Private/MuT/ASTOpConstantProjector.cpp
+++ b/Engine/Plugins/Experimental/Mutable/Source/MutableTools/Private/MuT/ASTOpConstantProjector.cpp
@@ -46,7 +46,7 @@ namespace mu
 	}
 
 
-	void ASTOpConstantProjector::Link(FProgram& program, const FLinkerOptions*)
+	void ASTOpConstantProjector::Link(FProgram& program, FLinkerOptions*)
 	{
 		if (!linkedAddress)
 		{
diff --git a/Engine/Plugins/Experimental/Mutable/Source/MutableTools/Private/MuT/ASTOpConstantProjector.h b/Engine/Plugins/Experimental/Mutable/Source/MutableTools/Private/MuT/ASTOpConstantProjector.h
--- a/Engine/Plugins/Experimental/Mutable/Source/MutableTools/Private/MuT/ASTOpConstantProjector.h
+++ b/Engine/Plugins/Experimental/Mutable/Source/MutableTools/Private/MuT/ASTOpConstantProjector.h
@@ -31,7 +31,7 @@ struct FProgram;
 		void ForEachChild(const TFunctionRef<void(ASTChild&)>) override;
 		bool IsEqual(const ASTOp& otherUntyped) const override;
 		Ptr<ASTOp> Clone(MapChildFuncRef mapChild) const override;
-		void Link(FProgram& program, const FLinkerOptions*) override;
+		void Link(FProgram& program, FLinkerOptions*) override;
 	};
 
 
diff --git a/Engine/Plugins/Experimental/Mutable/Source/MutableTools/Private/MuT/ASTOpConstantResource.cpp b/Engine/Plugins/Experimental/Mutable/Source/MutableTools/Private/MuT/ASTOpConstantResource.cpp
--- a/Engine/Plugins/Experimental/Mutable/Source/MutableTools/Private/MuT/ASTOpConstantResource.cpp
+++ b/Engine/Plugins/Experimental/Mutable/Source/MutableTools/Private/MuT/ASTOpConstantResource.cpp
@@ -65,8 +65,10 @@ namespace mu
 
 
 	//-------------------------------------------------------------------------------------------------
-	void ASTOpConstantResource::Link(FProgram& program, const FLinkerOptions* Options)
+	void ASTOpConstantResource::Link(FProgram& program, FLinkerOptions* Options)
 	{
+		MUTABLE_CPUPROFILER_SCOPE(ASTOpConstantResource_Link);
+
 		if (!linkedAddress && !bLinkedAndNull)
 		{
 			if (type == OP_TYPE::ME_CONSTANT)
@@ -74,24 +76,35 @@ namespace mu
 				OP::MeshConstantArgs args;
 				FMemory::Memset(&args, 0, sizeof(args));
 
-				Ptr<Mesh> pTyped = static_cast<const Mesh*>(GetValue().get())->Clone();
-				check(pTyped);
+				Ptr<Mesh> MeshData = static_cast<const Mesh*>(GetValue().get())->Clone();
+				check(MeshData);
 
 				args.skeleton = -1;
-				if (Ptr<const Skeleton> pSkeleton = pTyped->GetSkeleton())
+				if (Ptr<const Skeleton> pSkeleton = MeshData->GetSkeleton())
 				{
 					args.skeleton = program.AddConstant(pSkeleton.get());
-					pTyped->SetSkeleton(nullptr);
+					MeshData->SetSkeleton(nullptr);
 				}
 
 				args.physicsBody = -1;
-				if (Ptr<const PhysicsBody> pPhysicsBody = pTyped->GetPhysicsBody())
+				if (Ptr<const PhysicsBody> pPhysicsBody = MeshData->GetPhysicsBody())
 				{
 					args.physicsBody = program.AddConstant(pPhysicsBody.get());
-					pTyped->SetPhysicsBody(nullptr);
+					MeshData->SetPhysicsBody(nullptr);
 				}
 
-				args.value = program.AddConstant(pTyped.get());
+				// Use a map-based deduplication
+				mu::Ptr<const mu::Mesh> Key = MeshData;
+				const int32* IndexPtr = Options->MeshConstantMap.Find(Key);
+				if (!IndexPtr)
+				{
+					args.value = program.AddConstant(MeshData.get());
+					Options->MeshConstantMap.Add(MeshData, int32(args.value));
+				}
+				else
+				{
+					args.value = *IndexPtr;
+				}
 
 				linkedAddress = (OP::ADDRESS)program.m_opAddress.Num();
 				program.m_opAddress.Add((uint32_t)program.m_byteCode.Num());
@@ -119,7 +132,6 @@ namespace mu
 					}
 					else
 					{
-						// TODO: Compilation options?
 						int32 MinTextureResidentMipCount = Options->MinTextureResidentMipCount;
 						args.value = program.AddConstant(pTyped, MinTextureResidentMipCount);
 					}
diff --git a/Engine/Plugins/Experimental/Mutable/Source/MutableTools/Private/MuT/ASTOpConstantResource.h b/Engine/Plugins/Experimental/Mutable/Source/MutableTools/Private/MuT/ASTOpConstantResource.h
--- a/Engine/Plugins/Experimental/Mutable/Source/MutableTools/Private/MuT/ASTOpConstantResource.h
+++ b/Engine/Plugins/Experimental/Mutable/Source/MutableTools/Private/MuT/ASTOpConstantResource.h
@@ -60,7 +60,7 @@ namespace mu
 		bool IsEqual(const ASTOp& otherUntyped) const override;
 		Ptr<ASTOp> Clone(MapChildFuncRef mapChild) const override;
 		uint64 Hash() const override;
-		void Link(FProgram& program, const FLinkerOptions*) override;
+		void Link(FProgram& program, FLinkerOptions*) override;
 		FImageDesc GetImageDesc(bool, class FGetImageDescContext*) const override;
 		void GetBlockLayoutSize(int blockIndex, int* pBlockX, int* pBlockY,
 			FBlockLayoutSizeCache* cache) override;
diff --git a/Engine/Plugins/Experimental/Mutable/Source/MutableTools/Private/MuT/ASTOpConstantString.cpp b/Engine/Plugins/Experimental/Mutable/Source/MutableTools/Private/MuT/ASTOpConstantString.cpp
--- a/Engine/Plugins/Experimental/Mutable/Source/MutableTools/Private/MuT/ASTOpConstantString.cpp
+++ b/Engine/Plugins/Experimental/Mutable/Source/MutableTools/Private/MuT/ASTOpConstantString.cpp
@@ -47,7 +47,7 @@ namespace mu
 	}
 
 
-	void ASTOpConstantString::Link(FProgram& program, const FLinkerOptions*)
+	void ASTOpConstantString::Link(FProgram& program, FLinkerOptions*)
 	{
 		if (!linkedAddress)
 		{
diff --git a/Engine/Plugins/Experimental/Mutable/Source/MutableTools/Private/MuT/ASTOpConstantString.h b/Engine/Plugins/Experimental/Mutable/Source/MutableTools/Private/MuT/ASTOpConstantString.h
--- a/Engine/Plugins/Experimental/Mutable/Source/MutableTools/Private/MuT/ASTOpConstantString.h
+++ b/Engine/Plugins/Experimental/Mutable/Source/MutableTools/Private/MuT/ASTOpConstantString.h
@@ -28,7 +28,7 @@ struct FProgram;
 		void ForEachChild(const TFunctionRef<void(ASTChild&)>) override;
 		bool IsEqual(const ASTOp& otherUntyped) const override;
 		Ptr<ASTOp> Clone(MapChildFuncRef mapChild) const override;
-		void Link(FProgram& program, const FLinkerOptions* Options) override;
+		void Link(FProgram& program, FLinkerOptions* Options) override;
 	};
 
 
diff --git a/Engine/Plugins/Experimental/Mutable/Source/MutableTools/Private/MuT/ASTOpImageCompose.cpp b/Engine/Plugins/Experimental/Mutable/Source/MutableTools/Private/MuT/ASTOpImageCompose.cpp
--- a/Engine/Plugins/Experimental/Mutable/Source/MutableTools/Private/MuT/ASTOpImageCompose.cpp
+++ b/Engine/Plugins/Experimental/Mutable/Source/MutableTools/Private/MuT/ASTOpImageCompose.cpp
@@ -95,7 +95,7 @@ void ASTOpImageCompose::ForEachChild(const TFunctionRef<void(ASTChild&)> f )
 
 
 //-------------------------------------------------------------------------------------------------
-void ASTOpImageCompose::Link( FProgram& program, const FLinkerOptions*)
+void ASTOpImageCompose::Link( FProgram& program, FLinkerOptions*)
 {
     // Already linked?
     if (!linkedAddress)
diff --git a/Engine/Plugins/Experimental/Mutable/Source/MutableTools/Private/MuT/ASTOpImageCompose.h b/Engine/Plugins/Experimental/Mutable/Source/MutableTools/Private/MuT/ASTOpImageCompose.h
--- a/Engine/Plugins/Experimental/Mutable/Source/MutableTools/Private/MuT/ASTOpImageCompose.h
+++ b/Engine/Plugins/Experimental/Mutable/Source/MutableTools/Private/MuT/ASTOpImageCompose.h
@@ -35,7 +35,7 @@ template <class SCALAR> class vec4;
 		bool IsEqual(const ASTOp& otherUntyped) const override;
 		Ptr<ASTOp> Clone(MapChildFuncRef mapChild) const override;
 		void ForEachChild(const TFunctionRef<void(ASTChild&)>) override;
-		void Link(FProgram& program, const FLinkerOptions* Options) override;
+		void Link(FProgram& program, FLinkerOptions* Options) override;
 		//TODO: Ptr<ASTOp> OptimiseSink(const FModelOptimizationOptions& options, FOptimizeSinkContext& context) const override;
 		Ptr<ASTOp> OptimiseSemantic(const FModelOptimizationOptions& options, int32 Pass) const;
 		FImageDesc GetImageDesc(bool returnBestOption, FGetImageDescContext* context) const override;
diff --git a/Engine/Plugins/Experimental/Mutable/Source/MutableTools/Private/MuT/ASTOpImageLayer.cpp b/Engine/Plugins/Experimental/Mutable/Source/MutableTools/Private/MuT/ASTOpImageLayer.cpp
--- a/Engine/Plugins/Experimental/Mutable/Source/MutableTools/Private/MuT/ASTOpImageLayer.cpp
+++ b/Engine/Plugins/Experimental/Mutable/Source/MutableTools/Private/MuT/ASTOpImageLayer.cpp
@@ -89,7 +89,7 @@ namespace mu
 
 
 	//-------------------------------------------------------------------------------------------------
-	void ASTOpImageLayer::Link(FProgram& program, const FLinkerOptions*)
+	void ASTOpImageLayer::Link(FProgram& program, FLinkerOptions*)
 	{
 		// Already linked?
 		if (!linkedAddress)
diff --git a/Engine/Plugins/Experimental/Mutable/Source/MutableTools/Private/MuT/ASTOpImageLayer.h b/Engine/Plugins/Experimental/Mutable/Source/MutableTools/Private/MuT/ASTOpImageLayer.h
--- a/Engine/Plugins/Experimental/Mutable/Source/MutableTools/Private/MuT/ASTOpImageLayer.h
+++ b/Engine/Plugins/Experimental/Mutable/Source/MutableTools/Private/MuT/ASTOpImageLayer.h
@@ -48,7 +48,7 @@ namespace mu
 		bool IsEqual(const ASTOp& otherUntyped) const override;
 		Ptr<ASTOp> Clone(MapChildFuncRef mapChild) const override;
 		void ForEachChild(const TFunctionRef<void(ASTChild&)>) override;
-		void Link(FProgram& program, const FLinkerOptions*) override;
+		void Link(FProgram& program, FLinkerOptions*) override;
 		FImageDesc GetImageDesc(bool returnBestOption, FGetImageDescContext*) const override;
 		void GetLayoutBlockSize(int* pBlockX, int* pBlockY) override;
 		Ptr<ImageSizeExpression> GetImageSizeExpression() const override;
diff --git a/Engine/Plugins/Experimental/Mutable/Source/MutableTools/Private/MuT/ASTOpImageLayerColor.cpp b/Engine/Plugins/Experimental/Mutable/Source/MutableTools/Private/MuT/ASTOpImageLayerColor.cpp
--- a/Engine/Plugins/Experimental/Mutable/Source/MutableTools/Private/MuT/ASTOpImageLayerColor.cpp
+++ b/Engine/Plugins/Experimental/Mutable/Source/MutableTools/Private/MuT/ASTOpImageLayerColor.cpp
@@ -85,7 +85,7 @@ namespace mu
 
 
 	//-------------------------------------------------------------------------------------------------
-	void ASTOpImageLayerColor::Link(FProgram& program, const FLinkerOptions*)
+	void ASTOpImageLayerColor::Link(FProgram& program, FLinkerOptions*)
 	{
 		// Already linked?
 		if (!linkedAddress)
diff --git a/Engine/Plugins/Experimental/Mutable/Source/MutableTools/Private/MuT/ASTOpImageLayerColor.h b/Engine/Plugins/Experimental/Mutable/Source/MutableTools/Private/MuT/ASTOpImageLayerColor.h
--- a/Engine/Plugins/Experimental/Mutable/Source/MutableTools/Private/MuT/ASTOpImageLayerColor.h
+++ b/Engine/Plugins/Experimental/Mutable/Source/MutableTools/Private/MuT/ASTOpImageLayerColor.h
@@ -48,7 +48,7 @@ namespace mu
 		bool IsEqual(const ASTOp& otherUntyped) const override;
 		Ptr<ASTOp> Clone(MapChildFuncRef mapChild) const override;
 		void ForEachChild(const TFunctionRef<void(ASTChild&)>) override;
-		void Link(FProgram& program, const FLinkerOptions* Options) override;
+		void Link(FProgram& program, FLinkerOptions* Options) override;
 		FImageDesc GetImageDesc(bool returnBestOption, FGetImageDescContext* context) const override;
 		void GetLayoutBlockSize(int* pBlockX, int* pBlockY) override;
 		Ptr<ImageSizeExpression> GetImageSizeExpression() const override;
diff --git a/Engine/Plugins/Experimental/Mutable/Source/MutableTools/Private/MuT/ASTOpImageMakeGrowMap.cpp b/Engine/Plugins/Experimental/Mutable/Source/MutableTools/Private/MuT/ASTOpImageMakeGrowMap.cpp
--- a/Engine/Plugins/Experimental/Mutable/Source/MutableTools/Private/MuT/ASTOpImageMakeGrowMap.cpp
+++ b/Engine/Plugins/Experimental/Mutable/Source/MutableTools/Private/MuT/ASTOpImageMakeGrowMap.cpp
@@ -72,7 +72,7 @@ namespace mu
 	}
 
 
-	void ASTOpImageMakeGrowMap::Link(FProgram& program, const FLinkerOptions*)
+	void ASTOpImageMakeGrowMap::Link(FProgram& program, FLinkerOptions*)
 	{
 		// Already linked?
 		if (!linkedAddress)
diff --git a/Engine/Plugins/Experimental/Mutable/Source/MutableTools/Private/MuT/ASTOpImageMakeGrowMap.h b/Engine/Plugins/Experimental/Mutable/Source/MutableTools/Private/MuT/ASTOpImageMakeGrowMap.h
--- a/Engine/Plugins/Experimental/Mutable/Source/MutableTools/Private/MuT/ASTOpImageMakeGrowMap.h
+++ b/Engine/Plugins/Experimental/Mutable/Source/MutableTools/Private/MuT/ASTOpImageMakeGrowMap.h
@@ -32,7 +32,7 @@ namespace mu
 		bool IsEqual(const ASTOp& otherUntyped) const override;
 		Ptr<ASTOp> Clone(MapChildFuncRef mapChild) const override;
 		void ForEachChild(const TFunctionRef<void(ASTChild&)>) override;
-		void Link(FProgram& program, const FLinkerOptions* Options) override;
+		void Link(FProgram& program, FLinkerOptions* Options) override;
 		Ptr<ASTOp> OptimiseSink(const FModelOptimizationOptions& options, FOptimizeSinkContext& context) const override;
 		FImageDesc GetImageDesc(bool returnBestOption, FGetImageDescContext* context) const override;
 		void GetLayoutBlockSize(int* pBlockX, int* pBlockY) override;
diff --git a/Engine/Plugins/Experimental/Mutable/Source/MutableTools/Private/MuT/ASTOpImageMipmap.cpp b/Engine/Plugins/Experimental/Mutable/Source/MutableTools/Private/MuT/ASTOpImageMipmap.cpp
--- a/Engine/Plugins/Experimental/Mutable/Source/MutableTools/Private/MuT/ASTOpImageMipmap.cpp
+++ b/Engine/Plugins/Experimental/Mutable/Source/MutableTools/Private/MuT/ASTOpImageMipmap.cpp
@@ -87,7 +87,7 @@ namespace mu
 	}
 
 
-	void ASTOpImageMipmap::Link(FProgram& program, const FLinkerOptions*)
+	void ASTOpImageMipmap::Link(FProgram& program, FLinkerOptions*)
 	{
 		// Already linked?
 		if (!linkedAddress)
diff --git a/Engine/Plugins/Experimental/Mutable/Source/MutableTools/Private/MuT/ASTOpImageMipmap.h b/Engine/Plugins/Experimental/Mutable/Source/MutableTools/Private/MuT/ASTOpImageMipmap.h
--- a/Engine/Plugins/Experimental/Mutable/Source/MutableTools/Private/MuT/ASTOpImageMipmap.h
+++ b/Engine/Plugins/Experimental/Mutable/Source/MutableTools/Private/MuT/ASTOpImageMipmap.h
@@ -48,7 +48,7 @@ namespace mu
 		bool IsEqual(const ASTOp& otherUntyped) const override;
 		Ptr<ASTOp> Clone(MapChildFuncRef mapChild) const override;
 		void ForEachChild(const TFunctionRef<void(ASTChild&)>) override;
-		void Link(FProgram& program, const FLinkerOptions* Options) override;
+		void Link(FProgram& program, FLinkerOptions* Options) override;
 		mu::Ptr<ASTOp> OptimiseSemantic(const FModelOptimizationOptions& options, int32 Pass) const override;
 		Ptr<ASTOp> OptimiseSink(const FModelOptimizationOptions& options, FOptimizeSinkContext& context) const override;
 		FImageDesc GetImageDesc(bool returnBestOption, FGetImageDescContext* context) const override;
diff --git a/Engine/Plugins/Experimental/Mutable/Source/MutableTools/Private/MuT/ASTOpImageMultiLayer.cpp b/Engine/Plugins/Experimental/Mutable/Source/MutableTools/Private/MuT/ASTOpImageMultiLayer.cpp
--- a/Engine/Plugins/Experimental/Mutable/Source/MutableTools/Private/MuT/ASTOpImageMultiLayer.cpp
+++ b/Engine/Plugins/Experimental/Mutable/Source/MutableTools/Private/MuT/ASTOpImageMultiLayer.cpp
@@ -92,7 +92,7 @@ namespace mu
 
 
 	//-------------------------------------------------------------------------------------------------
-	void ASTOpImageMultiLayer::Link(FProgram& program, const FLinkerOptions*)
+	void ASTOpImageMultiLayer::Link(FProgram& program, FLinkerOptions*)
 	{
 		// Already linked?
 		if (!linkedAddress)
diff --git a/Engine/Plugins/Experimental/Mutable/Source/MutableTools/Private/MuT/ASTOpImageMultiLayer.h b/Engine/Plugins/Experimental/Mutable/Source/MutableTools/Private/MuT/ASTOpImageMultiLayer.h
--- a/Engine/Plugins/Experimental/Mutable/Source/MutableTools/Private/MuT/ASTOpImageMultiLayer.h
+++ b/Engine/Plugins/Experimental/Mutable/Source/MutableTools/Private/MuT/ASTOpImageMultiLayer.h
@@ -49,7 +49,7 @@ namespace mu
 		bool IsEqual(const ASTOp& otherUntyped) const override;
 		Ptr<ASTOp> Clone(MapChildFuncRef mapChild) const override;
 		void ForEachChild(const TFunctionRef<void(ASTChild&)>) override;
-		void Link(FProgram& program, const FLinkerOptions* Options) override;
+		void Link(FProgram& program, FLinkerOptions* Options) override;
 		FImageDesc GetImageDesc(bool returnBestOption, FGetImageDescContext* context) const override;
 		void GetLayoutBlockSize(int* pBlockX, int* pBlockY) override;
 		Ptr<ImageSizeExpression> GetImageSizeExpression() const override;
diff --git a/Engine/Plugins/Experimental/Mutable/Source/MutableTools/Private/MuT/ASTOpImageNormalComposite.cpp b/Engine/Plugins/Experimental/Mutable/Source/MutableTools/Private/MuT/ASTOpImageNormalComposite.cpp
--- a/Engine/Plugins/Experimental/Mutable/Source/MutableTools/Private/MuT/ASTOpImageNormalComposite.cpp
+++ b/Engine/Plugins/Experimental/Mutable/Source/MutableTools/Private/MuT/ASTOpImageNormalComposite.cpp
@@ -70,7 +70,7 @@ namespace mu
 	}
 
 
-	void ASTOpImageNormalComposite::Link(FProgram& program, const FLinkerOptions*)
+	void ASTOpImageNormalComposite::Link(FProgram& program, FLinkerOptions*)
 	{
 		// Already linked?
 		if (!linkedAddress)
diff --git a/Engine/Plugins/Experimental/Mutable/Source/MutableTools/Private/MuT/ASTOpImageNormalComposite.h b/Engine/Plugins/Experimental/Mutable/Source/MutableTools/Private/MuT/ASTOpImageNormalComposite.h
--- a/Engine/Plugins/Experimental/Mutable/Source/MutableTools/Private/MuT/ASTOpImageNormalComposite.h
+++ b/Engine/Plugins/Experimental/Mutable/Source/MutableTools/Private/MuT/ASTOpImageNormalComposite.h
@@ -32,7 +32,7 @@ struct FProgram;
 		bool IsEqual(const ASTOp& otherUntyped) const override;
 		Ptr<ASTOp> Clone(MapChildFuncRef mapChild) const override;
 		void ForEachChild(const TFunctionRef<void(ASTChild&)>) override;
-		void Link(FProgram& program, const FLinkerOptions* Options) override;
+		void Link(FProgram& program, FLinkerOptions* Options) override;
 		FImageDesc GetImageDesc(bool returnBestOption, FGetImageDescContext* context) const override;
 		void GetLayoutBlockSize(int* pBlockX, int* pBlockY) override;
 		Ptr<ImageSizeExpression> GetImageSizeExpression() const override;
diff --git a/Engine/Plugins/Experimental/Mutable/Source/MutableTools/Private/MuT/ASTOpImagePatch.cpp b/Engine/Plugins/Experimental/Mutable/Source/MutableTools/Private/MuT/ASTOpImagePatch.cpp
--- a/Engine/Plugins/Experimental/Mutable/Source/MutableTools/Private/MuT/ASTOpImagePatch.cpp
+++ b/Engine/Plugins/Experimental/Mutable/Source/MutableTools/Private/MuT/ASTOpImagePatch.cpp
@@ -70,7 +70,7 @@ namespace mu
 
 
 	//-------------------------------------------------------------------------------------------------
-	void ASTOpImagePatch::Link(FProgram& program, const FLinkerOptions*)
+	void ASTOpImagePatch::Link(FProgram& program, FLinkerOptions*)
 	{
 		// Already linked?
 		if (!linkedAddress)
diff --git a/Engine/Plugins/Experimental/Mutable/Source/MutableTools/Private/MuT/ASTOpImagePatch.h b/Engine/Plugins/Experimental/Mutable/Source/MutableTools/Private/MuT/ASTOpImagePatch.h
--- a/Engine/Plugins/Experimental/Mutable/Source/MutableTools/Private/MuT/ASTOpImagePatch.h
+++ b/Engine/Plugins/Experimental/Mutable/Source/MutableTools/Private/MuT/ASTOpImagePatch.h
@@ -33,7 +33,7 @@ struct FProgram;
 		bool IsEqual(const ASTOp& otherUntyped) const override;
 		Ptr<ASTOp> Clone(MapChildFuncRef mapChild) const override;
 		void ForEachChild(const TFunctionRef<void(ASTChild&)>) override;
-		void Link(FProgram& program, const FLinkerOptions* Options) override;
+		void Link(FProgram& program, FLinkerOptions* Options) override;
 		FImageDesc GetImageDesc(bool returnBestOption, FGetImageDescContext* context) const override;
 		Ptr<ImageSizeExpression> GetImageSizeExpression() const override;
 		void GetLayoutBlockSize(int* pBlockX, int* pBlockY) override;
diff --git a/Engine/Plugins/Experimental/Mutable/Source/MutableTools/Private/MuT/ASTOpImagePixelFormat.cpp b/Engine/Plugins/Experimental/Mutable/Source/MutableTools/Private/MuT/ASTOpImagePixelFormat.cpp
--- a/Engine/Plugins/Experimental/Mutable/Source/MutableTools/Private/MuT/ASTOpImagePixelFormat.cpp
+++ b/Engine/Plugins/Experimental/Mutable/Source/MutableTools/Private/MuT/ASTOpImagePixelFormat.cpp
@@ -78,7 +78,7 @@ namespace mu
 	}
 
 
-	void ASTOpImagePixelFormat::Link(FProgram& program, const FLinkerOptions*)
+	void ASTOpImagePixelFormat::Link(FProgram& program, FLinkerOptions*)
 	{
 		// Already linked?
 		if (!linkedAddress)
diff --git a/Engine/Plugins/Experimental/Mutable/Source/MutableTools/Private/MuT/ASTOpImagePixelFormat.h b/Engine/Plugins/Experimental/Mutable/Source/MutableTools/Private/MuT/ASTOpImagePixelFormat.h
--- a/Engine/Plugins/Experimental/Mutable/Source/MutableTools/Private/MuT/ASTOpImagePixelFormat.h
+++ b/Engine/Plugins/Experimental/Mutable/Source/MutableTools/Private/MuT/ASTOpImagePixelFormat.h
@@ -33,7 +33,7 @@ template <class SCALAR> class vec4;
 		bool IsEqual(const ASTOp& otherUntyped) const override;
 		Ptr<ASTOp> Clone(MapChildFuncRef mapChild) const override;
 		void ForEachChild(const TFunctionRef<void(ASTChild&)>) override;
-		void Link(FProgram& program, const FLinkerOptions* Options) override;
+		void Link(FProgram& program, FLinkerOptions* Options) override;
 		Ptr<ASTOp> OptimiseSink(const FModelOptimizationOptions& options, FOptimizeSinkContext& context) const override;
 		FImageDesc GetImageDesc(bool returnBestOption, FGetImageDescContext* context) const override;
 		void GetLayoutBlockSize(int* pBlockX, int* pBlockY) override;
diff --git a/Engine/Plugins/Experimental/Mutable/Source/MutableTools/Private/MuT/ASTOpImageRasterMesh.cpp b/Engine/Plugins/Experimental/Mutable/Source/MutableTools/Private/MuT/ASTOpImageRasterMesh.cpp
--- a/Engine/Plugins/Experimental/Mutable/Source/MutableTools/Private/MuT/ASTOpImageRasterMesh.cpp
+++ b/Engine/Plugins/Experimental/Mutable/Source/MutableTools/Private/MuT/ASTOpImageRasterMesh.cpp
@@ -127,7 +127,7 @@ namespace mu
 
 
 	//-------------------------------------------------------------------------------------------------
-	void ASTOpImageRasterMesh::Link(FProgram& program, const FLinkerOptions*)
+	void ASTOpImageRasterMesh::Link(FProgram& program, FLinkerOptions*)
 	{
 		// Already linked?
 		if (!linkedAddress)
diff --git a/Engine/Plugins/Experimental/Mutable/Source/MutableTools/Private/MuT/ASTOpImageRasterMesh.h b/Engine/Plugins/Experimental/Mutable/Source/MutableTools/Private/MuT/ASTOpImageRasterMesh.h
--- a/Engine/Plugins/Experimental/Mutable/Source/MutableTools/Private/MuT/ASTOpImageRasterMesh.h
+++ b/Engine/Plugins/Experimental/Mutable/Source/MutableTools/Private/MuT/ASTOpImageRasterMesh.h
@@ -54,7 +54,7 @@ namespace mu
 		bool IsEqual(const ASTOp& otherUntyped) const override;
 		Ptr<ASTOp> Clone(MapChildFuncRef mapChild) const override;
 		void ForEachChild(const TFunctionRef<void(ASTChild&)>) override;
-		void Link(FProgram& program, const FLinkerOptions* Options) override;
+		void Link(FProgram& program, FLinkerOptions* Options) override;
 		FImageDesc GetImageDesc(bool returnBestOption, FGetImageDescContext* context) const override;
 		Ptr<ImageSizeExpression> GetImageSizeExpression() const override;
 		Ptr<ASTOp> OptimiseSemantic(const FModelOptimizationOptions& options, int32 Pass) const;
diff --git a/Engine/Plugins/Experimental/Mutable/Source/MutableTools/Private/MuT/ASTOpImageSwizzle.cpp b/Engine/Plugins/Experimental/Mutable/Source/MutableTools/Private/MuT/ASTOpImageSwizzle.cpp
--- a/Engine/Plugins/Experimental/Mutable/Source/MutableTools/Private/MuT/ASTOpImageSwizzle.cpp
+++ b/Engine/Plugins/Experimental/Mutable/Source/MutableTools/Private/MuT/ASTOpImageSwizzle.cpp
@@ -96,7 +96,7 @@ namespace mu
 	}
 
 
-	void ASTOpImageSwizzle::Link(FProgram& program, const FLinkerOptions*)
+	void ASTOpImageSwizzle::Link(FProgram& program, FLinkerOptions*)
 	{
 		// Already linked?
 		if (!linkedAddress)
diff --git a/Engine/Plugins/Experimental/Mutable/Source/MutableTools/Private/MuT/ASTOpImageSwizzle.h b/Engine/Plugins/Experimental/Mutable/Source/MutableTools/Private/MuT/ASTOpImageSwizzle.h
--- a/Engine/Plugins/Experimental/Mutable/Source/MutableTools/Private/MuT/ASTOpImageSwizzle.h
+++ b/Engine/Plugins/Experimental/Mutable/Source/MutableTools/Private/MuT/ASTOpImageSwizzle.h
@@ -34,7 +34,7 @@ namespace mu
 		bool IsEqual(const ASTOp& otherUntyped) const override;
 		Ptr<ASTOp> Clone(MapChildFuncRef mapChild) const override;
 		void ForEachChild(const TFunctionRef<void(ASTChild&)>) override;
-		void Link(FProgram& program, const FLinkerOptions* Options) override;
+		void Link(FProgram& program, FLinkerOptions* Options) override;
 		Ptr<ASTOp> OptimiseSemantic(const FModelOptimizationOptions&, int32 Pass) const override;
 		Ptr<ASTOp> OptimiseSink(const FModelOptimizationOptions& options, FOptimizeSinkContext& context) const override;
 		FImageDesc GetImageDesc(bool returnBestOption, FGetImageDescContext* context) const override;
diff --git a/Engine/Plugins/Experimental/Mutable/Source/MutableTools/Private/MuT/ASTOpImageTransform.cpp b/Engine/Plugins/Experimental/Mutable/Source/MutableTools/Private/MuT/ASTOpImageTransform.cpp
--- a/Engine/Plugins/Experimental/Mutable/Source/MutableTools/Private/MuT/ASTOpImageTransform.cpp
+++ b/Engine/Plugins/Experimental/Mutable/Source/MutableTools/Private/MuT/ASTOpImageTransform.cpp
@@ -96,7 +96,7 @@ namespace mu
 
 
 	//-------------------------------------------------------------------------------------------------
-	void ASTOpImageTransform::Link(FProgram& program, const FLinkerOptions*)
+	void ASTOpImageTransform::Link(FProgram& program, FLinkerOptions*)
 	{
 		// Already linked?
 		if (!linkedAddress)
diff --git a/Engine/Plugins/Experimental/Mutable/Source/MutableTools/Private/MuT/ASTOpImageTransform.h b/Engine/Plugins/Experimental/Mutable/Source/MutableTools/Private/MuT/ASTOpImageTransform.h
--- a/Engine/Plugins/Experimental/Mutable/Source/MutableTools/Private/MuT/ASTOpImageTransform.h
+++ b/Engine/Plugins/Experimental/Mutable/Source/MutableTools/Private/MuT/ASTOpImageTransform.h
@@ -41,7 +41,7 @@ struct FProgram;
 		bool IsEqual(const ASTOp& otherUntyped) const override;
 		Ptr<ASTOp> Clone(MapChildFuncRef mapChild) const override;
 		void ForEachChild(const TFunctionRef<void(ASTChild&)>) override;
-		void Link(FProgram& program, const FLinkerOptions* Options) override;
+		void Link(FProgram& program, FLinkerOptions* Options) override;
 		FImageDesc GetImageDesc(bool returnBestOption, FGetImageDescContext* context) const override;
 		void GetLayoutBlockSize(int* pBlockX, int* pBlockY) override;
 		Ptr<ImageSizeExpression> GetImageSizeExpression() const override;
diff --git a/Engine/Plugins/Experimental/Mutable/Source/MutableTools/Private/MuT/ASTOpInstanceAdd.cpp b/Engine/Plugins/Experimental/Mutable/Source/MutableTools/Private/MuT/ASTOpInstanceAdd.cpp
--- a/Engine/Plugins/Experimental/Mutable/Source/MutableTools/Private/MuT/ASTOpInstanceAdd.cpp
+++ b/Engine/Plugins/Experimental/Mutable/Source/MutableTools/Private/MuT/ASTOpInstanceAdd.cpp
@@ -110,7 +110,7 @@ namespace mu
 
 
 	//-------------------------------------------------------------------------------------------------
-	void ASTOpInstanceAdd::Link(FProgram& program, const FLinkerOptions*)
+	void ASTOpInstanceAdd::Link(FProgram& program, FLinkerOptions*)
 	{
 		// Already linked?
 		if (!linkedAddress)
diff --git a/Engine/Plugins/Experimental/Mutable/Source/MutableTools/Private/MuT/ASTOpInstanceAdd.h b/Engine/Plugins/Experimental/Mutable/Source/MutableTools/Private/MuT/ASTOpInstanceAdd.h
--- a/Engine/Plugins/Experimental/Mutable/Source/MutableTools/Private/MuT/ASTOpInstanceAdd.h
+++ b/Engine/Plugins/Experimental/Mutable/Source/MutableTools/Private/MuT/ASTOpInstanceAdd.h
@@ -43,7 +43,7 @@ struct FProgram;
 		Ptr<ASTOp> Clone(MapChildFuncRef mapChild) const override;
 		uint64 Hash() const override;
 		void Assert() override;
-		void Link(FProgram& program, const FLinkerOptions* Options) override;
+		void Link(FProgram& program, FLinkerOptions* Options) override;
 	};
 
 
diff --git a/Engine/Plugins/Experimental/Mutable/Source/MutableTools/Private/MuT/ASTOpLayoutFromMesh.cpp b/Engine/Plugins/Experimental/Mutable/Source/MutableTools/Private/MuT/ASTOpLayoutFromMesh.cpp
--- a/Engine/Plugins/Experimental/Mutable/Source/MutableTools/Private/MuT/ASTOpLayoutFromMesh.cpp
+++ b/Engine/Plugins/Experimental/Mutable/Source/MutableTools/Private/MuT/ASTOpLayoutFromMesh.cpp
@@ -77,7 +77,7 @@ namespace mu
 	}
 
 
-	void ASTOpLayoutFromMesh::Link(FProgram& program, const FLinkerOptions*)
+	void ASTOpLayoutFromMesh::Link(FProgram& program, FLinkerOptions*)
 	{
 		// Already linked?
 		if (!linkedAddress)
diff --git a/Engine/Plugins/Experimental/Mutable/Source/MutableTools/Private/MuT/ASTOpLayoutFromMesh.h b/Engine/Plugins/Experimental/Mutable/Source/MutableTools/Private/MuT/ASTOpLayoutFromMesh.h
--- a/Engine/Plugins/Experimental/Mutable/Source/MutableTools/Private/MuT/ASTOpLayoutFromMesh.h
+++ b/Engine/Plugins/Experimental/Mutable/Source/MutableTools/Private/MuT/ASTOpLayoutFromMesh.h
@@ -30,7 +30,7 @@ namespace mu
 		bool IsEqual(const ASTOp& otherUntyped) const override;
 		Ptr<ASTOp> Clone(MapChildFuncRef mapChild) const override;
 		void ForEachChild(const TFunctionRef<void(ASTChild&)>) override;
-		void Link(FProgram& program, const FLinkerOptions* Options) override;
+		void Link(FProgram& program, FLinkerOptions* Options) override;
 		Ptr<ASTOp> OptimiseSink(const FModelOptimizationOptions& options, FOptimizeSinkContext& context) const override;
 		void GetBlockLayoutSize(int blockIndex, int* pBlockX, int* pBlockY, FBlockLayoutSizeCache* cache) override;
 
diff --git a/Engine/Plugins/Experimental/Mutable/Source/MutableTools/Private/MuT/ASTOpLayoutMerge.cpp b/Engine/Plugins/Experimental/Mutable/Source/MutableTools/Private/MuT/ASTOpLayoutMerge.cpp
--- a/Engine/Plugins/Experimental/Mutable/Source/MutableTools/Private/MuT/ASTOpLayoutMerge.cpp
+++ b/Engine/Plugins/Experimental/Mutable/Source/MutableTools/Private/MuT/ASTOpLayoutMerge.cpp
@@ -74,7 +74,7 @@ namespace mu
 	}
 
 
-	void ASTOpLayoutMerge::Link(FProgram& program, const FLinkerOptions*)
+	void ASTOpLayoutMerge::Link(FProgram& program, FLinkerOptions*)
 	{
 		// Already linked?
 		if (!linkedAddress)
diff --git a/Engine/Plugins/Experimental/Mutable/Source/MutableTools/Private/MuT/ASTOpLayoutMerge.h b/Engine/Plugins/Experimental/Mutable/Source/MutableTools/Private/MuT/ASTOpLayoutMerge.h
--- a/Engine/Plugins/Experimental/Mutable/Source/MutableTools/Private/MuT/ASTOpLayoutMerge.h
+++ b/Engine/Plugins/Experimental/Mutable/Source/MutableTools/Private/MuT/ASTOpLayoutMerge.h
@@ -30,7 +30,7 @@ namespace mu
 		bool IsEqual(const ASTOp& otherUntyped) const override;
 		Ptr<ASTOp> Clone(MapChildFuncRef mapChild) const override;
 		void ForEachChild(const TFunctionRef<void(ASTChild&)>) override;
-		void Link(FProgram& program, const FLinkerOptions* Options) override;
+		void Link(FProgram& program, FLinkerOptions* Options) override;
 		void GetBlockLayoutSize(int blockIndex, int* pBlockX, int* pBlockY, FBlockLayoutSizeCache* cache) override;
 
 	};
diff --git a/Engine/Plugins/Experimental/Mutable/Source/MutableTools/Private/MuT/ASTOpLayoutPack.cpp b/Engine/Plugins/Experimental/Mutable/Source/MutableTools/Private/MuT/ASTOpLayoutPack.cpp
--- a/Engine/Plugins/Experimental/Mutable/Source/MutableTools/Private/MuT/ASTOpLayoutPack.cpp
+++ b/Engine/Plugins/Experimental/Mutable/Source/MutableTools/Private/MuT/ASTOpLayoutPack.cpp
@@ -70,7 +70,7 @@ namespace mu
 	}
 
 
-	void ASTOpLayoutPack::Link(FProgram& program, const FLinkerOptions*)
+	void ASTOpLayoutPack::Link(FProgram& program, FLinkerOptions*)
 	{
 		// Already linked?
 		if (!linkedAddress)
diff --git a/Engine/Plugins/Experimental/Mutable/Source/MutableTools/Private/MuT/ASTOpLayoutPack.h b/Engine/Plugins/Experimental/Mutable/Source/MutableTools/Private/MuT/ASTOpLayoutPack.h
--- a/Engine/Plugins/Experimental/Mutable/Source/MutableTools/Private/MuT/ASTOpLayoutPack.h
+++ b/Engine/Plugins/Experimental/Mutable/Source/MutableTools/Private/MuT/ASTOpLayoutPack.h
@@ -29,7 +29,7 @@ namespace mu
 		bool IsEqual(const ASTOp& otherUntyped) const override;
 		Ptr<ASTOp> Clone(MapChildFuncRef mapChild) const override;
 		void ForEachChild(const TFunctionRef<void(ASTChild&)>) override;
-		void Link(FProgram& program, const FLinkerOptions* Options) override;
+		void Link(FProgram& program, FLinkerOptions* Options) override;
 		void GetBlockLayoutSize(int blockIndex, int* pBlockX, int* pBlockY, FBlockLayoutSizeCache* cache) override;
 
 	};
diff --git a/Engine/Plugins/Experimental/Mutable/Source/MutableTools/Private/MuT/ASTOpLayoutRemoveBlocks.cpp b/Engine/Plugins/Experimental/Mutable/Source/MutableTools/Private/MuT/ASTOpLayoutRemoveBlocks.cpp
--- a/Engine/Plugins/Experimental/Mutable/Source/MutableTools/Private/MuT/ASTOpLayoutRemoveBlocks.cpp
+++ b/Engine/Plugins/Experimental/Mutable/Source/MutableTools/Private/MuT/ASTOpLayoutRemoveBlocks.cpp
@@ -74,7 +74,7 @@ namespace mu
 	}
 
 
-	void ASTOpLayoutRemoveBlocks::Link(FProgram& program, const FLinkerOptions*)
+	void ASTOpLayoutRemoveBlocks::Link(FProgram& program, FLinkerOptions*)
 	{
 		// Already linked?
 		if (!linkedAddress)
diff --git a/Engine/Plugins/Experimental/Mutable/Source/MutableTools/Private/MuT/ASTOpLayoutRemoveBlocks.h b/Engine/Plugins/Experimental/Mutable/Source/MutableTools/Private/MuT/ASTOpLayoutRemoveBlocks.h
--- a/Engine/Plugins/Experimental/Mutable/Source/MutableTools/Private/MuT/ASTOpLayoutRemoveBlocks.h
+++ b/Engine/Plugins/Experimental/Mutable/Source/MutableTools/Private/MuT/ASTOpLayoutRemoveBlocks.h
@@ -30,7 +30,7 @@ namespace mu
 		bool IsEqual(const ASTOp& otherUntyped) const override;
 		Ptr<ASTOp> Clone(MapChildFuncRef mapChild) const override;
 		void ForEachChild(const TFunctionRef<void(ASTChild&)>) override;
-		void Link(FProgram& program, const FLinkerOptions* Options) override;
+		void Link(FProgram& program, FLinkerOptions* Options) override;
 		void GetBlockLayoutSize(int blockIndex, int* pBlockX, int* pBlockY, FBlockLayoutSizeCache* cache) override;
 
 	};
diff --git a/Engine/Plugins/Experimental/Mutable/Source/MutableTools/Private/MuT/ASTOpMeshApplyPose.cpp b/Engine/Plugins/Experimental/Mutable/Source/MutableTools/Private/MuT/ASTOpMeshApplyPose.cpp
--- a/Engine/Plugins/Experimental/Mutable/Source/MutableTools/Private/MuT/ASTOpMeshApplyPose.cpp
+++ b/Engine/Plugins/Experimental/Mutable/Source/MutableTools/Private/MuT/ASTOpMeshApplyPose.cpp
@@ -66,7 +66,7 @@ namespace mu
 	}
 
 
-	void ASTOpMeshApplyPose::Link(FProgram& program, const FLinkerOptions*)
+	void ASTOpMeshApplyPose::Link(FProgram& program, FLinkerOptions*)
 	{
 		// Already linked?
 		if (!linkedAddress)
diff --git a/Engine/Plugins/Experimental/Mutable/Source/MutableTools/Private/MuT/ASTOpMeshApplyPose.h b/Engine/Plugins/Experimental/Mutable/Source/MutableTools/Private/MuT/ASTOpMeshApplyPose.h
--- a/Engine/Plugins/Experimental/Mutable/Source/MutableTools/Private/MuT/ASTOpMeshApplyPose.h
+++ b/Engine/Plugins/Experimental/Mutable/Source/MutableTools/Private/MuT/ASTOpMeshApplyPose.h
@@ -34,7 +34,7 @@ struct FProgram;
 		bool IsEqual(const ASTOp& otherUntyped) const override;
 		Ptr<ASTOp> Clone(MapChildFuncRef mapChild) const override;
 		void ForEachChild(const TFunctionRef<void(ASTChild&)>) override;
-		void Link(FProgram& program, const FLinkerOptions* Options) override;
+		void Link(FProgram& program, FLinkerOptions* Options) override;
 	};
 
 
diff --git a/Engine/Plugins/Experimental/Mutable/Source/MutableTools/Private/MuT/ASTOpMeshApplyShape.cpp b/Engine/Plugins/Experimental/Mutable/Source/MutableTools/Private/MuT/ASTOpMeshApplyShape.cpp
--- a/Engine/Plugins/Experimental/Mutable/Source/MutableTools/Private/MuT/ASTOpMeshApplyShape.cpp
+++ b/Engine/Plugins/Experimental/Mutable/Source/MutableTools/Private/MuT/ASTOpMeshApplyShape.cpp
@@ -78,7 +78,7 @@ namespace mu
 	}
 
 
-	void ASTOpMeshApplyShape::Link(FProgram& program, const FLinkerOptions*)
+	void ASTOpMeshApplyShape::Link(FProgram& program, FLinkerOptions*)
 	{
 		// Already linked?
 		if (!linkedAddress)
diff --git a/Engine/Plugins/Experimental/Mutable/Source/MutableTools/Private/MuT/ASTOpMeshApplyShape.h b/Engine/Plugins/Experimental/Mutable/Source/MutableTools/Private/MuT/ASTOpMeshApplyShape.h
--- a/Engine/Plugins/Experimental/Mutable/Source/MutableTools/Private/MuT/ASTOpMeshApplyShape.h
+++ b/Engine/Plugins/Experimental/Mutable/Source/MutableTools/Private/MuT/ASTOpMeshApplyShape.h
@@ -37,7 +37,7 @@ struct FProgram;
 		bool IsEqual(const ASTOp& otherUntyped) const override;
 		Ptr<ASTOp> Clone(MapChildFuncRef mapChild) const override;
 		void ForEachChild(const TFunctionRef<void(ASTChild&)>) override;
-		void Link(FProgram& program, const FLinkerOptions* Options) override;
+		void Link(FProgram& program, FLinkerOptions* Options) override;
 	};
 
 }
diff --git a/Engine/Plugins/Experimental/Mutable/Source/MutableTools/Private/MuT/ASTOpMeshBindShape.cpp b/Engine/Plugins/Experimental/Mutable/Source/MutableTools/Private/MuT/ASTOpMeshBindShape.cpp
--- a/Engine/Plugins/Experimental/Mutable/Source/MutableTools/Private/MuT/ASTOpMeshBindShape.cpp
+++ b/Engine/Plugins/Experimental/Mutable/Source/MutableTools/Private/MuT/ASTOpMeshBindShape.cpp
@@ -107,7 +107,7 @@ namespace mu
 	}
 
 
-	void ASTOpMeshBindShape::Link(FProgram& program, const FLinkerOptions*)
+	void ASTOpMeshBindShape::Link(FProgram& program, FLinkerOptions*)
 	{
 		// Already linked?
 		if (!linkedAddress)
diff --git a/Engine/Plugins/Experimental/Mutable/Source/MutableTools/Private/MuT/ASTOpMeshBindShape.h b/Engine/Plugins/Experimental/Mutable/Source/MutableTools/Private/MuT/ASTOpMeshBindShape.h
--- a/Engine/Plugins/Experimental/Mutable/Source/MutableTools/Private/MuT/ASTOpMeshBindShape.h
+++ b/Engine/Plugins/Experimental/Mutable/Source/MutableTools/Private/MuT/ASTOpMeshBindShape.h
@@ -45,7 +45,7 @@ struct FProgram;
 		bool IsEqual(const ASTOp& otherUntyped) const override;
 		Ptr<ASTOp> Clone(MapChildFuncRef mapChild) const override;
 		void ForEachChild(const TFunctionRef<void(ASTChild&)>) override;
-		void Link(FProgram& program, const FLinkerOptions* Options) override;
+		void Link(FProgram& program, FLinkerOptions* Options) override;
 		Ptr<ASTOp> OptimiseSink(const FModelOptimizationOptions&, FOptimizeSinkContext&) const override;
 
 	};
diff --git a/Engine/Plugins/Experimental/Mutable/Source/MutableTools/Private/MuT/ASTOpMeshClipDeform.cpp b/Engine/Plugins/Experimental/Mutable/Source/MutableTools/Private/MuT/ASTOpMeshClipDeform.cpp
--- a/Engine/Plugins/Experimental/Mutable/Source/MutableTools/Private/MuT/ASTOpMeshClipDeform.cpp
+++ b/Engine/Plugins/Experimental/Mutable/Source/MutableTools/Private/MuT/ASTOpMeshClipDeform.cpp
@@ -65,7 +65,7 @@ namespace mu
 	}
 
 
-	void ASTOpMeshClipDeform::Link(FProgram& program, const FLinkerOptions*)
+	void ASTOpMeshClipDeform::Link(FProgram& program, FLinkerOptions*)
 	{
 		// Already linked?
 		if (!linkedAddress)
diff --git a/Engine/Plugins/Experimental/Mutable/Source/MutableTools/Private/MuT/ASTOpMeshClipDeform.h b/Engine/Plugins/Experimental/Mutable/Source/MutableTools/Private/MuT/ASTOpMeshClipDeform.h
--- a/Engine/Plugins/Experimental/Mutable/Source/MutableTools/Private/MuT/ASTOpMeshClipDeform.h
+++ b/Engine/Plugins/Experimental/Mutable/Source/MutableTools/Private/MuT/ASTOpMeshClipDeform.h
@@ -34,7 +34,7 @@ struct FProgram;
 		bool IsEqual(const ASTOp& otherUntyped) const override;
 		Ptr<ASTOp> Clone(MapChildFuncRef mapChild) const override;
 		void ForEachChild(const TFunctionRef<void(ASTChild&)>) override;
-		void Link(FProgram& program, const FLinkerOptions* Options) override;
+		void Link(FProgram& program, FLinkerOptions* Options) override;
 	};
 
 
diff --git a/Engine/Plugins/Experimental/Mutable/Source/MutableTools/Private/MuT/ASTOpMeshClipMorphPlane.cpp b/Engine/Plugins/Experimental/Mutable/Source/MutableTools/Private/MuT/ASTOpMeshClipMorphPlane.cpp
--- a/Engine/Plugins/Experimental/Mutable/Source/MutableTools/Private/MuT/ASTOpMeshClipMorphPlane.cpp
+++ b/Engine/Plugins/Experimental/Mutable/Source/MutableTools/Private/MuT/ASTOpMeshClipMorphPlane.cpp
@@ -77,7 +77,7 @@ namespace mu
 	}
 
 
-	void ASTOpMeshClipMorphPlane::Link(FProgram& program, const FLinkerOptions*)
+	void ASTOpMeshClipMorphPlane::Link(FProgram& program, FLinkerOptions*)
 	{
 		// Already linked?
 		if (!linkedAddress)
diff --git a/Engine/Plugins/Experimental/Mutable/Source/MutableTools/Private/MuT/ASTOpMeshClipMorphPlane.h b/Engine/Plugins/Experimental/Mutable/Source/MutableTools/Private/MuT/ASTOpMeshClipMorphPlane.h
--- a/Engine/Plugins/Experimental/Mutable/Source/MutableTools/Private/MuT/ASTOpMeshClipMorphPlane.h
+++ b/Engine/Plugins/Experimental/Mutable/Source/MutableTools/Private/MuT/ASTOpMeshClipMorphPlane.h
@@ -41,7 +41,7 @@ struct FProgram;
 		bool IsEqual(const ASTOp& otherUntyped) const override;
 		Ptr<ASTOp> Clone(MapChildFuncRef mapChild) const override;
 		void ForEachChild(const TFunctionRef<void(ASTChild&)>) override;
-		void Link(FProgram& program, const FLinkerOptions* Options) override;
+		void Link(FProgram& program, FLinkerOptions* Options) override;
 	};
 
 
diff --git a/Engine/Plugins/Experimental/Mutable/Source/MutableTools/Private/MuT/ASTOpMeshDifference.cpp b/Engine/Plugins/Experimental/Mutable/Source/MutableTools/Private/MuT/ASTOpMeshDifference.cpp
--- a/Engine/Plugins/Experimental/Mutable/Source/MutableTools/Private/MuT/ASTOpMeshDifference.cpp
+++ b/Engine/Plugins/Experimental/Mutable/Source/MutableTools/Private/MuT/ASTOpMeshDifference.cpp
@@ -75,7 +75,7 @@ namespace mu
 
 
 	//---------------------------------------------------------------------------------------------
-	void ASTOpMeshDifference::Link(FProgram& program, const FLinkerOptions*)
+	void ASTOpMeshDifference::Link(FProgram& program, FLinkerOptions*)
 	{
 		// Already linked?
 		if (!linkedAddress)
diff --git a/Engine/Plugins/Experimental/Mutable/Source/MutableTools/Private/MuT/ASTOpMeshDifference.h b/Engine/Plugins/Experimental/Mutable/Source/MutableTools/Private/MuT/ASTOpMeshDifference.h
--- a/Engine/Plugins/Experimental/Mutable/Source/MutableTools/Private/MuT/ASTOpMeshDifference.h
+++ b/Engine/Plugins/Experimental/Mutable/Source/MutableTools/Private/MuT/ASTOpMeshDifference.h
@@ -51,7 +51,7 @@ namespace mu
 		void ForEachChild(const TFunctionRef<void(ASTChild&)>) override;
 		bool IsEqual(const ASTOp& otherUntyped) const override;
 		Ptr<ASTOp> Clone(MapChildFuncRef mapChild) const override;
-		void Link(FProgram& program, const FLinkerOptions*) override;
+		void Link(FProgram& program, FLinkerOptions*) override;
 		Ptr<ASTOp> OptimiseSink(const FModelOptimizationOptions&, FOptimizeSinkContext&) const override;
 
 	};
diff --git a/Engine/Plugins/Experimental/Mutable/Source/MutableTools/Private/MuT/ASTOpMeshExtractLayoutBlocks.cpp b/Engine/Plugins/Experimental/Mutable/Source/MutableTools/Private/MuT/ASTOpMeshExtractLayoutBlocks.cpp
--- a/Engine/Plugins/Experimental/Mutable/Source/MutableTools/Private/MuT/ASTOpMeshExtractLayoutBlocks.cpp
+++ b/Engine/Plugins/Experimental/Mutable/Source/MutableTools/Private/MuT/ASTOpMeshExtractLayoutBlocks.cpp
@@ -69,7 +69,7 @@ namespace mu
 	}
 
 
-	void ASTOpMeshExtractLayoutBlocks::Link(FProgram& program, const FLinkerOptions*)
+	void ASTOpMeshExtractLayoutBlocks::Link(FProgram& program, FLinkerOptions*)
 	{
 		// Already linked?
 		if (!linkedAddress)
diff --git a/Engine/Plugins/Experimental/Mutable/Source/MutableTools/Private/MuT/ASTOpMeshExtractLayoutBlocks.h b/Engine/Plugins/Experimental/Mutable/Source/MutableTools/Private/MuT/ASTOpMeshExtractLayoutBlocks.h
--- a/Engine/Plugins/Experimental/Mutable/Source/MutableTools/Private/MuT/ASTOpMeshExtractLayoutBlocks.h
+++ b/Engine/Plugins/Experimental/Mutable/Source/MutableTools/Private/MuT/ASTOpMeshExtractLayoutBlocks.h
@@ -42,7 +42,7 @@ struct FProgram;
 		bool IsEqual(const ASTOp& otherUntyped) const override;
 		Ptr<ASTOp> Clone(MapChildFuncRef mapChild) const override;
 		void Assert() override;
-		void Link(FProgram& program, const FLinkerOptions* Options) override;
+		void Link(FProgram& program, FLinkerOptions* Options) override;
 
 	};
 
diff --git a/Engine/Plugins/Experimental/Mutable/Source/MutableTools/Private/MuT/ASTOpMeshFormat.cpp b/Engine/Plugins/Experimental/Mutable/Source/MutableTools/Private/MuT/ASTOpMeshFormat.cpp
--- a/Engine/Plugins/Experimental/Mutable/Source/MutableTools/Private/MuT/ASTOpMeshFormat.cpp
+++ b/Engine/Plugins/Experimental/Mutable/Source/MutableTools/Private/MuT/ASTOpMeshFormat.cpp
@@ -78,7 +78,7 @@ void ASTOpMeshFormat::ForEachChild(const TFunctionRef<void(ASTChild&)> f )
 }
 
 
-void ASTOpMeshFormat::Link( FProgram& program, const FLinkerOptions* )
+void ASTOpMeshFormat::Link( FProgram& program, FLinkerOptions* )
 {
     // Already linked?
     if (!linkedAddress)
diff --git a/Engine/Plugins/Experimental/Mutable/Source/MutableTools/Private/MuT/ASTOpMeshFormat.h b/Engine/Plugins/Experimental/Mutable/Source/MutableTools/Private/MuT/ASTOpMeshFormat.h
--- a/Engine/Plugins/Experimental/Mutable/Source/MutableTools/Private/MuT/ASTOpMeshFormat.h
+++ b/Engine/Plugins/Experimental/Mutable/Source/MutableTools/Private/MuT/ASTOpMeshFormat.h
@@ -31,7 +31,7 @@ namespace mu
 		bool IsEqual(const ASTOp& otherUntyped) const override;
 		Ptr<ASTOp> Clone(MapChildFuncRef mapChild) const override;
 		void ForEachChild(const TFunctionRef<void(ASTChild&)>) override;
-		void Link(FProgram& program, const FLinkerOptions* Options) override;
+		void Link(FProgram& program, FLinkerOptions* Options) override;
 		Ptr<ASTOp> OptimiseSink(const FModelOptimizationOptions& options, FOptimizeSinkContext& context) const override;
 
 	};
diff --git a/Engine/Plugins/Experimental/Mutable/Source/MutableTools/Private/MuT/ASTOpMeshGeometryOperation.cpp b/Engine/Plugins/Experimental/Mutable/Source/MutableTools/Private/MuT/ASTOpMeshGeometryOperation.cpp
--- a/Engine/Plugins/Experimental/Mutable/Source/MutableTools/Private/MuT/ASTOpMeshGeometryOperation.cpp
+++ b/Engine/Plugins/Experimental/Mutable/Source/MutableTools/Private/MuT/ASTOpMeshGeometryOperation.cpp
@@ -70,7 +70,7 @@ namespace mu
 	}
 
 
-	void ASTOpMeshGeometryOperation::Link(FProgram& program, const FLinkerOptions*)
+	void ASTOpMeshGeometryOperation::Link(FProgram& program, FLinkerOptions*)
 	{
 		// Already linked?
 		if (!linkedAddress)
diff --git a/Engine/Plugins/Experimental/Mutable/Source/MutableTools/Private/MuT/ASTOpMeshGeometryOperation.h b/Engine/Plugins/Experimental/Mutable/Source/MutableTools/Private/MuT/ASTOpMeshGeometryOperation.h
--- a/Engine/Plugins/Experimental/Mutable/Source/MutableTools/Private/MuT/ASTOpMeshGeometryOperation.h
+++ b/Engine/Plugins/Experimental/Mutable/Source/MutableTools/Private/MuT/ASTOpMeshGeometryOperation.h
@@ -36,7 +36,7 @@ struct FProgram;
 		bool IsEqual(const ASTOp& otherUntyped) const override;
 		Ptr<ASTOp> Clone(MapChildFuncRef mapChild) const override;
 		void ForEachChild(const TFunctionRef<void(ASTChild&)>) override;
-		void Link(FProgram& program, const FLinkerOptions* Options) override;
+		void Link(FProgram& program, FLinkerOptions* Options) override;
 	};
 
 
diff --git a/Engine/Plugins/Experimental/Mutable/Source/MutableTools/Private/MuT/ASTOpMeshMaskClipMesh.cpp b/Engine/Plugins/Experimental/Mutable/Source/MutableTools/Private/MuT/ASTOpMeshMaskClipMesh.cpp
--- a/Engine/Plugins/Experimental/Mutable/Source/MutableTools/Private/MuT/ASTOpMeshMaskClipMesh.cpp
+++ b/Engine/Plugins/Experimental/Mutable/Source/MutableTools/Private/MuT/ASTOpMeshMaskClipMesh.cpp
@@ -73,7 +73,7 @@ namespace mu
 
 
 	//-------------------------------------------------------------------------------------------------
-	void ASTOpMeshMaskClipMesh::Link(FProgram& program, const FLinkerOptions*)
+	void ASTOpMeshMaskClipMesh::Link(FProgram& program, FLinkerOptions*)
 	{
 		// Already linked?
 		if (!linkedAddress)
diff --git a/Engine/Plugins/Experimental/Mutable/Source/MutableTools/Private/MuT/ASTOpMeshMaskClipMesh.h b/Engine/Plugins/Experimental/Mutable/Source/MutableTools/Private/MuT/ASTOpMeshMaskClipMesh.h
--- a/Engine/Plugins/Experimental/Mutable/Source/MutableTools/Private/MuT/ASTOpMeshMaskClipMesh.h
+++ b/Engine/Plugins/Experimental/Mutable/Source/MutableTools/Private/MuT/ASTOpMeshMaskClipMesh.h
@@ -34,7 +34,7 @@ struct FProgram;
 		bool IsEqual(const ASTOp& otherUntyped) const override;
 		Ptr<ASTOp> Clone(MapChildFuncRef mapChild) const override;
 		void ForEachChild(const TFunctionRef<void(ASTChild&)>) override;
-		void Link(FProgram& program, const FLinkerOptions* Options) override;
+		void Link(FProgram& program, FLinkerOptions* Options) override;
 		Ptr<ASTOp> OptimiseSink(const FModelOptimizationOptions&, FOptimizeSinkContext&) const override;
 
 	};
diff --git a/Engine/Plugins/Experimental/Mutable/Source/MutableTools/Private/MuT/ASTOpMeshMorph.cpp b/Engine/Plugins/Experimental/Mutable/Source/MutableTools/Private/MuT/ASTOpMeshMorph.cpp
--- a/Engine/Plugins/Experimental/Mutable/Source/MutableTools/Private/MuT/ASTOpMeshMorph.cpp
+++ b/Engine/Plugins/Experimental/Mutable/Source/MutableTools/Private/MuT/ASTOpMeshMorph.cpp
@@ -70,7 +70,7 @@ namespace mu
 
 
 	//---------------------------------------------------------------------------------------------
-	void ASTOpMeshMorph::Link(FProgram& program, const FLinkerOptions*)
+	void ASTOpMeshMorph::Link(FProgram& program, FLinkerOptions*)
 	{
 		// Already linked?
 		if (!linkedAddress)
diff --git a/Engine/Plugins/Experimental/Mutable/Source/MutableTools/Private/MuT/ASTOpMeshMorph.h b/Engine/Plugins/Experimental/Mutable/Source/MutableTools/Private/MuT/ASTOpMeshMorph.h
--- a/Engine/Plugins/Experimental/Mutable/Source/MutableTools/Private/MuT/ASTOpMeshMorph.h
+++ b/Engine/Plugins/Experimental/Mutable/Source/MutableTools/Private/MuT/ASTOpMeshMorph.h
@@ -40,7 +40,7 @@ namespace mu
 		void ForEachChild(const TFunctionRef<void(ASTChild&)>) override;
 		bool IsEqual(const ASTOp& otherUntyped) const override;
 		Ptr<ASTOp> Clone(MapChildFuncRef mapChild) const override;
-		void Link(FProgram& program, const FLinkerOptions*) override;
+		void Link(FProgram& program, FLinkerOptions*) override;
 		Ptr<ASTOp> OptimiseSink(const FModelOptimizationOptions&, FOptimizeSinkContext&) const override;
 
 	};
diff --git a/Engine/Plugins/Experimental/Mutable/Source/MutableTools/Private/MuT/ASTOpMeshMorphReshape.cpp b/Engine/Plugins/Experimental/Mutable/Source/MutableTools/Private/MuT/ASTOpMeshMorphReshape.cpp
--- a/Engine/Plugins/Experimental/Mutable/Source/MutableTools/Private/MuT/ASTOpMeshMorphReshape.cpp
+++ b/Engine/Plugins/Experimental/Mutable/Source/MutableTools/Private/MuT/ASTOpMeshMorphReshape.cpp
@@ -66,7 +66,7 @@ namespace mu
 	}
 
 
-	void ASTOpMeshMorphReshape::Link(FProgram& program, const FLinkerOptions*)
+	void ASTOpMeshMorphReshape::Link(FProgram& program, FLinkerOptions*)
 	{
 		// Already linked?
 		if (!linkedAddress)
diff --git a/Engine/Plugins/Experimental/Mutable/Source/MutableTools/Private/MuT/ASTOpMeshMorphReshape.h b/Engine/Plugins/Experimental/Mutable/Source/MutableTools/Private/MuT/ASTOpMeshMorphReshape.h
--- a/Engine/Plugins/Experimental/Mutable/Source/MutableTools/Private/MuT/ASTOpMeshMorphReshape.h
+++ b/Engine/Plugins/Experimental/Mutable/Source/MutableTools/Private/MuT/ASTOpMeshMorphReshape.h
@@ -33,7 +33,7 @@ struct FProgram;
 		bool IsEqual(const ASTOp& otherUntyped) const override;
 		Ptr<ASTOp> Clone(MapChildFuncRef mapChild) const override;
 		void ForEachChild(const TFunctionRef<void(ASTChild&)>) override;
-		void Link(FProgram& program, const FLinkerOptions* Options) override;
+		void Link(FProgram& program, FLinkerOptions* Options) override;
 	};
 
 
diff --git a/Engine/Plugins/Experimental/Mutable/Source/MutableTools/Private/MuT/ASTOpMeshOptimizeSkinning.cpp b/Engine/Plugins/Experimental/Mutable/Source/MutableTools/Private/MuT/ASTOpMeshOptimizeSkinning.cpp
--- a/Engine/Plugins/Experimental/Mutable/Source/MutableTools/Private/MuT/ASTOpMeshOptimizeSkinning.cpp
+++ b/Engine/Plugins/Experimental/Mutable/Source/MutableTools/Private/MuT/ASTOpMeshOptimizeSkinning.cpp
@@ -68,7 +68,7 @@ namespace mu
 
 
 	//-------------------------------------------------------------------------------------------------
-	void ASTOpMeshOptimizeSkinning::Link(FProgram& program, const FLinkerOptions*)
+	void ASTOpMeshOptimizeSkinning::Link(FProgram& program, FLinkerOptions*)
 	{
 		// Already linked?
 		if (!linkedAddress)
diff --git a/Engine/Plugins/Experimental/Mutable/Source/MutableTools/Private/MuT/ASTOpMeshOptimizeSkinning.h b/Engine/Plugins/Experimental/Mutable/Source/MutableTools/Private/MuT/ASTOpMeshOptimizeSkinning.h
--- a/Engine/Plugins/Experimental/Mutable/Source/MutableTools/Private/MuT/ASTOpMeshOptimizeSkinning.h
+++ b/Engine/Plugins/Experimental/Mutable/Source/MutableTools/Private/MuT/ASTOpMeshOptimizeSkinning.h
@@ -32,7 +32,7 @@ namespace mu
 		void ForEachChild(const TFunctionRef<void(ASTChild&)>) override;
 		bool IsEqual(const ASTOp& otherUntyped) const override;
 		Ptr<ASTOp> Clone(MapChildFuncRef mapChild) const override;
-		void Link(FProgram& program, const FLinkerOptions*) override;
+		void Link(FProgram& program, FLinkerOptions*) override;
 	};
 
 
diff --git a/Engine/Plugins/Experimental/Mutable/Source/MutableTools/Private/MuT/ASTOpMeshRemapIndices.cpp b/Engine/Plugins/Experimental/Mutable/Source/MutableTools/Private/MuT/ASTOpMeshRemapIndices.cpp
--- a/Engine/Plugins/Experimental/Mutable/Source/MutableTools/Private/MuT/ASTOpMeshRemapIndices.cpp
+++ b/Engine/Plugins/Experimental/Mutable/Source/MutableTools/Private/MuT/ASTOpMeshRemapIndices.cpp
@@ -73,7 +73,7 @@ namespace mu
 
 
 	//-------------------------------------------------------------------------------------------------
-	void ASTOpMeshRemapIndices::Link(FProgram& program, const FLinkerOptions*)
+	void ASTOpMeshRemapIndices::Link(FProgram& program, FLinkerOptions*)
 	{
 		// Already linked?
 		if (!linkedAddress)
diff --git a/Engine/Plugins/Experimental/Mutable/Source/MutableTools/Private/MuT/ASTOpMeshRemapIndices.h b/Engine/Plugins/Experimental/Mutable/Source/MutableTools/Private/MuT/ASTOpMeshRemapIndices.h
--- a/Engine/Plugins/Experimental/Mutable/Source/MutableTools/Private/MuT/ASTOpMeshRemapIndices.h
+++ b/Engine/Plugins/Experimental/Mutable/Source/MutableTools/Private/MuT/ASTOpMeshRemapIndices.h
@@ -38,7 +38,7 @@ struct FProgram;
 		void ForEachChild(const TFunctionRef<void(ASTChild&)>) override;
 		bool IsEqual(const ASTOp& otherUntyped) const override;
 		Ptr<ASTOp> Clone(MapChildFuncRef mapChild) const override;
-		void Link(FProgram& program, const FLinkerOptions*) override;
+		void Link(FProgram& program, FLinkerOptions*) override;
 	};
 
 
diff --git a/Engine/Plugins/Experimental/Mutable/Source/MutableTools/Private/MuT/ASTOpMeshRemoveMask.cpp b/Engine/Plugins/Experimental/Mutable/Source/MutableTools/Private/MuT/ASTOpMeshRemoveMask.cpp
--- a/Engine/Plugins/Experimental/Mutable/Source/MutableTools/Private/MuT/ASTOpMeshRemoveMask.cpp
+++ b/Engine/Plugins/Experimental/Mutable/Source/MutableTools/Private/MuT/ASTOpMeshRemoveMask.cpp
@@ -88,7 +88,7 @@ namespace mu
 
 
 	//---------------------------------------------------------------------------------------------
-	void ASTOpMeshRemoveMask::Link(FProgram& program, const FLinkerOptions*)
+	void ASTOpMeshRemoveMask::Link(FProgram& program, FLinkerOptions*)
 	{
 		// Already linked?
 		if (!linkedAddress)
diff --git a/Engine/Plugins/Experimental/Mutable/Source/MutableTools/Private/MuT/ASTOpMeshRemoveMask.h b/Engine/Plugins/Experimental/Mutable/Source/MutableTools/Private/MuT/ASTOpMeshRemoveMask.h
--- a/Engine/Plugins/Experimental/Mutable/Source/MutableTools/Private/MuT/ASTOpMeshRemoveMask.h
+++ b/Engine/Plugins/Experimental/Mutable/Source/MutableTools/Private/MuT/ASTOpMeshRemoveMask.h
@@ -37,7 +37,7 @@ namespace mu
 		void ForEachChild(const TFunctionRef<void(ASTChild&)>) override;
 		bool IsEqual(const ASTOp& otherUntyped) const override;
 		Ptr<ASTOp> Clone(MapChildFuncRef mapChild) const override;
-		void Link(FProgram& program, const FLinkerOptions*) override;
+		void Link(FProgram& program, FLinkerOptions*) override;
 		Ptr<ASTOp> OptimiseSink(const FModelOptimizationOptions&, FOptimizeSinkContext&) const override;
 
 		// Own interface
diff --git a/Engine/Plugins/Experimental/Mutable/Source/MutableTools/Private/MuT/ASTOpMeshTransform.cpp b/Engine/Plugins/Experimental/Mutable/Source/MutableTools/Private/MuT/ASTOpMeshTransform.cpp
--- a/Engine/Plugins/Experimental/Mutable/Source/MutableTools/Private/MuT/ASTOpMeshTransform.cpp
+++ b/Engine/Plugins/Experimental/Mutable/Source/MutableTools/Private/MuT/ASTOpMeshTransform.cpp
@@ -62,7 +62,7 @@ namespace mu
 	}
 
 
-	void ASTOpMeshTransform::Link(FProgram& program, const FLinkerOptions*)
+	void ASTOpMeshTransform::Link(FProgram& program, FLinkerOptions*)
 	{
 		// Already linked?
 		if (!linkedAddress)
diff --git a/Engine/Plugins/Experimental/Mutable/Source/MutableTools/Private/MuT/ASTOpMeshTransform.h b/Engine/Plugins/Experimental/Mutable/Source/MutableTools/Private/MuT/ASTOpMeshTransform.h
--- a/Engine/Plugins/Experimental/Mutable/Source/MutableTools/Private/MuT/ASTOpMeshTransform.h
+++ b/Engine/Plugins/Experimental/Mutable/Source/MutableTools/Private/MuT/ASTOpMeshTransform.h
@@ -36,7 +36,7 @@ struct FProgram;
 		void ForEachChild(const TFunctionRef<void(ASTChild&)>) override;
 		bool IsEqual(const ASTOp& otherUntyped) const override;
 		Ptr<ASTOp> Clone(MapChildFuncRef mapChild) const override;
-		void Link(FProgram& program, const FLinkerOptions* Options) override;
+		void Link(FProgram& program, FLinkerOptions* Options) override;
 
 	};
 
diff --git a/Engine/Plugins/Experimental/Mutable/Source/MutableTools/Private/MuT/ASTOpParameter.cpp b/Engine/Plugins/Experimental/Mutable/Source/MutableTools/Private/MuT/ASTOpParameter.cpp
--- a/Engine/Plugins/Experimental/Mutable/Source/MutableTools/Private/MuT/ASTOpParameter.cpp
+++ b/Engine/Plugins/Experimental/Mutable/Source/MutableTools/Private/MuT/ASTOpParameter.cpp
@@ -74,7 +74,7 @@ namespace mu
 	}
 
 
-	void ASTOpParameter::Link(FProgram& program, const FLinkerOptions*)
+	void ASTOpParameter::Link(FProgram& program, FLinkerOptions*)
 	{
 		// Already linked?
 		if (!linkedAddress)
diff --git a/Engine/Plugins/Experimental/Mutable/Source/MutableTools/Private/MuT/ASTOpParameter.h b/Engine/Plugins/Experimental/Mutable/Source/MutableTools/Private/MuT/ASTOpParameter.h
--- a/Engine/Plugins/Experimental/Mutable/Source/MutableTools/Private/MuT/ASTOpParameter.h
+++ b/Engine/Plugins/Experimental/Mutable/Source/MutableTools/Private/MuT/ASTOpParameter.h
@@ -44,7 +44,7 @@ namespace mu
 		bool IsEqual(const ASTOp& otherUntyped) const override;
 		Ptr<ASTOp> Clone(MapChildFuncRef mapChild) const override;
 		void Assert() override;
-		void Link(FProgram& program, const FLinkerOptions*) override;
+		void Link(FProgram& program, FLinkerOptions*) override;
 		int EvaluateInt(ASTOpList& facts, bool& unknown) const override;
 		FBoolEvalResult EvaluateBool(ASTOpList& /*facts*/, FEvaluateBoolCache* = nullptr) const override;
 		FImageDesc GetImageDesc(bool, FGetImageDescContext*) const override;
diff --git a/Engine/Plugins/Experimental/Mutable/Source/MutableTools/Private/MuT/ASTOpReferenceResource.cpp b/Engine/Plugins/Experimental/Mutable/Source/MutableTools/Private/MuT/ASTOpReferenceResource.cpp
--- a/Engine/Plugins/Experimental/Mutable/Source/MutableTools/Private/MuT/ASTOpReferenceResource.cpp
+++ b/Engine/Plugins/Experimental/Mutable/Source/MutableTools/Private/MuT/ASTOpReferenceResource.cpp
@@ -49,7 +49,7 @@ namespace mu
 
 
 	//-------------------------------------------------------------------------------------------------
-	void ASTOpReferenceResource::Link(FProgram& program, const FLinkerOptions* Options)
+	void ASTOpReferenceResource::Link(FProgram& program, FLinkerOptions* Options)
 	{
 		if (!linkedAddress)
 		{			
diff --git a/Engine/Plugins/Experimental/Mutable/Source/MutableTools/Private/MuT/ASTOpReferenceResource.h b/Engine/Plugins/Experimental/Mutable/Source/MutableTools/Private/MuT/ASTOpReferenceResource.h
--- a/Engine/Plugins/Experimental/Mutable/Source/MutableTools/Private/MuT/ASTOpReferenceResource.h
+++ b/Engine/Plugins/Experimental/Mutable/Source/MutableTools/Private/MuT/ASTOpReferenceResource.h
@@ -44,7 +44,7 @@ namespace mu
 		bool IsEqual(const ASTOp& otherUntyped) const override;
 		Ptr<ASTOp> Clone(MapChildFuncRef mapChild) const override;
 		uint64 Hash() const override;
-		void Link(FProgram& program, const FLinkerOptions*) override;
+		void Link(FProgram& program, FLinkerOptions*) override;
 		FImageDesc GetImageDesc(bool, class FGetImageDescContext*) const override;
 		void GetBlockLayoutSize(int blockIndex, int* pBlockX, int* pBlockY, FBlockLayoutSizeCache* cache) override;
 		void GetLayoutBlockSize(int* pBlockX, int* pBlockY) override;
diff --git a/Engine/Plugins/Experimental/Mutable/Source/MutableTools/Private/MuT/ASTOpScalarCurve.cpp b/Engine/Plugins/Experimental/Mutable/Source/MutableTools/Private/MuT/ASTOpScalarCurve.cpp
--- a/Engine/Plugins/Experimental/Mutable/Source/MutableTools/Private/MuT/ASTOpScalarCurve.cpp
+++ b/Engine/Plugins/Experimental/Mutable/Source/MutableTools/Private/MuT/ASTOpScalarCurve.cpp
@@ -62,7 +62,7 @@ namespace mu
 	}
 
 
-	void ASTOpScalarCurve::Link(FProgram& program, const FLinkerOptions*)
+	void ASTOpScalarCurve::Link(FProgram& program, FLinkerOptions*)
 	{
 		if (!linkedAddress)
 		{
diff --git a/Engine/Plugins/Experimental/Mutable/Source/MutableTools/Private/MuT/ASTOpScalarCurve.h b/Engine/Plugins/Experimental/Mutable/Source/MutableTools/Private/MuT/ASTOpScalarCurve.h
--- a/Engine/Plugins/Experimental/Mutable/Source/MutableTools/Private/MuT/ASTOpScalarCurve.h
+++ b/Engine/Plugins/Experimental/Mutable/Source/MutableTools/Private/MuT/ASTOpScalarCurve.h
@@ -37,7 +37,7 @@ struct FProgram;
 		void ForEachChild(const TFunctionRef<void(ASTChild&)>) override;
 		bool IsEqual(const ASTOp& otherUntyped) const override;
 		Ptr<ASTOp> Clone(MapChildFuncRef mapChild) const override;
-		void Link(FProgram& program, const FLinkerOptions* Options) override;
+		void Link(FProgram& program, FLinkerOptions* Options) override;
 	};
 
 
diff --git a/Engine/Plugins/Experimental/Mutable/Source/MutableTools/Private/MuT/ASTOpSwitch.cpp b/Engine/Plugins/Experimental/Mutable/Source/MutableTools/Private/MuT/ASTOpSwitch.cpp
--- a/Engine/Plugins/Experimental/Mutable/Source/MutableTools/Private/MuT/ASTOpSwitch.cpp
+++ b/Engine/Plugins/Experimental/Mutable/Source/MutableTools/Private/MuT/ASTOpSwitch.cpp
@@ -168,7 +168,7 @@ namespace mu
 
 
 	//-------------------------------------------------------------------------------------------------
-	void ASTOpSwitch::Link(FProgram& program, const FLinkerOptions*)
+	void ASTOpSwitch::Link(FProgram& program, FLinkerOptions*)
 	{
 		// Already linked?
 		if (!linkedAddress)
diff --git a/Engine/Plugins/Experimental/Mutable/Source/MutableTools/Private/MuT/ASTOpSwitch.h b/Engine/Plugins/Experimental/Mutable/Source/MutableTools/Private/MuT/ASTOpSwitch.h
--- a/Engine/Plugins/Experimental/Mutable/Source/MutableTools/Private/MuT/ASTOpSwitch.h
+++ b/Engine/Plugins/Experimental/Mutable/Source/MutableTools/Private/MuT/ASTOpSwitch.h
@@ -65,7 +65,7 @@ template <class SCALAR> class vec4;
 		bool IsEqual(const ASTOp& otherUntyped) const override;
 		Ptr<ASTOp> Clone(MapChildFuncRef mapChild) const override;
 		void Assert() override;
-		void Link(FProgram& program, const FLinkerOptions* Options) override;
+		void Link(FProgram& program, FLinkerOptions* Options) override;
 		FImageDesc GetImageDesc(bool returnBestOption, class FGetImageDescContext* context) const  override;
 		void GetBlockLayoutSize(int blockIndex, int* pBlockX, int* pBlockY, FBlockLayoutSizeCache* cache) override;
 		void GetLayoutBlockSize(int* pBlockX, int* pBlockY) override;
diff --git a/Engine/Plugins/Experimental/Mutable/Source/MutableTools/Private/MuT/CodeGenerator.cpp b/Engine/Plugins/Experimental/Mutable/Source/MutableTools/Private/MuT/CodeGenerator.cpp
--- a/Engine/Plugins/Experimental/Mutable/Source/MutableTools/Private/MuT/CodeGenerator.cpp
+++ b/Engine/Plugins/Experimental/Mutable/Source/MutableTools/Private/MuT/CodeGenerator.cpp
@@ -149,7 +149,6 @@ namespace mu
 		m_generatedTables.clear();
 		m_firstPass = FirstPassGenerator();
 		m_currentBottomUpState = BOTTOM_UP_STATE();
-		m_imageState.Empty();
 		m_currentParents.Empty();
 		m_currentObject.Empty();
 		m_additionalComponents.clear();
@@ -427,10 +426,7 @@ namespace mu
 				blend = GenerateMissingImageCode(TEXT("Blend top image"), EImageFormat::IF_RGB_UBYTE, pPatch->GetPrivate()->m_errorContext, ImageOptions);
 			}
 			blend = GenerateImageFormat(blend, blockAd->GetImageDesc().m_format);
-			blend = GenerateImageSize(
-				blend,
-				FImageSize((uint16)m_imageState.Last().m_imageRect.size[0], (uint16)m_imageState.Last().m_imageRect.size[1])
-			);
+			blend = GenerateImageSize(blend, ImageOptions.RectSize);
 			op->blend = blend;
 
 			NodeImage* pMask = pPatch->GetPrivate()->m_pMask.get();
@@ -445,13 +441,11 @@ namespace mu
 			{
 				// Set the argument default value: affect all pixels.
 				// TODO: Special operation code without mask
-				mask = GeneratePlainImageCode(vec3<float>(1, 1, 1), ImageOptions);
+				FImageGenerationOptions MissingMaskOptions;
+				mask = GeneratePlainImageCode(FVector4f(1,1,1,1), MissingMaskOptions);
 			}
 			mask = GenerateImageFormat(mask, EImageFormat::IF_L_UBYTE);
-			mask = GenerateImageSize(
-				mask,
-				FImageSize((uint16)m_imageState.Last().m_imageRect.size[0], (uint16)m_imageState.Last().m_imageRect.size[1])
-			);
+			mask = GenerateImageSize(mask, ImageOptions.RectSize);
 			op->mask = mask;
 
 			blendAd = op;
@@ -675,7 +669,7 @@ namespace mu
 
 
 	//---------------------------------------------------------------------------------------------
-	Ptr<ASTOp> CodeGenerator::ApplyTiling(Ptr<ASTOp> Source, FImageSize Size, EImageFormat Format)
+	Ptr<ASTOp> CodeGenerator::ApplyTiling(Ptr<ASTOp> Source, UE::Math::TIntVector2<int32> Size, EImageFormat Format)
 	{
 		// For now always apply tiling
 		if (m_compilerOptions->ImageTiling==0)
@@ -826,7 +820,7 @@ namespace mu
 			MeshOptions.bUniqueVertexIDs = true;
 			MeshOptions.bLayouts = true;
 			MeshOptions.State = m_currentStateIndex;
-			MeshOptions.ActiveTags = m_activeTags.Last();
+			MeshOptions.ActiveTags = node.m_tags;
 
 			const FMeshGenerationResult* SharedMeshResults = nullptr;
 			if (bShareSurface)
@@ -877,7 +871,7 @@ namespace mu
 						MergedMeshOptions.bClampUVIslands = bShareSurface;
 						MergedMeshOptions.bNormalizeUVs = bNormalizeUVs;
 						MergedMeshOptions.State = m_currentStateIndex;
-						MergedMeshOptions.ActiveTags = m_activeTags.Last();
+						MergedMeshOptions.ActiveTags = e.node->m_tags;
 
 						if (SharedMeshResults)
 						{
@@ -945,10 +939,7 @@ namespace mu
 						RemoveMeshOptions.bUniqueVertexIDs = false;
 						RemoveMeshOptions.bLayouts = false;
 						RemoveMeshOptions.State = m_currentStateIndex;
-						if (!m_activeTags.IsEmpty())
-						{
-							RemoveMeshOptions.ActiveTags = m_activeTags.Last();
-						}
+						RemoveMeshOptions.ActiveTags = e.node->m_tags;
 
                         GenerateMesh(RemoveMeshOptions, removeResults, pRemove );
 
@@ -1268,16 +1259,11 @@ namespace mu
 						check(desc.m_format != EImageFormat::IF_NONE);
 
 						// Generate the image
-						IMAGE_STATE newState;
-						newState.m_imageSize[0] = desc.m_size[0];
-						newState.m_imageSize[1] = desc.m_size[1];
-						newState.m_imageRect = box<UE::Math::TIntVector2<int>>::FromMinSize(UE::Math::TIntVector2<int>(0, 0), UE::Math::TIntVector2<int>(desc.m_size));
-						newState.m_layoutBlockId = -1;
-						newState.m_pLayout = nullptr;
-						m_imageState.Add(newState);
-
 						FImageGenerationOptions ImageOptions;
+						ImageOptions.CurrentStateIndex = m_currentStateIndex;
 						ImageOptions.ImageLayoutStrategy = ImageLayoutStrategy;
+						ImageOptions.ActiveTags = node.m_tags;
+						ImageOptions.RectSize = UE::Math::TIntVector2<int32>(desc.m_size);
 						FImageGenerationResult Result;
 						GenerateImage(ImageOptions, Result, pImageNode);
 						Ptr<ASTOp> imageAd = Result.op;
@@ -1298,8 +1284,6 @@ namespace mu
 
 						check(imageAd);
 
-						m_imageState.Pop();
-
 						if (mipmapNode)
 						{
 							Ptr<ASTOpImageMipmap> op = new ASTOpImageMipmap();
@@ -1423,16 +1407,13 @@ namespace mu
 								rect.size[1] *= blockSizeY;
 
 								// Generate the image
-								IMAGE_STATE newState;
-								newState.m_imageSize = UE::Math::TIntVector2<int32>(desc.m_size);
-								newState.m_imageRect.min = UE::Math::TIntVector2<int32>(rect.min);
-								newState.m_imageRect.size = UE::Math::TIntVector2<int32>(rect.size);
-								newState.m_layoutBlockId = pLayout->m_blocks[b].m_id;
-								newState.m_pLayout = pLayout;
-								m_imageState.Add(newState);
-
 								FImageGenerationOptions ImageOptions;
+								ImageOptions.CurrentStateIndex = m_currentStateIndex;
 								ImageOptions.ImageLayoutStrategy = ImageLayoutStrategy;
+								ImageOptions.ActiveTags = node.m_tags;
+								ImageOptions.RectSize = UE::Math::TIntVector2<int32>(rect.size);
+								ImageOptions.LayoutToApply = pLayout;
+								ImageOptions.LayoutBlockId = pLayout->m_blocks[b].m_id;
 								FImageGenerationResult Result;
 								GenerateImage(ImageOptions, Result, pImageNode);
 								Ptr<ASTOp> blockAd = Result.op;
@@ -1458,7 +1439,7 @@ namespace mu
 								Ptr<ASTOp> FormattedBlock = GenerateImageFormat(blockAd, baseFormat);
 
 								// Apply tiling to avoid generating chunks of image that are too big.
-								FormattedBlock = ApplyTiling(FormattedBlock, FImageSize(rect.size), desc.m_format);
+								FormattedBlock = ApplyTiling(FormattedBlock, ImageOptions.RectSize, desc.m_format);
 
 								// Compose layout operation
 								Ptr<ASTOpImageCompose> composeOp = new ASTOpImageCompose();
@@ -1471,8 +1452,6 @@ namespace mu
 								composeOp->BlockIndex = pLayout->m_blocks[b].m_id;
 
 								imageAd = composeOp;
-
-								m_imageState.Pop();
 							}
 							check(imageAd);
 
@@ -1528,26 +1507,22 @@ namespace mu
 												rect.size[1] = (blockRect.size[1] * extendDesc.m_size[1]) / extlayout[1];
 
 												// Generate the image block
-												IMAGE_STATE newState;
-												newState.m_imageSize = UE::Math::TIntVector2<int32>(extendDesc.m_size);
-												newState.m_imageRect = rect;
-												newState.m_layoutBlockId = pExtendLayout->m_blocks[b].m_id;
-												newState.m_pLayout = pExtendLayout;
-												m_imageState.Add(newState);
-
 												FImageGenerationOptions ImageOptions;
+												ImageOptions.CurrentStateIndex = m_currentStateIndex;
 												ImageOptions.ImageLayoutStrategy = ImageLayoutStrategy;
+												ImageOptions.ActiveTags = node.m_tags;
+												ImageOptions.RectSize = UE::Math::TIntVector2<int32>(extendDesc.m_size);
+												ImageOptions.LayoutToApply = pExtendLayout;
+												ImageOptions.LayoutBlockId = pExtendLayout->m_blocks[b].m_id;
 												FImageGenerationResult ExtendResult;
 												GenerateImage(ImageOptions, ExtendResult, pExtend);
 												Ptr<ASTOp> fragmentAd = ExtendResult.op;
 
-												m_imageState.Pop();
-
 												// Adjust the format and size of the block to be added
 												Ptr<ASTOp> formatted = GenerateImageFormat(fragmentAd, GetUncompressedFormat(desc.m_format));
-												FImageSize expectedSize;
-												expectedSize[0] = (uint16)(blockSizeX * blockRect.size[0]);
-												expectedSize[1] = (uint16)(blockSizeY * blockRect.size[1]);
+												UE::Math::TIntVector2<int32> expectedSize;
+												expectedSize[0] = blockSizeX * blockRect.size[0];
+												expectedSize[1] = blockSizeY * blockRect.size[1];
 												formatted = GenerateImageSize(formatted, expectedSize);
 												fragmentAd = formatted;
 
diff --git a/Engine/Plugins/Experimental/Mutable/Source/MutableTools/Private/MuT/CodeGenerator.h b/Engine/Plugins/Experimental/Mutable/Source/MutableTools/Private/MuT/CodeGenerator.h
--- a/Engine/Plugins/Experimental/Mutable/Source/MutableTools/Private/MuT/CodeGenerator.h
+++ b/Engine/Plugins/Experimental/Mutable/Source/MutableTools/Private/MuT/CodeGenerator.h
@@ -150,7 +150,6 @@ namespace mu
 		//!
 		FirstPassGenerator m_firstPass;
 
-
         struct FVisitedKeyMap
         {
             FVisitedKeyMap()
@@ -161,12 +160,6 @@ namespace mu
 			{
 				uint32 KeyHash = 0;
 				KeyHash = HashCombine(KeyHash, ::GetTypeHash(InKey.pNode.get()));
-				KeyHash = HashCombine(KeyHash, ::GetTypeHash(InKey.imageSize[0]));
-				KeyHash = HashCombine(KeyHash, ::GetTypeHash(InKey.imageSize[1]));
-				KeyHash = HashCombine(KeyHash, ::GetTypeHash(InKey.imageRect.min[0]));
-				KeyHash = HashCombine(KeyHash, ::GetTypeHash(InKey.imageRect.min[1]));
-				KeyHash = HashCombine(KeyHash, ::GetTypeHash(InKey.imageRect.size[0]));
-				KeyHash = HashCombine(KeyHash, ::GetTypeHash(InKey.imageRect.size[1]));
 				KeyHash = HashCombine(KeyHash, ::GetTypeHash(InKey.state));
 				KeyHash = HashCombine(KeyHash, ::GetTypeHash((uint64)InKey.activeTags.Num()));
 				return KeyHash;
@@ -176,22 +169,15 @@ namespace mu
 			{
 				if (pNode != InKey.pNode) return false;
 				if (state != InKey.state) return false;
-				if (imageSize != InKey.imageSize) return false;
-				if (imageRect.min != InKey.imageRect.min) return false;
-				if (imageRect.size != InKey.imageRect.size) return false;
 				if (activeTags != InKey.activeTags) return false;
-				if (overrideLayouts != InKey.overrideLayouts) return false;
 				return true;
 			}
 
             // This reference has to be the smart pointer to avoid memory aliasing, keeping
             // processed nodes alive.
             NodePtrConst pNode;
-			UE::Math::TIntVector2<int32> imageSize = UE::Math::TIntVector2<int32>(0, 0);
-            box< UE::Math::TIntVector2<int32> > imageRect;
             int state = -1;
 			TArray<mu::string> activeTags;
-			TArray<LayoutPtrConst> overrideLayouts;
         };
 
         //! This struct contains additional state propagated from bottom to top of the object node graph.
@@ -232,17 +218,6 @@ namespace mu
         //! First free index to be used to identify mesh vertices.
         uint32 m_freeVertexIndex = 0;
 
-        //! When generating images, here we have the entire source image size and the rect of the
-        //! image that we are generating.
-        struct IMAGE_STATE
-        {
-            UE::Math::TIntVector2<int32> m_imageSize = UE::Math::TIntVector2<int32>(0, 0);
-            box< UE::Math::TIntVector2<int32> > m_imageRect;
-            int32 m_layoutBlockId;
-            LayoutPtrConst m_pLayout;
-        };
-		TArray<IMAGE_STATE> m_imageState;
-
 		// (top-down) Tags that are active when generating nodes.
 		TArray< TArray<mu::string> > m_activeTags;
 
@@ -347,11 +322,6 @@ namespace mu
 			FVisitedKeyMap key;
 			key.pNode = InNode;
 			key.state = m_currentStateIndex;
-			if (!m_imageState.IsEmpty())
-			{
-				key.imageSize = m_imageState.Last().m_imageSize;
-				key.imageRect = m_imageState.Last().m_imageRect;
-			}
 			if (!m_activeTags.IsEmpty())
 			{
 				key.activeTags = m_activeTags.Last();
@@ -362,33 +332,61 @@ namespace mu
 
 		//-----------------------------------------------------------------------------------------
 		// Images
-				
+		
 		/** Options that affect the generation of images. It is like list of what required data we want while parsing down the image node graph. */
 		struct FImageGenerationOptions
 		{
 			/** */
 			CompilerOptions::TextureLayoutStrategy ImageLayoutStrategy = CompilerOptions::TextureLayoutStrategy::None;
 
+			/** This is used to introduce additional image generation safety. \TODO: Move this "safety" to optimization? */
+			UE::Math::TIntVector2<int32> RectSize;
+
+			/** */
+			int32 CurrentStateIndex = -1;
+
+			/** Layout block that we are trying to generate if any. */
+			int32 LayoutBlockId = -1;
+			Ptr<const Layout> LayoutToApply;
+
+			/** Tags that are active at this point of the generation. */
+			TArray<mu::string> ActiveTags;
+
 			friend FORCEINLINE uint32 GetTypeHash(const FImageGenerationOptions& InKey)
 			{
 				uint32 KeyHash = 0;
 				KeyHash = HashCombine(KeyHash, ::GetTypeHash(InKey.ImageLayoutStrategy));
+				KeyHash = HashCombine(KeyHash, GetTypeHash(InKey.RectSize));
+				KeyHash = HashCombine(KeyHash, ::GetTypeHash(InKey.CurrentStateIndex));
+				KeyHash = HashCombine(KeyHash, ::GetTypeHash(InKey.LayoutBlockId));
+				KeyHash = HashCombine(KeyHash, ::GetTypeHash(InKey.LayoutToApply.get()));
 				return KeyHash;
 			}
 
 			FORCEINLINE bool operator==(const FImageGenerationOptions& Other) const
 			{
-				return ImageLayoutStrategy == Other.ImageLayoutStrategy;
+				return ImageLayoutStrategy == Other.ImageLayoutStrategy
+					&&
+					RectSize == Other.RectSize
+					&&
+					CurrentStateIndex == Other.CurrentStateIndex
+					&&
+					LayoutBlockId == Other.LayoutBlockId
+					&&
+					LayoutToApply == Other.LayoutToApply
+					&&
+					ActiveTags == Other.ActiveTags;
 			}
 
 		};
 
+		/** */
 		struct FImageGenerationResult
 		{
 			Ptr<ASTOp> op;
 		};
 
-
+		/** */
 		struct FGeneratedImageCacheKey
 		{
 			NodePtrConst Node;
@@ -408,7 +406,7 @@ namespace mu
 			}
 		};
 
-		typedef TMap<FVisitedKeyMap, FImageGenerationResult> GeneratedImagesMap;
+		typedef TMap<FGeneratedImageCacheKey, FImageGenerationResult> GeneratedImagesMap;
 		GeneratedImagesMap m_generatedImages;
 
 		void GenerateImage(const FImageGenerationOptions&, FImageGenerationResult& result, const NodeImagePtrConst& node);
@@ -445,7 +443,7 @@ namespace mu
 		Ptr<ASTOp> GenerateMissingImageCode(const TCHAR* strWhere, EImageFormat, const void* errorContext, const FImageGenerationOptions& Options);
 
 		//!
-		Ptr<ASTOp> GeneratePlainImageCode(const vec3<float>& colour, const FImageGenerationOptions& Options);
+		Ptr<ASTOp> GeneratePlainImageCode(const FVector4f& Color, const FImageGenerationOptions& Options);
 
 		//!
 		Ptr<ASTOp> GenerateImageFormat(Ptr<ASTOp>, EImageFormat);
@@ -454,13 +452,13 @@ namespace mu
 		Ptr<ASTOp> GenerateImageUncompressed(Ptr<ASTOp>);
 
 		//!
-		Ptr<ASTOp> GenerateImageSize(Ptr<ASTOp>, FImageSize);
+		Ptr<ASTOp> GenerateImageSize(Ptr<ASTOp>, UE::Math::TIntVector2<int32>);
 
 		//!
 		FImageDesc CalculateImageDesc(const Node::Private&);
 
 		/** Evaluate if the image to generate is big enough to be split in separate operations and tiled afterwards. */
-		Ptr<ASTOp> ApplyTiling(Ptr<ASTOp> Source, FImageSize Size, EImageFormat Format);
+		Ptr<ASTOp> ApplyTiling(Ptr<ASTOp> Source, UE::Math::TIntVector2<int32> Size, EImageFormat Format);
 
 		//!
 		Ptr<ASTOp> GenerateImageBlockPatch(Ptr<ASTOp> blockAd, const NodePatchImage* pPatch, Ptr<ASTOp> conditionAd, const FImageGenerationOptions& ImageOptions);
@@ -474,7 +472,9 @@ namespace mu
 		struct FMeshGenerationOptions
 		{
 			/** TODO: Review and document. */
-			int State = 0;
+			int32 State = 0;
+
+			/** Tags that are active at this point of the generation. */
 			TArray<mu::string> ActiveTags;
 
 			/** Whatever mesh we reach at the leaves of the graph will need to have unique ids for its vertices.
@@ -513,11 +513,11 @@ namespace mu
 			FORCEINLINE bool operator==(const FMeshGenerationOptions& Other) const
 			{
 				return State==Other.State 
-					&& bUniqueVertexIDs==Other.bUniqueVertexIDs && bLayouts==Other.bLayouts
+					&& bUniqueVertexIDs==Other.bUniqueVertexIDs && bLayouts==Other.bLayouts 
+					&& bClampUVIslands == Other.bClampUVIslands && bNormalizeUVs == Other.bNormalizeUVs
 					&& ActiveTags==Other.ActiveTags
 					&& OverrideLayouts ==Other.OverrideLayouts;
 			}
-
 		};
 
 		//! Store the results of the code generation of a mesh.
diff --git a/Engine/Plugins/Experimental/Mutable/Source/MutableTools/Private/MuT/CodeGenerator_Color.cpp b/Engine/Plugins/Experimental/Mutable/Source/MutableTools/Private/MuT/CodeGenerator_Color.cpp
--- a/Engine/Plugins/Experimental/Mutable/Source/MutableTools/Private/MuT/CodeGenerator_Color.cpp
+++ b/Engine/Plugins/Experimental/Mutable/Source/MutableTools/Private/MuT/CodeGenerator_Color.cpp
@@ -303,28 +303,19 @@ namespace mu
 
 		// Source image
 		FImageGenerationOptions ImageOptions;
-		ImageOptions.ImageLayoutStrategy = CompilerOptions::TextureLayoutStrategy::None;
+		ImageOptions.CurrentStateIndex = m_currentStateIndex;
+		if (!m_activeTags.IsEmpty())
+		{
+			ImageOptions.ActiveTags = m_activeTags.Last();
+		}
+
 		Ptr<ASTOp> base;
 		if (node.m_pImage)
 		{
-			// We take whatever size will be produced
-			FImageDesc desc = CalculateImageDesc(*node.m_pImage->GetBasePrivate());
-			IMAGE_STATE newState;
-			newState.m_imageSize[0] = desc.m_size[0];
-			newState.m_imageSize[1] = desc.m_size[1];
-			newState.m_imageRect.min[0] = 0;
-			newState.m_imageRect.min[1] = 0;
-			newState.m_imageRect.size[0] = desc.m_size[0];
-			newState.m_imageRect.size[1] = desc.m_size[1];
-			m_imageState.Add(newState);
-
 			// Generate
 			FImageGenerationResult MapResult;
 			GenerateImage(ImageOptions, MapResult, node.m_pImage);
 			base = MapResult.op;
-
-			// Restore rect
-			m_imageState.Pop();
 		}
 		else
 		{
@@ -485,7 +476,7 @@ namespace mu
 			{
 				NodeColourConstantPtr CellData = new NodeColourConstant();
 				FVector4f Colour = node.m_pTable->GetPrivate()->m_rows[row].m_values[colIndex].m_colour;
-				CellData->SetValue(Colour.X, Colour.Y, Colour.Z, Colour.W);
+				CellData->SetValue(Colour);
 				return Generate(CellData);
 			});
 	}
@@ -500,7 +491,7 @@ namespace mu
 
 		// Create a constant colour node
 		NodeColourConstantPtr pNode = new NodeColourConstant();
-		pNode->SetValue(1, 1, 0, 1);
+		pNode->SetValue(FVector4f(1, 1, 0, 1));
 
 		FColorGenerationResult Result;
 		GenerateColor(Result, pNode);
diff --git a/Engine/Plugins/Experimental/Mutable/Source/MutableTools/Private/MuT/CodeGenerator_Image.cpp b/Engine/Plugins/Experimental/Mutable/Source/MutableTools/Private/MuT/CodeGenerator_Image.cpp
--- a/Engine/Plugins/Experimental/Mutable/Source/MutableTools/Private/MuT/CodeGenerator_Image.cpp
+++ b/Engine/Plugins/Experimental/Mutable/Source/MutableTools/Private/MuT/CodeGenerator_Image.cpp
@@ -46,6 +46,7 @@
 #include "MuT/NodeColour.h"
 #include "MuT/NodeColourFromScalars.h"
 #include "MuT/NodeColourSampleImage.h"
+#include "MuT/NodeColourConstant.h"
 #include "MuT/NodeImage.h"
 #include "MuT/NodeImageBinarise.h"
 #include "MuT/NodeImageBinarisePrivate.h"
@@ -111,38 +112,22 @@ namespace mu
 
 
 	//-------------------------------------------------------------------------------------------------
-	void CodeGenerator::GenerateImage(const FImageGenerationOptions& Options, FImageGenerationResult& result, const NodeImagePtrConst& Untyped)
+	void CodeGenerator::GenerateImage(const FImageGenerationOptions& Options, FImageGenerationResult& Result, const NodeImagePtrConst& Untyped)
 	{
 		if (!Untyped)
 		{
-			result = FImageGenerationResult();
+			Result = FImageGenerationResult();
 			return;
 		}
 
-		// Generate the block size in case we are generating an expression whose root is an image
-		bool addedImageState = false;
-		if ( m_imageState.IsEmpty() )
-		{
-
-			FImageDesc desc = CalculateImageDesc(*Untyped->GetBasePrivate());
-			IMAGE_STATE newState;
-			newState.m_imageSize[0] = desc.m_size[0] ? desc.m_size[0] : 256;
-			newState.m_imageSize[1] = desc.m_size[1] ? desc.m_size[1] : 256;
-			newState.m_imageRect.size = UE::Math::TIntVector2<int32>(desc.m_size);
-			newState.m_imageRect.min[0] = 0;
-			newState.m_imageRect.min[1] = 0;
-			newState.m_layoutBlockId = -1;
-			m_imageState.Add(newState);
-			addedImageState = true;
-		}
-
 		// See if it was already generated
-		// \TODO: complete and use FGeneratedImageCacheKey instead of generic keys
-		FVisitedKeyMap key = GetCurrentCacheKey(Untyped);
-		GeneratedImagesMap::ValueType* it = m_generatedImages.Find(key);
-		if (it)
+		FGeneratedImageCacheKey key;
+		key.Options = Options;
+		key.Node = Untyped;
+		GeneratedImagesMap::ValueType* CachedPtr = m_generatedImages.Find(key);
+		if (CachedPtr)
 		{
-			result = *it;
+			Result = *CachedPtr;
 		}
 		else
 		{ 
@@ -151,50 +136,43 @@ namespace mu
 			// Generate for each different type of node
 			switch (Untyped->GetImageNodeType())
 			{
-			case NodeImage::EType::Constant: GenerateImage_Constant(Options, result, static_cast<const NodeImageConstant*>(Node)); break;
+			case NodeImage::EType::Constant: GenerateImage_Constant(Options, Result, static_cast<const NodeImageConstant*>(Node)); break;
 			case NodeImage::EType::Difference_Deprecated: check(false); break;
-			case NodeImage::EType::Interpolate: GenerateImage_Interpolate(Options, result, static_cast<const NodeImageInterpolate*>(Node)); break;
-			case NodeImage::EType::Saturate: GenerateImage_Saturate(Options, result, static_cast<const NodeImageSaturate*>(Node)); break;
-			case NodeImage::EType::Table: GenerateImage_Table(Options, result, static_cast<const NodeImageTable*>(Node)); break;
-			case NodeImage::EType::Swizzle: GenerateImage_Swizzle(Options, result, static_cast<const NodeImageSwizzle*>(Node)); break;
-			case NodeImage::EType::ColourMap: GenerateImage_ColourMap(Options, result, static_cast<const NodeImageColourMap*>(Node)); break;
-			case NodeImage::EType::Gradient: GenerateImage_Gradient(Options, result, static_cast<const NodeImageGradient*>(Node)); break;
-			case NodeImage::EType::Binarise: GenerateImage_Binarise(Options, result, static_cast<const NodeImageBinarise*>(Node)); break;
-			case NodeImage::EType::Luminance: GenerateImage_Luminance(Options, result, static_cast<const NodeImageLuminance*>(Node)); break;
-			case NodeImage::EType::Layer: GenerateImage_Layer(Options, result, static_cast<const NodeImageLayer*>(Node)); break;
-			case NodeImage::EType::LayerColour: GenerateImage_LayerColour(Options, result, static_cast<const NodeImageLayerColour*>(Node)); break;
-			case NodeImage::EType::Resize: GenerateImage_Resize(Options, result, static_cast<const NodeImageResize*>(Node)); break;
-			case NodeImage::EType::PlainColour: GenerateImage_PlainColour(Options, result, static_cast<const NodeImagePlainColour*>(Node)); break;
-			case NodeImage::EType::Project: GenerateImage_Project(Options, result, static_cast<const NodeImageProject*>(Node)); break;
-			case NodeImage::EType::Mipmap: GenerateImage_Mipmap(Options, result, static_cast<const NodeImageMipmap*>(Node)); break;
-			case NodeImage::EType::Switch: GenerateImage_Switch(Options, result, static_cast<const NodeImageSwitch*>(Node)); break;
-			case NodeImage::EType::Conditional: GenerateImage_Conditional(Options, result, static_cast<const NodeImageConditional*>(Node)); break;
-			case NodeImage::EType::Format: GenerateImage_Format(Options, result, static_cast<const NodeImageFormat*>(Node)); break;
-			case NodeImage::EType::Parameter: GenerateImage_Parameter(Options, result, static_cast<const NodeImageParameter*>(Node)); break;
-			case NodeImage::EType::MultiLayer: GenerateImage_MultiLayer(Options, result, static_cast<const NodeImageMultiLayer*>(Node)); break;
-			case NodeImage::EType::Invert: GenerateImage_Invert(Options, result, static_cast<const NodeImageInvert*>(Node)); break;
-			case NodeImage::EType::Variation: GenerateImage_Variation(Options, result, static_cast<const NodeImageVariation*>(Node)); break;
-			case NodeImage::EType::NormalComposite: GenerateImage_NormalComposite(Options, result, static_cast<const NodeImageNormalComposite*>(Node)); break;
-			case NodeImage::EType::Transform: GenerateImage_Transform(Options, result, static_cast<const NodeImageTransform*>(Node)); break;
-			case NodeImage::EType::Reference: GenerateImage_Reference(Options, result, static_cast<const NodeImageReference*>(Node)); break;
+			case NodeImage::EType::Interpolate: GenerateImage_Interpolate(Options, Result, static_cast<const NodeImageInterpolate*>(Node)); break;
+			case NodeImage::EType::Saturate: GenerateImage_Saturate(Options, Result, static_cast<const NodeImageSaturate*>(Node)); break;
+			case NodeImage::EType::Table: GenerateImage_Table(Options, Result, static_cast<const NodeImageTable*>(Node)); break;
+			case NodeImage::EType::Swizzle: GenerateImage_Swizzle(Options, Result, static_cast<const NodeImageSwizzle*>(Node)); break;
+			case NodeImage::EType::ColourMap: GenerateImage_ColourMap(Options, Result, static_cast<const NodeImageColourMap*>(Node)); break;
+			case NodeImage::EType::Gradient: GenerateImage_Gradient(Options, Result, static_cast<const NodeImageGradient*>(Node)); break;
+			case NodeImage::EType::Binarise: GenerateImage_Binarise(Options, Result, static_cast<const NodeImageBinarise*>(Node)); break;
+			case NodeImage::EType::Luminance: GenerateImage_Luminance(Options, Result, static_cast<const NodeImageLuminance*>(Node)); break;
+			case NodeImage::EType::Layer: GenerateImage_Layer(Options, Result, static_cast<const NodeImageLayer*>(Node)); break;
+			case NodeImage::EType::LayerColour: GenerateImage_LayerColour(Options, Result, static_cast<const NodeImageLayerColour*>(Node)); break;
+			case NodeImage::EType::Resize: GenerateImage_Resize(Options, Result, static_cast<const NodeImageResize*>(Node)); break;
+			case NodeImage::EType::PlainColour: GenerateImage_PlainColour(Options, Result, static_cast<const NodeImagePlainColour*>(Node)); break;
+			case NodeImage::EType::Project: GenerateImage_Project(Options, Result, static_cast<const NodeImageProject*>(Node)); break;
+			case NodeImage::EType::Mipmap: GenerateImage_Mipmap(Options, Result, static_cast<const NodeImageMipmap*>(Node)); break;
+			case NodeImage::EType::Switch: GenerateImage_Switch(Options, Result, static_cast<const NodeImageSwitch*>(Node)); break;
+			case NodeImage::EType::Conditional: GenerateImage_Conditional(Options, Result, static_cast<const NodeImageConditional*>(Node)); break;
+			case NodeImage::EType::Format: GenerateImage_Format(Options, Result, static_cast<const NodeImageFormat*>(Node)); break;
+			case NodeImage::EType::Parameter: GenerateImage_Parameter(Options, Result, static_cast<const NodeImageParameter*>(Node)); break;
+			case NodeImage::EType::MultiLayer: GenerateImage_MultiLayer(Options, Result, static_cast<const NodeImageMultiLayer*>(Node)); break;
+			case NodeImage::EType::Invert: GenerateImage_Invert(Options, Result, static_cast<const NodeImageInvert*>(Node)); break;
+			case NodeImage::EType::Variation: GenerateImage_Variation(Options, Result, static_cast<const NodeImageVariation*>(Node)); break;
+			case NodeImage::EType::NormalComposite: GenerateImage_NormalComposite(Options, Result, static_cast<const NodeImageNormalComposite*>(Node)); break;
+			case NodeImage::EType::Transform: GenerateImage_Transform(Options, Result, static_cast<const NodeImageTransform*>(Node)); break;
+			case NodeImage::EType::Reference: GenerateImage_Reference(Options, Result, static_cast<const NodeImageReference*>(Node)); break;
 			case NodeImage::EType::None: check(false);
 			}
 
-			// Cache the result
-			m_generatedImages.Add(key, result);
-		}
-
-		// Restore the modified image state
-		if (addedImageState)
-		{
-			m_imageState.Pop();
+			// Cache the Result
+			m_generatedImages.Add(key, Result);
 		}
-
 	}
 
 
     //---------------------------------------------------------------------------------------------
-    void CodeGenerator::GenerateImage_Constant(const FImageGenerationOptions& Options, FImageGenerationResult& result, const NodeImageConstant* InNode)
+    void CodeGenerator::GenerateImage_Constant(const FImageGenerationOptions& Options, FImageGenerationResult& Result, const NodeImageConstant* InNode)
     {
 		const NodeImageConstant::Private& node = *InNode->GetPrivate();
 		
@@ -217,52 +195,67 @@ namespace mu
             m_pErrorLog->GetPrivate()->Add( "Constant image not set.", ELMT_WARNING, node.m_errorContext );
         }
 
-		FIntVector2 imageSize( pImage->GetSizeX(), pImage->GetSizeY() );
-
-        // The constant image size may be different than the parent rect we are generating
-        // In that case we need to crop the proportional part and the code generator will
-        // add scaling operations later
-        box< vec2< int > > cropRect;
+		if (Options.LayoutToApply)
+		{
+			// We want to generate only a block from the image.
 
-        // Order of the operations is important: multiply first to avoid losing precision.
-        // It will not overflow since image sizes are limited to 16 bit
-		FIntVector2 RectDivisor = FIntVector2( FMath::Max(1, m_imageState.Last().m_imageSize[0]), FMath::Max(1, m_imageState.Last().m_imageSize[1]));
-		cropRect.min[0] = (m_imageState.Last().m_imageRect.min[0] * imageSize[0]) / RectDivisor[0];
-		cropRect.min[1] = (m_imageState.Last().m_imageRect.min[1] * imageSize[1]) / RectDivisor[1];
-		cropRect.size[0] = (m_imageState.Last().m_imageRect.size[0] * imageSize[0]) / RectDivisor[0];
-		cropRect.size[1] = (m_imageState.Last().m_imageRect.size[1] * imageSize[1]) / RectDivisor[1];
+			FIntVector2 SourceImageSize(pImage->GetSizeX(), pImage->GetSizeY());
 
-        //check( cropRect.size[0]>0 && cropRect.size[1]>0 );
-        cropRect.size[0] = FMath::Max( cropRect.size[0], 1 );
-        cropRect.size[1] = FMath::Max( cropRect.size[1], 1 );
+			int32 BlockIndex = Options.LayoutToApply->FindBlock(Options.LayoutBlockId);
+			check( BlockIndex>=0 );
 
-        //
-        if ( pImage->GetSizeX()!=cropRect.size[0] || pImage->GetSizeY()!=cropRect.size[1] )
-        {
-			FImageOperator ImOp
+			// Block in layout grid units
+			box< UE::Math::TIntVector2<uint16> > RectInCells;
+			Options.LayoutToApply->GetBlock
 			(
-				[](int32 x, int32 y, int32 m, EImageFormat f, EInitializationType i) { return new Image(x, y, m, f, i); },
-				[](Ptr<Image>& i) {i = nullptr; },
-				[](const Image* i) { return i->Clone(); }
+				BlockIndex,
+				&RectInCells.min[0], &RectInCells.min[1],
+				&RectInCells.size[0], &RectInCells.size[1]
 			);
 
+			FIntPoint grid = Options.LayoutToApply->GetGridSize();
+			grid[0] = FMath::Max(1, grid[0]);
+			grid[1] = FMath::Max(1, grid[1]);
 
-            Ptr<Image> pCropped = new Image(cropRect.size[0], cropRect.size[1], 1, pImage->GetFormat(), EInitializationType::NotInitialized);
-			ImOp.ImageCrop(pCropped.get(), m_compilerOptions->ImageCompressionQuality, pImage.get(), cropRect);
-            op->SetValue( pCropped, m_compilerOptions->OptimisationOptions.bUseDiskCache );
-        }
+			// Transform to pixels
+			box< UE::Math::TIntVector2<int32> > rect;
+			rect.min[0]  = (RectInCells.min[0]  * SourceImageSize[0]) / grid[0];
+			rect.min[1]  = (RectInCells.min[1]  * SourceImageSize[1]) / grid[1];
+			rect.size[0] = (RectInCells.size[0] * SourceImageSize[0]) / grid[0];
+			rect.size[1] = (RectInCells.size[1] * SourceImageSize[1]) / grid[1];
+
+			// Do we need to crop?
+			if (rect.min[0]!=0 || rect.min[1]!=0 || pImage->GetSizeX() != rect.size[0] || pImage->GetSizeY() != rect.size[1])
+			{
+				FImageOperator ImOp
+				(
+					[](int32 x, int32 y, int32 m, EImageFormat f, EInitializationType i) { return new Image(x, y, m, f, i); },
+					[](Ptr<Image>& i) {i = nullptr; },
+					[](const Image* i) { return i->Clone(); }
+				);
+
+
+				Ptr<Image> pCropped = new Image(rect.size[0], rect.size[1], 1, pImage->GetFormat(), EInitializationType::NotInitialized);
+				ImOp.ImageCrop(pCropped.get(), m_compilerOptions->ImageCompressionQuality, pImage.get(), rect);
+				op->SetValue(pCropped, m_compilerOptions->OptimisationOptions.bUseDiskCache);
+			}
+			else
+			{
+				// No need to crop.
+				op->SetValue(pImage, m_compilerOptions->OptimisationOptions.bUseDiskCache);
+			}
+		}
         else
         {
-            check( cropRect.min[0]==0 && cropRect.min[1]==0 );
             op->SetValue( pImage, m_compilerOptions->OptimisationOptions.bUseDiskCache );
         }
 
-		result.op = op;
+		Result.op = op;
     }
 
 
 	//---------------------------------------------------------------------------------------------
-	void CodeGenerator::GenerateImage_Reference(const FImageGenerationOptions& Options, FImageGenerationResult& result, const NodeImageReference* InNode)
+	void CodeGenerator::GenerateImage_Reference(const FImageGenerationOptions& Options, FImageGenerationResult& Result, const NodeImageReference* InNode)
 	{
 		const NodeImageReference::Private& node = *InNode->GetPrivate();
 
@@ -271,12 +264,12 @@ namespace mu
 		op->ID = node.ImageReferenceID;
 
 		// TODO: check no crop
-		result.op = op;
+		Result.op = op;
 	}
 
 
     //---------------------------------------------------------------------------------------------
-	void CodeGenerator::GenerateImage_Parameter(const FImageGenerationOptions& Options, FImageGenerationResult& result, const NodeImageParameter* InNode )
+	void CodeGenerator::GenerateImage_Parameter(const FImageGenerationOptions& Options, FImageGenerationResult& Result, const NodeImageParameter* InNode )
     {
 		const NodeImageParameter::Private& node = *InNode->GetPrivate();
 
@@ -308,12 +301,12 @@ namespace mu
             op = it->second;
         }
 
-		result.op = op;
+		Result.op = op;
     }
 
 
     //---------------------------------------------------------------------------------------------
-	void CodeGenerator::GenerateImage_Layer(const FImageGenerationOptions& Options, FImageGenerationResult& result, const NodeImageLayer* InNode)
+	void CodeGenerator::GenerateImage_Layer(const FImageGenerationOptions& Options, FImageGenerationResult& Result, const NodeImageLayer* InNode)
 	{
 		const NodeImageLayer::Private& node = *InNode->GetPrivate();
 		
@@ -337,7 +330,7 @@ namespace mu
             base = GenerateMissingImageCode(TEXT("Image Layer base"), EImageFormat::IF_RGB_UBYTE, node.m_errorContext, Options);
         }
 
-		FImageSize TargetSize((uint16)m_imageState.Last().m_imageRect.size[0],(uint16)m_imageState.Last().m_imageRect.size[1]);
+		UE::Math::TIntVector2<int32> TargetSize = Options.RectSize;
 
 		EImageFormat baseFormat = base->GetImageDesc( true ).m_format;
         //base = GenerateImageFormat( base, EImageFormat::IF_RGB_UBYTE );
@@ -368,19 +361,19 @@ namespace mu
         else
         {
             // This argument is required
-            blended = GeneratePlainImageCode( vec3<float>( 1,1,0 ), Options );
+            blended = GeneratePlainImageCode(FVector4f( 1,1,0,1 ), Options );
         }
         //blended = GenerateImageFormat( blended, EImageFormat::IF_RGB_UBYTE );
         blended = GenerateImageFormat( blended, baseFormat );
         blended = GenerateImageSize( blended, TargetSize);
         op->blend = blended;
 
-        result.op = op;
+        Result.op = op;
     }
 
 
     //---------------------------------------------------------------------------------------------
-	void CodeGenerator::GenerateImage_LayerColour(const FImageGenerationOptions& Options, FImageGenerationResult& result, const NodeImageLayerColour* InNode)
+	void CodeGenerator::GenerateImage_LayerColour(const FImageGenerationOptions& Options, FImageGenerationResult& Result, const NodeImageLayerColour* InNode)
 	{
 		const NodeImageLayerColour::Private& node = *InNode->GetPrivate();
 		
@@ -403,7 +396,7 @@ namespace mu
             base = GenerateMissingImageCode(TEXT("Layer base image"), EImageFormat::IF_RGB_UBYTE, node.m_errorContext, Options );
         }
         base = GenerateImageFormat( base, EImageFormat::IF_RGB_UBYTE );
-        base = GenerateImageSize( base, FImageSize((uint16)m_imageState.Last().m_imageRect.size[0],(uint16)m_imageState.Last().m_imageRect.size[1]) );
+        base = GenerateImageSize( base,Options.RectSize );
         op->base = base;
 
         // Mask of the effect
@@ -415,7 +408,7 @@ namespace mu
 			mask = MaskResult.op;
 			
 			mask = GenerateImageFormat( mask, EImageFormat::IF_L_UBYTE );
-            mask = GenerateImageSize( mask, FImageSize((uint16)m_imageState.Last().m_imageRect.size[0],(uint16)m_imageState.Last().m_imageRect.size[1]) );
+            mask = GenerateImageSize( mask, Options.RectSize);
         }
         op->mask = mask;
 
@@ -434,12 +427,12 @@ namespace mu
         }
         op->color = colour;
 
-        result.op = op;
+        Result.op = op;
     }
 
 
     //---------------------------------------------------------------------------------------------
-	void CodeGenerator::GenerateImage_MultiLayer(const FImageGenerationOptions& Options, FImageGenerationResult& result, const NodeImageMultiLayer* InNode)
+	void CodeGenerator::GenerateImage_MultiLayer(const FImageGenerationOptions& Options, FImageGenerationResult& Result, const NodeImageMultiLayer* InNode)
 	{
 		const NodeImageMultiLayer::Private& node = *InNode->GetPrivate();
 
@@ -464,8 +457,7 @@ namespace mu
         }
 
 		EImageFormat baseFormat = base->GetImageDesc().m_format;
-        base = GenerateImageSize( base, FImageSize((uint16)m_imageState.Last().m_imageRect.size[0], 
-			                                       (uint16)m_imageState.Last().m_imageRect.size[1]) );
+        base = GenerateImageSize( base, Options.RectSize);
         op->base = base;
 
         // Mask of the effect
@@ -477,9 +469,7 @@ namespace mu
 			mask = MaskResult.op;
 
 			mask = GenerateImageFormat( mask, EImageFormat::IF_L_UBYTE );
-            mask = GenerateImageSize
-                    ( mask, FImageSize((uint16)m_imageState.Last().m_imageRect.size[0],
-                                       (uint16)m_imageState.Last().m_imageRect.size[1]) );
+            mask = GenerateImageSize( mask, Options.RectSize);
         }
         op->mask = mask;
 
@@ -494,12 +484,10 @@ namespace mu
         else
         {
             // This argument is required
-            blended = GeneratePlainImageCode( vec3<float>( 1,1,0 ), Options);
+            blended = GeneratePlainImageCode(FVector4f( 1,1,0,1 ), Options);
         }
         blended = GenerateImageFormat( blended, baseFormat );
-        blended = GenerateImageSize
-                ( blended, FImageSize((uint16)m_imageState.Last().m_imageRect.size[0],
-                                      (uint16)m_imageState.Last().m_imageRect.size[1]) );
+        blended = GenerateImageSize( blended, Options.RectSize);
         op->blend = blended;
 
         // Range of iteration
@@ -513,12 +501,12 @@ namespace mu
             op->range.rangeUID = rangeResult.rangeUID;
         }
 
-        result.op = op;
+        Result.op = op;
     }
 
 
 	//---------------------------------------------------------------------------------------------
-	void CodeGenerator::GenerateImage_NormalComposite(const FImageGenerationOptions& Options, FImageGenerationResult& result, const NodeImageNormalComposite* InNode)
+	void CodeGenerator::GenerateImage_NormalComposite(const FImageGenerationOptions& Options, FImageGenerationResult& Result, const NodeImageNormalComposite* InNode)
 	{
 		const NodeImageNormalComposite::Private& node = *InNode->GetPrivate();
 
@@ -544,9 +532,7 @@ namespace mu
         }
 
 		EImageFormat baseFormat = base->GetImageDesc().m_format;
-        base = GenerateImageSize
-                ( base, FImageSize((uint16)m_imageState.Last().m_imageRect.size[0],
-                                   (uint16)m_imageState.Last().m_imageRect.size[1]) );
+        base = GenerateImageSize( base, Options.RectSize);
         op->Base = base;
 
         Ptr<ASTOp> normal;
@@ -557,9 +543,7 @@ namespace mu
 			normal = BaseResult.op;
 
             normal = GenerateImageFormat( normal, EImageFormat::IF_RGB_UBYTE );
-            normal = GenerateImageSize
-                    ( normal, FImageSize((uint16)m_imageState.Last().m_imageRect.size[0],
-                                         (uint16)m_imageState.Last().m_imageRect.size[1]) );
+            normal = GenerateImageSize( normal, Options.RectSize);
         }
 		else
 		{
@@ -569,10 +553,10 @@ namespace mu
 
         op->Normal = normal;
         
-		result.op = op;
+		Result.op = op;
     }
 
-	void CodeGenerator::GenerateImage_Transform(const FImageGenerationOptions& Options, FImageGenerationResult& result, const NodeImageTransform* InNode)
+	void CodeGenerator::GenerateImage_Transform(const FImageGenerationOptions& Options, FImageGenerationResult& Result, const NodeImageTransform* InNode)
     {
 		const NodeImageTransform::Private& node = *InNode->GetPrivate();
 
@@ -633,16 +617,14 @@ namespace mu
         }
 
 		EImageFormat baseFormat = base->GetImageDesc().m_format;
-        base = GenerateImageSize
-                ( base, FImageSize((uint16)m_imageState.Last().m_imageRect.size[0],
-									  (uint16)m_imageState.Last().m_imageRect.size[1]) );
+        base = GenerateImageSize( base, Options.RectSize);
         op->base = base;
 
-        result.op = op; 
+        Result.op = op; 
     }
 
     //---------------------------------------------------------------------------------------------
-	void CodeGenerator::GenerateImage_Interpolate(const FImageGenerationOptions& Options, FImageGenerationResult& result, const NodeImageInterpolate* InNode)
+	void CodeGenerator::GenerateImage_Interpolate(const FImageGenerationOptions& Options, FImageGenerationResult& Result, const NodeImageInterpolate* InNode)
 	{
 		const NodeImageInterpolate::Private& node = *InNode->GetPrivate();
 
@@ -678,8 +660,7 @@ namespace mu
 
                 // TODO: Support other formats
                 target = GenerateImageFormat( target, EImageFormat::IF_RGB_UBYTE );
-                target = GenerateImageSize
-                        ( target, FImageSize((uint16)m_imageState.Last().m_imageRect.size[0],(uint16)m_imageState.Last().m_imageRect.size[1]) );
+                target = GenerateImageSize( target, Options.RectSize);
 
                 op->SetChild( op->op.args.ImageInterpolate.targets[numTargets], target);
                 numTargets++;
@@ -690,23 +671,23 @@ namespace mu
         if (!op->op.args.ImageInterpolate.targets[0])
         {
             Ptr<ASTOp> target = GenerateMissingImageCode(TEXT("First interpolation image"), EImageFormat::IF_RGB_UBYTE, node.m_errorContext, Options);
-            target = GenerateImageSize( target, FImageSize((uint16)m_imageState.Last().m_imageRect.size[0],(uint16)m_imageState.Last().m_imageRect.size[1]) );
+            target = GenerateImageSize( target, Options.RectSize);
             op->SetChild( op->op.args.ImageInterpolate.targets[0], target);
         }
 
-        result.op = op;
+        Result.op = op;
     }
 
 
 	//---------------------------------------------------------------------------------------------
-	void CodeGenerator::GenerateImage_Swizzle(const FImageGenerationOptions& Options, FImageGenerationResult& result, const NodeImageSwizzle* InNode)
+	void CodeGenerator::GenerateImage_Swizzle(const FImageGenerationOptions& Options, FImageGenerationResult& Result, const NodeImageSwizzle* InNode)
 	{
 		const NodeImageSwizzle::Private& node = *InNode->GetPrivate();
 
 		//MUTABLE_CPUPROFILER_SCOPE(NodeImageSwizzle);
 
         // This node always produces a swizzle operation and sometimes it may produce a pixelformat
-		// operation to compress the result
+		// operation to compress the Result
         Ptr<ASTOpImageSwizzle> op = new ASTOpImageSwizzle();
 
 		// Format
@@ -807,12 +788,12 @@ namespace mu
 			resultOp = fop;
 		}
 
-        result.op = resultOp;
+        Result.op = resultOp;
 	}
 
 
 	//---------------------------------------------------------------------------------------------
-	void CodeGenerator::GenerateImage_Format(const FImageGenerationOptions& Options, FImageGenerationResult& result, const NodeImageFormat* InNode)
+	void CodeGenerator::GenerateImage_Format(const FImageGenerationOptions& Options, FImageGenerationResult& Result, const NodeImageFormat* InNode)
 	{
 		const NodeImageFormat::Private& node = *InNode->GetPrivate();
 
@@ -836,12 +817,12 @@ namespace mu
 			fop->Source = BaseResult.op;
 		}
 
-        result.op = fop;
+        Result.op = fop;
 	}
 
 
     //---------------------------------------------------------------------------------------------
-	void CodeGenerator::GenerateImage_Saturate(const FImageGenerationOptions& Options, FImageGenerationResult& result, const NodeImageSaturate* InNode)
+	void CodeGenerator::GenerateImage_Saturate(const FImageGenerationOptions& Options, FImageGenerationResult& Result, const NodeImageSaturate* InNode)
 	{
 		const NodeImageSaturate::Private& node = *InNode->GetPrivate();
 
@@ -864,8 +845,7 @@ namespace mu
         }
 		
         base = GenerateImageFormat(base, GetRGBOrRGBAFormat(base->GetImageDesc().m_format));
-        base = GenerateImageSize
-                ( base, FImageSize((uint16)m_imageState.Last().m_imageRect.size[0],(uint16)m_imageState.Last().m_imageRect.size[1]) );
+        base = GenerateImageSize( base, Options.RectSize);
         op->SetChild( op->op.args.ImageSaturate.base, base);
 
 
@@ -880,12 +860,12 @@ namespace mu
             op->SetChild( op->op.args.ImageSaturate.factor, GenerateMissingScalarCode(TEXT("Saturation factor"), 0.5f, node.m_errorContext ) );
         }
 
-        result.op = op;
+        Result.op = op;
     }
 
 
     //---------------------------------------------------------------------------------------------
-	void CodeGenerator::GenerateImage_Luminance(const FImageGenerationOptions& Options, FImageGenerationResult& result, const NodeImageLuminance* InNode)
+	void CodeGenerator::GenerateImage_Luminance(const FImageGenerationOptions& Options, FImageGenerationResult& Result, const NodeImageLuminance* InNode)
 	{
 		const NodeImageLuminance::Private& node = *InNode->GetPrivate();
 
@@ -908,12 +888,12 @@ namespace mu
         base = GenerateImageFormat( base, EImageFormat::IF_RGB_UBYTE );
         op->SetChild( op->op.args.ImageLuminance.base, base);
 
-        result.op = op;
+        Result.op = op;
     }
 
 
     //---------------------------------------------------------------------------------------------
-	void CodeGenerator::GenerateImage_ColourMap(const FImageGenerationOptions& Options, FImageGenerationResult& result, const NodeImageColourMap* InNode)
+	void CodeGenerator::GenerateImage_ColourMap(const FImageGenerationOptions& Options, FImageGenerationResult& Result, const NodeImageColourMap* InNode)
 	{
 		const NodeImageColourMap::Private& node = *InNode->GetPrivate();
 
@@ -933,9 +913,7 @@ namespace mu
             // This argument is required
             base = GenerateMissingImageCode(TEXT("Colourmap base image"), EImageFormat::IF_RGB_UBYTE, node.m_errorContext, Options);
         }
-        base = GenerateImageSize
-                ( base, FImageSize( (uint16)m_imageState.Last().m_imageRect.size[0],
-                                    (uint16)m_imageState.Last().m_imageRect.size[1]) );
+        base = GenerateImageSize( base, Options.RectSize);
         op->SetChild( op->op.args.ImageColourMap.base, base);
 
         // Mask of the effect
@@ -950,40 +928,40 @@ namespace mu
         {
             // Set the argument default value: affect all pixels.
             // TODO: Special operation code without mask
-            mask = GeneratePlainImageCode( vec3<float>( 1,1,1 ), Options);
+            mask = GeneratePlainImageCode(FVector4f( 1,1,1,1 ), Options);
         }
         mask = GenerateImageFormat( mask, EImageFormat::IF_L_UBYTE );
-        mask = GenerateImageSize
-                ( mask, FImageSize( (uint16)m_imageState.Last().m_imageRect.size[0],
-                                    (uint16)m_imageState.Last().m_imageRect.size[1]) );
+        mask = GenerateImageSize( mask, Options.RectSize);
         op->SetChild( op->op.args.ImageColourMap.mask, mask);
 
         // Map image
 		// \TODO: We are forcing an map image size as it was a normal texture, and this is incorrect.
-        Ptr<ASTOp> mapImage;
+        Ptr<ASTOp> MapImageOp;
         if ( node.m_pMap )
         {
+			FImageGenerationOptions MapOptions = Options;
+			MapOptions.ImageLayoutStrategy = CompilerOptions::TextureLayoutStrategy::None;
+			MapOptions.LayoutToApply = nullptr;
+			MapOptions.LayoutBlockId = -1;
+
 			FImageGenerationResult BaseResult;
-			GenerateImage(Options, BaseResult, node.m_pMap);
-			mapImage = BaseResult.op;
-        }
+			GenerateImage(MapOptions, BaseResult, node.m_pMap);
+			MapImageOp = BaseResult.op;
+			MapImageOp = GenerateImageFormat(MapImageOp, EImageFormat::IF_RGB_UBYTE);
+		}
         else
         {
             // This argument is required
-            mapImage = GenerateMissingImageCode(TEXT("Map image"), EImageFormat::IF_RGB_UBYTE, node.m_errorContext, Options);
+			MapImageOp = GenerateMissingImageCode(TEXT("Map image"), EImageFormat::IF_RGB_UBYTE, node.m_errorContext, Options);
         }
-        mapImage = GenerateImageFormat( mapImage, EImageFormat::IF_RGB_UBYTE );
-        mapImage = GenerateImageSize
-                ( mapImage, FImageSize( (uint16)m_imageState.Last().m_imageRect.size[0],
-                                        (uint16)m_imageState.Last().m_imageRect.size[1]) );
-        op->SetChild( op->op.args.ImageColourMap.map, mapImage);
+        op->SetChild( op->op.args.ImageColourMap.map, MapImageOp);
 
-        result.op = op;
+        Result.op = op;
     }
 
 
     //---------------------------------------------------------------------------------------------
-	void CodeGenerator::GenerateImage_Gradient(const FImageGenerationOptions& Options, FImageGenerationResult& result, const NodeImageGradient* InNode)
+	void CodeGenerator::GenerateImage_Gradient(const FImageGenerationOptions& Options, FImageGenerationResult& Result, const NodeImageGradient* InNode)
 	{
 		const NodeImageGradient::Private& node = *InNode->GetPrivate();
 
@@ -1019,12 +997,12 @@ namespace mu
         op->op.args.ImageGradient.size[0] = (uint16)FMath::Max( 2, FMath::Min( 1024, node.m_size[0] ) );
         op->op.args.ImageGradient.size[1] = (uint16)FMath::Max( 1, FMath::Min( 1024, node.m_size[1] ) );
 
-        result.op = op;
+        Result.op = op;
     }
 
 
     //---------------------------------------------------------------------------------------------
-	void CodeGenerator::GenerateImage_Binarise(const FImageGenerationOptions& Options, FImageGenerationResult& result, const NodeImageBinarise* InNode)
+	void CodeGenerator::GenerateImage_Binarise(const FImageGenerationOptions& Options, FImageGenerationResult& Result, const NodeImageBinarise* InNode)
 	{
 		const NodeImageBinarise::Private& node = *InNode->GetPrivate();
 
@@ -1045,9 +1023,7 @@ namespace mu
             a = GenerateMissingImageCode(TEXT("Image Binarise Base"), EImageFormat::IF_RGB_UBYTE, node.m_errorContext, Options);
         }
         a = GenerateImageFormat( a, EImageFormat::IF_RGB_UBYTE );
-        a = GenerateImageSize
-                ( a, FImageSize( (uint16)m_imageState.Last().m_imageRect.size[0],
-                                 (uint16)m_imageState.Last().m_imageRect.size[1]) );
+        a = GenerateImageSize( a, Options.RectSize);
         op->SetChild( op->op.args.ImageBinarise.base, a );
 
         // Threshold
@@ -1063,12 +1039,12 @@ namespace mu
         }
         op->SetChild( op->op.args.ImageBinarise.threshold, b );
 
-        result.op = op;
+        Result.op = op;
     }
 
 
     //---------------------------------------------------------------------------------------------
-	void CodeGenerator::GenerateImage_Resize(const FImageGenerationOptions& Options, FImageGenerationResult& result, const NodeImageResize* InNode)
+	void CodeGenerator::GenerateImage_Resize(const FImageGenerationOptions& Options, FImageGenerationResult& Result, const NodeImageResize* InNode)
 	{
 		const NodeImageResize::Private& node = *InNode->GetPrivate();
 
@@ -1080,24 +1056,13 @@ namespace mu
         Ptr<ASTOp> base;
         if ( node.m_pBase )
         {
-            IMAGE_STATE newState = m_imageState.Last();
-            if ( node.m_relative )
-            {
-                newState.m_imageSize[0]=(int)std::roundf(newState.m_imageSize[0]/node.m_sizeX);
-                newState.m_imageSize[1]=(int)std::roundf(newState.m_imageSize[1]/node.m_sizeY);
-                newState.m_imageRect.min[0]=(int)std::roundf(newState.m_imageRect.min[0]/node.m_sizeX);
-                newState.m_imageRect.min[1]=(int)std::roundf(newState.m_imageRect.min[1]/node.m_sizeY);
-                newState.m_imageRect.size[0]=(int)std::roundf(newState.m_imageRect.size[0]/node.m_sizeX);
-                newState.m_imageRect.size[1]=(int)std::roundf(newState.m_imageRect.size[1]/node.m_sizeY);
-            }
-
-            m_imageState.Add(newState);
+			FImageGenerationOptions NewOptions = Options;
+			NewOptions.RectSize[0] = FMath::RoundToInt32(NewOptions.RectSize[0] / node.m_sizeX);
+			NewOptions.RectSize[1] = FMath::RoundToInt32(NewOptions.RectSize[1] / node.m_sizeY);
 
 			FImageGenerationResult BaseResult;
-			GenerateImage(Options, BaseResult, node.m_pBase);
+			GenerateImage(NewOptions, BaseResult, node.m_pBase);
 			base = BaseResult.op;
-
-            m_imageState.Pop();
         }
         else
         {
@@ -1125,12 +1090,12 @@ namespace mu
             at = op;
         }
 
-        result.op = at;
+        Result.op = at;
     }
 
 
     //---------------------------------------------------------------------------------------------
-	void CodeGenerator::GenerateImage_PlainColour(const FImageGenerationOptions& Options, FImageGenerationResult& result, const NodeImagePlainColour* InNode)
+	void CodeGenerator::GenerateImage_PlainColour(const FImageGenerationOptions& Options, FImageGenerationResult& Result, const NodeImagePlainColour* InNode)
 	{
 		const NodeImagePlainColour::Private& node = *InNode->GetPrivate();
 
@@ -1150,30 +1115,16 @@ namespace mu
         op->op.type = OP_TYPE::IM_PLAINCOLOUR;
         op->SetChild( op->op.args.ImagePlainColour.colour, base);
 		op->op.args.ImagePlainColour.format = node.Format;
-        op->op.args.ImagePlainColour.size[0] = uint16(node.m_sizeX);
-		op->op.args.ImagePlainColour.size[1] = uint16(node.m_sizeY);
+        op->op.args.ImagePlainColour.size[0] = (Options.RectSize[0] > 0) ? Options.RectSize[0] : node.m_sizeX;
+		op->op.args.ImagePlainColour.size[1] = (Options.RectSize[1] > 0) ? Options.RectSize[1] : node.m_sizeY;
 		op->op.args.ImagePlainColour.LODs = 1;
 
-        Ptr<ASTOpFixed> opSize = new ASTOpFixed();
-        opSize->op.type = OP_TYPE::IM_RESIZE;
-        if (m_imageState.Num())
-        {
-            opSize->op.args.ImageResize.size[0] = (uint16)m_imageState.Last().m_imageRect.size[0];
-            opSize->op.args.ImageResize.size[1] = (uint16)m_imageState.Last().m_imageRect.size[1];
-        }
-        else
-        {
-            opSize->op.args.ImageResize.size[0] = (uint16)node.m_sizeX;
-            opSize->op.args.ImageResize.size[1] = (uint16)node.m_sizeY;
-        }
-        opSize->SetChild( opSize->op.args.ImageResize.source, op);
-
-        result.op = opSize;
+        Result.op = op;
     }
 
 
     //---------------------------------------------------------------------------------------------
-	void CodeGenerator::GenerateImage_Switch(const FImageGenerationOptions& Options, FImageGenerationResult& result, const NodeImageSwitch* InNode)
+	void CodeGenerator::GenerateImage_Switch(const FImageGenerationOptions& Options, FImageGenerationResult& Result, const NodeImageSwitch* InNode)
 	{
 		const NodeImageSwitch::Private& node = *InNode->GetPrivate();
 
@@ -1183,7 +1134,7 @@ namespace mu
 		{
 			// No options in the switch!
             Ptr<ASTOp> missingOp = GenerateMissingImageCode(TEXT("Switch option"), EImageFormat::IF_RGBA_UBYTE, node.m_errorContext, Options);
-			result.op = missingOp;
+			Result.op = missingOp;
 			return;
 		}
 
@@ -1252,12 +1203,12 @@ namespace mu
         //    switchAt = fop;
         //}
 
-        result.op = switchAt;
+        Result.op = switchAt;
     }
 
 
     //---------------------------------------------------------------------------------------------
-	void CodeGenerator::GenerateImage_Conditional(const FImageGenerationOptions& Options, FImageGenerationResult& result, const NodeImageConditional* InNode)
+	void CodeGenerator::GenerateImage_Conditional(const FImageGenerationOptions& Options, FImageGenerationResult& Result, const NodeImageConditional* InNode)
 	{
 		const NodeImageConditional::Private& node = *InNode->GetPrivate();
 
@@ -1284,12 +1235,12 @@ namespace mu
 		GenerateImage(Options, NoResult, node.m_false);
 		op->no = NoResult.op;
 		
-        result.op = op;
+        Result.op = op;
     }
 
 
     //---------------------------------------------------------------------------------------------
-	void CodeGenerator::GenerateImage_Project(const FImageGenerationOptions& Options, FImageGenerationResult& result, const NodeImageProject* InNode)
+	void CodeGenerator::GenerateImage_Project(const FImageGenerationOptions& Options, FImageGenerationResult& Result, const NodeImageProject* InNode)
 	{
 		const NodeImageProject::Private& node = *InNode->GetPrivate();
 
@@ -1318,28 +1269,29 @@ namespace mu
         pop->SetChild( pop->op.args.MeshProject.projector, projectorResult.op );
 
 		int32 LayoutBlockIndex = -1;
-		if (m_imageState.Last().m_pLayout)
+		if (Options.LayoutToApply)
 		{
-			LayoutBlockIndex = m_imageState.Last().m_pLayout->m_blocks.IndexOfByPredicate([&](const Layout::FBlock& Block) { return Block.m_id == m_imageState.Last().m_layoutBlockId; });
+			LayoutBlockIndex = Options.LayoutToApply->m_blocks.IndexOfByPredicate([&](const Layout::FBlock& Block) { return Block.m_id == Options.LayoutBlockId; });
 		}
 		int32 GeneratedLayoutBlockId = -1;
 
         // Mesh
         if ( node.m_pMesh )
         {
-			// TODO: This will probably result in a duplicated mesh subgraph, with the original mesh but new layout block ids.
+			// TODO: This will probably Result in a duplicated mesh subgraph, with the original mesh but new layout block ids.
 			// See if it can be optimized and try to reuse the existing layout block ids instead of generating new ones.
-            FMeshGenerationResult MeshResult;
 			FMeshGenerationOptions MeshOptions;
-			MeshOptions.State = m_currentStateIndex;
-			if (m_activeTags.Num())
-			{
-				MeshOptions.ActiveTags = m_activeTags.Last();
-			}
+			MeshOptions.State = Options.CurrentStateIndex;
+			MeshOptions.ActiveTags = Options.ActiveTags;
 			MeshOptions.bLayouts = true;			// We need the layout that we will use to render
 			MeshOptions.bNormalizeUVs = true;		// We need normalized UVs for the projection
-			MeshOptions.bUniqueVertexIDs = false;	// We don't need the IDs at this point.
-            GenerateMesh( MeshOptions, MeshResult, node.m_pMesh );
+
+			// We may need IDs at this point if there are modifiers for the mesh.
+			// \TODO: Detect this at mesh constant generation and enable there istead of here?
+			MeshOptions.bUniqueVertexIDs = true;
+			
+			FMeshGenerationResult MeshResult;
+			GenerateMesh( MeshOptions, MeshResult, node.m_pMesh );
 
 			// Match the block id of the block we are generating with the id that resulted in the generated mesh
 			GeneratedLayoutBlockId = -1;
@@ -1384,7 +1336,7 @@ namespace mu
             else
             {
                 // Extract the mesh layout block
-                if ( m_imageState.Num() && GeneratedLayoutBlockId>=0 )
+                if ( GeneratedLayoutBlockId>=0 )
                 {
                     Ptr<ASTOpMeshExtractLayoutBlocks> eop = new ASTOpMeshExtractLayoutBlocks();
                     eop->source = pop->children[pop->op.args.MeshProject.mesh].child();
@@ -1435,30 +1387,19 @@ namespace mu
         // Image
         if ( node.m_pImage )
         {
-            // Remeber previous rect values
-            IMAGE_STATE newState;
-
-            // We take whatever size will be produced
-            FImageDesc desc = CalculateImageDesc( *node.m_pImage->GetBasePrivate() );
-            newState.m_imageSize = UE::Math::TIntVector2<int32>(desc.m_size);
-            newState.m_imageRect.min = UE::Math::TIntVector2<int32>(0,0);
-            newState.m_imageRect.size = UE::Math::TIntVector2<int32>(desc.m_size);
-            newState.m_layoutBlockId = -1;
-            m_imageState.Add( newState );
-
             // Generate
-			// \TODO: Build new Options with the above data
+			FImageGenerationOptions NewOptions = Options;
+			NewOptions.ImageLayoutStrategy = CompilerOptions::TextureLayoutStrategy::None;
+			NewOptions.LayoutToApply = nullptr;
+			NewOptions.LayoutBlockId = -1;
+
 			FImageGenerationResult ImageResult;
-			GenerateImage(Options, ImageResult, node.m_pImage);
+			GenerateImage(NewOptions, ImageResult, node.m_pImage);
 			op->image = ImageResult.op;
 
+			FImageDesc desc = CalculateImageDesc(*node.m_pImage->GetBasePrivate());
 			op->SourceSizeX = desc.m_size[0];
 			op->SourceSizeY = desc.m_size[1];
-
-
-            // Restore rect
-            m_imageState.Pop();
-
         }
         else
         {
@@ -1467,8 +1408,8 @@ namespace mu
         }
 
         // Image size, from the current block being generated
-        op->SizeX = (uint16)m_imageState.Last().m_imageRect.size[0];
-        op->SizeY = (uint16)m_imageState.Last().m_imageRect.size[1];
+        op->SizeX = Options.RectSize[0];
+        op->SizeY = Options.RectSize[1];
 		// \TODO: Review naming of arg
         op->BlockIndex = GeneratedLayoutBlockId;
 
@@ -1502,7 +1443,7 @@ namespace mu
 			Ptr<ASTOp> mask = MaskResult.op;
 
             mask = GenerateImageFormat( mask, EImageFormat::IF_L_UBYTE );
-            op->mask = GenerateImageSize( mask, FImageSize( m_imageState.Last().m_imageRect.size ) );
+            op->mask = GenerateImageSize( mask, Options.RectSize);
         }
 
         // Seam correction operations
@@ -1539,12 +1480,12 @@ namespace mu
         disop->SetChild(disop->op.args.ImageDisplace.displacementMap, MakeGrowMapOp);
         disop->SetChild(disop->op.args.ImageDisplace.source, op );
 
-        result.op = disop;
+        Result.op = disop;
     }
 
 
     //---------------------------------------------------------------------------------------------
-	void CodeGenerator::GenerateImage_Mipmap(const FImageGenerationOptions& Options, FImageGenerationResult& result, const NodeImageMipmap* InNode)
+	void CodeGenerator::GenerateImage_Mipmap(const FImageGenerationOptions& Options, FImageGenerationResult& Result, const NodeImageMipmap* InNode)
 	{
 		const NodeImageMipmap::Private& node = *InNode->GetPrivate();
 
@@ -1605,12 +1546,12 @@ namespace mu
 
         res = op;
 
-        result.op = res;
+        Result.op = res;
     }
 
 
 	//---------------------------------------------------------------------------------------------
-	void CodeGenerator::GenerateImage_Invert(const FImageGenerationOptions& Options, FImageGenerationResult& result, const NodeImageInvert* InNode)
+	void CodeGenerator::GenerateImage_Invert(const FImageGenerationOptions& Options, FImageGenerationResult& Result, const NodeImageInvert* InNode)
 	{
 		const NodeImageInvert::Private& node = *InNode->GetPrivate();
 
@@ -1631,17 +1572,15 @@ namespace mu
 			a = GenerateMissingImageCode(TEXT("Image Invert Color"), EImageFormat::IF_RGB_UBYTE, node.m_errorContext, Options);
 		}
 		a = GenerateImageFormat(a, EImageFormat::IF_RGB_UBYTE);
-		a = GenerateImageSize
-		(a, FImageSize((uint16)m_imageState.Last().m_imageRect.size[0],
-			(uint16)m_imageState.Last().m_imageRect.size[1]));
+		a = GenerateImageSize(a, Options.RectSize);
 		op->SetChild(op->op.args.ImageInvert.base, a);
 
-		result.op = op;
+		Result.op = op;
 	}
 
 
     //---------------------------------------------------------------------------------------------
-	void CodeGenerator::GenerateImage_Variation(const FImageGenerationOptions& Options, FImageGenerationResult& result, const NodeImageVariation* InNode)
+	void CodeGenerator::GenerateImage_Variation(const FImageGenerationOptions& Options, FImageGenerationResult& Result, const NodeImageVariation* InNode)
 	{
 		const NodeImageVariation::Private& node = *InNode->GetPrivate();
 
@@ -1684,9 +1623,9 @@ namespace mu
             Ptr<ASTOp> variationOp;
             if ( node.m_variations[t].m_image )
             {
-				FImageGenerationResult Result;
-				GenerateImage(Options, Result, node.m_variations[t].m_image);
-				variationOp = Result.op;
+				FImageGenerationResult VariationResult;
+				GenerateImage(Options, VariationResult, node.m_variations[t].m_image);
+				variationOp = VariationResult.op;
             }
             else
             {
@@ -1730,16 +1669,16 @@ namespace mu
             currentOp = fop;
         }
 
-        result.op = currentOp;
+        Result.op = currentOp;
     }
 
 
 	//---------------------------------------------------------------------------------------------
-	void CodeGenerator::GenerateImage_Table(const FImageGenerationOptions& Options, FImageGenerationResult& result, const NodeImageTable* InNode)
+	void CodeGenerator::GenerateImage_Table(const FImageGenerationOptions& Options, FImageGenerationResult& Result, const NodeImageTable* InNode)
 	{
 		const NodeImageTable::Private& node = *InNode->GetPrivate();
 
-		result.op = GenerateTableSwitch<NodeImageTable::Private, TCT_IMAGE, OP_TYPE::IM_SWITCH>(node,
+		Result.op = GenerateTableSwitch<NodeImageTable::Private, TCT_IMAGE, OP_TYPE::IM_SWITCH>(node,
 			[this,Options](const NodeImageTable::Private& node, int colIndex, int row, ErrorLog* pErrorLog)
 			{
 				TABLE_VALUE CellData = node.m_pTable->GetPrivate()->m_rows[row].m_values[colIndex];
@@ -1896,28 +1835,19 @@ namespace mu
 
 
     //---------------------------------------------------------------------------------------------
-    Ptr<ASTOp> CodeGenerator::GeneratePlainImageCode( const vec3<float>& colour, const FImageGenerationOptions& Options )
+    Ptr<ASTOp> CodeGenerator::GeneratePlainImageCode( const FVector4f& InColor, const FImageGenerationOptions& Options )
     {
-        const int size = 4;
-        ImagePtr pImage = new Image( size, size, 1, EImageFormat::IF_RGB_UBYTE, EInitializationType::NotInitialized);
+		Ptr<NodeColourConstant> ConstantColor = new NodeColourConstant();
+		ConstantColor->SetValue(InColor);
 
-        uint8_t* pData = pImage->GetData();
-        for ( int p=0; p<size*size; ++p )
-        {
-            pData[0] = (uint8_t)FMath::Min( 255.0f, FMath::Max( 0.0f, 255*colour[0] ) );
-            pData[1] = (uint8_t)FMath::Min( 255.0f, FMath::Max( 0.0f, 255*colour[1] ) );
-            pData[2] = (uint8_t)FMath::Min( 255.0f, FMath::Max( 0.0f, 255*colour[2] ) );
-            pData += 3;
-        }
-
-        NodeImageConstantPtr ConstantNode = new NodeImageConstant();
-		ConstantNode->SetValue( pImage.get() );
+        Ptr<NodeImagePlainColour> PlainNode = new NodeImagePlainColour();
+		PlainNode->SetColour(ConstantColor);
 
 		FImageGenerationResult TempResult;
-		GenerateImage(Options, TempResult, ConstantNode);
-        Ptr<ASTOp> result = TempResult.op;
+		GenerateImage(Options, TempResult, PlainNode);
+        Ptr<ASTOp> Result = TempResult.op;
 
-        return result;
+        return Result;
     }
 
 
@@ -1942,7 +1872,7 @@ namespace mu
     //---------------------------------------------------------------------------------------------
     Ptr<ASTOp> CodeGenerator::GenerateImageUncompressed( Ptr<ASTOp> at )
     {
-        Ptr<ASTOp> result = at;
+        Ptr<ASTOp> Result = at;
 
         if (at)
         {
@@ -1955,33 +1885,33 @@ namespace mu
                 Ptr<ASTOpImagePixelFormat> op = new ASTOpImagePixelFormat();
                 op->Source = at;
                 op->Format = targetFormat;
-                result = op;
+                Result = op;
             }
         }
 
-        return result;
+        return Result;
     }
 
 
     //---------------------------------------------------------------------------------------------
-    Ptr<ASTOp> CodeGenerator::GenerateImageSize( Ptr<ASTOp> at, FImageSize size )
+    Ptr<ASTOp> CodeGenerator::GenerateImageSize( Ptr<ASTOp> at, UE::Math::TIntVector2<int32> size )
     {
-        Ptr<ASTOp> result = at;
+        Ptr<ASTOp> Result = at;
 
 		if (size[0] > 0 && size[1] > 0)
 		{
-			if (at->GetImageDesc().m_size != size)
+			if (UE::Math::TIntVector2<int32>(at->GetImageDesc().m_size) != size)
 			{
 				Ptr<ASTOpFixed> op = new ASTOpFixed();
 				op->op.type = OP_TYPE::IM_RESIZE;
 				op->SetChild(op->op.args.ImageResize.source, at);
 				op->op.args.ImageResize.size[0] = size[0];
 				op->op.args.ImageResize.size[1] = size[1];
-				result = op;
+				Result = op;
 			}
 		}
 
-        return result;
+        return Result;
     }
 
 
diff --git a/Engine/Plugins/Experimental/Mutable/Source/MutableTools/Private/MuT/CodeGenerator_Mesh.cpp b/Engine/Plugins/Experimental/Mutable/Source/MutableTools/Private/MuT/CodeGenerator_Mesh.cpp
--- a/Engine/Plugins/Experimental/Mutable/Source/MutableTools/Private/MuT/CodeGenerator_Mesh.cpp
+++ b/Engine/Plugins/Experimental/Mutable/Source/MutableTools/Private/MuT/CodeGenerator_Mesh.cpp
@@ -1443,8 +1443,7 @@ class Node;
 		BOTTOM_UP_STATE temp = m_currentBottomUpState;
 
 		bool bModifiersForBeforeOperations = true;
-		OutResult.meshOp = ApplyMeshModifiers(op, InOptions.ActiveTags,
-			bModifiersForBeforeOperations, node.m_errorContext);
+		OutResult.meshOp = ApplyMeshModifiers(op, InOptions.ActiveTags, bModifiersForBeforeOperations, node.m_errorContext);
 
 		m_currentBottomUpState = temp;
 
@@ -1544,7 +1543,9 @@ class Node;
         // Base
         if (node.m_pSource)
         {
-            GenerateMesh(InOptions, OutResult, node.m_pSource);
+			FMeshGenerationOptions BaseOptions = InOptions;
+			BaseOptions.bUniqueVertexIDs = true;
+            GenerateMesh(BaseOptions, OutResult, node.m_pSource);
             op->source = OutResult.meshOp;
         }
         else
@@ -1616,7 +1617,10 @@ class Node;
         // Base
         if (node.m_pSource)
         {
-            GenerateMesh( InOptions, OutResult, node.m_pSource );
+			FMeshGenerationOptions BaseOptions = InOptions;
+			BaseOptions.bUniqueVertexIDs = true;
+
+            GenerateMesh(BaseOptions, OutResult, node.m_pSource );
             op->SetChild( op->op.args.MeshClipWithMesh.source, OutResult.meshOp );
         }
         else
@@ -1660,7 +1664,10 @@ class Node;
 		// Base Mesh
 		if (Node.m_pBaseMesh)
 		{
-			GenerateMesh(InOptions, Result, Node.m_pBaseMesh);
+			FMeshGenerationOptions BaseOptions = InOptions;
+			BaseOptions.bUniqueVertexIDs = true;
+
+			GenerateMesh(BaseOptions, Result, Node.m_pBaseMesh);
 			OpBind->Mesh = Result.meshOp;
 		}
 		else
@@ -1799,7 +1806,10 @@ class Node;
 		// Base Mesh
 		if (node.m_pBaseMesh)
 		{
-			GenerateMesh(InOptions, OutResult, node.m_pBaseMesh);
+			FMeshGenerationOptions BaseOptions = InOptions;
+			BaseOptions.bUniqueVertexIDs = true;
+
+			GenerateMesh(BaseOptions, OutResult, node.m_pBaseMesh);
 			OpBind->Mesh = OutResult.meshOp;
 		}
 		else
diff --git a/Engine/Plugins/Experimental/Mutable/Source/MutableTools/Private/MuT/Compiler.cpp b/Engine/Plugins/Experimental/Mutable/Source/MutableTools/Private/MuT/Compiler.cpp
--- a/Engine/Plugins/Experimental/Mutable/Source/MutableTools/Private/MuT/Compiler.cpp
+++ b/Engine/Plugins/Experimental/Mutable/Source/MutableTools/Private/MuT/Compiler.cpp
@@ -203,7 +203,7 @@ namespace mu
         }
 
         vector<Ptr<ASTOp>> roots;
-        for( const auto& s: states)
+        for( const STATE_COMPILATION_DATA& s: states)
         {
             roots.push_back(s.root);
         }
@@ -220,10 +220,18 @@ namespace mu
 
         // Link the program and generate state data.
 		TSharedPtr<Model> pResult = MakeShared<Model>();
-        auto& program = pResult->GetPrivate()->m_program;
-        for( auto& s: states )
+        FProgram& program = pResult->GetPrivate()->m_program;
+
+		// Preallocate ample memory
+		program.m_byteCode.Reserve(16 * 1024 * 1024);
+		program.m_opAddress.Reserve(1024 * 1024);
+
+		// Keep the link options outside the scope because it is also used to cache constant data that has already been 
+		// added and could be reused across states.
+		FLinkerOptions LinkerOptions;
+
+		for( STATE_COMPILATION_DATA& s: states )
         {
-			FLinkerOptions LinkerOptions;
 			LinkerOptions.MinTextureResidentMipCount = m_pD->m_options->GetPrivate()->MinTextureResidentMipCount;
 
             ASTOp::FullLink(s.root,program, &LinkerOptions);
@@ -237,13 +245,16 @@ namespace mu
             }
         }
 
+		program.m_byteCode.Shrink();
+		program.m_opAddress.Shrink();
+
         // Set the runtime parameter indices.
         for( STATE_COMPILATION_DATA& s: states )
         {
             for ( int32 p=0; p<s.nodeState.m_runtimeParams.Num(); ++p )
             {
-                int paramIndex = -1;
-                for ( size_t i=0; paramIndex<0 && i<program.m_parameters.Num(); ++i )
+                int32 paramIndex = -1;
+                for ( int32 i=0; paramIndex<0 && i<program.m_parameters.Num(); ++i )
                 {
                     if ( program.m_parameters[i].m_name
                          ==
diff --git a/Engine/Plugins/Experimental/Mutable/Source/MutableTools/Private/MuT/NodeColourConstant.cpp b/Engine/Plugins/Experimental/Mutable/Source/MutableTools/Private/MuT/NodeColourConstant.cpp
--- a/Engine/Plugins/Experimental/Mutable/Source/MutableTools/Private/MuT/NodeColourConstant.cpp
+++ b/Engine/Plugins/Experimental/Mutable/Source/MutableTools/Private/MuT/NodeColourConstant.cpp
@@ -64,34 +64,16 @@ namespace mu
 	//---------------------------------------------------------------------------------------------
 	// Own Interface
 	//---------------------------------------------------------------------------------------------
-	void NodeColourConstant::GetValue( float* pR, float *pG, float* pB, float* pA) const
+	FVector4f NodeColourConstant::GetValue() const
 	{
-		if (pR)
-		{
-			*pR = m_pD->m_value[0];
-		}
-
-		if (pG)
-		{
-			*pG = m_pD->m_value[1];
-		}
-
-		if (pB)
-		{
-			*pB = m_pD->m_value[2];
-		}
-
-		if (pA)
-		{
-			*pA = m_pD->m_value[3];
-		}
+		return m_pD->m_value;
 	}
 
 
 	//---------------------------------------------------------------------------------------------
-	void NodeColourConstant::SetValue( float r, float g, float b, float a)
+	void NodeColourConstant::SetValue(FVector4f Value)
 	{
-		m_pD->m_value = FVector4f(r, g, b, a);
+		m_pD->m_value = Value;
 	}
 
 
diff --git a/Engine/Plugins/Experimental/Mutable/Source/MutableTools/Public/MuT/NodeColourConstant.h b/Engine/Plugins/Experimental/Mutable/Source/MutableTools/Public/MuT/NodeColourConstant.h
--- a/Engine/Plugins/Experimental/Mutable/Source/MutableTools/Public/MuT/NodeColourConstant.h
+++ b/Engine/Plugins/Experimental/Mutable/Source/MutableTools/Public/MuT/NodeColourConstant.h
@@ -51,10 +51,10 @@ namespace mu
 		//-----------------------------------------------------------------------------------------
 
 		//! Get the value that this node returns
-		void GetValue( float* pR, float* pG, float* pB, float* pA) const;
+		FVector4f GetValue() const;
 
 		//! Set the value to be returned by this node
-		void SetValue( float r, float g, float b, float a);
+		void SetValue(FVector4f);
 
 		//-----------------------------------------------------------------------------------------
 		// Interface pattern
