[Mutable] Fix TextureParameters being removed from the ImageProvider's cache while still in use.

#rnx

diff --git a/Engine/Plugins/Experimental/Mutable/Source/CustomizableObject/Classes/MuCO/CustomizableObjectSystem.h b/Engine/Plugins/Experimental/Mutable/Source/CustomizableObject/Classes/MuCO/CustomizableObjectSystem.h
--- a/Engine/Plugins/Experimental/Mutable/Source/CustomizableObject/Classes/MuCO/CustomizableObjectSystem.h
+++ b/Engine/Plugins/Experimental/Mutable/Source/CustomizableObject/Classes/MuCO/CustomizableObjectSystem.h
@@ -184,9 +184,9 @@ public:
 };
 
 
-// TODO GMT Remove hotfix. External images are uncached images if Parameters is valid. Before the Mutable Queue rework this made sense, but this is no longer the case. Remove this when doing MTBL-1409.
+// Before the Mutable Queue rework this made sense, but this is no longer the case. Remove this when doing MTBL-1409.
 /** End a Customizable Object Instance Update. All code paths of an update have to end here. */
-void FinishUpdateGlobal(UCustomizableObjectInstance* Instance, EUpdateResult UpdateResult, FInstanceUpdateDelegate* UpdateCallback, mu::ParametersPtrConst Parameters, const FDescriptorRuntimeHash InUpdatedHash = FDescriptorRuntimeHash());
+void FinishUpdateGlobal(UCustomizableObjectInstance* Instance, EUpdateResult UpdateResult, FInstanceUpdateDelegate* UpdateCallback, const FDescriptorRuntimeHash InUpdatedHash = FDescriptorRuntimeHash());
 
 
 UCLASS(Blueprintable, BlueprintType)
diff --git a/Engine/Plugins/Experimental/Mutable/Source/CustomizableObject/Private/MuCO/CustomizableObjectInstance.cpp b/Engine/Plugins/Experimental/Mutable/Source/CustomizableObject/Private/MuCO/CustomizableObjectInstance.cpp
--- a/Engine/Plugins/Experimental/Mutable/Source/CustomizableObject/Private/MuCO/CustomizableObjectInstance.cpp
+++ b/Engine/Plugins/Experimental/Mutable/Source/CustomizableObject/Private/MuCO/CustomizableObjectInstance.cpp
@@ -362,12 +362,15 @@ void UCustomizableObjectInstance::ReleaseMutableResources(bool bCalledFromBeginD
 			}
 		}
 
-		// Remove all references to cached Texture Parameters
-		FUnrealMutableImageProvider* ImageProvider = UCustomizableObjectSystem::GetInstance()->GetPrivateChecked()->GetImageProviderChecked();
-
-		for (const FString& TextureParameter : PrivateData->UpdateTextureParameters)
+		// Remove all references to cached Texture Parameters. Only if we're destroying the COI.
+		if (bCalledFromBeginDestroy)
 		{
-			ImageProvider->UnCacheImage(TextureParameter, false);				
+			FUnrealMutableImageProvider* ImageProvider = UCustomizableObjectSystem::GetInstance()->GetPrivateChecked()->GetImageProviderChecked();
+
+			for (const FString& TextureParameter : PrivateData->UpdateTextureParameters)
+			{
+				ImageProvider->UnCacheImage(TextureParameter, false);
+			}
 		}
 	}
 
@@ -853,17 +856,9 @@ void UCustomizableObjectInstance::DoUpdateSkeletalMesh(bool bIsCloseDistTick, bo
 
 	UCustomizableObjectSystem* System = UCustomizableObjectSystem::GetInstance();
 
-	const mu::ParametersPtr Parameters = Descriptor.GetParameters();
-	
-	// Cache Texture Parameters being used during the update:
-	if (Parameters)
-	{
-		System->GetPrivate()->GetImageProviderChecked()->CacheImages(*Parameters);		
-	}
-
 	if (!CanUpdateInstance())
 	{
-		FinishUpdateGlobal(this, EUpdateResult::ErrorDiscarded, UpdateCallback, Parameters);
+		FinishUpdateGlobal(this, EUpdateResult::ErrorDiscarded, UpdateCallback);
 		return;
 	}
 	
@@ -872,7 +867,7 @@ void UCustomizableObjectInstance::DoUpdateSkeletalMesh(bool bIsCloseDistTick, bo
 	{
 	case EUpdateRequired::NoUpdate:
 	{	
-		FinishUpdateGlobal(this, EUpdateResult::ErrorDiscarded, UpdateCallback, Parameters);
+		FinishUpdateGlobal(this, EUpdateResult::ErrorDiscarded, UpdateCallback);
 		break;
 	}		
 	case EUpdateRequired::Update:
@@ -912,7 +907,7 @@ void UCustomizableObjectInstance::DoUpdateSkeletalMesh(bool bIsCloseDistTick, bo
 	{
 		System->GetPrivate()->InitDiscardResourcesSkeletalMesh(this);
 		
-		FinishUpdateGlobal(this, EUpdateResult::ErrorDiscarded, UpdateCallback, Parameters);
+		FinishUpdateGlobal(this, EUpdateResult::ErrorDiscarded, UpdateCallback);
 		break;
 	}
 
diff --git a/Engine/Plugins/Experimental/Mutable/Source/CustomizableObject/Private/MuCO/CustomizableObjectSystem.cpp b/Engine/Plugins/Experimental/Mutable/Source/CustomizableObject/Private/MuCO/CustomizableObjectSystem.cpp
--- a/Engine/Plugins/Experimental/Mutable/Source/CustomizableObject/Private/MuCO/CustomizableObjectSystem.cpp
+++ b/Engine/Plugins/Experimental/Mutable/Source/CustomizableObject/Private/MuCO/CustomizableObjectSystem.cpp
@@ -133,7 +133,7 @@ void FMutablePendingInstanceWork::AddUpdate(const FMutablePendingInstanceUpdate&
 		FMutablePendingInstanceUpdate TaskToEnqueue = UpdateToAdd;
 
 		FInstanceUpdateDelegate* Callback = ExistingUpdate->Callback;
-		FinishUpdateGlobal(ExistingUpdate->CustomizableObjectInstance.Get(), EUpdateResult::ErrorReplaced, Callback,ExistingUpdate->CustomizableObjectInstance.Get()->GetDescriptor().GetParameters());
+		FinishUpdateGlobal(ExistingUpdate->CustomizableObjectInstance.Get(), EUpdateResult::ErrorReplaced, Callback);
 
 		TaskToEnqueue.PriorityType = FMath::Min(ExistingUpdate->PriorityType, UpdateToAdd.PriorityType);
 		TaskToEnqueue.SecondsAtUpdate = FMath::Min(ExistingUpdate->SecondsAtUpdate, UpdateToAdd.SecondsAtUpdate);
@@ -148,7 +148,7 @@ void FMutablePendingInstanceWork::AddUpdate(const FMutablePendingInstanceUpdate&
 
 	if (FMutablePendingInstanceDiscard* ExistingDiscard = PendingInstanceDiscards.Find(UpdateToAdd.CustomizableObjectInstance))
 	{
-		FinishUpdateGlobal(ExistingDiscard->CustomizableObjectInstance.Get(), EUpdateResult::ErrorReplaced, nullptr, ExistingDiscard->CustomizableObjectInstance.Get()->GetDescriptor().GetParameters());
+		FinishUpdateGlobal(ExistingDiscard->CustomizableObjectInstance.Get(), EUpdateResult::ErrorReplaced, nullptr);
 
 		PendingInstanceDiscards.Remove(ExistingDiscard->CustomizableObjectInstance);
 	}
@@ -171,7 +171,7 @@ void FMutablePendingInstanceWork::AddDiscard(const FMutablePendingInstanceDiscar
 {
 	if (FMutablePendingInstanceUpdate* ExistingUpdate = PendingInstanceUpdates.Find(TaskToEnqueue.CustomizableObjectInstance.Get()))
 	{
-		FinishUpdateGlobal(ExistingUpdate->CustomizableObjectInstance.Get(), EUpdateResult::ErrorReplaced, ExistingUpdate->Callback, ExistingUpdate->CustomizableObjectInstance.Get()->GetDescriptor().GetParameters());
+		FinishUpdateGlobal(ExistingUpdate->CustomizableObjectInstance.Get(), EUpdateResult::ErrorReplaced, ExistingUpdate->Callback);
 		PendingInstanceUpdates.Remove(ExistingUpdate->CustomizableObjectInstance);
 	}
 
@@ -742,7 +742,7 @@ FAutoConsoleVariableRef CVarMaxTextureSizeToGenerate(
 		"If a texture doesn't have small enough mips, mip 0 will be the last mip available."));
 
 
-void FinishUpdateGlobal(UCustomizableObjectInstance* Instance, EUpdateResult UpdateResult, FInstanceUpdateDelegate* UpdateCallback, mu::ParametersPtrConst Parameters, const FDescriptorRuntimeHash InUpdatedHash)
+void FinishUpdateGlobal(UCustomizableObjectInstance* Instance, EUpdateResult UpdateResult, FInstanceUpdateDelegate* UpdateCallback, const FDescriptorRuntimeHash InUpdatedHash)
 {
 	check(IsInGameThread())
 	
@@ -774,17 +774,12 @@ void FinishUpdateGlobal(UCustomizableObjectInstance* Instance, EUpdateResult Upd
 		UpdateCallback->ExecuteIfBound(Context);
 	}
 
-	if (Parameters)
-	{
-		UCustomizableObjectSystem::GetInstance()->GetPrivateChecked()->GetImageProviderChecked()->UnCacheImages(*Parameters);
-	}
-
 	UCustomizableObjectSystem::GetInstance()->GetPrivate()->MutableTaskGraph.AllowLaunchingMutableTaskLowPriority(true, false);
 }
 
 
 /** Update the given Instance Skeletal Meshes and call its callbacks. */
-void UpdateSkeletalMesh(UCustomizableObjectInstance& CustomizableObjectInstance, const FDescriptorRuntimeHash& UpdatedDescriptorRuntimeHash, EUpdateResult UpdateResult, FInstanceUpdateDelegate* UpdateCallback, mu::ParametersPtrConst Parameters = nullptr)
+void UpdateSkeletalMesh(UCustomizableObjectInstance& CustomizableObjectInstance, const FDescriptorRuntimeHash& UpdatedDescriptorRuntimeHash, EUpdateResult UpdateResult, FInstanceUpdateDelegate* UpdateCallback)
 {
 	MUTABLE_CPUPROFILER_SCOPE(UpdateSkeletalMesh);
 
@@ -832,7 +827,7 @@ void UpdateSkeletalMesh(UCustomizableObjectInstance& CustomizableObjectInstance,
 
 	CustomizableObjectInstance.bEditorPropertyChanged = false;
 
-	FinishUpdateGlobal(&CustomizableObjectInstance, UpdateResult, UpdateCallback, Parameters, UpdatedDescriptorRuntimeHash);
+	FinishUpdateGlobal(&CustomizableObjectInstance, UpdateResult, UpdateCallback, UpdatedDescriptorRuntimeHash);
 }
 
 
@@ -879,7 +874,7 @@ void FCustomizableObjectSystemPrivate::InitUpdateSkeletalMesh(UCustomizableObjec
 	{
 		if (UpdateDescriptorHash.IsSubset(FDescriptorRuntimeHash(QueueElem->InstanceDescriptor)))
 		{
-			FinishUpdateGlobal(&Instance, EUpdateResult::ErrorOptimized, UpdateCallback, Instance.GetDescriptor().GetParameters());			
+			FinishUpdateGlobal(&Instance, EUpdateResult::ErrorOptimized, UpdateCallback);			
 			return; // The the requested update is equal to the last enqueued update.
 		}
 	}	
@@ -888,7 +883,7 @@ void FCustomizableObjectSystemPrivate::InitUpdateSkeletalMesh(UCustomizableObjec
 		&Instance == CurrentMutableOperation->CustomizableObjectInstance &&
 		UpdateDescriptorHash.IsSubset(CurrentMutableOperation->InstanceDescriptorRuntimeHash))
 	{
-		FinishUpdateGlobal(&Instance, EUpdateResult::ErrorOptimized, UpdateCallback, Instance.GetDescriptor().GetParameters());			
+		FinishUpdateGlobal(&Instance, EUpdateResult::ErrorOptimized, UpdateCallback);			
 		return; // The requested update is equal to the running update.
 	}
 	
@@ -902,7 +897,42 @@ void FCustomizableObjectSystemPrivate::InitUpdateSkeletalMesh(UCustomizableObjec
 		UpdateSkeletalMesh(Instance, Instance.GetDescriptorRuntimeHash(), EUpdateResult::Success, UpdateCallback);
 	}
 	else
-	{	
+	{
+		// Cache Texture Parameters being used during the update:
+		check(ImageProvider);
+
+		// Cache new Texture Parameters
+		for (const FCustomizableObjectTextureParameterValue& TextureParameters : Instance.GetDescriptor().GetTextureParameters())
+		{
+			ImageProvider->CacheImage(TextureParameters.ParameterValue, false);
+
+			for (const FString& TextureParameter : TextureParameters.ParameterRangeValues)
+			{
+				ImageProvider->CacheImage(TextureParameter, false);
+			}
+		}
+
+		UCustomizableInstancePrivateData* InstancePrivate = Instance.GetPrivate();
+		check(InstancePrivate);
+
+		// Uncache old Texture Parameters
+		for (const FString& TextureParameter : InstancePrivate->UpdateTextureParameters)
+		{
+			ImageProvider->UnCacheImage(TextureParameter, false);
+		}
+
+		// Update which ones are currently are being used
+		InstancePrivate->UpdateTextureParameters.Reset();
+		for (const FCustomizableObjectTextureParameterValue& TextureParameters : Instance.GetDescriptor().GetTextureParameters())
+		{
+			InstancePrivate->UpdateTextureParameters.Add(TextureParameters.ParameterValue);
+
+			for (const FString& TextureParameter : TextureParameters.ParameterRangeValues)
+			{
+				InstancePrivate->UpdateTextureParameters.Add(TextureParameter);
+			}
+		}
+
 		Instance.SkeletalMeshStatus = ESkeletalMeshState::AsyncUpdatePending;
 
 		if (!bIsCloseDistTick) // When called from bIsCloseDistTick, the update operation is directly processed without going to a queue first
@@ -1862,7 +1892,7 @@ namespace impl
 		UCustomizableObjectSystem* System = UCustomizableObjectSystem::GetInstance();
 		if (!System || !System->IsValidLowLevel() || System->HasAnyFlags(RF_BeginDestroyed))
 		{
-			FinishUpdateGlobal(CustomizableObjectInstancePtr.Get(), EUpdateResult::Error, &OperationData->UpdateCallback, OperationData->MutableParameters);
+			FinishUpdateGlobal(CustomizableObjectInstancePtr.Get(), EUpdateResult::Error, &OperationData->UpdateCallback);
 			return;
 		}
 
@@ -1872,7 +1902,7 @@ namespace impl
 		if (!CustomizableObjectInstance || !CustomizableObjectInstance->IsValidLowLevel() )
 		{
 			System->ClearCurrentMutableOperation();
-			FinishUpdateGlobal(CustomizableObjectInstance, EUpdateResult::Error, &OperationData->UpdateCallback, OperationData->MutableParameters);
+			FinishUpdateGlobal(CustomizableObjectInstance, EUpdateResult::Error, &OperationData->UpdateCallback);
 			return;
 		}
 
@@ -1881,7 +1911,7 @@ namespace impl
 
 		// Actual work
 		// TODO MTBL-391: Review This hotfix
-		UpdateSkeletalMesh(*CustomizableObjectInstance, CustomizableObjectSystemPrivateData->CurrentMutableOperation->InstanceDescriptorRuntimeHash, OperationData->UpdateResult, &OperationData->UpdateCallback, OperationData->MutableParameters);
+		UpdateSkeletalMesh(*CustomizableObjectInstance, CustomizableObjectSystemPrivateData->CurrentMutableOperation->InstanceDescriptorRuntimeHash, OperationData->UpdateResult, &OperationData->UpdateCallback);
 
 		// All work is done, release unused textures.
 		if (CustomizableObjectSystemPrivateData->bReleaseTexturesImmediately)
@@ -1924,7 +1954,7 @@ namespace impl
 		UCustomizableObjectSystem* System = UCustomizableObjectSystem::GetInstance();
 		if (!System || !System->IsValidLowLevel() || System->HasAnyFlags(RF_BeginDestroyed))
 		{
-			FinishUpdateGlobal(CustomizableObjectInstancePtr.Get(), EUpdateResult::Error, &OperationData->UpdateCallback, OperationData->MutableParameters);
+			FinishUpdateGlobal(CustomizableObjectInstancePtr.Get(), EUpdateResult::Error, &OperationData->UpdateCallback);
 			return;
 		}
 
@@ -2057,7 +2087,7 @@ namespace impl
 		}
 		else
 		{
-			FinishUpdateGlobal(CustomizableObjectInstance, EUpdateResult::Error, &OperationData->UpdateCallback, OperationData->MutableParameters);
+			FinishUpdateGlobal(CustomizableObjectInstance, EUpdateResult::Error, &OperationData->UpdateCallback);
 		}
 	}
 
@@ -2081,7 +2111,7 @@ namespace impl
 		if (!ObjectInstance)
 		{
 			System->ClearCurrentMutableOperation();
-			FinishUpdateGlobal(ObjectInstance, EUpdateResult::Error, &OperationData->UpdateCallback, OperationData->MutableParameters);
+			FinishUpdateGlobal(ObjectInstance, EUpdateResult::Error, &OperationData->UpdateCallback);
 			return;
 		}
 
@@ -2103,7 +2133,7 @@ namespace impl
 		if (!CustomizableObject)
 		{
 			System->ClearCurrentMutableOperation();
-			FinishUpdateGlobal(CustomizableObjectInstancePtr.Get(), EUpdateResult::Error, &OperationData->UpdateCallback, OperationData->MutableParameters);
+			FinishUpdateGlobal(CustomizableObjectInstancePtr.Get(), EUpdateResult::Error, &OperationData->UpdateCallback);
 			return;
 		}
 
@@ -2115,37 +2145,6 @@ namespace impl
 			ObjectInstancePrivateData->RelevantParameters = OperationData->RelevantParametersInProgress;
 		}
 
-		// Cache Texture Parameters
-		FUnrealMutableImageProvider* ImageProvider = UCustomizableObjectSystem::GetInstance()->GetPrivateChecked()->GetImageProviderChecked();
-
-		// Cache new Texture Parameters
-		for (const FCustomizableObjectTextureParameterValue& TextureParameters : ObjectInstance->GetDescriptor().GetTextureParameters())
-		{
-			ImageProvider->CacheImage(TextureParameters.ParameterValue, false);				
-		
-			for (const FString& TextureParameter: TextureParameters.ParameterRangeValues)
-			{
-				ImageProvider->CacheImage(TextureParameter, false);
-			}
-		}
-		
-		// Uncache old Texture Parameters
-		for (const FString& TextureParameter : ObjectInstancePrivateData->UpdateTextureParameters)
-		{
-			ImageProvider->UnCacheImage(TextureParameter, false);
-		}
-
-		// Update which ones are currently are being used
-		ObjectInstancePrivateData->UpdateTextureParameters.Reset();
-		for (const FCustomizableObjectTextureParameterValue& TextureParameters : ObjectInstance->GetDescriptor().GetTextureParameters())
-		{
-			ObjectInstancePrivateData->UpdateTextureParameters.Add(TextureParameters.ParameterValue);
-		
-			for (const FString& TextureParameter: TextureParameters.ParameterRangeValues)
-			{
-				ObjectInstancePrivateData->UpdateTextureParameters.Add(TextureParameter);
-			}
-		}
 		
 		// Selectively lock the resource cache for the object used by this instance to avoid the destruction of resources that we may want to reuse.
 		// When protecting textures there mustn't be any left from a previous update
@@ -2287,7 +2286,7 @@ namespace impl
 		if (!Operation->CustomizableObjectInstance.IsValid() || !Operation->CustomizableObjectInstance->IsValidLowLevel()) // Only start if it hasn't been already destroyed (i.e. GC after finish PIE)
 		{
 			System->ClearCurrentMutableOperation();
-			FinishUpdateGlobal(nullptr, EUpdateResult::Error, &Operation->UpdateCallback, Operation->GetParameters());
+			FinishUpdateGlobal(nullptr, EUpdateResult::Error, &Operation->UpdateCallback);
 			return;
 		}
 
@@ -2297,7 +2296,7 @@ namespace impl
 		if (!CandidateInstancePrivateData)
 		{
 			System->ClearCurrentMutableOperation();
-			FinishUpdateGlobal(nullptr, EUpdateResult::Error, &Operation->UpdateCallback, Operation->GetParameters());
+			FinishUpdateGlobal(nullptr, EUpdateResult::Error, &Operation->UpdateCallback);
 			return;
 		}
 		
@@ -2339,7 +2338,7 @@ namespace impl
 		if (Operation->InstanceDescriptorRuntimeHash.IsSubset(CandidateInstance->GetDescriptorRuntimeHash()))
 		{
 			CandidateInstance->SkeletalMeshStatus = ESkeletalMeshState::Correct;
-			UpdateSkeletalMesh(*CandidateInstance, CandidateInstance->GetDescriptorRuntimeHash(), EUpdateResult::Success, &Operation->UpdateCallback, Operation->GetParameters());
+			UpdateSkeletalMesh(*CandidateInstance, CandidateInstance->GetDescriptorRuntimeHash(), EUpdateResult::Success, &Operation->UpdateCallback);
 			bCancel = true;
 		}
 
@@ -2355,7 +2354,7 @@ namespace impl
 
 			System->ClearCurrentMutableOperation();
 
-			FinishUpdateGlobal(CandidateInstance, EUpdateResult::Error, &Operation->UpdateCallback, Operation->GetParameters());
+			FinishUpdateGlobal(CandidateInstance, EUpdateResult::Error, &Operation->UpdateCallback);
 			return;
 		}
 
diff --git a/Engine/Plugins/Experimental/Mutable/Source/CustomizableObject/Private/MuCO/UnrealMutableImageProvider.cpp b/Engine/Plugins/Experimental/Mutable/Source/CustomizableObject/Private/MuCO/UnrealMutableImageProvider.cpp
--- a/Engine/Plugins/Experimental/Mutable/Source/CustomizableObject/Private/MuCO/UnrealMutableImageProvider.cpp
+++ b/Engine/Plugins/Experimental/Mutable/Source/CustomizableObject/Private/MuCO/UnrealMutableImageProvider.cpp
@@ -579,6 +579,10 @@ void FUnrealMutableImageProvider::UnCacheImage(const FExternalImageID& Id, bool
 			GlobalExternalImages.Remove(Id);		
 		}
 	}
+	else
+	{
+		UE_LOG(LogMutable, Warning, TEXT("Failed to uncache external image %s. Possible double free!"), *Id.GetName());
+	}
 }
 
 
