[mutable] Remove invalid assert on null image generation.
#rnx

diff --git a/Engine/Plugins/Experimental/Mutable/Source/MutableTools/Private/MuT/CodeOptimiser.cpp b/Engine/Plugins/Experimental/Mutable/Source/MutableTools/Private/MuT/CodeOptimiser.cpp
--- a/Engine/Plugins/Experimental/Mutable/Source/MutableTools/Private/MuT/CodeOptimiser.cpp
+++ b/Engine/Plugins/Experimental/Mutable/Source/MutableTools/Private/MuT/CodeOptimiser.cpp
@@ -505,8 +505,6 @@ namespace mu
 					constantOp->SetValue( pImage, bUseDiskCache );
 					m_result = constantOp;
 				}
-
-				check(m_result);
 				break;
 			}
 
