[mutable] Fix texture nodes pin rename from image

#rnx

diff --git a/Engine/Plugins/Experimental/Mutable/Source/CustomizableObject/Public/MuCO/CustomizableObjectCustomVersion.h b/Engine/Plugins/Experimental/Mutable/Source/CustomizableObject/Public/MuCO/CustomizableObjectCustomVersion.h
--- a/Engine/Plugins/Experimental/Mutable/Source/CustomizableObject/Public/MuCO/CustomizableObjectCustomVersion.h
+++ b/Engine/Plugins/Experimental/Mutable/Source/CustomizableObject/Public/MuCO/CustomizableObjectCustomVersion.h
@@ -88,6 +88,8 @@ struct CUSTOMIZABLEOBJECT_API FCustomizableObjectCustomVersion
 		
 		AddedTableMaterialSwitch,
 
+		FixPinsNamesImageToTexture2,
+
 		// -----<new versions can be added above this line>-------------------------------------------------
 		VersionPlusOne,
 		LatestVersion = VersionPlusOne - 1
diff --git a/Engine/Plugins/Experimental/Mutable/Source/CustomizableObjectEditor/Classes/MuCOE/Nodes/CustomizableObjectNodeTexture.h b/Engine/Plugins/Experimental/Mutable/Source/CustomizableObjectEditor/Classes/MuCOE/Nodes/CustomizableObjectNodeTexture.h
--- a/Engine/Plugins/Experimental/Mutable/Source/CustomizableObjectEditor/Classes/MuCOE/Nodes/CustomizableObjectNodeTexture.h
+++ b/Engine/Plugins/Experimental/Mutable/Source/CustomizableObjectEditor/Classes/MuCOE/Nodes/CustomizableObjectNodeTexture.h
@@ -82,7 +82,8 @@ public:
 	GENERATED_BODY()
 
 	// UCustomizableObjectNode interface
-	void AllocateDefaultPins(UCustomizableObjectNodeRemapPins* RemapPins) override;
+	virtual void AllocateDefaultPins(UCustomizableObjectNodeRemapPins* RemapPins) override;
+	virtual void BackwardsCompatibleFixup() override;
 
 	// Begin EdGraphNode interface
 	FText GetNodeTitle(ENodeTitleType::Type TitleType) const override;
diff --git a/Engine/Plugins/Experimental/Mutable/Source/CustomizableObjectEditor/Classes/MuCOE/Nodes/CustomizableObjectNodeTextureBinarise.h b/Engine/Plugins/Experimental/Mutable/Source/CustomizableObjectEditor/Classes/MuCOE/Nodes/CustomizableObjectNodeTextureBinarise.h
--- a/Engine/Plugins/Experimental/Mutable/Source/CustomizableObjectEditor/Classes/MuCOE/Nodes/CustomizableObjectNodeTextureBinarise.h
+++ b/Engine/Plugins/Experimental/Mutable/Source/CustomizableObjectEditor/Classes/MuCOE/Nodes/CustomizableObjectNodeTextureBinarise.h
@@ -24,7 +24,8 @@ public:
 	virtual void Serialize(FArchive& Ar) override;
 	
 	// UCustomizableObjectNode interface
-	void AllocateDefaultPins(UCustomizableObjectNodeRemapPins* RemapPins) override;
+	virtual void AllocateDefaultPins(UCustomizableObjectNodeRemapPins* RemapPins) override;
+	virtual void BackwardsCompatibleFixup() override;
 
 	UEdGraphPin* GetBaseImagePin() const;
 
diff --git a/Engine/Plugins/Experimental/Mutable/Source/CustomizableObjectEditor/Classes/MuCOE/Nodes/CustomizableObjectNodeTextureColourMap.h b/Engine/Plugins/Experimental/Mutable/Source/CustomizableObjectEditor/Classes/MuCOE/Nodes/CustomizableObjectNodeTextureColourMap.h
--- a/Engine/Plugins/Experimental/Mutable/Source/CustomizableObjectEditor/Classes/MuCOE/Nodes/CustomizableObjectNodeTextureColourMap.h
+++ b/Engine/Plugins/Experimental/Mutable/Source/CustomizableObjectEditor/Classes/MuCOE/Nodes/CustomizableObjectNodeTextureColourMap.h
@@ -25,7 +25,8 @@ public:
 	FText GetTooltipText() const override;
 
 	// UCustomizableObjectNode interface
-	void AllocateDefaultPins(UCustomizableObjectNodeRemapPins* RemapPins) override;
+	virtual void AllocateDefaultPins(UCustomizableObjectNodeRemapPins* RemapPins) override;
+	virtual void BackwardsCompatibleFixup() override;
 
 	UEdGraphPin* GetMapPin() const
 	{
diff --git a/Engine/Plugins/Experimental/Mutable/Source/CustomizableObjectEditor/Classes/MuCOE/Nodes/CustomizableObjectNodeTextureFromChannels.h b/Engine/Plugins/Experimental/Mutable/Source/CustomizableObjectEditor/Classes/MuCOE/Nodes/CustomizableObjectNodeTextureFromChannels.h
--- a/Engine/Plugins/Experimental/Mutable/Source/CustomizableObjectEditor/Classes/MuCOE/Nodes/CustomizableObjectNodeTextureFromChannels.h
+++ b/Engine/Plugins/Experimental/Mutable/Source/CustomizableObjectEditor/Classes/MuCOE/Nodes/CustomizableObjectNodeTextureFromChannels.h
@@ -25,7 +25,8 @@ public:
 	FText GetTooltipText() const override;
 
 	// UCustomizableObjectNode interface
-	void AllocateDefaultPins(UCustomizableObjectNodeRemapPins* RemapPins) override;
+	virtual void AllocateDefaultPins(UCustomizableObjectNodeRemapPins* RemapPins) override;
+	virtual void BackwardsCompatibleFixup() override;
 
 	UEdGraphPin* RPin() const
 	{
diff --git a/Engine/Plugins/Experimental/Mutable/Source/CustomizableObjectEditor/Classes/MuCOE/Nodes/CustomizableObjectNodeTextureFromColor.h b/Engine/Plugins/Experimental/Mutable/Source/CustomizableObjectEditor/Classes/MuCOE/Nodes/CustomizableObjectNodeTextureFromColor.h
--- a/Engine/Plugins/Experimental/Mutable/Source/CustomizableObjectEditor/Classes/MuCOE/Nodes/CustomizableObjectNodeTextureFromColor.h
+++ b/Engine/Plugins/Experimental/Mutable/Source/CustomizableObjectEditor/Classes/MuCOE/Nodes/CustomizableObjectNodeTextureFromColor.h
@@ -25,7 +25,8 @@ public:
 	FText GetTooltipText() const override;
 
 	// UCustomizableObjectNode interface
-	void AllocateDefaultPins(UCustomizableObjectNodeRemapPins* RemapPins) override;
+	virtual void AllocateDefaultPins(UCustomizableObjectNodeRemapPins* RemapPins) override;
+	virtual void BackwardsCompatibleFixup() override;
 
 	UEdGraphPin* ColorPin() const
 	{
diff --git a/Engine/Plugins/Experimental/Mutable/Source/CustomizableObjectEditor/Classes/MuCOE/Nodes/CustomizableObjectNodeTextureInvert.h b/Engine/Plugins/Experimental/Mutable/Source/CustomizableObjectEditor/Classes/MuCOE/Nodes/CustomizableObjectNodeTextureInvert.h
--- a/Engine/Plugins/Experimental/Mutable/Source/CustomizableObjectEditor/Classes/MuCOE/Nodes/CustomizableObjectNodeTextureInvert.h
+++ b/Engine/Plugins/Experimental/Mutable/Source/CustomizableObjectEditor/Classes/MuCOE/Nodes/CustomizableObjectNodeTextureInvert.h
@@ -24,7 +24,8 @@ public:
 	virtual void Serialize(FArchive& Ar) override;
 	
 	// UCustomizableObjectNode interface
-	void AllocateDefaultPins(UCustomizableObjectNodeRemapPins* RemapPins) override;
+	virtual void AllocateDefaultPins(UCustomizableObjectNodeRemapPins* RemapPins) override;
+	virtual void BackwardsCompatibleFixup() override;
 
 	UEdGraphPin* GetBaseImagePin() const;
 
diff --git a/Engine/Plugins/Experimental/Mutable/Source/CustomizableObjectEditor/Classes/MuCOE/Nodes/CustomizableObjectNodeTextureLayer.h b/Engine/Plugins/Experimental/Mutable/Source/CustomizableObjectEditor/Classes/MuCOE/Nodes/CustomizableObjectNodeTextureLayer.h
--- a/Engine/Plugins/Experimental/Mutable/Source/CustomizableObjectEditor/Classes/MuCOE/Nodes/CustomizableObjectNodeTextureLayer.h
+++ b/Engine/Plugins/Experimental/Mutable/Source/CustomizableObjectEditor/Classes/MuCOE/Nodes/CustomizableObjectNodeTextureLayer.h
@@ -70,7 +70,8 @@ public:
 	void PostPasteNode() override;
 
 	// UCustomizableObjectNode interface
-	void AllocateDefaultPins(UCustomizableObjectNodeRemapPins* RemapPins) override;
+	virtual void AllocateDefaultPins(UCustomizableObjectNodeRemapPins* RemapPins) override;
+	virtual void BackwardsCompatibleFixup() override;
 	virtual bool CanConnect(const UEdGraphPin* InOwnedInputPin, const UEdGraphPin* InOutputPin, bool& bOutIsOtherNodeBlocklisted, bool& bOutArePinsCompatible) const override;
 		
 	UEdGraphPin* BasePin() const
diff --git a/Engine/Plugins/Experimental/Mutable/Source/CustomizableObjectEditor/Classes/MuCOE/Nodes/CustomizableObjectNodeTextureProject.h b/Engine/Plugins/Experimental/Mutable/Source/CustomizableObjectEditor/Classes/MuCOE/Nodes/CustomizableObjectNodeTextureProject.h
--- a/Engine/Plugins/Experimental/Mutable/Source/CustomizableObjectEditor/Classes/MuCOE/Nodes/CustomizableObjectNodeTextureProject.h
+++ b/Engine/Plugins/Experimental/Mutable/Source/CustomizableObjectEditor/Classes/MuCOE/Nodes/CustomizableObjectNodeTextureProject.h
@@ -80,7 +80,8 @@ public:
 	FText GetTooltipText() const override;
 
 	// UCustomizableObjectNode interface
-	void AllocateDefaultPins(UCustomizableObjectNodeRemapPins* RemapPins) override;
+	virtual void AllocateDefaultPins(UCustomizableObjectNodeRemapPins* RemapPins) override;
+	virtual void BackwardsCompatibleFixup() override;
 
 	UEdGraphPin* MeshPin() const
 	{
diff --git a/Engine/Plugins/Experimental/Mutable/Source/CustomizableObjectEditor/Classes/MuCOE/Nodes/CustomizableObjectNodeTextureSwitch.h b/Engine/Plugins/Experimental/Mutable/Source/CustomizableObjectEditor/Classes/MuCOE/Nodes/CustomizableObjectNodeTextureSwitch.h
--- a/Engine/Plugins/Experimental/Mutable/Source/CustomizableObjectEditor/Classes/MuCOE/Nodes/CustomizableObjectNodeTextureSwitch.h
+++ b/Engine/Plugins/Experimental/Mutable/Source/CustomizableObjectEditor/Classes/MuCOE/Nodes/CustomizableObjectNodeTextureSwitch.h
@@ -19,6 +19,9 @@ public:
 	// UObject interface
 	virtual void Serialize(FArchive& Ar) override;
 
+	// UCustomizableObjectNode interface
+	virtual void BackwardsCompatibleFixup() override;
+
 	// UCustomizableObjectNodeSwitchBase interface
 	FString GetOutputPinName() const override;
 
diff --git a/Engine/Plugins/Experimental/Mutable/Source/CustomizableObjectEditor/Classes/MuCOE/Nodes/CustomizableObjectNodeTextureToChannels.h b/Engine/Plugins/Experimental/Mutable/Source/CustomizableObjectEditor/Classes/MuCOE/Nodes/CustomizableObjectNodeTextureToChannels.h
--- a/Engine/Plugins/Experimental/Mutable/Source/CustomizableObjectEditor/Classes/MuCOE/Nodes/CustomizableObjectNodeTextureToChannels.h
+++ b/Engine/Plugins/Experimental/Mutable/Source/CustomizableObjectEditor/Classes/MuCOE/Nodes/CustomizableObjectNodeTextureToChannels.h
@@ -30,7 +30,8 @@ public:
 	FText GetTooltipText() const override;
 
 	// UCustomizableObjectNode interface
-	void AllocateDefaultPins(UCustomizableObjectNodeRemapPins* RemapPins) override;
+	virtual void AllocateDefaultPins(UCustomizableObjectNodeRemapPins* RemapPins) override;
+	virtual void BackwardsCompatibleFixup() override;
 
 	UEdGraphPin* InputPin() const;
 
diff --git a/Engine/Plugins/Experimental/Mutable/Source/CustomizableObjectEditor/Private/MuCOE/Nodes/CustomizableObjectNodeTexture.cpp b/Engine/Plugins/Experimental/Mutable/Source/CustomizableObjectEditor/Private/MuCOE/Nodes/CustomizableObjectNodeTexture.cpp
--- a/Engine/Plugins/Experimental/Mutable/Source/CustomizableObjectEditor/Private/MuCOE/Nodes/CustomizableObjectNodeTexture.cpp
+++ b/Engine/Plugins/Experimental/Mutable/Source/CustomizableObjectEditor/Private/MuCOE/Nodes/CustomizableObjectNodeTexture.cpp
@@ -33,6 +33,22 @@ void UCustomizableObjectNodeTexture::AllocateDefaultPins(UCustomizableObjectNode
 }
 
 
+void UCustomizableObjectNodeTexture::BackwardsCompatibleFixup()
+{
+	Super::BackwardsCompatibleFixup();
+
+	const int32 CustomizableObjectCustomVersion = GetLinkerCustomVersion(FCustomizableObjectCustomVersion::GUID);
+
+	if (CustomizableObjectCustomVersion < FCustomizableObjectCustomVersion::FixPinsNamesImageToTexture2)
+	{
+		if (UEdGraphPin* TexturePin = FindPin(TEXT("Image"))) {
+			TexturePin->PinName = TEXT("Texture");
+			UCustomizableObjectNode::ReconstructNode();
+		}
+	}
+}
+
+
 FText UCustomizableObjectNodeTexture::GetNodeTitle(ENodeTitleType::Type TitleType) const
 {
 	if (Texture)
diff --git a/Engine/Plugins/Experimental/Mutable/Source/CustomizableObjectEditor/Private/MuCOE/Nodes/CustomizableObjectNodeTextureBinarise.cpp b/Engine/Plugins/Experimental/Mutable/Source/CustomizableObjectEditor/Private/MuCOE/Nodes/CustomizableObjectNodeTextureBinarise.cpp
--- a/Engine/Plugins/Experimental/Mutable/Source/CustomizableObjectEditor/Private/MuCOE/Nodes/CustomizableObjectNodeTextureBinarise.cpp
+++ b/Engine/Plugins/Experimental/Mutable/Source/CustomizableObjectEditor/Private/MuCOE/Nodes/CustomizableObjectNodeTextureBinarise.cpp
@@ -27,6 +27,30 @@ void UCustomizableObjectNodeTextureBinarise::AllocateDefaultPins(UCustomizableOb
 }
 
 
+void UCustomizableObjectNodeTextureBinarise::BackwardsCompatibleFixup()
+{
+	Super::BackwardsCompatibleFixup();
+
+	const int32 CustomizableObjectCustomVersion = GetLinkerCustomVersion(FCustomizableObjectCustomVersion::GUID);
+
+	if (CustomizableObjectCustomVersion < FCustomizableObjectCustomVersion::FixPinsNamesImageToTexture2)
+	{
+		if (UEdGraphPin* TexturePin = FindPin(TEXT("Image"))) {
+			TexturePin->PinName = TEXT("Texture");
+			UCustomizableObjectNode::ReconstructNode();
+		}
+	}
+
+	if (CustomizableObjectCustomVersion < FCustomizableObjectCustomVersion::FixPinsNamesImageToTexture2)
+	{
+		if (UEdGraphPin* TexturePin = FindPin(TEXT("Base Image"))) {
+			TexturePin->PinName = TEXT("Base Texture");
+			UCustomizableObjectNode::ReconstructNode();
+		}
+	}
+}
+
+
 UEdGraphPin* UCustomizableObjectNodeTextureBinarise::GetBaseImagePin() const
 {
 	return BaseImagePinReference.Get();
diff --git a/Engine/Plugins/Experimental/Mutable/Source/CustomizableObjectEditor/Private/MuCOE/Nodes/CustomizableObjectNodeTextureColourMap.cpp b/Engine/Plugins/Experimental/Mutable/Source/CustomizableObjectEditor/Private/MuCOE/Nodes/CustomizableObjectNodeTextureColourMap.cpp
--- a/Engine/Plugins/Experimental/Mutable/Source/CustomizableObjectEditor/Private/MuCOE/Nodes/CustomizableObjectNodeTextureColourMap.cpp
+++ b/Engine/Plugins/Experimental/Mutable/Source/CustomizableObjectEditor/Private/MuCOE/Nodes/CustomizableObjectNodeTextureColourMap.cpp
@@ -28,7 +28,23 @@ void UCustomizableObjectNodeTextureColourMap::AllocateDefaultPins(UCustomizableO
 	PinName = TEXT("Map");
 	UEdGraphPin* GradientPin = CustomCreatePin(EGPD_Input, Schema->PC_Image, FName(*PinName));
 	GradientPin->bDefaultValueIsIgnored = true;
+}
+
 
+void UCustomizableObjectNodeTextureColourMap::BackwardsCompatibleFixup()
+{
+	Super::BackwardsCompatibleFixup();
+
+	const int32 CustomizableObjectCustomVersion = GetLinkerCustomVersion(FCustomizableObjectCustomVersion::GUID);
+
+	if (CustomizableObjectCustomVersion < FCustomizableObjectCustomVersion::FixPinsNamesImageToTexture2)
+	{
+		if (UEdGraphPin* TexturePin = FindPin(TEXT("Image")))
+		{
+			TexturePin->PinName = TEXT("Texture");
+			UCustomizableObjectNode::ReconstructNode();
+		}
+	}
 }
 
 
diff --git a/Engine/Plugins/Experimental/Mutable/Source/CustomizableObjectEditor/Private/MuCOE/Nodes/CustomizableObjectNodeTextureFromChannels.cpp b/Engine/Plugins/Experimental/Mutable/Source/CustomizableObjectEditor/Private/MuCOE/Nodes/CustomizableObjectNodeTextureFromChannels.cpp
--- a/Engine/Plugins/Experimental/Mutable/Source/CustomizableObjectEditor/Private/MuCOE/Nodes/CustomizableObjectNodeTextureFromChannels.cpp
+++ b/Engine/Plugins/Experimental/Mutable/Source/CustomizableObjectEditor/Private/MuCOE/Nodes/CustomizableObjectNodeTextureFromChannels.cpp
@@ -35,6 +35,23 @@ void UCustomizableObjectNodeTextureFromChannels::AllocateDefaultPins(UCustomizab
 }
 
 
+void UCustomizableObjectNodeTextureFromChannels::BackwardsCompatibleFixup()
+{
+	Super::BackwardsCompatibleFixup();
+
+	const int32 CustomizableObjectCustomVersion = GetLinkerCustomVersion(FCustomizableObjectCustomVersion::GUID);
+
+	if (CustomizableObjectCustomVersion < FCustomizableObjectCustomVersion::FixPinsNamesImageToTexture2)
+	{
+		if (UEdGraphPin* TexturePin = FindPin(TEXT("Image")))
+		{
+			TexturePin->PinName = TEXT("Texture");
+			UCustomizableObjectNode::ReconstructNode();
+		}
+	}
+}
+
+
 FText UCustomizableObjectNodeTextureFromChannels::GetNodeTitle(ENodeTitleType::Type TitleType) const
 {
 	return LOCTEXT("Texture_From_Channels", "Make Texture");
diff --git a/Engine/Plugins/Experimental/Mutable/Source/CustomizableObjectEditor/Private/MuCOE/Nodes/CustomizableObjectNodeTextureFromColor.cpp b/Engine/Plugins/Experimental/Mutable/Source/CustomizableObjectEditor/Private/MuCOE/Nodes/CustomizableObjectNodeTextureFromColor.cpp
--- a/Engine/Plugins/Experimental/Mutable/Source/CustomizableObjectEditor/Private/MuCOE/Nodes/CustomizableObjectNodeTextureFromColor.cpp
+++ b/Engine/Plugins/Experimental/Mutable/Source/CustomizableObjectEditor/Private/MuCOE/Nodes/CustomizableObjectNodeTextureFromColor.cpp
@@ -23,6 +23,23 @@ void UCustomizableObjectNodeTextureFromColor::AllocateDefaultPins(UCustomizableO
 }
 
 
+void UCustomizableObjectNodeTextureFromColor::BackwardsCompatibleFixup()
+{
+	Super::BackwardsCompatibleFixup();
+
+	const int32 CustomizableObjectCustomVersion = GetLinkerCustomVersion(FCustomizableObjectCustomVersion::GUID);
+
+	if (CustomizableObjectCustomVersion < FCustomizableObjectCustomVersion::FixPinsNamesImageToTexture2)
+	{
+		if (UEdGraphPin* TexturePin = FindPin(TEXT("Image")))
+		{
+			TexturePin->PinName = TEXT("Texture");
+			UCustomizableObjectNode::ReconstructNode();
+		}
+	}
+}
+
+
 FText UCustomizableObjectNodeTextureFromColor::GetNodeTitle(ENodeTitleType::Type TitleType) const
 {
 	return LOCTEXT("Texture_From_Color", "Texture From Color");
diff --git a/Engine/Plugins/Experimental/Mutable/Source/CustomizableObjectEditor/Private/MuCOE/Nodes/CustomizableObjectNodeTextureInvert.cpp b/Engine/Plugins/Experimental/Mutable/Source/CustomizableObjectEditor/Private/MuCOE/Nodes/CustomizableObjectNodeTextureInvert.cpp
--- a/Engine/Plugins/Experimental/Mutable/Source/CustomizableObjectEditor/Private/MuCOE/Nodes/CustomizableObjectNodeTextureInvert.cpp
+++ b/Engine/Plugins/Experimental/Mutable/Source/CustomizableObjectEditor/Private/MuCOE/Nodes/CustomizableObjectNodeTextureInvert.cpp
@@ -24,6 +24,35 @@ void UCustomizableObjectNodeTextureInvert::AllocateDefaultPins(UCustomizableObje
 }
 
 
+void UCustomizableObjectNodeTextureInvert::BackwardsCompatibleFixup()
+{
+	Super::BackwardsCompatibleFixup();
+
+	const int32 CustomizableObjectCustomVersion = GetLinkerCustomVersion(FCustomizableObjectCustomVersion::GUID);
+
+	if (CustomizableObjectCustomVersion < FCustomizableObjectCustomVersion::FixPinsNamesImageToTexture2)
+	{
+		bool Replaced = false;
+		if (UEdGraphPin* TexturePin = FindPin(TEXT("Image")))
+		{
+			TexturePin->PinName = TEXT("Texture");
+			Replaced = true;
+		}
+
+		if (UEdGraphPin* TexturePin = FindPin(TEXT("Base Image")))
+		{
+			TexturePin->PinName = TEXT("Base Texture");
+			Replaced = true;
+		}
+
+		if(Replaced)
+		{
+			UCustomizableObjectNode::ReconstructNode();
+		}
+	}
+}
+
+
 UEdGraphPin* UCustomizableObjectNodeTextureInvert::GetBaseImagePin() const
 {
 	return BaseImagePinReference.Get();
@@ -59,6 +88,14 @@ void UCustomizableObjectNodeTextureInvert::Serialize(FArchive& Ar)
 	{
 		BaseImagePinReference = FEdGraphPinReference(FindPin(TEXT("Base Image")));
 	}
+
+	if (Ar.CustomVer(FCustomizableObjectCustomVersion::GUID) < FCustomizableObjectCustomVersion::FixPinsNamesImageToTexture2)
+	{
+		if (UEdGraphPin* TexturePin = FindPin(TEXT("Base Image"))) {
+			TexturePin->PinName = TEXT("Base Texture");
+			UCustomizableObjectNode::ReconstructNode();
+		}
+	}
 }
 
 
diff --git a/Engine/Plugins/Experimental/Mutable/Source/CustomizableObjectEditor/Private/MuCOE/Nodes/CustomizableObjectNodeTextureLayer.cpp b/Engine/Plugins/Experimental/Mutable/Source/CustomizableObjectEditor/Private/MuCOE/Nodes/CustomizableObjectNodeTextureLayer.cpp
--- a/Engine/Plugins/Experimental/Mutable/Source/CustomizableObjectEditor/Private/MuCOE/Nodes/CustomizableObjectNodeTextureLayer.cpp
+++ b/Engine/Plugins/Experimental/Mutable/Source/CustomizableObjectEditor/Private/MuCOE/Nodes/CustomizableObjectNodeTextureLayer.cpp
@@ -69,6 +69,23 @@ void UCustomizableObjectNodeTextureLayer::AllocateDefaultPins(UCustomizableObjec
 }
 
 
+void UCustomizableObjectNodeTextureLayer::BackwardsCompatibleFixup()
+{
+	Super::BackwardsCompatibleFixup();
+
+	const int32 CustomizableObjectCustomVersion = GetLinkerCustomVersion(FCustomizableObjectCustomVersion::GUID);
+
+	if (CustomizableObjectCustomVersion < FCustomizableObjectCustomVersion::FixPinsNamesImageToTexture2)
+	{
+		if (UEdGraphPin* TexturePin = FindPin(TEXT("Image")))
+		{
+			TexturePin->PinName = TEXT("Texture");
+			UCustomizableObjectNode::ReconstructNode();
+		}
+	}
+}
+
+
 FText UCustomizableObjectNodeTextureLayer::GetNodeTitle(ENodeTitleType::Type TitleType) const
 {
 	return LOCTEXT("Texture_Layer", "Texture Layer");
diff --git a/Engine/Plugins/Experimental/Mutable/Source/CustomizableObjectEditor/Private/MuCOE/Nodes/CustomizableObjectNodeTextureProject.cpp b/Engine/Plugins/Experimental/Mutable/Source/CustomizableObjectEditor/Private/MuCOE/Nodes/CustomizableObjectNodeTextureProject.cpp
--- a/Engine/Plugins/Experimental/Mutable/Source/CustomizableObjectEditor/Private/MuCOE/Nodes/CustomizableObjectNodeTextureProject.cpp
+++ b/Engine/Plugins/Experimental/Mutable/Source/CustomizableObjectEditor/Private/MuCOE/Nodes/CustomizableObjectNodeTextureProject.cpp
@@ -116,6 +116,38 @@ void UCustomizableObjectNodeTextureProject::AllocateDefaultPins(UCustomizableObj
 }
 
 
+void UCustomizableObjectNodeTextureProject::BackwardsCompatibleFixup()
+{
+	Super::BackwardsCompatibleFixup();
+
+	const int32 CustomizableObjectCustomVersion = GetLinkerCustomVersion(FCustomizableObjectCustomVersion::GUID);
+
+	if (CustomizableObjectCustomVersion < FCustomizableObjectCustomVersion::FixPinsNamesImageToTexture2)
+	{
+		TArray<UEdGraphPin*> AllPins = GetAllPins();
+		uint32 TexturePinsCount = 0;
+		uint32 OutputPinsCount = 0;
+
+		for (UEdGraphPin* Pin : AllPins)
+		{
+			if (Pin && Pin->GetName().StartsWith(TEXT("Image "))) {
+				if (Pin->Direction == EGPD_Input)
+				{
+					Pin->PinName = FName(*FString::Printf(TEXT("Texture %d"), TexturePinsCount));
+					TexturePinsCount++;
+				}
+				else if (Pin->Direction == EGPD_Output)
+				{
+					Pin->PinName = FName(*FString::Printf(TEXT("Texture %d"), OutputPinsCount));
+					OutputPinsCount++;
+				}
+			}
+		}
+		UCustomizableObjectNode::ReconstructNode();
+	}
+}
+
+
 UEdGraphPin* UCustomizableObjectNodeTextureProject::TexturePins(int32 Index) const
 {
 	return TexturePinsReferences[Index].Get();	
diff --git a/Engine/Plugins/Experimental/Mutable/Source/CustomizableObjectEditor/Private/MuCOE/Nodes/CustomizableObjectNodeTextureSwitch.cpp b/Engine/Plugins/Experimental/Mutable/Source/CustomizableObjectEditor/Private/MuCOE/Nodes/CustomizableObjectNodeTextureSwitch.cpp
--- a/Engine/Plugins/Experimental/Mutable/Source/CustomizableObjectEditor/Private/MuCOE/Nodes/CustomizableObjectNodeTextureSwitch.cpp
+++ b/Engine/Plugins/Experimental/Mutable/Source/CustomizableObjectEditor/Private/MuCOE/Nodes/CustomizableObjectNodeTextureSwitch.cpp
@@ -15,7 +15,24 @@ void UCustomizableObjectNodeTextureSwitch::Serialize(FArchive& Ar)
 
 	if (Ar.CustomVer(FCustomizableObjectCustomVersion::GUID) < FCustomizableObjectCustomVersion::PinsNamesImageToTexture)
 	{
-		OutputPinReference = FindPin(TEXT("Image"));	
+		OutputPinReference = FindPin(TEXT("Image"));
+	}
+}
+
+
+void UCustomizableObjectNodeTextureSwitch::BackwardsCompatibleFixup()
+{
+	Super::BackwardsCompatibleFixup();
+
+	const int32 CustomizableObjectCustomVersion = GetLinkerCustomVersion(FCustomizableObjectCustomVersion::GUID);
+
+	if (CustomizableObjectCustomVersion < FCustomizableObjectCustomVersion::FixPinsNamesImageToTexture2)
+	{
+		if (UEdGraphPin* TexturePin = FindPin(TEXT("Image")))
+		{
+			TexturePin->PinName = TEXT("Texture");
+			UCustomizableObjectNode::ReconstructNode();
+		}
 	}
 }
 
diff --git a/Engine/Plugins/Experimental/Mutable/Source/CustomizableObjectEditor/Private/MuCOE/Nodes/CustomizableObjectNodeTextureToChannels.cpp b/Engine/Plugins/Experimental/Mutable/Source/CustomizableObjectEditor/Private/MuCOE/Nodes/CustomizableObjectNodeTextureToChannels.cpp
--- a/Engine/Plugins/Experimental/Mutable/Source/CustomizableObjectEditor/Private/MuCOE/Nodes/CustomizableObjectNodeTextureToChannels.cpp
+++ b/Engine/Plugins/Experimental/Mutable/Source/CustomizableObjectEditor/Private/MuCOE/Nodes/CustomizableObjectNodeTextureToChannels.cpp
@@ -67,6 +67,22 @@ void UCustomizableObjectNodeTextureToChannels::AllocateDefaultPins(UCustomizable
 }
 
 
+void UCustomizableObjectNodeTextureToChannels::BackwardsCompatibleFixup()
+{
+	Super::BackwardsCompatibleFixup();
+
+	const int32 CustomizableObjectCustomVersion = GetLinkerCustomVersion(FCustomizableObjectCustomVersion::GUID);
+
+	if (CustomizableObjectCustomVersion < FCustomizableObjectCustomVersion::FixPinsNamesImageToTexture2)
+	{
+		if (UEdGraphPin* TexturePin = FindPin(TEXT("Image"))) {
+			TexturePin->PinName = TEXT("Texture");
+			UCustomizableObjectNode::ReconstructNode();
+		}
+	}
+}
+
+
 UEdGraphPin* UCustomizableObjectNodeTextureToChannels::InputPin() const
 {
 	return InputPinReference.Get();
