[mutable] More compact packing of binary blobs for parameter values in generated resources keys.
Changed external image IDs to use FNames.
Single cache for all models.
#rnx

diff --git a/Engine/Plugins/Experimental/Mutable/Source/CustomizableObject/Classes/MuCO/CustomizableObject.h b/Engine/Plugins/Experimental/Mutable/Source/CustomizableObject/Classes/MuCO/CustomizableObject.h
--- a/Engine/Plugins/Experimental/Mutable/Source/CustomizableObject/Classes/MuCO/CustomizableObject.h
+++ b/Engine/Plugins/Experimental/Mutable/Source/CustomizableObject/Classes/MuCO/CustomizableObject.h
@@ -1408,7 +1408,7 @@ public:
 	 * @param InParameterName The name of the Projector parameter to get the default value of.
 	 * @return An id representing the default parameter's texture.
 	 */
-	FString GetTextureParameterDefaultValue (const FString& InParameterName) const;
+	FName GetTextureParameterDefaultValue (const FString& InParameterName) const;
 
 	/** Return true or false depending if the parameter at the index provided is multidimensional or not.
 	 * @param InParamIndex The index of the parameter to check.
@@ -1439,7 +1439,7 @@ private:
 	// This is a manual version number for the binary blobs in this asset.
 	// Increasing it invalidates all the previously compiled models.
 	// Warning: If while merging code both versions have changed, take the highest+1.
-	static const int32 CurrentSupportedVersion = 395;
+	static const int32 CurrentSupportedVersion = 396;
 
 public:
 
diff --git a/Engine/Plugins/Experimental/Mutable/Source/CustomizableObject/Classes/MuCO/CustomizableObjectInstance.h b/Engine/Plugins/Experimental/Mutable/Source/CustomizableObject/Classes/MuCO/CustomizableObjectInstance.h
--- a/Engine/Plugins/Experimental/Mutable/Source/CustomizableObject/Classes/MuCO/CustomizableObjectInstance.h
+++ b/Engine/Plugins/Experimental/Mutable/Source/CustomizableObject/Classes/MuCO/CustomizableObjectInstance.h
@@ -385,7 +385,7 @@ public:
 	
 	/** Gets the value of a texture parameter with name "TextureParamName". */
 	UFUNCTION(BlueprintCallable, Category = CustomizableObjectInstance)
-	FString GetTextureParameterSelectedOption(const FString& TextureParamName, int32 RangeIndex = -1) const;
+	FName GetTextureParameterSelectedOption(const FString& TextureParamName, int32 RangeIndex = -1) const;
 
 	/** Sets the texture value "TextureValue" of a texture parameter with index "TextureParamIndex". */
 	UFUNCTION(BlueprintCallable, Category = CustomizableObjectInstance)
diff --git a/Engine/Plugins/Experimental/Mutable/Source/CustomizableObject/Classes/MuCO/CustomizableObjectSystem.h b/Engine/Plugins/Experimental/Mutable/Source/CustomizableObject/Classes/MuCO/CustomizableObjectSystem.h
--- a/Engine/Plugins/Experimental/Mutable/Source/CustomizableObject/Classes/MuCO/CustomizableObjectSystem.h
+++ b/Engine/Plugins/Experimental/Mutable/Source/CustomizableObject/Classes/MuCO/CustomizableObjectSystem.h
@@ -135,7 +135,7 @@ struct FCustomizableObjectExternalTexture
 	FCustomizableObjectExternalTexture() = default;
 
 	FString Name;
-	FString Value;
+	FName Value;
 };
 
 
@@ -166,17 +166,17 @@ public:
 
 	// Query that Mutable will run to find out if a texture will be provided as an Unreal UTexture2D,
 	// or as a raw data blob.
-	virtual ValueType HasTextureParameterValue(const FString& ID) { return ValueType::None; }
+	virtual ValueType HasTextureParameterValue(const FName& ID) { return ValueType::None; }
 
 	// In case IsTextureParameterValueUnreal returns false, this will be used to query the texture size data.
-	virtual FIntVector GetTextureParameterValueSize(const FString& ID) { return FIntVector(0, 0, 0); }
+	virtual FIntVector GetTextureParameterValueSize(const FName& ID) { return FIntVector(0, 0, 0); }
 
 	// In case IsTextureParameterValueUnreal returns false, this will be used to query the texture data that must
 	// be copied in the preallocated buffer. The pixel format is assumed to be 4-channel RGBA, uint8_t per channel.
-	virtual void GetTextureParameterValueData(const FString& ID, uint8* OutData) {}
+	virtual void GetTextureParameterValueData(const FName& ID, uint8* OutData) {}
 
 	// In case IsTextureParameterValueUnreal returns true, this will be used to query the texture.
-	virtual UTexture2D* GetTextureParameterValue(const FString& ID) { return nullptr; }
+	virtual UTexture2D* GetTextureParameterValue(const FName& ID) { return nullptr; }
 
 	// Used in the editor to show the list of available options.
 	// Only necessary if the images are required in editor previews.
@@ -285,9 +285,9 @@ public:
 		memory is used and the real texel data is loaded when needed during an update and then immediately discarded */
 
 	/** [Texture Parameters] Cache an image which has to have been previously registered by an Image provider with the parameter id. */
-	void CacheImage(FString ImageId);
+	void CacheImage(FName ImageId);
 	/** [Texture Parameters] Remove an image from the cache. */
-	void UnCacheImage(FString ImageId);
+	void UnCacheImage(FName ImageId);
 	/** [Texture Parameters] Remove all images from the cache. */
 	void ClearImageCache();
 
diff --git a/Engine/Plugins/Experimental/Mutable/Source/CustomizableObject/Private/MuCO/CustomizableInstancePrivateData.h b/Engine/Plugins/Experimental/Mutable/Source/CustomizableObject/Private/MuCO/CustomizableInstancePrivateData.h
--- a/Engine/Plugins/Experimental/Mutable/Source/CustomizableObject/Private/MuCO/CustomizableInstancePrivateData.h
+++ b/Engine/Plugins/Experimental/Mutable/Source/CustomizableObject/Private/MuCO/CustomizableInstancePrivateData.h
@@ -199,7 +199,7 @@ public:
 	/** Cached Texture Parameters from the last update.
 	 * This cache is required since the Instance can have a LOD Update at any time.
 	 * So we need to make sure that the initially provided Texture Parameters by the user will be available until the user decides to change them. */
-	TArray<FString> UpdateTextureParameters;
+	TArray<FName> UpdateTextureParameters;
 	
 	FCustomizableInstanceComponentData* GetComponentData(int32 ComponentIndex);
 	const FCustomizableInstanceComponentData* GetComponentData(int32 ComponentIndex) const;
diff --git a/Engine/Plugins/Experimental/Mutable/Source/CustomizableObject/Private/MuCO/CustomizableObject.cpp b/Engine/Plugins/Experimental/Mutable/Source/CustomizableObject/Private/MuCO/CustomizableObject.cpp
--- a/Engine/Plugins/Experimental/Mutable/Source/CustomizableObject/Private/MuCO/CustomizableObject.cpp
+++ b/Engine/Plugins/Experimental/Mutable/Source/CustomizableObject/Private/MuCO/CustomizableObject.cpp
@@ -1662,23 +1662,23 @@ FCustomizableObjectProjector UCustomizableObject::GetProjectorParameterDefaultVa
 }
 
 
-FString UCustomizableObject::GetTextureParameterDefaultValue(const FString& InParameterName) const
+FName UCustomizableObject::GetTextureParameterDefaultValue(const FString& InParameterName) const
 {
 	const int32 ParameterIndex = FindParameter(InParameterName);
 	if (ParameterIndex == INDEX_NONE)
 	{
 		checkNoEntry();
-		return FCustomizableObjectTextureParameterValue::DEFAULT_PARAMETER_VALUE;
+		return FName();
 	}
 
 	const TSharedPtr<mu::Model> Model = GetModel();
 	if (!Model)
 	{
 		checkNoEntry();
-		return FCustomizableObjectTextureParameterValue::DEFAULT_PARAMETER_VALUE;
+		return FName();
 	}
 	
-	return Model->GetImageDefaultValue(ParameterIndex).GetName();
+	return Model->GetImageDefaultValue(ParameterIndex);
 }
 
 
diff --git a/Engine/Plugins/Experimental/Mutable/Source/CustomizableObject/Private/MuCO/CustomizableObjectInstance.cpp b/Engine/Plugins/Experimental/Mutable/Source/CustomizableObject/Private/MuCO/CustomizableObjectInstance.cpp
--- a/Engine/Plugins/Experimental/Mutable/Source/CustomizableObject/Private/MuCO/CustomizableObjectInstance.cpp
+++ b/Engine/Plugins/Experimental/Mutable/Source/CustomizableObject/Private/MuCO/CustomizableObjectInstance.cpp
@@ -367,7 +367,7 @@ void UCustomizableObjectInstance::ReleaseMutableResources(bool bCalledFromBeginD
 		{
 			FUnrealMutableImageProvider* ImageProvider = UCustomizableObjectSystem::GetInstance()->GetPrivateChecked()->GetImageProviderChecked();
 
-			for (const FString& TextureParameter : PrivateData->UpdateTextureParameters)
+			for (const FName& TextureParameter : PrivateData->UpdateTextureParameters)
 			{
 				ImageProvider->UnCacheImage(TextureParameter, false);
 			}
@@ -2428,7 +2428,7 @@ void UCustomizableObjectInstance::SetFloatParameterSelectedOption(const FString&
 }
 
 
-FString UCustomizableObjectInstance::GetTextureParameterSelectedOption(const FString& TextureParamName, const int32 RangeIndex) const
+FName UCustomizableObjectInstance::GetTextureParameterSelectedOption(const FString& TextureParamName, const int32 RangeIndex) const
 {
 	return Descriptor.GetTextureParameterSelectedOption(TextureParamName, RangeIndex);
 }
@@ -6381,6 +6381,10 @@ void UCustomizableObjectInstance::SetRequestedLODs(int32 InMinLOD, int32 InMaxLO
 				bUpdateRequestedLODs |= RequestedLODsDifferent;
 
 				// Save new RequestedLODs
+				if (ComponentIndex >= MutableUpdateCandidate.RequestedLODLevels.Num())
+				{
+					MutableUpdateCandidate.RequestedLODLevels.SetNumZeroed(ComponentIndex+1);
+				}
 				MutableUpdateCandidate.RequestedLODLevels[ComponentIndex] = RequestedLODs;
 			}
 		}
diff --git a/Engine/Plugins/Experimental/Mutable/Source/CustomizableObject/Private/MuCO/CustomizableObjectInstanceDescriptor.cpp b/Engine/Plugins/Experimental/Mutable/Source/CustomizableObject/Private/MuCO/CustomizableObjectInstanceDescriptor.cpp
--- a/Engine/Plugins/Experimental/Mutable/Source/CustomizableObject/Private/MuCO/CustomizableObjectInstanceDescriptor.cpp
+++ b/Engine/Plugins/Experimental/Mutable/Source/CustomizableObject/Private/MuCO/CustomizableObjectInstanceDescriptor.cpp
@@ -175,8 +175,8 @@ void FCustomizableObjectInstanceDescriptor::SaveDescriptor(FArchive& Ar, bool bU
 
 		case EMutableParameterType::Texture:
 		{
-			FString Value;
-			TArray<FString> RangeValues;
+			FName Value;
+			TArray<FName> RangeValues;
 
 			for (const FCustomizableObjectTextureParameterValue& P : TextureParameters)
 			{
@@ -368,8 +368,8 @@ void FCustomizableObjectInstanceDescriptor::LoadDescriptor(FArchive& Ar)
 
 		case EMutableParameterType::Texture:
 		{
-			FString Value;
-			TArray<FString> RangeValues;
+			FName Value;
+			TArray<FName> RangeValues;
 			Ar << Value;
 			Ar << RangeValues;
 
@@ -969,16 +969,16 @@ void FCustomizableObjectInstanceDescriptor::ReloadParameters()
 
 							for (int32 Index = PreviousNum; Index < Param.ParameterRangeValues.Num(); ++Index)
 							{
-								Param.ParameterRangeValues[Index] = FCustomizableObjectTextureParameterValue::DEFAULT_PARAMETER_VALUE;
+								Param.ParameterRangeValues[Index] = FName();
 							}
 						}
 
-						Param.ParameterRangeValues[RangeIndex] = MutableParameters->GetImageValue(ParamIndex, RangeValueIdxPtr).GetName();
+						Param.ParameterRangeValues[RangeIndex] = MutableParameters->GetImageValue(ParamIndex, RangeValueIdxPtr);
 					}
 				}
 				else
 				{
-					Param.ParameterValue = MutableParameters->GetImageValue(ParamIndex).GetName();
+					Param.ParameterValue = MutableParameters->GetImageValue(ParamIndex);
 				}
 			}
 
@@ -1327,7 +1327,7 @@ void FCustomizableObjectInstanceDescriptor::SetFloatParameterSelectedOption(cons
 }
 
 
-FString FCustomizableObjectInstanceDescriptor::GetTextureParameterSelectedOption(const FString& TextureParamName, const int32 RangeIndex) const
+FName FCustomizableObjectInstanceDescriptor::GetTextureParameterSelectedOption(const FString& TextureParamName, const int32 RangeIndex) const
 {
 	check(CustomizableObject);
 
@@ -1354,7 +1354,7 @@ FString FCustomizableObjectInstanceDescriptor::GetTextureParameterSelectedOption
 	
 	LogParameterNotFoundWarning(TextureParamName, ParameterIndexInObject, ParameterIndexInInstance, CustomizableObject, __FUNCTION__);
 
-	return FCustomizableObjectTextureParameterValue::DEFAULT_PARAMETER_VALUE; 
+	return FName(); 
 }
 
 
@@ -1376,7 +1376,7 @@ void FCustomizableObjectInstanceDescriptor::SetTextureParameterSelectedOption(co
 	if (RangeIndex == INDEX_NONE)
 	{
 		check(!CustomizableObject->IsParameterMultidimensional(ParameterIndexInObject)); // This param is multidimensional, it must have a RangeIndex of 0 or more
-		TextureParameters[ParameterIndexInInstance].ParameterValue = TextureValue;
+		TextureParameters[ParameterIndexInInstance].ParameterValue = FName(TextureValue);
 	}
 	else
 	{
@@ -1390,7 +1390,7 @@ void FCustomizableObjectInstanceDescriptor::SetTextureParameterSelectedOption(co
 		}
 
 		check(TextureParameters[ParameterIndexInInstance].ParameterRangeValues.IsValidIndex(RangeIndex));
-		TextureParameters[ParameterIndexInInstance].ParameterRangeValues[RangeIndex] = TextureValue;
+		TextureParameters[ParameterIndexInInstance].ParameterRangeValues[RangeIndex] = FName(TextureValue);
 	}
 }
 
@@ -1990,7 +1990,7 @@ int32 FCustomizableObjectInstanceDescriptor::AddValueToTextureRange(const FStrin
 	if (TextureParameterIndex != -1)
 	{
 		FCustomizableObjectTextureParameterValue& TextureParameter = TextureParameters[TextureParameterIndex];
-		return TextureParameter.ParameterRangeValues.Add(FCustomizableObjectTextureParameterValue::DEFAULT_PARAMETER_VALUE);		
+		return TextureParameter.ParameterRangeValues.Add(FName());		
 	}
 	
 	return -1;
diff --git a/Engine/Plugins/Experimental/Mutable/Source/CustomizableObject/Private/MuCO/CustomizableObjectSystem.cpp b/Engine/Plugins/Experimental/Mutable/Source/CustomizableObject/Private/MuCO/CustomizableObjectSystem.cpp
--- a/Engine/Plugins/Experimental/Mutable/Source/CustomizableObject/Private/MuCO/CustomizableObjectSystem.cpp
+++ b/Engine/Plugins/Experimental/Mutable/Source/CustomizableObject/Private/MuCO/CustomizableObjectSystem.cpp
@@ -906,7 +906,7 @@ void FCustomizableObjectSystemPrivate::InitUpdateSkeletalMesh(UCustomizableObjec
 		{
 			ImageProvider->CacheImage(TextureParameters.ParameterValue, false);
 
-			for (const FString& TextureParameter : TextureParameters.ParameterRangeValues)
+			for (const FName& TextureParameter : TextureParameters.ParameterRangeValues)
 			{
 				ImageProvider->CacheImage(TextureParameter, false);
 			}
@@ -916,7 +916,7 @@ void FCustomizableObjectSystemPrivate::InitUpdateSkeletalMesh(UCustomizableObjec
 		check(InstancePrivate);
 
 		// Uncache old Texture Parameters
-		for (const FString& TextureParameter : InstancePrivate->UpdateTextureParameters)
+		for (const FName& TextureParameter : InstancePrivate->UpdateTextureParameters)
 		{
 			ImageProvider->UnCacheImage(TextureParameter, false);
 		}
@@ -927,7 +927,7 @@ void FCustomizableObjectSystemPrivate::InitUpdateSkeletalMesh(UCustomizableObjec
 		{
 			InstancePrivate->UpdateTextureParameters.Add(TextureParameters.ParameterValue);
 
-			for (const FString& TextureParameter : TextureParameters.ParameterRangeValues)
+			for (const FName& TextureParameter : TextureParameters.ParameterRangeValues)
 			{
 				InstancePrivate->UpdateTextureParameters.Add(TextureParameter);
 			}
@@ -3286,13 +3286,13 @@ uint64 UCustomizableObjectSystem::GetMaxChunkSizeForPlatform(const ITargetPlatfo
 #endif // WITH_EDITOR
 
 
-void UCustomizableObjectSystem::CacheImage(FString ImageId)
+void UCustomizableObjectSystem::CacheImage(FName ImageId)
 {
 	GetPrivateChecked()->GetImageProviderChecked()->CacheImage(ImageId, true);
 }
 
 
-void UCustomizableObjectSystem::UnCacheImage(FString ImageId)
+void UCustomizableObjectSystem::UnCacheImage(FName ImageId)
 {
 	GetPrivateChecked()->GetImageProviderChecked()->UnCacheImage(ImageId, true);
 }
diff --git a/Engine/Plugins/Experimental/Mutable/Source/CustomizableObject/Private/MuCO/DefaultImageProvider.cpp b/Engine/Plugins/Experimental/Mutable/Source/CustomizableObject/Private/MuCO/DefaultImageProvider.cpp
--- a/Engine/Plugins/Experimental/Mutable/Source/CustomizableObject/Private/MuCO/DefaultImageProvider.cpp
+++ b/Engine/Plugins/Experimental/Mutable/Source/CustomizableObject/Private/MuCO/DefaultImageProvider.cpp
@@ -34,7 +34,7 @@ namespace UDefaultImageProviderCVars
 }
 
 
-UCustomizableSystemImageProvider::ValueType UDefaultImageProvider::HasTextureParameterValue(const FString& ID)
+UCustomizableSystemImageProvider::ValueType UDefaultImageProvider::HasTextureParameterValue(const FName& ID)
 {
 	const TObjectPtr<UTexture2D>* Texture = Textures.Find(ID);
 	
@@ -44,7 +44,7 @@ UCustomizableSystemImageProvider::ValueType UDefaultImageProvider::HasTexturePar
 }
 
 
-UTexture2D* UDefaultImageProvider::GetTextureParameterValue(const FString& ID)
+UTexture2D* UDefaultImageProvider::GetTextureParameterValue(const FName& ID)
 {
 	return Textures[ID];
 }
@@ -52,7 +52,7 @@ UTexture2D* UDefaultImageProvider::GetTextureParameterValue(const FString& ID)
 
 void UDefaultImageProvider::GetTextureParameterValues(TArray<FCustomizableObjectExternalTexture>& OutValues)
 {
-	for (TTuple<FString, TObjectPtr<UTexture2D>>& Pair : Textures)
+	for (TTuple<FName, TObjectPtr<UTexture2D>>& Pair : Textures)
 	{
 		FCustomizableObjectExternalTexture Data;
 		Data.Name = Pair.Value.GetName();
@@ -66,13 +66,13 @@ FString UDefaultImageProvider::Add(UTexture2D* Texture)
 {
 	if (!Texture)
 	{
-		return FCustomizableObjectTextureParameterValue::DEFAULT_PARAMETER_VALUE;
+		return FString();
 	}
 
-	const FString Id = Texture->GetFullName();
+	const FName Id(Texture->GetFullName());
 	Textures.Add(Id, Texture);	
 
-	return Id;
+	return Id.ToString();
 }
 
 
@@ -83,6 +83,6 @@ void UDefaultImageProvider::Remove(UTexture2D* Texture)
 		return;
 	}
 	
-	Textures.Remove(Texture->GetFullName());
+	Textures.Remove(FName(Texture->GetFullName()));
 }
 
diff --git a/Engine/Plugins/Experimental/Mutable/Source/CustomizableObject/Private/MuCO/UnrealMutableImageProvider.cpp b/Engine/Plugins/Experimental/Mutable/Source/CustomizableObject/Private/MuCO/UnrealMutableImageProvider.cpp
--- a/Engine/Plugins/Experimental/Mutable/Source/CustomizableObject/Private/MuCO/UnrealMutableImageProvider.cpp
+++ b/Engine/Plugins/Experimental/Mutable/Source/CustomizableObject/Private/MuCO/UnrealMutableImageProvider.cpp
@@ -117,9 +117,9 @@ mu::EImageFormat GetMutablePixelFormat(EPixelFormat InTextureFormat)
 
 //-------------------------------------------------------------------------------------------------
 #ifdef MUTABLE_USE_NEW_TASKGRAPH
-	TTuple<UE::Tasks::FTask, TFunction<void()>> FUnrealMutableImageProvider::GetImageAsync(const FExternalImageID& id, uint8 MipmapsToSkip, TFunction<void(mu::Ptr<mu::Image>)>& ResultCallback)
+	TTuple<UE::Tasks::FTask, TFunction<void()>> FUnrealMutableImageProvider::GetImageAsync(FName Id, uint8 MipmapsToSkip, TFunction<void(mu::Ptr<mu::Image>)>& ResultCallback)
 #else
-	TTuple<FGraphEventRef, TFunction<void()>> FUnrealMutableImageProvider::GetImageAsync(const FExternalImageID& id, uint8 MipmapsToSkip, TFunction<void(mu::Ptr<mu::Image>)>& ResultCallback)
+	TTuple<FGraphEventRef, TFunction<void()>> FUnrealMutableImageProvider::GetImageAsync(FName Id, uint8 MipmapsToSkip, TFunction<void(mu::Ptr<mu::Image>)>& ResultCallback)
 #endif
 {
 	// Thread: worker
@@ -161,7 +161,7 @@ mu::EImageFormat GetMutablePixelFormat(EPixelFormat InTextureFormat)
 		// Inside this scope it's safe to access GlobalExternalImages
 
 
-		if (!GlobalExternalImages.Contains(id))
+		if (!GlobalExternalImages.Contains(Id))
 		{
 			// Null case, no image was provided
 			UE_LOG(LogMutable, Warning, TEXT("Failed to get external image. GlobalExternalImage not found."));
@@ -170,7 +170,7 @@ mu::EImageFormat GetMutablePixelFormat(EPixelFormat InTextureFormat)
 			return Invoke(TrivialReturn);
 		}
 
-		const FUnrealMutableImageInfo& ImageInfo = GlobalExternalImages[id];
+		const FUnrealMutableImageInfo& ImageInfo = GlobalExternalImages[Id];
 
 		if (ImageInfo.Image)
 		{
@@ -393,7 +393,7 @@ mu::EImageFormat GetMutablePixelFormat(EPixelFormat InTextureFormat)
 
 
 // This should mantain parity with the descriptor of the images generated by GetImageAsync 
-mu::FImageDesc FUnrealMutableImageProvider::GetImageDesc(const FExternalImageID& id, uint8 MipmapsToSkip) 
+mu::FImageDesc FUnrealMutableImageProvider::GetImageDesc(FName Id, uint8 MipmapsToSkip)
 {
 	MUTABLE_CPUPROFILER_SCOPE(FUnrealMutableImageProvider::GetImageDesc);
 
@@ -402,13 +402,13 @@ mu::FImageDesc FUnrealMutableImageProvider::GetImageDesc(const FExternalImageID&
 		FScopeLock Lock(&ExternalImagesLock);
 		// Inside this scope it's safe to access GlobalExternalImages
 
-		if (!GlobalExternalImages.Contains(id))
+		if (!GlobalExternalImages.Contains(Id))
 		{
 			// Null case, no image was provided
 			return CreateDummyDesc();
 		}
 
-		const FUnrealMutableImageInfo& ImageInfo = GlobalExternalImages[id];
+		const FUnrealMutableImageInfo& ImageInfo = GlobalExternalImages[Id];
 
 		if (ImageInfo.Image)
 		{
@@ -463,11 +463,11 @@ mu::FImageDesc FUnrealMutableImageProvider::GetImageDesc(const FExternalImageID&
 }
 
 
-void FUnrealMutableImageProvider::CacheImage(const FExternalImageID& Id, bool bUser)
+void FUnrealMutableImageProvider::CacheImage(FName Id, bool bUser)
 {
 	check( IsInGameThread() );
 
-	if (Id == FCustomizableObjectTextureParameterValue::DEFAULT_PARAMETER_VALUE)
+	if (Id == FName())
 	{
 		return;	
 	}	
@@ -499,27 +499,27 @@ void FUnrealMutableImageProvider::CacheImage(const FExternalImageID& Id, bool bU
 				if (Provider.IsValid())
 				{
 					// \TODO: all these queries could probably be optimized into a single call.
-					switch (Provider->HasTextureParameterValue(Id.GetName()))
+					switch (Provider->HasTextureParameterValue(Id))
 					{
 
 					case UCustomizableSystemImageProvider::ValueType::Raw:
 					{
-						FIntVector desc = Provider->GetTextureParameterValueSize(Id.GetName());
+						FIntVector desc = Provider->GetTextureParameterValueSize(Id);
 						pResult = new mu::Image(desc[0], desc[1], 1, mu::EImageFormat::IF_RGBA_UBYTE, mu::EInitializationType::Black);
-						Provider->GetTextureParameterValueData(Id.GetName(), pResult->GetData());
+						Provider->GetTextureParameterValueData(Id, pResult->GetData());
 						break;
 					}
 
 					case UCustomizableSystemImageProvider::ValueType::Unreal:
 					{
-						UTexture2D* UnrealTexture = Provider->GetTextureParameterValue(Id.GetName());
+						UTexture2D* UnrealTexture = Provider->GetTextureParameterValue(Id);
 						pResult = ConvertTextureUnrealToMutable(UnrealTexture, 0);
 						break;
 					}
 
 					case UCustomizableSystemImageProvider::ValueType::Unreal_Deferred:
 					{
-						UnrealDeferredTexture = Provider->GetTextureParameterValue(Id.GetName());
+						UnrealDeferredTexture = Provider->GetTextureParameterValue(Id);
 						break;
 					}
 
@@ -531,7 +531,7 @@ void FUnrealMutableImageProvider::CacheImage(const FExternalImageID& Id, bool bU
 
 			if (!pResult && !UnrealDeferredTexture)
 			{
-				UE_LOG(LogMutable, Warning, TEXT("Failed to cache external image %s."), *Id.GetName());
+				UE_LOG(LogMutable, Warning, TEXT("Failed to cache external image %s."), *Id.ToString());
 				pResult = CreateDummy();
 			}
 			
@@ -552,11 +552,11 @@ void FUnrealMutableImageProvider::CacheImage(const FExternalImageID& Id, bool bU
 }
 
 
-void FUnrealMutableImageProvider::UnCacheImage(const FExternalImageID& Id, bool bUser)
+void FUnrealMutableImageProvider::UnCacheImage(FName Id, bool bUser)
 {
 	check(IsInGameThread());
 
-	if (Id == FCustomizableObjectTextureParameterValue::DEFAULT_PARAMETER_VALUE)
+	if (Id == FName())
 	{
 		return;	
 	}
@@ -581,7 +581,7 @@ void FUnrealMutableImageProvider::UnCacheImage(const FExternalImageID& Id, bool
 	}
 	else
 	{
-		UE_LOG(LogMutable, Warning, TEXT("Failed to uncache external image %s. Possible double free!"), *Id.GetName());
+		UE_LOG(LogMutable, Warning, TEXT("Failed to uncache external image %s. Possible double free!"), *Id.ToString());
 	}
 }
 
@@ -591,7 +591,7 @@ void FUnrealMutableImageProvider::ClearCache(bool bUser)
 	check(IsInGameThread());
 
 	FScopeLock Lock(&ExternalImagesLock);
-	for (TTuple<FExternalImageID, FUnrealMutableImageInfo> Tuple : GlobalExternalImages)
+	for (TTuple<FName, FUnrealMutableImageInfo> Tuple : GlobalExternalImages)
 	{
 		UnCacheImage(Tuple.Key, bUser);
 	}	
@@ -609,7 +609,7 @@ void FUnrealMutableImageProvider::CacheImages(const mu::Parameters& Parameters)
 		}
 	
 		{
-			const FExternalImageID TextureId = Parameters.GetImageValue(ParamIndex);
+			const FName TextureId = Parameters.GetImageValue(ParamIndex);
 			CacheImage(TextureId, false);
 		}
 
@@ -617,7 +617,7 @@ void FUnrealMutableImageProvider::CacheImages(const mu::Parameters& Parameters)
 		for (int32 ValueIndex = 0; ValueIndex < NumValues; ++ValueIndex)
 		{
 			mu::Ptr<const mu::RangeIndex> Range = Parameters.GetValueIndex(ParamIndex, ValueIndex);
-			const FExternalImageID TextureId = Parameters.GetImageValue(ParamIndex, Range);
+			const FName TextureId = Parameters.GetImageValue(ParamIndex, Range);
 
 			CacheImage(TextureId, false);
 		}
@@ -636,7 +636,7 @@ void FUnrealMutableImageProvider::UnCacheImages(const mu::Parameters& Parameters
 		}
 	
 		{
-			const FExternalImageID TextureId = Parameters.GetImageValue(ParamIndex);
+			const FName TextureId = Parameters.GetImageValue(ParamIndex);
 			UnCacheImage(TextureId, false);				
 		}
 
@@ -644,7 +644,7 @@ void FUnrealMutableImageProvider::UnCacheImages(const mu::Parameters& Parameters
 		for (int32 ValueIndex = 0; ValueIndex < NumValues; ++ValueIndex)
 		{
 			mu::Ptr<const mu::RangeIndex> Range = Parameters.GetValueIndex(ParamIndex, ValueIndex);
-			const FExternalImageID TextureId = Parameters.GetImageValue(ParamIndex, Range);
+			const FName TextureId = Parameters.GetImageValue(ParamIndex, Range);
 
 			UnCacheImage(TextureId, false);
 		}
@@ -693,7 +693,7 @@ mu::FImageDesc FUnrealMutableImageProvider::CreateDummyDesc()
 
 void FUnrealMutableImageProvider::AddReferencedObjects(FReferenceCollector& Collector)
 {
-	for (TPair<FExternalImageID, FUnrealMutableImageInfo>& Image : GlobalExternalImages)
+	for (TPair<FName, FUnrealMutableImageInfo>& Image : GlobalExternalImages)
 	{
 		if (Image.Value.TextureToLoad)
 		{
diff --git a/Engine/Plugins/Experimental/Mutable/Source/CustomizableObject/Private/MuCO/UnrealMutableImageProvider.h b/Engine/Plugins/Experimental/Mutable/Source/CustomizableObject/Private/MuCO/UnrealMutableImageProvider.h
--- a/Engine/Plugins/Experimental/Mutable/Source/CustomizableObject/Private/MuCO/UnrealMutableImageProvider.h
+++ b/Engine/Plugins/Experimental/Mutable/Source/CustomizableObject/Private/MuCO/UnrealMutableImageProvider.h
@@ -23,23 +23,23 @@ public:
 	// mu::ImageParameterGenerator interface
 	// Thread: worker
 #ifdef MUTABLE_USE_NEW_TASKGRAPH
-	virtual TTuple<UE::Tasks::FTask, TFunction<void()>> GetImageAsync(const FExternalImageID& Id, uint8 MipmapsToSkip, TFunction<void(mu::Ptr<mu::Image>)>& ResultCallback) override;
+	virtual TTuple<UE::Tasks::FTask, TFunction<void()>> GetImageAsync(FName Id, uint8 MipmapsToSkip, TFunction<void(mu::Ptr<mu::Image>)>& ResultCallback) override;
 #else
-	virtual TTuple<FGraphEventRef, TFunction<void()>> GetImageAsync(const FExternalImageID& Id, uint8 MipmapsToSkip, TFunction<void(mu::Ptr<mu::Image>)>& ResultCallback) override;
+	virtual TTuple<FGraphEventRef, TFunction<void()>> GetImageAsync(FName Id, uint8 MipmapsToSkip, TFunction<void(mu::Ptr<mu::Image>)>& ResultCallback) override;
 #endif
 
-	virtual mu::FImageDesc GetImageDesc(const FExternalImageID& Id, uint8 MipmapsToSkip) override;
+	virtual mu::FImageDesc GetImageDesc(FName Id, uint8 MipmapsToSkip) override;
 
 	
 	// Own interface	
 	// Thread: Game
 	/** Add a reference to the image. If it was not cached it caches it.
 	 * @param bUser if true, adds a reference to the user reference counter. If false, adds a reference to the system reference counter. */
-	void CacheImage(const FExternalImageID& Id, bool bUser);
+	void CacheImage(FName Id, bool bUser);
 
 	/** Removes a reference to the image. If all references are removed, it uncaches the image.
 	 * @param bUser if true, removes a reference from the user reference counter. If false, removes a reference from the system reference counter. */
-	void UnCacheImage(const FExternalImageID& Id, bool bUser);
+	void UnCacheImage(FName Id, bool bUser);
 
 	/** Removes a reference to all images. All images which no longer have references will be uncached.
 	 * @param bUser if true, removes a references from the user reference counter. If false, removes a reference from the system reference counter. */
@@ -89,7 +89,7 @@ private:
 	* This is only safely written from the game thread protected by the following critical section, and it
 	* is safely read from the mutable thread during the update of the instance or texture mips
 	*/
-	TMap<FExternalImageID, FUnrealMutableImageInfo> GlobalExternalImages;
+	TMap<FName, FUnrealMutableImageInfo> GlobalExternalImages;
 	
 	/** Access to GlobalExternalImages must be protected with this because it may be accessed concurrently from the 
 	Game thread to modify it and from the Mutable thread to read it. */
diff --git a/Engine/Plugins/Experimental/Mutable/Source/CustomizableObject/Public/MuCO/CustomizableObjectInstanceDescriptor.h b/Engine/Plugins/Experimental/Mutable/Source/CustomizableObject/Public/MuCO/CustomizableObjectInstanceDescriptor.h
--- a/Engine/Plugins/Experimental/Mutable/Source/CustomizableObject/Public/MuCO/CustomizableObjectInstanceDescriptor.h
+++ b/Engine/Plugins/Experimental/Mutable/Source/CustomizableObject/Public/MuCO/CustomizableObjectInstanceDescriptor.h
@@ -128,7 +128,7 @@ struct CUSTOMIZABLEOBJECT_API FCustomizableObjectInstanceDescriptor
 	void SetFloatParameterSelectedOption(const FString& FloatParamName, float FloatValue, int32 RangeIndex = -1);
 
 	/** Gets the value of a texture parameter with name "TextureParamName". */
-	FString GetTextureParameterSelectedOption(const FString& TextureParamName, int32 RangeIndex) const;
+	FName GetTextureParameterSelectedOption(const FString& TextureParamName, int32 RangeIndex) const;
 
 	/** Sets the texture value "TextureValue" of a texture parameter with index "TextureParamIndex". */
 	void SetTextureParameterSelectedOption(const FString& TextureParamName, const FString& TextureValue, int32 RangeIndex);
diff --git a/Engine/Plugins/Experimental/Mutable/Source/CustomizableObject/Public/MuCO/CustomizableObjectParameterTypeDefinitions.h b/Engine/Plugins/Experimental/Mutable/Source/CustomizableObject/Public/MuCO/CustomizableObjectParameterTypeDefinitions.h
--- a/Engine/Plugins/Experimental/Mutable/Source/CustomizableObject/Public/MuCO/CustomizableObjectParameterTypeDefinitions.h
+++ b/Engine/Plugins/Experimental/Mutable/Source/CustomizableObject/Public/MuCO/CustomizableObjectParameterTypeDefinitions.h
@@ -158,20 +158,18 @@ USTRUCT(BlueprintType)
 struct FCustomizableObjectTextureParameterValue
 {
 	GENERATED_USTRUCT_BODY()
-
-	inline static const FString DEFAULT_PARAMETER_VALUE = {}; // Means no Texture Parameter assigned.
 	
 	UPROPERTY(Category = CustomizableObjectTextureParameterValue, VisibleAnywhere)
 	FString ParameterName;
 
 	UPROPERTY(Category = CustomizableObjectTextureParameterValue, VisibleAnywhere)
-	FString ParameterValue = DEFAULT_PARAMETER_VALUE;
+	FName ParameterValue;
 
 	UPROPERTY(Category = CustomizableObjectTextureParameterValue, VisibleAnywhere)
 	FString Uid;
 
 	UPROPERTY(Category = CustomizableObjectTextureParameterValue, VisibleAnywhere)
-	TArray<FString> ParameterRangeValues;
+	TArray<FName> ParameterRangeValues;
 };
 
 
@@ -180,7 +178,7 @@ inline uint32 GetTypeHash(const FCustomizableObjectTextureParameterValue& Key)
 	uint32 Hash = GetTypeHash(Key.ParameterName);
 	Hash = HashCombine(Hash, GetTypeHash(Key.ParameterValue));
 
-	for (const FString& Value : Key.ParameterRangeValues)
+	for (const FName& Value : Key.ParameterRangeValues)
     {
     	Hash = HashCombine(Hash, GetTypeHash(Value));
     }
diff --git a/Engine/Plugins/Experimental/Mutable/Source/CustomizableObject/Public/MuCO/DefaultImageProvider.h b/Engine/Plugins/Experimental/Mutable/Source/CustomizableObject/Public/MuCO/DefaultImageProvider.h
--- a/Engine/Plugins/Experimental/Mutable/Source/CustomizableObject/Public/MuCO/DefaultImageProvider.h
+++ b/Engine/Plugins/Experimental/Mutable/Source/CustomizableObject/Public/MuCO/DefaultImageProvider.h
@@ -24,12 +24,12 @@ class CUSTOMIZABLEOBJECT_API UDefaultImageProvider : public UCustomizableSystemI
 
 public:
 	// UCustomizableSystemImageProvider interface
-	virtual ValueType HasTextureParameterValue(const FString& ID) override;
-	virtual UTexture2D* GetTextureParameterValue(const FString& ID) override;
+	virtual ValueType HasTextureParameterValue(const FName& ID) override;
+	virtual UTexture2D* GetTextureParameterValue(const FName& ID) override;
 	virtual void GetTextureParameterValues(TArray<FCustomizableObjectExternalTexture>& OutValues) override;
 	
 	/** Add a Texture to the provider. */
-	FString Add(UTexture2D* Texture); 
+	FString Add(UTexture2D* Texture);
 
 	/** Remove a Texture from the provider. */
 	void Remove(UTexture2D* Texture); 
@@ -37,6 +37,6 @@ public:
 private:
 	// Always contains valid pointers
 	UPROPERTY()
-	TMap<FString, TObjectPtr<UTexture2D>> Textures;
+	TMap<FName, TObjectPtr<UTexture2D>> Textures;
 };
 
diff --git a/Engine/Plugins/Experimental/Mutable/Source/CustomizableObjectEditor/Private/MuCOE/SCustomizableInstanceProperties.cpp b/Engine/Plugins/Experimental/Mutable/Source/CustomizableObjectEditor/Private/MuCOE/SCustomizableInstanceProperties.cpp
--- a/Engine/Plugins/Experimental/Mutable/Source/CustomizableObjectEditor/Private/MuCOE/SCustomizableInstanceProperties.cpp
+++ b/Engine/Plugins/Experimental/Mutable/Source/CustomizableObjectEditor/Private/MuCOE/SCustomizableInstanceProperties.cpp
@@ -78,7 +78,7 @@ void SCustomizableInstanceProperties::Construct(const FArguments& InArgs)
 
 		// Store the texture parameters data required for the ui.
 		TextureParameterValueNames.Add(MakeShareable(new FString("None")));
-		TextureParameterValues.Add(FCustomizableObjectTextureParameterValue::DEFAULT_PARAMETER_VALUE);
+		TextureParameterValues.Add(FName());
 		TArray<FCustomizableObjectExternalTexture> Textures = UCustomizableObjectSystem::GetInstance()->GetTextureParameterValues();
 		for (int i = 0; i < Textures.Num(); ++i)
 		{
@@ -816,7 +816,7 @@ void SCustomizableInstanceProperties::AddParameter(int32 ParamIndexInObject)
 		const TArray<FCustomizableObjectTextureParameterValue>& TextureParameters = CustomInstance->GetTextureParameters();
 		TSharedPtr<FString> InitiallySelected = TextureParameterValueNames[0]; // First index is always the None option.
 
-		const FString ParameterValue = CustomInstance->GetTextureParameterSelectedOption(ParamName);
+		const FName ParameterValue = CustomInstance->GetTextureParameterSelectedOption(ParamName);
 		
 		// Look for the value index
 		for (int32 ValueIndex = 0; ValueIndex < TextureParameterValueNames.Num(); ++ValueIndex)
@@ -2400,7 +2400,7 @@ bool SCustomizableInstanceProperties::SetParameterValueToDefault(int32 Parameter
 	}
 	case EMutableParameterType::Texture:
 	{
-		FString DefaultValue = CustomObject->GetTextureParameterDefaultValue(ParameterName);
+		FName DefaultValue = CustomObject->GetTextureParameterDefaultValue(ParameterName);
 
 		if (CustomObject->IsParameterMultidimensional(ParameterName))
 		{
@@ -2408,12 +2408,12 @@ bool SCustomizableInstanceProperties::SetParameterValueToDefault(int32 Parameter
 
 			for (int32 RangeIndex = 0; RangeIndex < NumRanges; ++RangeIndex)
 			{
-				CustomInstance->SetTextureParameterSelectedOption(ParameterName, DefaultValue, RangeIndex);
+				CustomInstance->SetTextureParameterSelectedOption(ParameterName, DefaultValue.ToString(), RangeIndex);
 			}
 		}
 		else
 		{
-			CustomInstance->SetTextureParameterSelectedOption(ParameterName, DefaultValue);
+			CustomInstance->SetTextureParameterSelectedOption(ParameterName, DefaultValue.ToString());
 		}
 		break;
 	}
diff --git a/Engine/Plugins/Experimental/Mutable/Source/CustomizableObjectEditor/Private/MuCOE/SMutableCodeViewer.cpp b/Engine/Plugins/Experimental/Mutable/Source/CustomizableObjectEditor/Private/MuCOE/SMutableCodeViewer.cpp
--- a/Engine/Plugins/Experimental/Mutable/Source/CustomizableObjectEditor/Private/MuCOE/SMutableCodeViewer.cpp
+++ b/Engine/Plugins/Experimental/Mutable/Source/CustomizableObjectEditor/Private/MuCOE/SMutableCodeViewer.cpp
@@ -2459,9 +2459,9 @@ namespace
 		
 	public:
 #ifdef MUTABLE_USE_NEW_TASKGRAPH
-		TTuple<UE::Tasks::FTask, TFunction<void()>> GetImageAsync(const FExternalImageID& id, uint8 MipmapsToSkip, TFunction<void(mu::Ptr<mu::Image>)>& ResultCallback) override
+		TTuple<UE::Tasks::FTask, TFunction<void()>> GetImageAsync(FName Id, uint8 MipmapsToSkip, TFunction<void(mu::Ptr<mu::Image>)>& ResultCallback) override
 #else
-		TTuple<FGraphEventRef, TFunction<void()>> GetImageAsync(const FExternalImageID& id, uint8 MipmapsToSkip, TFunction<void(mu::Ptr<mu::Image>)>& ResultCallback) override
+		TTuple<FGraphEventRef, TFunction<void()>> GetImageAsync(FName Id, uint8 MipmapsToSkip, TFunction<void(mu::Ptr<mu::Image>)>& ResultCallback) override
 #endif
 		{
 			MUTABLE_CPUPROFILER_SCOPE(TestImageProvider_GetImage);
@@ -2512,7 +2512,7 @@ namespace
 			return MakeTuple(CompletionEvent, []() -> void {});
 		}
 
-		mu::FImageDesc GetImageDesc(const FExternalImageID& Id, uint8 MipmapsToSkip) override
+		mu::FImageDesc GetImageDesc(FName Id, uint8 MipmapsToSkip) override
 		{
 			return IMAGE_DESC;
 		}
diff --git a/Engine/Plugins/Experimental/Mutable/Source/CustomizableObjectEditor/Private/MuCOE/StressTest.cpp b/Engine/Plugins/Experimental/Mutable/Source/CustomizableObjectEditor/Private/MuCOE/StressTest.cpp
--- a/Engine/Plugins/Experimental/Mutable/Source/CustomizableObjectEditor/Private/MuCOE/StressTest.cpp
+++ b/Engine/Plugins/Experimental/Mutable/Source/CustomizableObjectEditor/Private/MuCOE/StressTest.cpp
@@ -666,7 +666,7 @@ void FRunningStressTest::PrintParameterData(const MostExpensiveInstanceData& Dat
 		{
 			Line = FString::Printf(TEXT("\t\t\t %s"), *it.ParameterName);
 			AppendWhitespace(Line, MaxLength - Line.Len());
-			Line += it.ParameterValue;
+			Line += it.ParameterValue.ToString();
 			UE_LOG(LogMutable, Log, TEXT("%s"), *Line);
 		}
 	}
diff --git a/Engine/Plugins/Experimental/Mutable/Source/CustomizableObjectEditor/Public/MuCOE/SCustomizableInstanceProperties.h b/Engine/Plugins/Experimental/Mutable/Source/CustomizableObjectEditor/Public/MuCOE/SCustomizableInstanceProperties.h
--- a/Engine/Plugins/Experimental/Mutable/Source/CustomizableObjectEditor/Public/MuCOE/SCustomizableInstanceProperties.h
+++ b/Engine/Plugins/Experimental/Mutable/Source/CustomizableObjectEditor/Public/MuCOE/SCustomizableInstanceProperties.h
@@ -99,7 +99,7 @@ private:
 	// These arrays store the textures available for texture parameters of the model.
 	// These come from the texture generators registered in the CustomizableObjectSystem
 	TArray<TSharedPtr<FString>> TextureParameterValueNames;
-	TArray<FString> TextureParameterValues;
+	TArray<FName> TextureParameterValues;
 
 	struct FSliderData
 	{
diff --git a/Engine/Plugins/Experimental/Mutable/Source/MutableRuntime/Private/MuR/CodeRunner.cpp b/Engine/Plugins/Experimental/Mutable/Source/MutableRuntime/Private/MuR/CodeRunner.cpp
--- a/Engine/Plugins/Experimental/Mutable/Source/MutableRuntime/Private/MuR/CodeRunner.cpp
+++ b/Engine/Plugins/Experimental/Mutable/Source/MutableRuntime/Private/MuR/CodeRunner.cpp
@@ -138,9 +138,9 @@ namespace mu
 
     //---------------------------------------------------------------------------------------------
 #ifdef MUTABLE_USE_NEW_TASKGRAPH
-	TTuple<UE::Tasks::FTask, TFunction<void()>> CodeRunner::LoadExternalImageAsync(FExternalImageID Id, uint8 MipmapsToSkip, TFunction<void(Ptr<Image>)>& ResultCallback)
+	TTuple<UE::Tasks::FTask, TFunction<void()>> CodeRunner::LoadExternalImageAsync(FName Id, uint8 MipmapsToSkip, TFunction<void(Ptr<Image>)>& ResultCallback)
 #else
-	TTuple<FGraphEventRef, TFunction<void()>> CodeRunner::LoadExternalImageAsync(FExternalImageID Id, uint8 MipmapsToSkip, TFunction<void(Ptr<Image>)>& ResultCallback)
+	TTuple<FGraphEventRef, TFunction<void()>> CodeRunner::LoadExternalImageAsync(FName Id, uint8 MipmapsToSkip, TFunction<void(Ptr<Image>)>& ResultCallback)
 #endif
     {
 		MUTABLE_CPUPROFILER_SCOPE(LoadExternalImageAsync);
@@ -150,10 +150,6 @@ namespace mu
 		if (m_pSystem->ImageParameterGenerator)
 		{
 			return m_pSystem->ImageParameterGenerator->GetImageAsync(Id, MipmapsToSkip, ResultCallback);
-
-			// Don't cache for now. Need to figure out how to invalidate them.
-			// \TODO: Like constants? attached to a cache level?
-			//m_pSystem->m_externalImages.push_back( pair<FExternalImageID,ImagePtr>(id,pResult) );
 		}
 		else
 		{
@@ -175,7 +171,7 @@ namespace mu
 
 	
     //---------------------------------------------------------------------------------------------
-	FImageDesc CodeRunner::GetExternalImageDesc(FExternalImageID Id, uint8 MipmapsToSkip)
+	FImageDesc CodeRunner::GetExternalImageDesc(FName Id, uint8 MipmapsToSkip)
 	{
 		MUTABLE_CPUPROFILER_SCOPE(GetExternalImageDesc);
 
@@ -777,22 +773,22 @@ namespace mu
 
 
     //---------------------------------------------------------------------------------------------
-    void CodeRunner::RunCode_InstanceAddResource(const FScheduledOp& item, const Model* pModel, const Parameters* pParams )
+    void CodeRunner::RunCode_InstanceAddResource(const FScheduledOp& item, const Model* InModel, const Parameters* InParams )
     {
 		MUTABLE_CPUPROFILER_SCOPE(RunCode_InstanceAddResource);
 
-		if (!pModel || !m_pSystem)
+		if (!InModel || !m_pSystem)
 		{
 			return;
 		}
 
-		OP_TYPE type = pModel->GetPrivate()->m_program.GetOpType(item.At);
+		OP_TYPE type = InModel->GetPrivate()->m_program.GetOpType(item.At);
         switch (type)
         {
 
         case OP_TYPE::IN_ADDMESH:
         {
-			OP::InstanceAddArgs args = pModel->GetPrivate()->m_program.GetOpArgs<OP::InstanceAddArgs>(item.At);
+			OP::InstanceAddArgs args = InModel->GetPrivate()->m_program.GetOpArgs<OP::InstanceAddArgs>(item.At);
             switch (item.Stage)
             {
             case 0:
@@ -818,16 +814,11 @@ namespace mu
 
                 if ( args.value )
                 {
-					FWorkingMemoryManager::FModelCacheEntry* ModelCache = m_pSystem->WorkingMemoryManager.FindModelCache(pModel);
-					check(ModelCache);
-					if (ModelCache)
-					{
-						FResourceID MeshId = ModelCache->GetResourceKey(args.relevantParametersListIndex, args.value, pParams, m_pSystem->WorkingMemoryManager.MaxGeneratedResourceCacheSize);
-						OP::ADDRESS NameAd = args.name;
-						check(NameAd < (uint32)pModel->GetPrivate()->m_program.m_constantStrings.Num());
-						const char* Name = pModel->GetPrivate()->m_program.m_constantStrings[NameAd].c_str();
-						pResult->GetPrivate()->AddMesh(0, 0, MeshId, Name);
-					}
+					FResourceID MeshId = m_pSystem->WorkingMemoryManager.GetResourceKey(InModel,InParams,args.relevantParametersListIndex, args.value);
+					OP::ADDRESS NameAd = args.name;
+					check(NameAd < (uint32)InModel->GetPrivate()->m_program.m_constantStrings.Num());
+					const char* Name = InModel->GetPrivate()->m_program.m_constantStrings[NameAd].c_str();
+					pResult->GetPrivate()->AddMesh(0, 0, MeshId, Name);
                 }
                 StoreInstance( item, pResult );
                 break;
@@ -841,7 +832,7 @@ namespace mu
 
         case OP_TYPE::IN_ADDIMAGE:
         {
-			OP::InstanceAddArgs args = pModel->GetPrivate()->m_program.GetOpArgs<OP::InstanceAddArgs>(item.At);
+			OP::InstanceAddArgs args = InModel->GetPrivate()->m_program.GetOpArgs<OP::InstanceAddArgs>(item.At);
             switch (item.Stage)
             {
             case 0:
@@ -864,16 +855,11 @@ namespace mu
 
                 if ( args.value )
                 {
-					FWorkingMemoryManager::FModelCacheEntry* ModelCache = m_pSystem->WorkingMemoryManager.FindModelCache(pModel);
-					check(ModelCache);
-					if (ModelCache)
-					{
-						FResourceID ImageId = ModelCache->GetResourceKey(args.relevantParametersListIndex, args.value, pParams, m_pSystem->WorkingMemoryManager.MaxGeneratedResourceCacheSize );
-						OP::ADDRESS NameAd = args.name;
-						check(NameAd < (uint32)pModel->GetPrivate()->m_program.m_constantStrings.Num());
-						const char* Name = pModel->GetPrivate()->m_program.m_constantStrings[NameAd].c_str();
-						pResult->GetPrivate()->AddImage(0, 0, 0, ImageId, Name);
-					}
+					FResourceID ImageId = m_pSystem->WorkingMemoryManager.GetResourceKey(InModel, InParams, args.relevantParametersListIndex, args.value);
+					OP::ADDRESS NameAd = args.name;
+					check(NameAd < (uint32)InModel->GetPrivate()->m_program.m_constantStrings.Num());
+					const char* Name = InModel->GetPrivate()->m_program.m_constantStrings[NameAd].c_str();
+					pResult->GetPrivate()->AddImage(0, 0, 0, ImageId, Name);
                 }
                 StoreInstance( item, pResult );
                 break;
@@ -5947,9 +5933,9 @@ namespace mu
 		{
 			check(item.Stage == 0);
 			OP::ParameterArgs args = program.GetOpArgs<OP::ParameterArgs>(item.At);
-			FExternalImageID id = pParams->GetImageValue(args.variable);
+			FName Id = pParams->GetImageValue(args.variable);
 			uint8 MipsToSkip = item.ExecutionOptions;
-			m_heapImageDesc[item.CustomState] = GetExternalImageDesc(id, MipsToSkip);
+			m_heapImageDesc[item.CustomState] = GetExternalImageDesc(Id, MipsToSkip);
 			StoreValidDesc(item);
 			break;
 		}
diff --git a/Engine/Plugins/Experimental/Mutable/Source/MutableRuntime/Private/MuR/CodeRunner.h b/Engine/Plugins/Experimental/Mutable/Source/MutableRuntime/Private/MuR/CodeRunner.h
--- a/Engine/Plugins/Experimental/Mutable/Source/MutableRuntime/Private/MuR/CodeRunner.h
+++ b/Engine/Plugins/Experimental/Mutable/Source/MutableRuntime/Private/MuR/CodeRunner.h
@@ -127,11 +127,11 @@ namespace  mu
 		//! Load an external image asynchronously, retuns an event to wait for complition and a cleanup function 
 		//! that must be called once the event has completed.
 #ifdef MUTABLE_USE_NEW_TASKGRAPH
-		TTuple<UE::Tasks::FTask, TFunction<void()>> LoadExternalImageAsync(FExternalImageID Id, uint8 MipmapsToSkip, TFunction<void(Ptr<Image>)>& ResultCallback);
+		TTuple<UE::Tasks::FTask, TFunction<void()>> LoadExternalImageAsync(FName Id, uint8 MipmapsToSkip, TFunction<void(Ptr<Image>)>& ResultCallback);
 #else
-		TTuple<FGraphEventRef, TFunction<void()>> LoadExternalImageAsync(FExternalImageID Id, uint8 MipmapsToSkip, TFunction<void(Ptr<Image>)>& ResultCallback);
+		TTuple<FGraphEventRef, TFunction<void()>> LoadExternalImageAsync(FName Id, uint8 MipmapsToSkip, TFunction<void(Ptr<Image>)>& ResultCallback);
 #endif
- 	    mu::FImageDesc GetExternalImageDesc(FExternalImageID Id, uint8 MipmapsToSkip);
+ 	    mu::FImageDesc GetExternalImageDesc(FName Id, uint8 MipmapsToSkip);
 
 		/** Settings that may affect the execution of some operations, like image conversion quality. */
 		Ptr<const Settings> m_pSettings;
diff --git a/Engine/Plugins/Experimental/Mutable/Source/MutableRuntime/Private/MuR/CodeRunnerTasked.cpp b/Engine/Plugins/Experimental/Mutable/Source/MutableRuntime/Private/MuR/CodeRunnerTasked.cpp
--- a/Engine/Plugins/Experimental/Mutable/Source/MutableRuntime/Private/MuR/CodeRunnerTasked.cpp
+++ b/Engine/Plugins/Experimental/Mutable/Source/MutableRuntime/Private/MuR/CodeRunnerTasked.cpp
@@ -2369,7 +2369,7 @@ namespace mu
 	class FImageExternalLoadTask : public CodeRunner::FIssuedTask
 	{
 	public:
-		FImageExternalLoadTask(const FScheduledOp& InItem, uint8 InMipmapsToSkip, FExternalImageID InId);
+		FImageExternalLoadTask(const FScheduledOp& InItem, uint8 InMipmapsToSkip, FName InId);
 
 		// FIssuedTask interface
 		virtual bool Prepare(CodeRunner*, bool& bOutFailed) override;
@@ -2377,7 +2377,7 @@ namespace mu
 		
 	private:
 		uint8 MipmapsToSkip;
-		FExternalImageID Id;
+		FName Id;
 
 		Ptr<Image> Result;
 		
@@ -2386,7 +2386,7 @@ namespace mu
 
 
 	//---------------------------------------------------------------------------------------------
-	FImageExternalLoadTask::FImageExternalLoadTask(const FScheduledOp& InOp,  uint8 InMipmapsToSkip, FExternalImageID InId)
+	FImageExternalLoadTask::FImageExternalLoadTask(const FScheduledOp& InOp,  uint8 InMipmapsToSkip, FName InId)
 		: FIssuedTask(InOp)
 	{
 		MipmapsToSkip = InMipmapsToSkip;
@@ -2522,7 +2522,7 @@ namespace mu
 			OP::ParameterArgs args = program.GetOpArgs<OP::ParameterArgs>(item.At);
 			Ptr<RangeIndex> Index = BuildCurrentOpRangeIndex(item, m_pParams, m_pModel.Get(), args.variable);
 
-			const FExternalImageID Id = m_pParams->GetImageValue(args.variable, Index);
+			const FName Id = m_pParams->GetImageValue(args.variable, Index);
 
 			check(ImageLOD < TNumericLimits<uint8>::Max() && ImageLOD >= 0);
 			check(ImageLOD + static_cast<int32>(item.ExecutionOptions) < TNumericLimits<uint8>::Max());
diff --git a/Engine/Plugins/Experimental/Mutable/Source/MutableRuntime/Private/MuR/Instance.cpp b/Engine/Plugins/Experimental/Mutable/Source/MutableRuntime/Private/MuR/Instance.cpp
--- a/Engine/Plugins/Experimental/Mutable/Source/MutableRuntime/Private/MuR/Instance.cpp
+++ b/Engine/Plugins/Experimental/Mutable/Source/MutableRuntime/Private/MuR/Instance.cpp
@@ -154,7 +154,7 @@ namespace mu
 
 
     //---------------------------------------------------------------------------------------------
-    uint32_t Instance::GetSurfaceId( int32 lod, int32 comp, int32 surf ) const
+    uint32 Instance::GetSurfaceId( int32 lod, int32 comp, int32 surf ) const
     {
         if ( lod>=0 && lod<m_pD->m_lods.Num() &&
              comp>=0 && comp<m_pD->m_lods[lod].m_components.Num() &&
@@ -172,12 +172,12 @@ namespace mu
 
 
     //---------------------------------------------------------------------------------------------
-    int Instance::FindSurfaceById( int lod, int comp, uint32_t id ) const
+    int32 Instance::FindSurfaceById( int32 lod, int32 comp, uint32 id ) const
     {
 		if (lod >= 0 && lod < m_pD->m_lods.Num() &&
 			comp >= 0 && comp < m_pD->m_lods[lod].m_components.Num())
 		{
-			for (int i = 0; i < m_pD->m_lods[lod].m_components[comp].m_surfaces.Num(); ++i)
+			for (int32 i = 0; i < m_pD->m_lods[lod].m_components[comp].m_surfaces.Num(); ++i)
 			{
 				if (m_pD->m_lods[lod].m_components[comp].m_surfaces[i].InternalId == id)
 				{
@@ -238,7 +238,7 @@ namespace mu
 	
 
     //---------------------------------------------------------------------------------------------
-    uint32_t Instance::GetSurfaceCustomId( int32 lod, int32 comp, int32 surf ) const
+    uint32 Instance::GetSurfaceCustomId( int32 lod, int32 comp, int32 surf ) const
     {
         if ( lod>=0 && lod<m_pD->m_lods.Num() &&
              comp>=0 && comp<m_pD->m_lods[lod].m_components.Num() &&
@@ -256,7 +256,7 @@ namespace mu
 
 
 	//---------------------------------------------------------------------------------------------
-    int Instance::GetMeshCount( int32 lod, int32 comp ) const
+    int32 Instance::GetMeshCount( int32 lod, int32 comp ) const
 	{
 		check(lod >= 0 && lod < m_pD->m_lods.Num());
 		check(comp >= 0 && comp < m_pD->m_lods[lod].m_components.Num());
@@ -469,14 +469,14 @@ namespace mu
     //---------------------------------------------------------------------------------------------
 	int32 Instance::Private::AddLOD()
 	{
-		int result = m_lods.Num();
+		int32 result = m_lods.Num();
 		m_lods.Add( INSTANCE_LOD() );
 		return result;
 	}
 
 
     //---------------------------------------------------------------------------------------------
-    int32 Instance::Private::AddComponent( int lod )
+    int32 Instance::Private::AddComponent( int32 lod )
     {
         // Automatically create the necessary lods and components
         while (lod >= m_lods.Num())
@@ -484,14 +484,14 @@ namespace mu
             AddLOD();
         }
 
-        int result = m_lods[lod].m_components.Num();
+        int32 result = m_lods[lod].m_components.Num();
         m_lods[lod].m_components.Add( INSTANCE_COMPONENT() );
         return result;
     }
 
 
     //---------------------------------------------------------------------------------------------
-    int32 Instance::Private::AddSurface( int lod, int comp )
+    int32 Instance::Private::AddSurface( int32 lod, int32 comp )
     {
         // Automatically create the necessary lods and components
         while (lod >= m_lods.Num())
@@ -503,7 +503,7 @@ namespace mu
             AddComponent(lod);
         }
 
-        int result = m_lods[lod].m_components[comp].m_surfaces.Num();
+        int32 result = m_lods[lod].m_components[comp].m_surfaces.Num();
         m_lods[lod].m_components[comp].m_surfaces.Add( INSTANCE_SURFACE() );
         return result;
     }
@@ -577,7 +577,7 @@ namespace mu
 		}
 
 		INSTANCE_COMPONENT& component = m_lods[lod].m_components[comp];
-        int result = component.m_meshes.Num();
+        int32 result = component.m_meshes.Num();
         component.m_meshes.Emplace( meshId, strName );
 
 		return result;
@@ -602,7 +602,7 @@ namespace mu
 		}
 
         INSTANCE_SURFACE& surface = m_lods[lod].m_components[comp].m_surfaces[surf];
-        int result = surface.m_images.Num();
+        int32 result = surface.m_images.Num();
         surface.m_images.Emplace( imageId, strName );
 
         return result;
@@ -627,7 +627,7 @@ namespace mu
 		}
 
         INSTANCE_SURFACE& surface = m_lods[lod].m_components[comp].m_surfaces[surf];
-        int result = surface.m_vectors.Num();
+        int32 result = surface.m_vectors.Num();
         surface.m_vectors.Emplace( vec, strName );
 
         return result;
@@ -652,7 +652,7 @@ namespace mu
 		}
 
         INSTANCE_SURFACE& surface = m_lods[lod].m_components[comp].m_surfaces[surf];
-        int result = surface.m_scalars.Num();
+        int32 result = surface.m_scalars.Num();
         surface.m_scalars.Emplace( sca, strName );
 
         return result;
@@ -677,7 +677,7 @@ namespace mu
 		}
 
         INSTANCE_SURFACE& surface = m_lods[lod].m_components[comp].m_surfaces[surf];
-        int result = surface.m_strings.Num();
+        int32 result = surface.m_strings.Num();
         surface.m_strings.Emplace( str, strName );
 
         return result;
diff --git a/Engine/Plugins/Experimental/Mutable/Source/MutableRuntime/Private/MuR/Model.cpp b/Engine/Plugins/Experimental/Mutable/Source/MutableRuntime/Private/MuR/Model.cpp
--- a/Engine/Plugins/Experimental/Mutable/Source/MutableRuntime/Private/MuR/Model.cpp
+++ b/Engine/Plugins/Experimental/Mutable/Source/MutableRuntime/Private/MuR/Model.cpp
@@ -356,7 +356,7 @@ namespace mu
 
 
 	//---------------------------------------------------------------------------------------------
-    FExternalImageID Model::GetImageDefaultValue(int32 Index) const
+	FName Model::GetImageDefaultValue(int32 Index) const
     {
 	    check(m_pD->m_program.m_parameters.IsValidIndex(Index));
 		check(m_pD->m_program.m_parameters[Index].m_type == PARAMETER_TYPE::T_IMAGE);
diff --git a/Engine/Plugins/Experimental/Mutable/Source/MutableRuntime/Private/MuR/Parameters.cpp b/Engine/Plugins/Experimental/Mutable/Source/MutableRuntime/Private/MuR/Parameters.cpp
--- a/Engine/Plugins/Experimental/Mutable/Source/MutableRuntime/Private/MuR/Parameters.cpp
+++ b/Engine/Plugins/Experimental/Mutable/Source/MutableRuntime/Private/MuR/Parameters.cpp
@@ -758,7 +758,7 @@ namespace mu
 
 
     //---------------------------------------------------------------------------------------------
-    FExternalImageID Parameters::GetImageValue( int index, const Ptr<const RangeIndex>& pos ) const
+	FName Parameters::GetImageValue( int index, const Ptr<const RangeIndex>& pos ) const
     {
         check( index >= 0 && index < (int)m_pD->m_values.Num() );
         check( GetType( index ) == PARAMETER_TYPE::T_IMAGE );
@@ -799,7 +799,7 @@ namespace mu
 
 
     //---------------------------------------------------------------------------------------------
-    void Parameters::SetImageValue( int index, FExternalImageID id, const Ptr<const RangeIndex>& pos )
+    void Parameters::SetImageValue( int index, FName id, const Ptr<const RangeIndex>& pos )
     {
         check( index >= 0 && index < (int)m_pD->m_values.Num() );
         check( GetType( index ) == PARAMETER_TYPE::T_IMAGE );
diff --git a/Engine/Plugins/Experimental/Mutable/Source/MutableRuntime/Private/MuR/ParametersPrivate.h b/Engine/Plugins/Experimental/Mutable/Source/MutableRuntime/Private/MuR/ParametersPrivate.h
--- a/Engine/Plugins/Experimental/Mutable/Source/MutableRuntime/Private/MuR/ParametersPrivate.h
+++ b/Engine/Plugins/Experimental/Mutable/Source/MutableRuntime/Private/MuR/ParametersPrivate.h
@@ -255,7 +255,7 @@ namespace mu
 	using ParamFloatType = float;
 	using ParamColorType = FVector3f;
 	using ParamProjectorType = FProjector;
-	using ParamImageType = FExternalImageID;
+	using ParamImageType = FName;
 	using ParamStringType = string;
 
 	
@@ -481,6 +481,20 @@ namespace mu
 
         //!
         FProjector GetProjectorValue( int index, const Ptr<const RangeIndex>& pos ) const;
+
+		/** Return true if the parameter has any multi-dimensional values set. This is independent to if the model
+		* accepts multi-dimensional parameters for this particular parameter.
+		*/
+		inline bool HasMultipleValues(int32 ParamIndex) const
+		{
+			if (ParamIndex >= m_multiValues.Num())
+			{
+				return false;
+			}
+
+			return m_multiValues[ParamIndex].Num()>0;
+		}
+
     };
 
 }
diff --git a/Engine/Plugins/Experimental/Mutable/Source/MutableRuntime/Private/MuR/System.cpp b/Engine/Plugins/Experimental/Mutable/Source/MutableRuntime/Private/MuR/System.cpp
--- a/Engine/Plugins/Experimental/Mutable/Source/MutableRuntime/Private/MuR/System.cpp
+++ b/Engine/Plugins/Experimental/Mutable/Source/MutableRuntime/Private/MuR/System.cpp
@@ -5,6 +5,7 @@
 #include "Containers/Array.h"
 #include "Containers/Map.h"
 #include "Containers/Set.h"
+#include "Serialization/BitWriter.h"
 #include "HAL/IConsoleManager.h"
 #include "HAL/LowLevelMemTracker.h"
 #include "HAL/PlatformCrt.h"
@@ -32,6 +33,7 @@
 #include "Templates/Tuple.h"
 #include "Trace/Detail/Channel.h"
 #include "ProfilingDebugging/CountersTrace.h"
+#include "PackedNormal.h"
 
 
 namespace mu
@@ -101,6 +103,7 @@ namespace mu
         Settings = InSettings;
 	
 		WorkingMemoryManager.BudgetBytes = Settings->WorkingMemoryBytes;
+		WorkingMemoryManager.GeneratedResources.Reserve(WorkingMemoryManager.MaxGeneratedResourceCacheSize);
 
 		ExtensionDataStreamer = InDataStreamer;
 		if (!ExtensionDataStreamer)
@@ -150,6 +153,12 @@ namespace mu
 		MUTABLE_CPUPROFILER_SCOPE(SetGeneratedCacheSize);
 
 		m_pD->WorkingMemoryManager.MaxGeneratedResourceCacheSize = InCount;
+		m_pD->WorkingMemoryManager.GeneratedResources.Reserve(m_pD->WorkingMemoryManager.MaxGeneratedResourceCacheSize);
+		if (m_pD->WorkingMemoryManager.GeneratedResources.Num()>int32(InCount))
+		{
+			// Discard some random resource keys.
+			m_pD->WorkingMemoryManager.GeneratedResources.SetNum(InCount);
+		}
 	}
 
 	
@@ -1397,51 +1406,39 @@ namespace mu
 
 
 	//---------------------------------------------------------------------------------------------
-	static int32 AddMultiValueKeys(FWorkingMemoryManager::TMemoryTrackedArray<uint8>& Blob, int32 pos, const TMap< TArray<int>, PARAMETER_VALUE >& multi)
+	static void AddMultiValueKeys(FBitWriter& Blob, const TMap< TArray<int32>, PARAMETER_VALUE >& Multi)
 	{
-		Blob.SetNum(pos + 4);
-		uint32 s = uint32(multi.Num());
-		FMemory::Memcpy(&Blob[pos], &s, sizeof(uint32));
-		pos += 4;
+		uint16 Num = uint16(Multi.Num());
+		Blob.Serialize(&Num, 2);
 
-		for (const auto& v : multi)
+		for (const TPair< TArray<int32>, PARAMETER_VALUE >& v : Multi)
 		{
-			uint32 ds = uint32(v.Key.Num());
-
-			Blob.SetNum(pos + 4 + ds * 4);
-
-			FMemory::Memcpy(&Blob[pos], &s, sizeof(uint32));
-			pos += 4;
-
-			FMemory::Memcpy(&Blob[pos], v.Key.GetData(), ds * sizeof(int32));
-			pos += ds * sizeof(int32);
+			uint16 RangeNum = uint16(v.Key.Num());
+			Blob.Serialize(&RangeNum, 2);
+			Blob.Serialize((void*)v.Key.GetData(), RangeNum * sizeof(int32));
 		}
-
-		return pos;
 	}
 
 
 	//---------------------------------------------------------------------------------------------
-	FResourceID FWorkingMemoryManager::FModelCacheEntry::GetResourceKey(uint32 paramListIndex, OP::ADDRESS RootAt, const Parameters* Params, int32 InMaxResourceKeys)
+	FResourceID FWorkingMemoryManager::GetResourceKey(const Model* Model, const Parameters* Params, uint32 ParamListIndex, OP::ADDRESS RootAt)
 	{
 		MUTABLE_CPUPROFILER_SCOPE(GetResourceKey);
 
 		constexpr uint32 ErrorId = 0xffff;
 
-		const mu::Model* ThisModel = Model.Pin().Get();
-		check(ThisModel);
-		if (!ThisModel)
+		if (!Model)
 		{
 			return ErrorId;
 		}
 
-		const FProgram& Program = ThisModel->GetPrivate()->m_program;
+		const FProgram& Program = Model->GetPrivate()->m_program;
 
 		// Find the list of relevant parameters
 		const TArray<uint16>* RelevantParams = nullptr;
-		if (paramListIndex < (uint32)Program.m_parameterLists.Num())
+		if (ParamListIndex < (uint32)Program.m_parameterLists.Num())
 		{
-			RelevantParams = &Program.m_parameterLists[paramListIndex];
+			RelevantParams = &Program.m_parameterLists[ParamListIndex];
 		}
 		check(RelevantParams);
 		if (!RelevantParams)
@@ -1450,137 +1447,179 @@ namespace mu
 		}
 
 		// Generate the relevant parameters blob
-		TMemoryTrackedArray<uint8> Blob;
-		Blob.Reserve(2048);
-		for (int32 ParamIndex : *RelevantParams)
+		FBitWriter Blob( 2048*8, true );
+
+		const TArray<FParameterDesc>& ParamDescs = Params->GetPrivate()->m_pModel->GetPrivate()->m_program.m_parameters;
+
+		// First make a mask with a bit for each relevant parameter. It will be on for parameters included in the blob.
+		// A parameter will be excluded from the blob if it has the deafult value, and no multivalues.
+		TBitArray IncludedParameters(0, RelevantParams->Num());
+		if (RelevantParams->Num())
 		{
-			int32 pos = Blob.Num();
-			int32 dataSize = 0;
+			for (int32 IndexIndex = 0; IndexIndex < RelevantParams->Num(); ++IndexIndex)
+			{
+				int32 ParamIndex = (*RelevantParams)[IndexIndex];
+				bool bInclude = Params->GetPrivate()->HasMultipleValues(ParamIndex);
+				if (!bInclude)
+				{
+					bInclude =
+						Params->GetPrivate()->m_values[ParamIndex]
+						!=
+						ParamDescs[ParamIndex].m_defaultValue;
+				}
+
+				IncludedParameters[IndexIndex] = bInclude;
+			}
+			Blob.SerializeBits(IncludedParameters.GetData(), IncludedParameters.Num());
+		}
+
+		// Second: serialize the value of the selected parameters.
+		for (int32 IndexIndex = 0; IndexIndex < RelevantParams->Num(); ++IndexIndex)
+		{
+			int32 ParamIndex = (*RelevantParams)[IndexIndex];
+			if (!IncludedParameters[IndexIndex])
+			{
+				continue;
+			}
+
+			int32 DataSize = 0;
 
 			switch (Program.m_parameters[ParamIndex].m_type)
 			{
 			case PARAMETER_TYPE::T_BOOL:
-				dataSize = 1;
-				Blob.Add(Params->GetPrivate()->m_values[ParamIndex].Get<ParamBoolType>() ? 1 : 0);
-				pos += dataSize;
+				Blob.WriteBit(Params->GetPrivate()->m_values[ParamIndex].Get<ParamBoolType>() ? 1 : 0);
 
 				// Multi-values
-				if (ParamIndex < Params->GetPrivate()->m_multiValues.Num())
+				if (Params->GetPrivate()->HasMultipleValues(ParamIndex))
 				{
-					const TMap< TArray<int32>, PARAMETER_VALUE >& multi = Params->GetPrivate()->m_multiValues[ParamIndex];
-					pos = AddMultiValueKeys(Blob, pos, multi);
-					for (const TPair< TArray<int32>, PARAMETER_VALUE >& v : multi)
+					const TMap< TArray<int32>, PARAMETER_VALUE >& Multi = Params->GetPrivate()->m_multiValues[ParamIndex];
+					AddMultiValueKeys(Blob, Multi);
+					for (const TPair< TArray<int32>, PARAMETER_VALUE >& v : Multi)
 					{
-						Blob.Add(v.Value.Get<ParamBoolType>() ? 1 : 0);
-						pos += dataSize;
+						Blob.WriteBit(v.Value.Get<ParamBoolType>() ? 1 : 0);
 					}
 				}
 				break;
 
 			case PARAMETER_TYPE::T_INT:
-				dataSize = sizeof(ParamIntType);
-				Blob.SetNum(pos + dataSize);
-				FMemory::Memcpy(&Blob[pos], &Params->GetPrivate()->m_values[ParamIndex].Get<ParamIntType>(), dataSize);
-				pos += dataSize;
+			{
+				int32 MaxValue = ParamDescs[ParamIndex].m_possibleValues.Num();
+				int32 Value = Params->GetPrivate()->m_values[ParamIndex].Get<ParamIntType>();
+				if (MaxValue)
+				{
+					// It is an enum
+					uint32 LimitedValue = FMath::Clamp( Params->GetIntValueIndex(ParamIndex,Value), 0, MaxValue-1 );
+					Blob.SerializeInt(LimitedValue, uint32(MaxValue));
+				}
+				else
+				{
+					// It may have any value
+					DataSize = sizeof(ParamIntType);
+					Blob.Serialize(&Value, DataSize);
+				}
 
 				// Multi-values
-				if (ParamIndex < Params->GetPrivate()->m_multiValues.Num())
+				if (Params->GetPrivate()->HasMultipleValues(ParamIndex))
 				{
-					const TMap< TArray<int32>, PARAMETER_VALUE >& multi = Params->GetPrivate()->m_multiValues[ParamIndex];
-					pos = AddMultiValueKeys(Blob, pos, multi);
-					Blob.SetNum(pos + multi.Num() * dataSize);
-					for (const TPair< TArray<int32>, PARAMETER_VALUE >& v : multi)
+					const TMap< TArray<int32>, PARAMETER_VALUE >& Multi = Params->GetPrivate()->m_multiValues[ParamIndex];
+					AddMultiValueKeys(Blob, Multi);
+					for (const TPair< TArray<int32>, PARAMETER_VALUE >& v : Multi)
 					{
-						FMemory::Memcpy(&Blob[pos], &v.Value.Get<ParamIntType>(), dataSize);
-						pos += dataSize;
+						Value = v.Value.Get<ParamIntType>();
+						if (MaxValue)
+						{
+							// It is an enum
+							uint32 LimitedValue = Value;
+							Blob.SerializeInt(LimitedValue, uint32(MaxValue));
+						}
+						else
+						{
+							// It may have any value
+							DataSize = sizeof(ParamIntType);
+							Blob.Serialize(&Value, DataSize);
+						}
 					}
 				}
 				break;
+			}
 
-				//! Floating point value in the range of 0.0 to 1.0
 			case PARAMETER_TYPE::T_FLOAT:
-				dataSize = sizeof(ParamFloatType);
-				Blob.SetNum(pos + dataSize);
-				FMemory::Memcpy(&Blob[pos], &Params->GetPrivate()->m_values[ParamIndex].Get<ParamFloatType>(), dataSize);
-				pos += dataSize;
+				DataSize = sizeof(ParamFloatType);
+				Blob.Serialize(&Params->GetPrivate()->m_values[ParamIndex].Get<ParamFloatType>(), DataSize);
 
 				// Multi-values
-				if (ParamIndex < Params->GetPrivate()->m_multiValues.Num())
+				if (Params->GetPrivate()->HasMultipleValues(ParamIndex))
 				{
-					const TMap< TArray<int32>, PARAMETER_VALUE >& multi = Params->GetPrivate()->m_multiValues[ParamIndex];
-					pos = AddMultiValueKeys(Blob, pos, multi);
-					Blob.SetNum(pos + multi.Num() * dataSize);
-					for (const TPair< TArray<int32>, PARAMETER_VALUE >& v : multi)
+					const TMap< TArray<int32>, PARAMETER_VALUE >& Multi = Params->GetPrivate()->m_multiValues[ParamIndex];
+					AddMultiValueKeys(Blob, Multi);
+					for (const TPair< TArray<int32>, PARAMETER_VALUE >& v : Multi)
 					{
-						FMemory::Memcpy(&Blob[pos], &v.Value.Get<ParamFloatType>(), dataSize);
-						pos += dataSize;
+						Blob.Serialize((void*)&v.Value.Get<ParamFloatType>(), DataSize);
 					}
 				}
 				break;
 
-				//! Floating point RGBA colour, with each channel ranging from 0.0 to 1.0
 			case PARAMETER_TYPE::T_COLOUR:
-				dataSize = sizeof(ParamColorType);
-				Blob.SetNum(pos + dataSize);
-				FMemory::Memcpy(&Blob[pos], &Params->GetPrivate()->m_values[ParamIndex].Get<ParamColorType>(), dataSize);
-				pos += dataSize;
+				DataSize = sizeof(ParamColorType);
+				Blob.Serialize(&Params->GetPrivate()->m_values[ParamIndex].Get<ParamColorType>(), DataSize);
 
 				// Multi-values
-				if (ParamIndex < Params->GetPrivate()->m_multiValues.Num())
+				if (Params->GetPrivate()->HasMultipleValues(ParamIndex))
 				{
-					const TMap< TArray<int32>, PARAMETER_VALUE >& multi = Params->GetPrivate()->m_multiValues[ParamIndex];
-					pos = AddMultiValueKeys(Blob, pos, multi);
-					Blob.SetNum(pos + multi.Num() * dataSize);
-					for (const TPair< TArray<int32>, PARAMETER_VALUE >& v : multi)
+					const TMap< TArray<int32>, PARAMETER_VALUE >& Multi = Params->GetPrivate()->m_multiValues[ParamIndex];
+					AddMultiValueKeys(Blob, Multi);
+					for (const TPair< TArray<int32>, PARAMETER_VALUE >& v : Multi)
 					{
-						FMemory::Memcpy(&Blob[pos], &v.Value.Get<ParamColorType>(), dataSize);
-						pos += dataSize;
+						Blob.Serialize((void*)&v.Value.Get<ParamColorType>(), DataSize);
 					}
 				}
 				break;
 
-				//! 3D Projector type, defining a position, scale and orientation. Basically used for
-				//! projected decals.
 			case PARAMETER_TYPE::T_PROJECTOR:
-				dataSize = sizeof(ParamProjectorType);
-
-				// \todo: padding will be random?
-				Blob.SetNum(pos + dataSize);
-				FMemory::Memcpy(&Blob[pos], &Params->GetPrivate()->m_values[ParamIndex].Get<ParamProjectorType>(), dataSize);
-				pos += dataSize;
+			{
+				FPackedNormal TempVec;
+				FProjector& Value = Params->GetPrivate()->m_values[ParamIndex].Get<ParamProjectorType>();
+				Blob.Serialize((void*)&Value.position, sizeof(FVector3f));
+				TempVec = FPackedNormal(Value.direction);
+				Blob.Serialize(&TempVec, sizeof(FPackedNormal));
+				TempVec = FPackedNormal(Value.up);
+				Blob.Serialize(&TempVec, sizeof(FPackedNormal));
+				Blob.Serialize((void*)&Value.scale, sizeof(FVector3f));
+				Blob.Serialize((void*)&Value.projectionAngle, sizeof(float));
 
 				// Multi-values
-				if (ParamIndex < Params->GetPrivate()->m_multiValues.Num())
+				if (Params->GetPrivate()->HasMultipleValues(ParamIndex))
 				{
-					const TMap< TArray<int32>, PARAMETER_VALUE >& multi = Params->GetPrivate()->m_multiValues[ParamIndex];
-					pos = AddMultiValueKeys(Blob, pos, multi);
-					Blob.SetNum(pos + multi.Num() * dataSize);
-					for (const TPair< TArray<int32>, PARAMETER_VALUE >& v : multi)
+					const TMap< TArray<int32>, PARAMETER_VALUE >& Multi = Params->GetPrivate()->m_multiValues[ParamIndex];
+					AddMultiValueKeys(Blob, Multi);
+					for (const TPair< TArray<int32>, PARAMETER_VALUE >& v : Multi)
 					{
-						FMemory::Memcpy(&Blob[pos], &v.Value.Get<ParamProjectorType>(), dataSize);
-						pos += dataSize;
+						Value = v.Value.Get<ParamProjectorType>();
+						Blob.Serialize((void*)&Value.position, sizeof(FVector3f));
+						TempVec = FPackedNormal(Value.direction);
+						Blob.Serialize(&TempVec, sizeof(FPackedNormal));
+						TempVec = FPackedNormal(Value.up);
+						Blob.Serialize(&TempVec, sizeof(FPackedNormal));
+						Blob.Serialize((void*)&Value.scale, sizeof(FVector3f));
+						Blob.Serialize((void*)&Value.projectionAngle, sizeof(float));
 					}
 				}
 				break;
+			}
 
 			case PARAMETER_TYPE::T_IMAGE:
 			{
-				dataSize = sizeof(ParamImageType::HashType);
-				Blob.SetNum(pos + dataSize);
-				uint64 Hash = Params->GetPrivate()->m_values[ParamIndex].Get<ParamImageType>().GetHash();
-				FMemory::Memcpy(&Blob[pos], &Hash, dataSize);
-				pos += dataSize;
+				DataSize = sizeof(FName);
+				Blob.Serialize((void*)&Params->GetPrivate()->m_values[ParamIndex].Get<ParamImageType>(), DataSize);
 
 				// Multi-values
-				if (ParamIndex < Params->GetPrivate()->m_multiValues.Num())
+				if (Params->GetPrivate()->HasMultipleValues(ParamIndex))
 				{
-					const TMap< TArray<int32>, PARAMETER_VALUE >& multi = Params->GetPrivate()->m_multiValues[ParamIndex];
-					pos = AddMultiValueKeys(Blob, pos, multi);
-					Blob.SetNum(pos + multi.Num() * dataSize);
-					for (const TPair< TArray<int32>, PARAMETER_VALUE >& v : multi)
+					const TMap< TArray<int32>, PARAMETER_VALUE >& Multi = Params->GetPrivate()->m_multiValues[ParamIndex];
+					AddMultiValueKeys(Blob, Multi);
+					for (const TPair< TArray<int32>, PARAMETER_VALUE >& v : Multi)
 					{
-						FMemory::Memcpy(&Blob[pos], &v.Value.Get<ParamImageType>(), dataSize);
-						pos += dataSize;
+						Blob.Serialize((void*)&v.Value.Get<ParamImageType>(), DataSize);
 					}
 				}
 				break;
@@ -1595,38 +1634,44 @@ namespace mu
 		// Increase the request id
 		++LastResourceResquestId;
 
+		FGeneratedResourceData NewKey;
+		int32 BlobBytes = Blob.GetNumBytes();
+		NewKey.ParameterValuesBlob.SetNum(BlobBytes);
+		FMemory::Memcpy(NewKey.ParameterValuesBlob.GetData(), Blob.GetData(), BlobBytes);
+
 		// See if we already have this id
 		int32 OldestCachePosition = 0;
+		uint32 OldestRequestId = 0;
 		for (int32 CacheIndex = 0; CacheIndex < GeneratedResources.Num(); ++CacheIndex)
 		{
 			FGeneratedResourceData& Data = GeneratedResources[CacheIndex];
-			if (GetResourceIDRoot(Data.Id) == RootAt
+			if (Data.Model.Pin().Get() == Model
+				&&
+				GetResourceIDRoot(Data.Id) == RootAt
 				&&
-				Data.ParameterValuesBlob == Blob)
+				Data.ParameterValuesBlob == NewKey.ParameterValuesBlob)
 			{
 				Data.LastRequestId = LastResourceResquestId;
 				return Data.Id;
 			}
 			else
 			{
-				if (GeneratedResources[OldestCachePosition].LastRequestId > Data.LastRequestId)
+				if (!Data.Model.Pin().Get() 
+					||
+					OldestRequestId > Data.LastRequestId)
 				{
 					OldestCachePosition = CacheIndex;
+					OldestRequestId = Data.Model.Pin() ? Data.LastRequestId : 0;
 				}
 			}
 		}
 
 		// Generate a new id
 		uint32 NewBlobId = ++LastResourceKeyId;
-		FGeneratedResourceData NewKey;
 		NewKey.Id = MakeResourceID( RootAt, NewBlobId);
 		NewKey.LastRequestId = LastResourceResquestId;
 
-		// Don't MoveTemp so that we get a tight array on the destination.
-		//NewKey.ParameterValuesBlob = MoveTemp(Blob);
-		NewKey.ParameterValuesBlob = Blob;
-
-		if (GeneratedResources.Num() >= InMaxResourceKeys)
+		if (GeneratedResources.Num() >= MaxGeneratedResourceCacheSize)
 		{
 			GeneratedResources[OldestCachePosition] = MoveTemp(NewKey);
 		}
diff --git a/Engine/Plugins/Experimental/Mutable/Source/MutableRuntime/Private/MuR/SystemPrivate.h b/Engine/Plugins/Experimental/Mutable/Source/MutableRuntime/Private/MuR/SystemPrivate.h
--- a/Engine/Plugins/Experimental/Mutable/Source/MutableRuntime/Private/MuR/SystemPrivate.h
+++ b/Engine/Plugins/Experimental/Mutable/Source/MutableRuntime/Private/MuR/SystemPrivate.h
@@ -1126,39 +1126,42 @@ namespace mu
 
 			/** Count of pending operations for every rom index. */
 			TMemoryTrackedArray<uint16> PendingOpsPerRom;
+        };
 
-			//! Management of generated resources
-			//! @{
+		//! Management of generated resources
+		//! @{
 
-			//! This is used to uniquely identify a generated resource like meshes or images.
-			struct FGeneratedResourceData
-			{
-				//! The id assigned to the generated resource.
-				FResourceID Id;
+		//! This is used to uniquely identify a generated resource like meshes or images.
+		struct FGeneratedResourceData
+		{
+			/** Model for this resource. */
+			TWeakPtr<const Model> Model;
 
-				//! The last request operation for this resource
-				uint32 LastRequestId;
+			//! The id assigned to the generated resource.
+			FResourceID Id;
 
-				//! An opaque blob with the values of the relevant parameters
-				TMemoryTrackedArray<uint8> ParameterValuesBlob;
-			};
+			//! The last request operation for this resource
+			uint32 LastRequestId;
 
-			//! The last id generated for a resource
-			uint32 LastResourceKeyId = 0;
+			//! An opaque blob with the values of the relevant parameters
+			TMemoryTrackedArray<uint8> ParameterValuesBlob;
+		};
 
-			//! The last id generated for a resource request. This is used to check the
-			//! relevancy of the resources when flushing the cache
-			uint32 LastResourceResquestId = 0;
+		//! The last id generated for a resource
+		uint32 LastResourceKeyId = 0;
 
-			//! Cached ids for returned assets
-			//! This is non-persistent runtime data
-			TMemoryTrackedArray<FGeneratedResourceData> GeneratedResources;
+		//! The last id generated for a resource request. This is used to check the
+		//! relevancy of the resources when flushing the cache
+		uint32 LastResourceResquestId = 0;
 
-			//! Get a resource key for a given resource with given parameter values.
-			FResourceID GetResourceKey(uint32 ParamListIndex, OP::ADDRESS RootAt, const Parameters*, int32 InMaxResourceKeys);
+		//! Cached ids for returned assets
+		//! This is non-persistent runtime data
+		TMemoryTrackedArray<FGeneratedResourceData> GeneratedResources;
 
-			//! @}
-        };
+		/** */
+		FResourceID GetResourceKey(const Model* Model, const Parameters* Params, uint32 ParamListIndex, OP::ADDRESS RootAt);
+
+		//! @}
 
 		/** Maximum working memory that mutable should be using. */
 		int64 BudgetBytes = 0;
diff --git a/Engine/Plugins/Experimental/Mutable/Source/MutableRuntime/Private/MuR/Types.cpp b/Engine/Plugins/Experimental/Mutable/Source/MutableRuntime/Private/MuR/Types.cpp
deleted file mode 100644
diff --git a/Engine/Plugins/Experimental/Mutable/Source/MutableRuntime/Public/MuR/Model.h b/Engine/Plugins/Experimental/Mutable/Source/MutableRuntime/Public/MuR/Model.h
--- a/Engine/Plugins/Experimental/Mutable/Source/MutableRuntime/Public/MuR/Model.h
+++ b/Engine/Plugins/Experimental/Mutable/Source/MutableRuntime/Public/MuR/Model.h
@@ -135,7 +135,7 @@ namespace mu
         //! \pre The parameter specified by index is a T_IMAGE.
         //! \param Index Index of the parameter from 0 to GetCount()-1
 		//! \return The externalId specified when setting the image value (\see SetImageValue)
-        FExternalImageID GetImageDefaultValue(int32 Index) const;
+		FName GetImageDefaultValue(int32 Index) const;
 
     	int32 GetRomCount() const;
 
diff --git a/Engine/Plugins/Experimental/Mutable/Source/MutableRuntime/Public/MuR/Parameters.h b/Engine/Plugins/Experimental/Mutable/Source/MutableRuntime/Public/MuR/Parameters.h
--- a/Engine/Plugins/Experimental/Mutable/Source/MutableRuntime/Public/MuR/Parameters.h
+++ b/Engine/Plugins/Experimental/Mutable/Source/MutableRuntime/Public/MuR/Parameters.h
@@ -351,13 +351,13 @@ namespace mu
         //! \param index Index of the parameter from 0 to GetCount()-1
 		//! \param pos Only for multidimensional parameters: relevant position to set in the ranges
 	    //! \return The externalId specified when setting the image value (\see SetImageValue)
-        FExternalImageID GetImageValue( int index, const Ptr<const RangeIndex>& pos = nullptr) const;
+		FName GetImageValue( int index, const Ptr<const RangeIndex>& pos = nullptr) const;
 
         //! If the parameter is of the image type, set its value.
         //! \param index Index of the parameter from 0 to GetCount()-1
         //! \param externalId Application-specific id used to identify this image during replication.
 		//! \param pos Only for multidimensional parameters: relevant position to set in the ranges
-		void SetImageValue( int index, FExternalImageID externalId, const Ptr<const RangeIndex>& pos = nullptr);
+		void SetImageValue( int index, FName externalId, const Ptr<const RangeIndex>& pos = nullptr);
 
         //! Return the value of a float parameter.
         //! \pre The parameter specified by index is a T_FLOAT.
diff --git a/Engine/Plugins/Experimental/Mutable/Source/MutableRuntime/Public/MuR/Serialisation.h b/Engine/Plugins/Experimental/Mutable/Source/MutableRuntime/Public/MuR/Serialisation.h
--- a/Engine/Plugins/Experimental/Mutable/Source/MutableRuntime/Public/MuR/Serialisation.h
+++ b/Engine/Plugins/Experimental/Mutable/Source/MutableRuntime/Public/MuR/Serialisation.h
@@ -378,8 +378,8 @@ namespace mu
 	
 	template<>
 	void operator>> <FString>(InputArchive& arch, FString& t);
-	
-	
+
+
 	//---------------------------------------------------------------------------------------------
 	template<typename T, typename Alloc> 
 	void operator<<(OutputArchive& arch, const TArray<T, Alloc>& v)
@@ -406,5 +406,21 @@ namespace mu
 		}
 	}
 
+
+	//---------------------------------------------------------------------------------------------
+	template<> 
+	inline void operator<<(OutputArchive& arch, const FName& v)
+	{
+		arch << v.ToString();
+	}
+
+	template<> 
+	inline void operator>>(InputArchive& arch, FName& v)
+	{
+		FString Temp;
+		arch >> Temp;
+		v = FName(Temp);
+	}
+
 }
 
diff --git a/Engine/Plugins/Experimental/Mutable/Source/MutableRuntime/Public/MuR/System.h b/Engine/Plugins/Experimental/Mutable/Source/MutableRuntime/Public/MuR/System.h
--- a/Engine/Plugins/Experimental/Mutable/Source/MutableRuntime/Public/MuR/System.h
+++ b/Engine/Plugins/Experimental/Mutable/Source/MutableRuntime/Public/MuR/System.h
@@ -91,12 +91,12 @@ namespace mu
 
         //! Returns the completion event and a cleanup function that must be called once event is completed.
 #ifdef MUTABLE_USE_NEW_TASKGRAPH
-		virtual TTuple<UE::Tasks::FTask, TFunction<void()>> GetImageAsync(const FExternalImageID& Id, uint8 MipmapsToSkip, TFunction<void(Ptr<Image>)>& ResultCallback) = 0;
+		virtual TTuple<UE::Tasks::FTask, TFunction<void()>> GetImageAsync(FName Id, uint8 MipmapsToSkip, TFunction<void(Ptr<Image>)>& ResultCallback) = 0;
 #else
-		virtual TTuple<FGraphEventRef, TFunction<void()>> GetImageAsync(const FExternalImageID& Id, uint8 MipmapsToSkip, TFunction<void(Ptr<Image>)>& ResultCallback) = 0;
+		virtual TTuple<FGraphEventRef, TFunction<void()>> GetImageAsync(FName Id, uint8 MipmapsToSkip, TFunction<void(Ptr<Image>)>& ResultCallback) = 0;
 #endif
 
-        virtual mu::FImageDesc GetImageDesc(const FExternalImageID& Id, uint8 MipmapsToSkip) = 0;
+        virtual mu::FImageDesc GetImageDesc(FName Id, uint8 MipmapsToSkip) = 0;
     };
 
 
diff --git a/Engine/Plugins/Experimental/Mutable/Source/MutableRuntime/Public/MuR/Types.h b/Engine/Plugins/Experimental/Mutable/Source/MutableRuntime/Public/MuR/Types.h
--- a/Engine/Plugins/Experimental/Mutable/Source/MutableRuntime/Public/MuR/Types.h
+++ b/Engine/Plugins/Experimental/Mutable/Source/MutableRuntime/Public/MuR/Types.h
@@ -16,39 +16,6 @@ namespace mu
     typedef uint64 FResourceID;
 }
 
-/** Unique identifier for images provided by code using mutable, of which mutable doesn't know much. */
-struct FExternalImageID
-{
-public:
-
-	typedef uint64 HashType;
-
-	FExternalImageID() {}
-	MUTABLERUNTIME_API FExternalImageID(const FString& InName);
-
-	inline bool operator==(const FExternalImageID& Other) const
-	{
-		return NameHash == Other.NameHash && Name == Other.Name;
-	}
-
-	void Serialise(mu::OutputArchive& arch) const;
-	void Unserialise(mu::InputArchive& arch);
-
-	inline const FString& GetName() const { return Name; }
-	inline HashType GetHash() const { return NameHash; }
-
-private:
-	FString Name;
-	HashType NameHash = 0;
-};
-
-
-/** */
-inline uint32 GetTypeHash(const FExternalImageID& Key)
-{
-	return GetTypeHash(Key.GetHash());
-}
-
 
 //
 template<typename T>
diff --git a/Engine/Plugins/Experimental/Mutable/Source/MutableTools/Private/MuT/CodeGenerator_Image.cpp b/Engine/Plugins/Experimental/Mutable/Source/MutableTools/Private/MuT/CodeGenerator_Image.cpp
--- a/Engine/Plugins/Experimental/Mutable/Source/MutableTools/Private/MuT/CodeGenerator_Image.cpp
+++ b/Engine/Plugins/Experimental/Mutable/Source/MutableTools/Private/MuT/CodeGenerator_Image.cpp
@@ -284,7 +284,7 @@ namespace mu
 			op->parameter.m_name = node.m_name;
 			op->parameter.m_uid = node.m_uid;
             op->parameter.m_type = PARAMETER_TYPE::T_IMAGE;
-        	op->parameter.m_defaultValue.Set<ParamImageType>(FExternalImageID());
+        	op->parameter.m_defaultValue.Set<ParamImageType>(FName());
 
 			// Generate the code for the ranges
 			for (int32 a = 0; a < node.m_ranges.Num(); ++a)
