[mutable] Add missing image release causing excessive use of memory in some cases.
#rnx

diff --git a/Engine/Plugins/Experimental/Mutable/Source/MutableRuntime/Private/MuR/Image.cpp b/Engine/Plugins/Experimental/Mutable/Source/MutableRuntime/Private/MuR/Image.cpp
--- a/Engine/Plugins/Experimental/Mutable/Source/MutableRuntime/Private/MuR/Image.cpp
+++ b/Engine/Plugins/Experimental/Mutable/Source/MutableRuntime/Private/MuR/Image.cpp
@@ -1289,6 +1289,8 @@ namespace mu
 				TargetData += FormatData.BytesPerBlock;
 				--BlockCount;
 			}
+
+			ReleaseImage(Converted);
 			break;
 		}
 
