[mutable] Fix performance regression with image projections.
Add additional optimization of resize+invert operation.
#rnx

diff --git a/Engine/Plugins/Experimental/Mutable/Source/MutableTools/Private/MuT/CodeGenerator_Image.cpp b/Engine/Plugins/Experimental/Mutable/Source/MutableTools/Private/MuT/CodeGenerator_Image.cpp
--- a/Engine/Plugins/Experimental/Mutable/Source/MutableTools/Private/MuT/CodeGenerator_Image.cpp
+++ b/Engine/Plugins/Experimental/Mutable/Source/MutableTools/Private/MuT/CodeGenerator_Image.cpp
@@ -195,7 +195,7 @@ namespace mu
             m_pErrorLog->GetPrivate()->Add( "Constant image not set.", ELMT_WARNING, node.m_errorContext );
         }
 
-		if (Options.LayoutToApply)
+		if (Options.ImageLayoutStrategy!= CompilerOptions::TextureLayoutStrategy::None && Options.LayoutToApply)
 		{
 			// We want to generate only a block from the image.
 
@@ -935,7 +935,6 @@ namespace mu
         op->SetChild( op->op.args.ImageColourMap.mask, mask);
 
         // Map image
-		// \TODO: We are forcing an map image size as it was a normal texture, and this is incorrect.
         Ptr<ASTOp> MapImageOp;
         if ( node.m_pMap )
         {
@@ -943,6 +942,7 @@ namespace mu
 			MapOptions.ImageLayoutStrategy = CompilerOptions::TextureLayoutStrategy::None;
 			MapOptions.LayoutToApply = nullptr;
 			MapOptions.LayoutBlockId = -1;
+			MapOptions.RectSize = {};
 
 			FImageGenerationResult BaseResult;
 			GenerateImage(MapOptions, BaseResult, node.m_pMap);
@@ -1392,6 +1392,7 @@ namespace mu
 			NewOptions.ImageLayoutStrategy = CompilerOptions::TextureLayoutStrategy::None;
 			NewOptions.LayoutToApply = nullptr;
 			NewOptions.LayoutBlockId = -1;
+			NewOptions.RectSize = {};
 
 			FImageGenerationResult ImageResult;
 			GenerateImage(NewOptions, ImageResult, node.m_pImage);
diff --git a/Engine/Plugins/Experimental/Mutable/Source/MutableTools/Private/MuT/CodeOptimiser_Semantic.cpp b/Engine/Plugins/Experimental/Mutable/Source/MutableTools/Private/MuT/CodeOptimiser_Semantic.cpp
--- a/Engine/Plugins/Experimental/Mutable/Source/MutableTools/Private/MuT/CodeOptimiser_Semantic.cpp
+++ b/Engine/Plugins/Experimental/Mutable/Source/MutableTools/Private/MuT/CodeOptimiser_Semantic.cpp
@@ -1253,7 +1253,7 @@ namespace mu
             case OP_TYPE::IM_RESIZE:
             {
                 // Keep top resize
-                auto sourceOp = dynamic_cast<const ASTOpFixed*>(sourceAt.get());
+                Ptr<const ASTOpFixed> sourceOp = dynamic_cast<const ASTOpFixed*>(sourceAt.get());
 
                 Ptr<ASTOpFixed> newOp = mu::Clone<ASTOpFixed>(this);
                 newOp->SetChild( newOp->op.args.ImageResize.source, sourceOp->children[sourceOp->op.args.ImageResize.source]);
@@ -1265,7 +1265,7 @@ namespace mu
             case OP_TYPE::IM_PLAINCOLOUR:
             {
                 // Set the size in the children and remove resize
-                auto sourceOp = mu::Clone<ASTOpFixed>(sourceAt.get());
+				Ptr<ASTOpFixed> sourceOp = mu::Clone<ASTOpFixed>(sourceAt.get());
                 sourceOp->op.args.ImagePlainColour.size[0] = op.args.ImageResize.size[0];
 				sourceOp->op.args.ImagePlainColour.size[1] = op.args.ImageResize.size[1];
 				sourceOp->op.args.ImagePlainColour.LODs = 1; // TODO
@@ -1407,6 +1407,20 @@ namespace mu
                 break;
             }
 
+			case OP_TYPE::IM_INVERT:
+			{
+				Ptr<ASTOpFixed> NewOp = mu::Clone<ASTOpFixed>(sourceAt);
+				Ptr<ASTOp> BaseAt = NewOp->children[NewOp->op.args.ImageInvert.base].child();
+
+				Ptr<ASTOpFixed> NewBase = mu::Clone<ASTOpFixed>(this);
+				NewBase->SetChild(NewBase->op.args.ImageResize.source, BaseAt);
+
+				NewOp->SetChild(NewOp->op.args.ImageInvert.base, NewBase);
+
+				at = NewOp;
+				break;
+			}
+
             case OP_TYPE::IM_LAYER:
             {
                 Ptr<ASTOpImageLayer> newOp = mu::Clone<ASTOpImageLayer>(sourceAt);
