[mutable] Crash fixes for unoptimized code, for latest code generation.
#rnx

diff --git a/Engine/Plugins/Experimental/Mutable/Source/MutableRuntime/Private/MuR/CodeRunner.cpp b/Engine/Plugins/Experimental/Mutable/Source/MutableRuntime/Private/MuR/CodeRunner.cpp
--- a/Engine/Plugins/Experimental/Mutable/Source/MutableRuntime/Private/MuR/CodeRunner.cpp
+++ b/Engine/Plugins/Experimental/Mutable/Source/MutableRuntime/Private/MuR/CodeRunner.cpp
@@ -3500,7 +3500,8 @@ namespace mu
 
                 if ( Base->GetSize()!=DestSize )
                 {
-					Ptr<Image> Result = CreateImage(DestSize[0], DestSize[1], Base->GetLODCount(), Base->GetFormat(), EInitializationType::NotInitialized);
+					int32 BaseLODCount = Base->GetLODCount();
+					Ptr<Image> Result = CreateImage(DestSize[0], DestSize[1], BaseLODCount, Base->GetFormat(), EInitializationType::NotInitialized);
 					ImOp.ImageResizeLinear( Result.get(), m_pSettings->ImageCompressionQuality, Base.get());
 					Release(Base);
 
@@ -3508,8 +3509,8 @@ namespace mu
                     // This shouldn't happen often since "ResizeLike" should be usually optimised out
                     // during model compilation. The mipmap generation below is not very precise with
                     // the number of mips that are needed and will probably generate too many
-                    bool sourceHasMips = Base->GetLODCount()>1;
-                    if (sourceHasMips)
+                    bool bSourceHasMips = BaseLODCount>1;
+                    if (bSourceHasMips)
                     {
                         int levelCount = Image::GetMipmapCount( Result->GetSizeX(), Result->GetSizeY() );
                         Ptr<Image> Mipmapped = CreateImage( Result->GetSizeX(), Result->GetSizeY(), levelCount, Result->GetFormat(), EInitializationType::NotInitialized);
