[mutable] Additional optimization for latest code generation.
#rnx

diff --git a/Engine/Plugins/Experimental/Mutable/Source/MutableTools/Private/MuT/ASTOpImageSwizzle.cpp b/Engine/Plugins/Experimental/Mutable/Source/MutableTools/Private/MuT/ASTOpImageSwizzle.cpp
--- a/Engine/Plugins/Experimental/Mutable/Source/MutableTools/Private/MuT/ASTOpImageSwizzle.cpp
+++ b/Engine/Plugins/Experimental/Mutable/Source/MutableTools/Private/MuT/ASTOpImageSwizzle.cpp
@@ -932,7 +932,7 @@ namespace mu
 			at = NewSaturate;
 		}
 
-
+		// Disable this for now.
 		// Swizzle with the same op as identity in RGB, a Layer op in the A that has one of the operands matching 
 		// the one in the swizzle RGB, but using its A. 
 		// The Layer operation can be flagged as alpha only and moved up the swizzle, then the swizzle is identity 
@@ -951,40 +951,70 @@ namespace mu
 		//	- A
 		//	- B
 		// In addition, if the blend operation done by LAYER is commutative, see if X is 3 from I instead.
-		if (!at
-			&&
-			(Sources[0] == Sources[1]) && (Sources[0] == Sources[2])
-			&&
-			SourceChannels[0] == 0 && SourceChannels[1] == 1 && SourceChannels[2] == 2
-			&&
-			Sources[3] && Sources[3]->GetOpType() == OP_TYPE::IM_LAYER
-			)
-		{
-			const ASTOpImageLayer* OldLayer = dynamic_cast<const ASTOpImageLayer*>(Sources[3].child().get());
-
-			// For now just check the case that we are observing in the working data: 
-			// A is in the blended of a multiply, and we take its alpha channel
-			// \TODO: Implement the other cases when we find instances of them.
-			if ( OldLayer->Flags==OP::ImageLayerArgs::FLAGS::F_BLENDED_RGB_FROM_ALPHA
-				&&
-				OldLayer->blend==Sources[0]
-				&& 
-				OldLayer->blendType==EBlendType::BT_MULTIPLY 
-				&&
-				OldLayer->blendTypeAlpha == EBlendType::BT_NONE 
-				&&
-				SourceChannels[3]==0)
-			{
-				Ptr<ASTOpImageLayer> NewLayer = mu::Clone<ASTOpImageLayer>(OldLayer);
-				NewLayer->blend = NewLayer->base.child();
-				NewLayer->base = Sources[0].child();
-				NewLayer->blendTypeAlpha = NewLayer->blendType;
-				NewLayer->blendType = EBlendType::BT_NONE;
-				NewLayer->Flags = OP::ImageLayerArgs::FLAGS::F_BASE_RGB_FROM_ALPHA;
-
-				at = NewLayer;
-			}
-		}
+		//if (!at
+		//	&&
+		//	(Sources[0] == Sources[1]) && (Sources[0] == Sources[2])
+		//	&&
+		//	SourceChannels[0] == 0 && SourceChannels[1] == 1 && SourceChannels[2] == 2
+		//	&&
+		//	Sources[3] && Sources[3]->GetOpType() == OP_TYPE::IM_LAYER
+		//	)
+		//{
+		//	const ASTOpImageLayer* OldLayer = dynamic_cast<const ASTOpImageLayer*>(Sources[3].child().get());
+
+		//	Ptr<ASTOp> SwizzleRGBOp = Sources[0].child();
+		//	Ptr<ASTOp> OldLayerBlendOp = OldLayer->blend.child();
+		//	//{
+		//	//	auto DiscardNeutralOps = [](Ptr<ASTOp> Op)
+		//	//	{
+		//	//		bool bUpdated = true;
+		//	//		while (bUpdated)
+		//	//		{
+		//	//			bUpdated = false;
+		//	//			switch (Op->GetOpType())
+		//	//			{
+		//	//			case OP_TYPE::IM_PIXELFORMAT:
+		//	//			{
+		//	//				const ASTOpImagePixelFormat* Typed = dynamic_cast<const ASTOpImagePixelFormat*>(Op.get());
+		//	//				Op = Typed->Source.child();
+		//	//				bUpdated = true;
+		//	//				break;
+		//	//			}
+
+		//	//			default: break;
+		//	//			}
+		//	//		}
+		//	//		return Op;
+		//	//	};
+
+		//	//	SwizzleRGBOp = DiscardNeutralOps(SwizzleRGBOp);
+		//	//	OldLayerBlendOp = DiscardNeutralOps(OldLayerBlendOp);
+		//	//}
+		//	bool bOldLayerBlendIsCompatibleWithSwizzleRGBs = OldLayerBlendOp == SwizzleRGBOp;
+
+		//	// For now just check the case that we are observing in the working data: 
+		//	// A is in the blended of a multiply, and we take its alpha channel
+		//	// \TODO: Implement the other cases when we find instances of them.
+		//	if ( OldLayer->Flags==OP::ImageLayerArgs::FLAGS::F_BLENDED_RGB_FROM_ALPHA
+		//		&&
+		//		bOldLayerBlendIsCompatibleWithSwizzleRGBs
+		//		&& 
+		//		OldLayer->blendType==EBlendType::BT_MULTIPLY 
+		//		&&
+		//		OldLayer->blendTypeAlpha == EBlendType::BT_NONE 
+		//		&&
+		//		SourceChannels[3]==0)
+		//	{
+		//		Ptr<ASTOpImageLayer> NewLayer = mu::Clone<ASTOpImageLayer>(OldLayer);
+		//		NewLayer->blend = NewLayer->base.child();
+		//		NewLayer->base = Sources[0].child();
+		//		NewLayer->blendTypeAlpha = NewLayer->blendType;
+		//		NewLayer->blendType = EBlendType::BT_NONE;
+		//		NewLayer->Flags = OP::ImageLayerArgs::FLAGS::F_BASE_RGB_FROM_ALPHA;
+
+		//		at = NewLayer;
+		//	}
+		//}
 
 
 		// If we have an alpha channel that has as children something that expands a single channel texture
@@ -1008,7 +1038,12 @@ namespace mu
 				const ASTOpImagePixelFormat* OldFormat = dynamic_cast<const ASTOpImagePixelFormat*>(OldLayer->blend.child().get());
 				if (OldFormat->Source->GetOpType() == OP_TYPE::IM_SWIZZLE
 					&&
-					OldFormat->Format == EImageFormat::IF_RGB_UBYTE)
+					(
+						OldFormat->Format == EImageFormat::IF_RGB_UBYTE
+						||
+						OldFormat->Format == EImageFormat::IF_RGBA_UBYTE
+						) 
+					)
 				{
 					const ASTOpImageSwizzle* OldChildSwizzle = dynamic_cast<const ASTOpImageSwizzle*>(OldFormat->Source.child().get());
 					if (OldChildSwizzle->Format == EImageFormat::IF_L_UBYTE)
diff --git a/Engine/Plugins/Experimental/Mutable/Source/MutableTools/Private/MuT/CodeOptimiser_Semantic.cpp b/Engine/Plugins/Experimental/Mutable/Source/MutableTools/Private/MuT/CodeOptimiser_Semantic.cpp
--- a/Engine/Plugins/Experimental/Mutable/Source/MutableTools/Private/MuT/CodeOptimiser_Semantic.cpp
+++ b/Engine/Plugins/Experimental/Mutable/Source/MutableTools/Private/MuT/CodeOptimiser_Semantic.cpp
@@ -974,7 +974,6 @@ namespace mu
 
 			case OP_TYPE::IM_PIXELFORMAT:
 			{
-				// We move the mask creation down format.
 				// \todo: only if shrinking?
 				Ptr<ASTOpImagePixelFormat> newOp = mu::Clone<ASTOpImagePixelFormat>(at);
 				newOp->Source = Visit(newOp->Source.child(), currentSinkingOp);
@@ -1421,6 +1420,21 @@ namespace mu
 				break;
 			}
 
+			case OP_TYPE::IM_PIXELFORMAT:
+			{
+				// \todo: only if shrinking?
+				Ptr<ASTOpImagePixelFormat> NewOp = mu::Clone<ASTOpImagePixelFormat>(sourceAt);
+				Ptr<ASTOp> BaseAt = NewOp->Source.child();
+
+				Ptr<ASTOpFixed> NewBase = mu::Clone<ASTOpFixed>(this);
+				NewBase->SetChild(NewBase->op.args.ImageResize.source, BaseAt);
+
+				NewOp->Source = NewBase;
+
+				at = NewOp;
+				break;
+			}
+
             case OP_TYPE::IM_LAYER:
             {
                 Ptr<ASTOpImageLayer> newOp = mu::Clone<ASTOpImageLayer>(sourceAt);
