[mutable] Optimize code for colour swizzle x swizzle and swizzle x from_scalars.
#rnx

diff --git a/Engine/Plugins/Experimental/Mutable/Source/MutableRuntime/Private/MuR/CodeRunner.cpp b/Engine/Plugins/Experimental/Mutable/Source/MutableRuntime/Private/MuR/CodeRunner.cpp
--- a/Engine/Plugins/Experimental/Mutable/Source/MutableRuntime/Private/MuR/CodeRunner.cpp
+++ b/Engine/Plugins/Experimental/Mutable/Source/MutableRuntime/Private/MuR/CodeRunner.cpp
@@ -5464,37 +5464,25 @@ namespace mu
             {
             case 0:
                     AddOp( FScheduledOp( item.At, item, 1),
-                           FScheduledOp( args.x, item),
-                           FScheduledOp( args.y, item),
-                           FScheduledOp( args.z, item),
-                           FScheduledOp( args.w, item));
+                           FScheduledOp( args.v[0], item),
+                           FScheduledOp( args.v[1], item),
+                           FScheduledOp( args.v[2], item),
+                           FScheduledOp( args.v[3], item));
                 break;
 
             case 1:
             {
-				FVector4f result = FVector4f(1, 1, 1, 1);
-
-                if (args.x)
-                {
-                    result[0] = LoadScalar( FCacheAddress(args.x,item) );
-                }
-
-                if (args.y)
-                {
-                    result[1] = LoadScalar( FCacheAddress(args.y,item) );
-                }
-
-                if (args.z)
-                {
-                    result[2] = LoadScalar( FCacheAddress(args.z,item) );
-                }
+				FVector4f Result = FVector4f(1, 1, 1, 1);
 
-                if (args.w)
-                {
-                    result[3] = LoadScalar( FCacheAddress(args.w,item) );
-                }
+				for (int32 t = 0; t < MUTABLE_OP_MAX_SWIZZLE_CHANNELS; ++t)
+				{
+					if (args.v[t])
+					{
+						Result[t] = LoadScalar(FCacheAddress(args.v[t], item));
+					}
+				}
 
-                StoreColor( item, result );
+                StoreColor( item, Result );
                 break;
             }
 
diff --git a/Engine/Plugins/Experimental/Mutable/Source/MutableRuntime/Private/MuR/Operations.cpp b/Engine/Plugins/Experimental/Mutable/Source/MutableRuntime/Private/MuR/Operations.cpp
--- a/Engine/Plugins/Experimental/Mutable/Source/MutableRuntime/Private/MuR/Operations.cpp
+++ b/Engine/Plugins/Experimental/Mutable/Source/MutableRuntime/Private/MuR/Operations.cpp
@@ -222,7 +222,7 @@ namespace mu
             break;
 
         case OP_TYPE::CO_SWIZZLE:
-            for (int t=0;t<MUTABLE_OP_MAX_SWIZZLE_CHANNELS;++t)
+            for (int32 t=0;t<MUTABLE_OP_MAX_SWIZZLE_CHANNELS;++t)
             {
                 f(&op.args.ColourSwizzle.sources[t] );
             }
@@ -237,10 +237,10 @@ namespace mu
             break;
 
         case OP_TYPE::CO_FROMSCALARS:
-            f(&op.args.ColourFromScalars.x);
-            f(&op.args.ColourFromScalars.y);
-            f(&op.args.ColourFromScalars.z);
-            f(&op.args.ColourFromScalars.w);
+			for (int32 t = 0; t < MUTABLE_OP_MAX_SWIZZLE_CHANNELS; ++t)
+			{
+				f(&op.args.ColourFromScalars.v[t]);
+			}
             break;
 
         case OP_TYPE::CO_ARITHMETIC:
@@ -558,10 +558,10 @@ namespace mu
         case OP_TYPE::CO_FROMSCALARS:
         {
 			OP::ColourFromScalarsArgs args = program.GetOpArgs<OP::ColourFromScalarsArgs>(at);
-            f(args.x);
-            f(args.y);
-            f(args.z);
-            f(args.w);
+			for (int32 t = 0; t < MUTABLE_OP_MAX_SWIZZLE_CHANNELS; ++t)
+			{
+				f(args.v[t]);
+			}
             break;
         }
 
diff --git a/Engine/Plugins/Experimental/Mutable/Source/MutableRuntime/Private/MuR/Operations.h b/Engine/Plugins/Experimental/Mutable/Source/MutableRuntime/Private/MuR/Operations.h
--- a/Engine/Plugins/Experimental/Mutable/Source/MutableRuntime/Private/MuR/Operations.h
+++ b/Engine/Plugins/Experimental/Mutable/Source/MutableRuntime/Private/MuR/Operations.h
@@ -499,7 +499,7 @@ namespace mu
 
 		struct ColourFromScalarsArgs
 		{
-			ADDRESS x, y, z, w;
+			ADDRESS v[MUTABLE_OP_MAX_SWIZZLE_CHANNELS];
 		};
 
         struct ArithmeticArgs
diff --git a/Engine/Plugins/Experimental/Mutable/Source/MutableTools/Private/MuT/CodeGenerator_Color.cpp b/Engine/Plugins/Experimental/Mutable/Source/MutableTools/Private/MuT/CodeGenerator_Color.cpp
--- a/Engine/Plugins/Experimental/Mutable/Source/MutableTools/Private/MuT/CodeGenerator_Color.cpp
+++ b/Engine/Plugins/Experimental/Mutable/Source/MutableTools/Private/MuT/CodeGenerator_Color.cpp
@@ -371,49 +371,49 @@ namespace mu
 		// X
 		if (Node* pX = node.m_pX.get())
 		{
-			op->SetChild(op->op.args.ColourFromScalars.x, Generate(pX));
+			op->SetChild(op->op.args.ColourFromScalars.v[0], Generate(pX));
 		}
 		else
 		{
 			NodeScalarConstantPtr pNode = new NodeScalarConstant();
 			pNode->SetValue(1.0f);
-			op->SetChild(op->op.args.ColourFromScalars.x, Generate(pNode));
+			op->SetChild(op->op.args.ColourFromScalars.v[0], Generate(pNode));
 		}
 
 		// Y
 		if (Node* pY = node.m_pY.get())
 		{
-			op->SetChild(op->op.args.ColourFromScalars.y, Generate(pY));
+			op->SetChild(op->op.args.ColourFromScalars.v[1], Generate(pY));
 		}
 		else
 		{
 			NodeScalarConstantPtr pNode = new NodeScalarConstant();
 			pNode->SetValue(1.0f);
-			op->SetChild(op->op.args.ColourFromScalars.y, Generate(pNode));
+			op->SetChild(op->op.args.ColourFromScalars.v[1], Generate(pNode));
 		}
 
 		// Z
 		if (Node* pZ = node.m_pZ.get())
 		{
-			op->SetChild(op->op.args.ColourFromScalars.z, Generate(pZ));
+			op->SetChild(op->op.args.ColourFromScalars.v[2], Generate(pZ));
 		}
 		else
 		{
 			NodeScalarConstantPtr pNode = new NodeScalarConstant();
 			pNode->SetValue(1.0f);
-			op->SetChild(op->op.args.ColourFromScalars.z, Generate(pNode));
+			op->SetChild(op->op.args.ColourFromScalars.v[2], Generate(pNode));
 		}
 
 		// W
 		if (Node* pW = node.m_pW.get())
 		{
-			op->SetChild(op->op.args.ColourFromScalars.w, Generate(pW));
+			op->SetChild(op->op.args.ColourFromScalars.v[3], Generate(pW));
 		}
 		else
 		{
 			NodeScalarConstantPtr pNode = new NodeScalarConstant();
 			pNode->SetValue(1.0f);
-			op->SetChild(op->op.args.ColourFromScalars.w, Generate(pNode));
+			op->SetChild(op->op.args.ColourFromScalars.v[3], Generate(pNode));
 		}
 
 		result.op = op;
diff --git a/Engine/Plugins/Experimental/Mutable/Source/MutableTools/Private/MuT/CodeOptimiser_Semantic.cpp b/Engine/Plugins/Experimental/Mutable/Source/MutableTools/Private/MuT/CodeOptimiser_Semantic.cpp
--- a/Engine/Plugins/Experimental/Mutable/Source/MutableTools/Private/MuT/CodeOptimiser_Semantic.cpp
+++ b/Engine/Plugins/Experimental/Mutable/Source/MutableTools/Private/MuT/CodeOptimiser_Semantic.cpp
@@ -32,6 +32,7 @@
 #include <memory>
 #include <utility>
 
+
 namespace mu
 {
 
@@ -251,6 +252,106 @@ namespace mu
             break;
         }
 
+		case OP_TYPE::CO_SWIZZLE:
+		{
+			// Optimizations that can be applied per-channel
+			{
+				Ptr<ASTOpFixed> NewAt;
+
+				for (int32 ChannelIndex = 0; ChannelIndex < MUTABLE_OP_MAX_SWIZZLE_CHANNELS; ++ChannelIndex)
+				{
+					Ptr<ASTOp> Candidate = children[op.args.ColourSwizzle.sources[ChannelIndex]].child();
+					if (!Candidate)
+					{
+						continue;
+					}
+
+					switch (Candidate->GetOpType())
+					{
+
+						// Swizzle + swizzle = swizzle
+					case OP_TYPE::CO_SWIZZLE:
+					{
+						if (!NewAt)
+						{
+							NewAt = mu::Clone<ASTOpFixed>(this);
+						}
+						const ASTOpFixed* TypedCandidate = dynamic_cast<const ASTOpFixed*>(Candidate.get());
+						int32 CandidateChannel = op.args.ColourSwizzle.sourceChannels[ChannelIndex];
+
+						NewAt->SetChild(NewAt->op.args.ColourSwizzle.sources[ChannelIndex], TypedCandidate->children[TypedCandidate->op.args.ColourSwizzle.sources[CandidateChannel]]);
+						NewAt->op.args.ColourSwizzle.sourceChannels[ChannelIndex] = TypedCandidate->op.args.ColourSwizzle.sourceChannels[CandidateChannel];
+						break;
+					}
+
+					default:
+						break;
+
+					}
+				}
+
+				at = NewAt;
+			}
+
+			// Not optimized yet?
+			if (!at)
+			{
+				// Optimizations that depend on all channels.
+				bool bAllChannelsSameType = true;
+				OP_TYPE AllChannelsType = OP_TYPE::NONE;
+				for (int32 ChannelIndex = 0; ChannelIndex < MUTABLE_OP_MAX_SWIZZLE_CHANNELS; ++ChannelIndex)
+				{
+					Ptr<ASTOp> Candidate = children[op.args.ColourSwizzle.sources[ChannelIndex]].child();
+					if (!Candidate)
+					{
+						continue;
+					}
+
+					if (ChannelIndex == 0)
+					{
+						AllChannelsType = Candidate->GetOpType();
+					}
+					else if (Candidate->GetOpType()!=AllChannelsType)
+					{
+						bAllChannelsSameType = false;
+						break;
+					}
+				}
+
+				if (bAllChannelsSameType)
+				{
+					switch (AllChannelsType)
+					{
+					case OP_TYPE::CO_FROMSCALARS:
+					{
+						// We can remove the swizzle and replace it with a new FromScalars actually swizzling the inputs.
+
+						Ptr<ASTOpFixed> NewAt = new ASTOpFixed();
+						NewAt->op.type = OP_TYPE::CO_FROMSCALARS;
+
+						for (int32 ChannelIndex = 0; ChannelIndex < MUTABLE_OP_MAX_SWIZZLE_CHANNELS; ++ChannelIndex)
+						{
+							const ASTOpFixed* SelectedSource = dynamic_cast<const ASTOpFixed*>(children[op.args.ColourSwizzle.sources[ChannelIndex]].child().get());
+							if (SelectedSource)
+							{
+								int32 SelectedChannel = op.args.ColourSwizzle.sourceChannels[ChannelIndex];
+								Ptr<ASTOp> SelectedFloatInput = SelectedSource->children[SelectedSource->op.args.ColourFromScalars.v[SelectedChannel]].child();
+								NewAt->SetChild(NewAt->op.args.ColourFromScalars.v[ChannelIndex], SelectedFloatInput);
+							}
+						}
+
+						at = NewAt;
+						break;
+					}
+
+					default: break;
+					}
+				}
+			}
+			
+			break;
+		}
+
 		//-------------------------------------------------------------------------------------
 		case OP_TYPE::IM_RESIZEREL:
 		{
@@ -1640,4 +1741,3 @@ namespace mu
     }
 
 }
-
