[Mutable] Add Mask to Reshape operations.

- Added vertex color channel usage selection to reshape node.
- Removed rigid parts node reshape checkbox. Now rigid parts is enabled by assigning at least one vertex color channel usage to Rigid Cluster Id.

#rnx

diff --git a/Engine/Plugins/Experimental/Mutable/Source/CustomizableObject/Classes/MuCO/CustomizableObject.h b/Engine/Plugins/Experimental/Mutable/Source/CustomizableObject/Classes/MuCO/CustomizableObject.h
--- a/Engine/Plugins/Experimental/Mutable/Source/CustomizableObject/Classes/MuCO/CustomizableObject.h
+++ b/Engine/Plugins/Experimental/Mutable/Source/CustomizableObject/Classes/MuCO/CustomizableObject.h
@@ -1439,7 +1439,7 @@ private:
 	// This is a manual version number for the binary blobs in this asset.
 	// Increasing it invalidates all the previously compiled models.
 	// Warning: If while merging code both versions have changed, take the highest+1.
-	static const int32 CurrentSupportedVersion = 396;
+	static const int32 CurrentSupportedVersion = 397;
 
 public:
 
diff --git a/Engine/Plugins/Experimental/Mutable/Source/CustomizableObject/Public/MuCO/CustomizableObjectCustomVersion.h b/Engine/Plugins/Experimental/Mutable/Source/CustomizableObject/Public/MuCO/CustomizableObjectCustomVersion.h
--- a/Engine/Plugins/Experimental/Mutable/Source/CustomizableObject/Public/MuCO/CustomizableObjectCustomVersion.h
+++ b/Engine/Plugins/Experimental/Mutable/Source/CustomizableObject/Public/MuCO/CustomizableObjectCustomVersion.h
@@ -90,6 +90,8 @@ struct CUSTOMIZABLEOBJECT_API FCustomizableObjectCustomVersion
 
 		FixPinsNamesImageToTexture2,
 
+		MeshReshapeVertexColorUsageSelection,
+
 		// -----<new versions can be added above this line>-------------------------------------------------
 		VersionPlusOne,
 		LatestVersion = VersionPlusOne - 1
diff --git a/Engine/Plugins/Experimental/Mutable/Source/CustomizableObjectEditor/Classes/MuCOE/Nodes/CustomizableObjectNodeMeshReshape.h b/Engine/Plugins/Experimental/Mutable/Source/CustomizableObjectEditor/Classes/MuCOE/Nodes/CustomizableObjectNodeMeshReshape.h
--- a/Engine/Plugins/Experimental/Mutable/Source/CustomizableObjectEditor/Classes/MuCOE/Nodes/CustomizableObjectNodeMeshReshape.h
+++ b/Engine/Plugins/Experimental/Mutable/Source/CustomizableObjectEditor/Classes/MuCOE/Nodes/CustomizableObjectNodeMeshReshape.h
@@ -12,6 +12,33 @@ class UObject;
 struct FMeshReshapeBoneReference;
 enum class EBoneDeformSelectionMethod : uint8;
 
+
+UENUM()
+enum class EMeshReshapeVertexColorChannelUsage
+{
+	None = 0,
+	RigidClusterId = 1,
+	MaskWeight = 2
+};
+
+USTRUCT()
+struct FMeshReshapeColorUsage
+{
+	GENERATED_BODY()
+
+	UPROPERTY(EditAnywhere, Category = CustomizableObject)
+	EMeshReshapeVertexColorChannelUsage R = EMeshReshapeVertexColorChannelUsage::None;
+
+	UPROPERTY(EditAnywhere, Category = CustomizableObject)
+	EMeshReshapeVertexColorChannelUsage G = EMeshReshapeVertexColorChannelUsage::None;
+
+	UPROPERTY(EditAnywhere, Category = CustomizableObject)
+	EMeshReshapeVertexColorChannelUsage B = EMeshReshapeVertexColorChannelUsage::None;
+
+	UPROPERTY(EditAnywhere, Category = CustomizableObject)
+	EMeshReshapeVertexColorChannelUsage A = EMeshReshapeVertexColorChannelUsage::None;
+};
+
 UCLASS()
 class CUSTOMIZABLEOBJECTEDITOR_API UCustomizableObjectNodeMeshReshape : public UCustomizableObjectNode
 {
@@ -58,9 +85,11 @@ public:
     UPROPERTY(EditAnywhere, Category = CustomizableObject)
     bool bReshapePhysicsVolumes = false;
 	
-	/** Enable rigid parts base mesh. */
+	UPROPERTY()
+	bool bEnableRigidParts_DEPRECATED = false;
+
 	UPROPERTY(EditAnywhere, Category = CustomizableObject)
-	bool bEnableRigidParts = false;
+	FMeshReshapeColorUsage VertexColorUsage;
 
 	/** Bone Reshape Selection Method */
 	UPROPERTY(EditAnywhere, Category = ReshapeBones, meta = (EditCondition = "bReshapeSkeleton"))
diff --git a/Engine/Plugins/Experimental/Mutable/Source/CustomizableObjectEditor/Private/MuCOE/GenerateMutableSource/GenerateMutableSourceMesh.cpp b/Engine/Plugins/Experimental/Mutable/Source/CustomizableObjectEditor/Private/MuCOE/GenerateMutableSource/GenerateMutableSourceMesh.cpp
--- a/Engine/Plugins/Experimental/Mutable/Source/CustomizableObjectEditor/Private/MuCOE/GenerateMutableSource/GenerateMutableSourceMesh.cpp
+++ b/Engine/Plugins/Experimental/Mutable/Source/CustomizableObjectEditor/Private/MuCOE/GenerateMutableSource/GenerateMutableSourceMesh.cpp
@@ -828,8 +828,8 @@ mu::MeshPtr ConvertSkeletalMeshToMutable(const USkeletalMesh* InSkeletalMesh, co
 	// Vertices
 	TArray<FSoftSkinVertex> Vertices;
 	ImportedModel->LODModels[LODIndex].GetVertices(Vertices);
-	int VertexStart = ImportedModel->LODModels[LODIndex].Sections[SectionIndex].GetVertexBufferIndex();
-	int VertexCount = ImportedModel->LODModels[LODIndex].Sections[SectionIndex].GetNumVertices();
+	int32 VertexStart = ImportedModel->LODModels[LODIndex].Sections[SectionIndex].GetVertexBufferIndex();
+	int32 VertexCount = ImportedModel->LODModels[LODIndex].Sections[SectionIndex].GetNumVertices();
 
 	MutableMesh->GetVertexBuffers().SetElementCount(VertexCount);
 
@@ -1773,7 +1773,6 @@ mu::MeshPtr ConvertSkeletalMeshToMutable(const USkeletalMesh* InSkeletalMesh, co
 	return MutableMesh;
 }
 
-
 mu::MeshPtr ConvertStaticMeshToMutable(const UStaticMesh* StaticMesh, int32 LODIndex, int32 SectionIndex, FMutableGraphGenerationContext& GenerationContext, const UCustomizableObjectNode* CurrentNode)
 {
 	if (!StaticMesh->GetRenderData() ||
@@ -3369,7 +3368,7 @@ mu::NodeMeshPtr GenerateMutableSourceMesh(const UEdGraphPin* Pin,
 		}
 		else
 		{
-			GenerationContext.Compiler->CompilerLog(LOCTEXT("MeshGeometryMissingDef", "Mesh variation node requires a default value."), Node);
+			GenerationContext.Compiler->CompilerLog(LOCTEXT("MeshReshapeMissingDef", "Mesh reshape node requires a default value."), Node);
 		}
 	
 		{
@@ -3377,7 +3376,58 @@ mu::NodeMeshPtr GenerateMutableSourceMesh(const UEdGraphPin* Pin,
 			MeshNode->SetReshapeVertices(TypedNodeReshape->bReshapeVertices);
 			MeshNode->SetReshapeSkeleton(TypedNodeReshape->bReshapeSkeleton);
 			MeshNode->SetReshapePhysicsVolumes(TypedNodeReshape->bReshapePhysicsVolumes);
-			MeshNode->SetEnableRigidParts(TypedNodeReshape->bEnableRigidParts);
+
+			EMeshReshapeVertexColorChannelUsage ChannelUsages[4] =
+			{
+				TypedNodeReshape->VertexColorUsage.R,
+				TypedNodeReshape->VertexColorUsage.G,
+				TypedNodeReshape->VertexColorUsage.B,
+				TypedNodeReshape->VertexColorUsage.A
+			};
+
+			{
+				int32 MaskWeightChannelNum = 0;
+				for (int32 I = 0; I < 4; ++I)
+				{
+					if (ChannelUsages[I] == EMeshReshapeVertexColorChannelUsage::MaskWeight)
+					{
+						++MaskWeightChannelNum;
+					}
+				}
+
+				if (MaskWeightChannelNum > 1)
+				{
+					for (int32 I = 0; I < 4; ++I)
+					{
+						if (ChannelUsages[I] == EMeshReshapeVertexColorChannelUsage::MaskWeight)
+						{
+							ChannelUsages[I] = EMeshReshapeVertexColorChannelUsage::None;
+						}
+					}
+
+					GenerationContext.Compiler->CompilerLog(
+						LOCTEXT("MeshReshapeColorUsageMask", 
+								"Only one color channel with mask weight usage is allowed, multiple found. Reshape masking disabled."),
+						Node);
+				}
+			}
+
+			auto ConvertColorUsage = [](EMeshReshapeVertexColorChannelUsage Usage) -> mu::EVertexColorUsage
+			{
+				switch (Usage)
+				{
+				case EMeshReshapeVertexColorChannelUsage::None:			  return mu::EVertexColorUsage::None;
+				case EMeshReshapeVertexColorChannelUsage::RigidClusterId: return mu::EVertexColorUsage::ReshapeClusterId;
+				case EMeshReshapeVertexColorChannelUsage::MaskWeight:     return mu::EVertexColorUsage::ReshapeMaskWeight;
+				default: check(false); return mu::EVertexColorUsage::None;
+				};
+			};
+
+			MeshNode->SetColorUsages(
+				ConvertColorUsage(ChannelUsages[0]),
+				ConvertColorUsage(ChannelUsages[1]),
+				ConvertColorUsage(ChannelUsages[2]),
+				ConvertColorUsage(ChannelUsages[3]));
 				
 			const UEdGraphPin* ConnectedPin = FollowInputPin(*TypedNodeReshape->BaseMeshPin());
 			const UEdGraphPin* SourceMeshPin = ConnectedPin ? FindMeshBaseSource(*ConnectedPin, false) : nullptr;
diff --git a/Engine/Plugins/Experimental/Mutable/Source/CustomizableObjectEditor/Private/MuCOE/Nodes/CustomizableObjectNodeMeshReshape.cpp b/Engine/Plugins/Experimental/Mutable/Source/CustomizableObjectEditor/Private/MuCOE/Nodes/CustomizableObjectNodeMeshReshape.cpp
--- a/Engine/Plugins/Experimental/Mutable/Source/CustomizableObjectEditor/Private/MuCOE/Nodes/CustomizableObjectNodeMeshReshape.cpp
+++ b/Engine/Plugins/Experimental/Mutable/Source/CustomizableObjectEditor/Private/MuCOE/Nodes/CustomizableObjectNodeMeshReshape.cpp
@@ -73,6 +73,17 @@ void UCustomizableObjectNodeMeshReshape::Serialize(FArchive& Ar)
 			BonesToDeform.Empty();
 		}
 	}
+
+	if (Ar.CustomVer(FCustomizableObjectCustomVersion::GUID) < FCustomizableObjectCustomVersion::MeshReshapeVertexColorUsageSelection)
+	{
+		if (bEnableRigidParts_DEPRECATED)
+		{
+			VertexColorUsage.R = EMeshReshapeVertexColorChannelUsage::RigidClusterId;
+			VertexColorUsage.G = EMeshReshapeVertexColorChannelUsage::RigidClusterId;
+			VertexColorUsage.B = EMeshReshapeVertexColorChannelUsage::RigidClusterId;
+			VertexColorUsage.A = EMeshReshapeVertexColorChannelUsage::RigidClusterId;
+		}
+	}
 }
 
 #undef LOCTEXT_NAMESPACE
diff --git a/Engine/Plugins/Experimental/Mutable/Source/MutableRuntime/Private/MuR/CodeRunner.cpp b/Engine/Plugins/Experimental/Mutable/Source/MutableRuntime/Private/MuR/CodeRunner.cpp
--- a/Engine/Plugins/Experimental/Mutable/Source/MutableRuntime/Private/MuR/CodeRunner.cpp
+++ b/Engine/Plugins/Experimental/Mutable/Source/MutableRuntime/Private/MuR/CodeRunner.cpp
@@ -2412,10 +2412,14 @@ namespace mu
 
 					const EMeshBindShapeFlags BindFlags = static_cast<EMeshBindShapeFlags>(Args.flags);
 
+					FMeshBindColorChannelUsages ColorChannelUsages;
+					FMemory::Memcpy(&ColorChannelUsages, &Args.ColorUsage, sizeof(ColorChannelUsages));
+					static_assert(sizeof(ColorChannelUsages) == sizeof(Args.ColorUsage));
+
 					Ptr<Mesh> BindMeshResult = CreateMesh();
 
 					bool bOutSuccess = false;
-					MeshBindShapeReshape(BindMeshResult.get(), BaseMesh.get(), Shape.get(), BonesToDeform, PhysicsToDeform, BindFlags, bOutSuccess);
+					MeshBindShapeReshape(BindMeshResult.get(), BaseMesh.get(), Shape.get(), BonesToDeform, PhysicsToDeform, BindFlags, ColorChannelUsages, bOutSuccess);
 				
 					Release(Shape);
 					// not success indicates nothing has bond so the base mesh can be reused.
diff --git a/Engine/Plugins/Experimental/Mutable/Source/MutableRuntime/Private/MuR/Mesh.cpp b/Engine/Plugins/Experimental/Mutable/Source/MutableRuntime/Private/MuR/Mesh.cpp
--- a/Engine/Plugins/Experimental/Mutable/Source/MutableRuntime/Private/MuR/Mesh.cpp
+++ b/Engine/Plugins/Experimental/Mutable/Source/MutableRuntime/Private/MuR/Mesh.cpp
@@ -1858,5 +1858,6 @@ void Mesh::Log( FString& out, int32 BufferElementLimit)
 MUTABLE_IMPLEMENT_ENUM_SERIALISABLE(EBoneUsageFlags)
 MUTABLE_IMPLEMENT_ENUM_SERIALISABLE(EMeshBufferType)
 MUTABLE_IMPLEMENT_ENUM_SERIALISABLE(EShapeBindingMethod)
+MUTABLE_IMPLEMENT_ENUM_SERIALISABLE(EVertexColorUsage)
 	
 }
diff --git a/Engine/Plugins/Experimental/Mutable/Source/MutableRuntime/Private/MuR/OpMeshBind.h b/Engine/Plugins/Experimental/Mutable/Source/MutableRuntime/Private/MuR/OpMeshBind.h
--- a/Engine/Plugins/Experimental/Mutable/Source/MutableRuntime/Private/MuR/OpMeshBind.h
+++ b/Engine/Plugins/Experimental/Mutable/Source/MutableRuntime/Private/MuR/OpMeshBind.h
@@ -31,7 +31,6 @@
 
 namespace mu
 {
-
 	template<class PositionCompType, class NormalCompType>
 	struct FShapeMeshDescriptor
 	{
@@ -41,7 +40,7 @@ namespace mu
 	};
 
 	using FShapeMeshDescriptorBind = FShapeMeshDescriptor<double, float>;
-		
+
 	struct FShapeMeshAdapter
 	{
 		const FShapeMeshDescriptorBind& Mesh;
@@ -224,13 +223,40 @@ namespace mu
 		}
 	};
 
+	struct FMeshBindColorChannelUsageMasks
+	{
+		uint32 MaskWeight = 0;
+		uint32 ClusterId  = 0;
+	};
+
 	//---------------------------------------------------------------------------------------------
     //! Generate the mesh-shape binding data for Reshape operations
     //---------------------------------------------------------------------------------------------
 
+	inline float GetVertexMaskWeight(const UntypedMeshBufferIteratorConst& ColorIter, const FMeshBindColorChannelUsageMasks& ChannelUsages)
+	{
+		check(ColorIter.GetFormat() == MBF_UINT8 || ColorIter.GetFormat() == MBF_NUINT8);
+		check(ColorIter.GetComponents() == 4);
+		check(ColorIter.ptr());
+
+		const uint32 Value = *reinterpret_cast<const uint32*>(ColorIter.ptr());
+		return static_cast<float>((Value >> FMath::CountTrailingZeros(ChannelUsages.MaskWeight) & 0xFF)) / 255.0f;
+	}
+
+	inline uint32 GetVertexClusterId(const UntypedMeshBufferIteratorConst& ColorIter, const FMeshBindColorChannelUsageMasks& ChannelUsages)
+	{	
+		check(ColorIter.GetFormat() == MBF_UINT8 || ColorIter.GetFormat() == MBF_NUINT8);
+		check(ColorIter.GetComponents() == 4);
+		check(ColorIter.ptr());
+
+		const uint32 MaskedValue = *reinterpret_cast<const uint32*>(ColorIter.ptr()) & ChannelUsages.ClusterId;
+		// Set all unused bits to 1 so we have a consistent value, in that case white, for the non clustered vertices.
+		return MaskedValue | ~ChannelUsages.ClusterId;
+	}
+
 	inline void BindReshapePoint(
 			FShapeMeshTree& ShapeMeshTree,
-			const FVector3f& Point, FReshapeVertexBindingData& OutBindData, 
+			const FVector3f& Point, float MaskWeight, FReshapeVertexBindingData& OutBindData, 
 			const float ValidityTolerance = UE_KINDA_SMALL_NUMBER)
 	{
 
@@ -242,11 +268,16 @@ namespace mu
 		OutBindData.Triangle = -1;
 		OutBindData.ShapeNormal = FVector3f::Zero();
 		OutBindData.AttachmentPoint = Point;
-		float Weight = 1.0f;
+		float Weight = MaskWeight;
+
+		if (FMath::IsNearlyZero(MaskWeight))
+		{
+			return;
+		}
 
 		double DistSqr = 0.0;
 
-		int32 FoundIndex = ShapeMeshTree.FindNearestTriangle( FVector3d(Point), DistSqr );
+		int32 FoundIndex = ShapeMeshTree.FindNearestTriangle(FVector3d(Point), DistSqr);
 			
 		if (FoundIndex < 0)
 		{
@@ -305,133 +336,36 @@ namespace mu
 
 		// Arbitrary factor, a binding will be considered valid (with its corresponding weight) to ErrorFalloffFactor times the validity tolerance.
 		constexpr float ErrorFalloffFactor = 4.0f;
-		OutBindData.Weight = 1.0f - FMath::Clamp( (ErrorEstimate - ValidityTolerance) / (ValidityTolerance * ErrorFalloffFactor), 0.0f, 1.0f );
+		OutBindData.Weight = FMath::Min(
+			MaskWeight, 
+			1.0f - FMath::Clamp( (ErrorEstimate - ValidityTolerance) / (ValidityTolerance * ErrorFalloffFactor), 0.0f, 1.0f));
 		
 		OutBindData.Triangle = FMath::IsNearlyZero(OutBindData.Weight) ? -1 : OutBindData.Triangle;
 	}
-
-	inline void BindPointReference(
-			const TArray<UE::Geometry::FIndex3i>& ShapeTriangles, const TArray<FVector3f>& ShapePositions, const TArray<FVector3f>& ShapeNormals,
-			const FVector3f& Point, FReshapeVertexBindingData& BindData,
-			const float ValidityTolerance = UE_KINDA_SMALL_NUMBER)
-	{
-		BindData.S = 0.0f;
-		BindData.T = 0.0f;
-		BindData.D = 0.0f;
-		BindData.Triangle = -1;
-		BindData.ShapeNormal = FVector3f::Zero();
-		BindData.AttachmentPoint = Point;
-		float Weight = 1.0f;
-
-		// Find the best shape triangle for each base mesh vertex
-
-		// Select the triangle closest to the vertex.
-		// \TODO: review this strategy.
-		int32 TriangleIndex = -1;
-		float BestDistSquared = TNumericLimits<float>::Max();
-		for (int ShapeTriangleIndex = 0; ShapeTriangleIndex < ShapeTriangles.Num(); ++ShapeTriangleIndex)
-		{
-			UE::Geometry::FIndex3i Triangle = ShapeTriangles[ShapeTriangleIndex];
-			FVector3f TriangleA = ShapePositions[Triangle.A];
-			FVector3f TriangleB = ShapePositions[Triangle.B];
-			FVector3f TriangleC = ShapePositions[Triangle.C];
-			FVector3f ClosestPoint = (FVector3f)FMath::ClosestPointOnTriangleToPoint((FVector)Point, (FVector)TriangleA, (FVector)TriangleB, (FVector)TriangleC);
-			float DistSquared = (ClosestPoint - Point).SquaredLength();
-			if (DistSquared < BestDistSquared)
-			{
-				TriangleIndex = ShapeTriangleIndex;
-				BestDistSquared = DistSquared;
-			}
-		}
-
-		if (TriangleIndex < 0)
-		{
-			//check(false);
-			return;
-		}
-
-		// Calculate the binding data of the base mesh vertex to its bound shape triangle
-		UE::Geometry::FIndex3i Triangle = ShapeTriangles[TriangleIndex];
-
-		// Project on the triangle, but using the vertex normals.
-		// See figure 3 in https://www.sciencedirect.com/science/article/pii/S1319157820303931
-		FVector3f TriangleA = ShapePositions[Triangle.A];
-		FVector3f TriangleB = ShapePositions[Triangle.B];
-		FVector3f TriangleC = ShapePositions[Triangle.C];
-
-		FPlane4f TrianglePlane(TriangleA, TriangleB, TriangleC);
-		FPlane4f VertexPlane(Point, TrianglePlane.GetNormal());
-	
-		// T1 = Triangle projected on the vertex plane along the triangle vertex normals
-		FVector3f TriangleA_VertexPlane = FMath::RayPlaneIntersection(TriangleA, ShapeNormals[Triangle.A], VertexPlane);
-		FVector3f TriangleB_VertexPlane = FMath::RayPlaneIntersection(TriangleB, ShapeNormals[Triangle.B], VertexPlane);
-		FVector3f TriangleC_VertexPlane = FMath::RayPlaneIntersection(TriangleC, ShapeNormals[Triangle.C], VertexPlane);
-	
-		// Barycentric coordinates of the vertex on in T1
-		FVector3f Barycentric = (FVector3f)FMath::ComputeBaryCentric2D((FVector)Point, (FVector)TriangleA_VertexPlane, (FVector)TriangleB_VertexPlane, (FVector)TriangleC_VertexPlane);
-
-		FVector3f ProjectedPoint = TriangleA * Barycentric.X + TriangleB * Barycentric.Y + TriangleC * Barycentric.Z;
-
-		FVector3f InterpolatedVertexNormal = ShapeNormals[Triangle.A] * Barycentric.X + ShapeNormals[Triangle.B] * Barycentric.Y + ShapeNormals[Triangle.C] * Barycentric.Z;
-		FVector3f ProjectedToVertex = (Point - ProjectedPoint);
-
-		// If the interpolated normal is not normalized the dot product gives the distance to the point times the length of the normal.
-		// We don't want that since when deforming a point, using the interpolated normal, we'd multiply by the length twice. See GetDeform().
-		// One option would be to normalize here (computing the distance) and when deforming the mesh.
-		// We can also compensate the signed distance so that the same interpolated normal (not normalized) gives us the correct point
-		// ( InterpolatedVertexNormal * d = ProjectedToVertex). That way modifications in the interpolated bound normal will affect the
-		// resulting point. We are doing the later.
-		const float InterpolatedNormalSizeSquared = InterpolatedVertexNormal.SizeSquared();
-		const float InvInterpolatedNormalSizeSquared = InterpolatedNormalSizeSquared > SMALL_NUMBER ? 1.0f / InterpolatedNormalSizeSquared : 0.0f;	
-		float d = FVector3f::DotProduct(ProjectedToVertex, InterpolatedVertexNormal) * InvInterpolatedNormalSizeSquared;
-	
-		BindData.S = Barycentric.Y;
-		BindData.T = Barycentric.Z;
-		BindData.D = d;
-		
-		BindData.Triangle = TriangleIndex;
-		
-		BindData.ShapeNormal = ((TriangleB - TriangleA) ^ (TriangleC - TriangleA)).GetSafeNormal();
-		
-		const FVector3f ReprojectedPoint = ProjectedPoint + InterpolatedVertexNormal * d;
-		
-		const FVector3f ReprojectedVector = ReprojectedPoint - Point;
-		const float ErrorEstimate = (ReprojectedPoint - Point).GetAbsMax();
-		
-		// If within the tolerance, 1.0, otherwise linear falloff based on the tolerance
-
-		// Arbitrary factor, a binding will be considered valid (with its corresponding weight) to ErrorFalloffFactor times the validity tolerance.
-		constexpr float ErrorFalloffFactor = 4.0f;
-		BindData.Weight = 1.0f - FMath::Clamp( (ErrorEstimate - ValidityTolerance) / (ValidityTolerance * ErrorFalloffFactor), 0.0f, 1.0f );
-		
-		BindData.Triangle = FMath::IsNearlyZero(BindData.Weight) ? -1 : BindData.Triangle;
-	}
 	
     //---------------------------------------------------------------------------------------------
     //! Find mesh clusters.
-	//! Colour {1,1,1} is reserved for the non rigid cluster
+	//! Colour {1,1,1,1} is reserved for the non rigid cluster
     //---------------------------------------------------------------------------------------------
 	
-	inline void FindRigidClusters( const Mesh* Mesh, TArray<TArray<int32>>& OutClusters, int32& OutNonRigidClusterIdx)
+	inline void FindRigidClusters(const Mesh* Mesh, const FMeshBindColorChannelUsageMasks& ColorUsageMasks, TArray<TArray<int32>>& OutClusters, int32& OutNonRigidClusterIdx)
 	{
 		MUTABLE_CPUPROFILER_SCOPE(MeshFindRigidClusters);
 		
-		const FVector3f NonRigidColour = FVector3f(1.0f, 1.0f, 1.0f);
+		constexpr uint32 NonRigidId = ~0;
 
 		OutNonRigidClusterIdx = -1;
 		
-		const UntypedMeshBufferIteratorConst ItColoursBase(Mesh->GetVertexBuffers(), MBS_COLOUR);
-		if (ItColoursBase.ptr())
+		const UntypedMeshBufferIteratorConst ItColorBase(Mesh->GetVertexBuffers(), MBS_COLOUR);
+		if (ItColorBase.ptr())
 		{
 			const int32 VertexCount = Mesh->GetVertexCount();
 
-			// Check if all colors are the same. (maybe should only search for white as it indicates no rigidity)
 			{
 				int32 VertexIndex = 0;
 				for (; VertexIndex < VertexCount; ++VertexIndex)
 				{
-					 // Colours are treated as ids, no need to use tolerance with the comparison.
-					 if (NonRigidColour != (ItColoursBase + VertexIndex).GetAsVec3f())
+					 if (NonRigidId != GetVertexClusterId(ItColorBase + VertexIndex, ColorUsageMasks))
 					 {
 						  break;
 					 }
@@ -446,13 +380,12 @@ namespace mu
 			
 			OutClusters.Empty(16);
 			TArray<TArray<int32>>& ClusterData = OutClusters;
-			// Consider using NUINT8 colors for id and pack them in a uint32 as key.
-			TMap<FVector3f, int32> ClusterSet;
+			TMap<uint32, int32> ClusterSet;
 
 			for (int32 VertexIndex = 0; VertexIndex < VertexCount; ++VertexIndex)
 			{
-				FVector3f Colour = (ItColoursBase + VertexIndex).GetAsVec3f();
-				int32& ClusterIdx = ClusterSet.FindOrAdd( Colour, INDEX_NONE );
+				const uint32 ClusterId = GetVertexClusterId(ItColorBase + VertexIndex, ColorUsageMasks);
+				int32& ClusterIdx = ClusterSet.FindOrAdd(ClusterId, INDEX_NONE);
 				TArray<int32>& Cluster = ClusterIdx < 0 ? ClusterData.Emplace_GetRef() : ClusterData[ClusterIdx];
 				if (ClusterIdx < 0)
 				{
@@ -460,12 +393,12 @@ namespace mu
 					Cluster.Reserve(32);
 				}
 
-				Cluster.Push(VertexIndex);	
+				Cluster.Add(VertexIndex);	
 			}
 
-			// Colour {1, 1, 1} is reserved for the nonrigid cluster. The choice if this value is not arbitrary,
+			// Cluster id 0xFFFFFFFF is reserved for the nonrigid cluster. The choice if this value is not arbitrary,
 			// meshes without colour will get white as default.
-			const int32* NonRigidClusterIdxFound = ClusterSet.Find( NonRigidColour );
+			const int32* NonRigidClusterIdxFound = ClusterSet.Find(NonRigidId);
 
 			// If not found, add an empty cluster for the non rigid id.
 			OutNonRigidClusterIdx = NonRigidClusterIdxFound ? *NonRigidClusterIdxFound : ClusterData.Emplace();
@@ -474,19 +407,22 @@ namespace mu
 	
 	inline void FindBindingForCluster( 
 			const Mesh* BaseMesh, FShapeMeshTree& ShapeMeshTree,
-			const TArray<int32>& Cluster, FReshapeVertexBindingData& OutBindingData, float BindTolerance )
+			const TArray<int32>& Cluster, FReshapeVertexBindingData& OutBindingData, float BindTolerance)
 	{
 		const UntypedMeshBufferIteratorConst ItPositionBase(BaseMesh->GetVertexBuffers(), MBS_POSITION);
 
 		FBox3f ClusterBoundingBox(ForceInit);
-		for ( const int32& V : Cluster)
+		for (const int32& V : Cluster)
 		{
 			ClusterBoundingBox += (ItPositionBase + V).GetAsVec3f();
 		}
 
 		FVector3f BoundPoint = ClusterBoundingBox.GetCenter();
-		
-		BindReshapePoint(ShapeMeshTree, BoundPoint, OutBindingData, BindTolerance);
+
+		// Mask weight is set on a vertex by vertex basis, ignore weight for the shared data.
+		// This will be filled in afterwards.
+		constexpr float MaskWeight = 1.0f;
+		BindReshapePoint(ShapeMeshTree, BoundPoint, MaskWeight, OutBindingData, BindTolerance);
 	}
 
 	inline TTuple<TArray<FReshapePointBindingData>, TArray<int32>, TArray<int32>> BindPhysicsBodies( 
@@ -742,7 +678,8 @@ namespace mu
 
 		for (int32 PointIndex = 0; PointIndex < TotalNumPoints; ++PointIndex)
 		{
-			BindReshapePoint(ShapeMeshTree, Points[PointIndex], VertexBindData, 0.1f);
+			constexpr float MaskWeight = 1.0f;
+			BindReshapePoint(ShapeMeshTree, Points[PointIndex], MaskWeight, VertexBindData, 0.1f);
 			PhysicsBodyBindData[PointIndex] = FReshapePointBindingData
 				{ VertexBindData.S, VertexBindData.T, VertexBindData.D, VertexBindData.Triangle, VertexBindData.Weight };
 		}
@@ -751,7 +688,7 @@ namespace mu
 	}
 
 	inline TArray<FReshapeVertexBindingData> BindVerticesReshape(
-			const Mesh* BaseMesh, FShapeMeshTree& ShapeMeshTree, bool bEnableRigidParts )	
+			const Mesh* BaseMesh, FShapeMeshTree& ShapeMeshTree, const FMeshBindColorChannelUsageMasks& ColorUsageMasks)	
 	{
 	
 		UE::Geometry::FAxisAlignedBox3d ShapeAABBox = ShapeMeshTree.GetBoundingBox();
@@ -763,9 +700,9 @@ namespace mu
 		TArray<TArray<int32>> VertexClusters;
 		int32 NonRigidClusterIdx = -1;
 
-		if (bEnableRigidParts)
+		if (ColorUsageMasks.ClusterId != 0)
 		{
-			FindRigidClusters( BaseMesh, VertexClusters, NonRigidClusterIdx );
+			FindRigidClusters(BaseMesh, ColorUsageMasks, VertexClusters, NonRigidClusterIdx);
 		}
 		
 		// Find nearest shape triangle for each base mesh vertex		
@@ -774,9 +711,14 @@ namespace mu
 		{
 			MUTABLE_CPUPROFILER_SCOPE(Project);
 
-			BindData.SetNum( MeshVertexCount );
+			BindData.SetNum(MeshVertexCount);
 
-			UntypedMeshBufferIteratorConst ItPositionBase(BaseMesh->GetVertexBuffers(), MBS_POSITION);
+			const UntypedMeshBufferIteratorConst ItPositionBase(BaseMesh->GetVertexBuffers(), MBS_POSITION);
+
+			// Disable vertex color reads if the color is not used for mask weights.
+			const UntypedMeshBufferIteratorConst ItColorBase = ColorUsageMasks.MaskWeight == 0 
+				? UntypedMeshBufferIteratorConst()
+				: UntypedMeshBufferIteratorConst(BaseMesh->GetVertexBuffers(), MBS_COLOUR);
 
 			// Special case for non rigid parts
 
@@ -787,8 +729,10 @@ namespace mu
 				const int32 VertexCount = BaseMesh->GetVertexCount();
 				for (int32 VertexIndex = 0; VertexIndex < VertexCount; ++VertexIndex)
 				{
-					FVector3f VertexPosition = (ItPositionBase + VertexIndex).GetAsVec3f();
-					BindReshapePoint(ShapeMeshTree, VertexPosition, BindData[VertexIndex], BindValidityTolerance);
+					const FVector3f VertexPosition = (ItPositionBase + VertexIndex).GetAsVec3f();
+					const float MaskWeight = ItColorBase.ptr() ? GetVertexMaskWeight(ItColorBase + VertexIndex, ColorUsageMasks) : 1.0f;
+
+					BindReshapePoint(ShapeMeshTree, VertexPosition, MaskWeight, BindData[VertexIndex], BindValidityTolerance);
 				}
 			}
 			else
@@ -799,8 +743,10 @@ namespace mu
 				for (int32 I = 0; I < NonRigidVertexCount; ++I)
 				{
 					const int32 VertexIndex = NonRigidCluster[I];
-					FVector3f VertexPosition = (ItPositionBase + VertexIndex).GetAsVec3f();
-					BindReshapePoint(ShapeMeshTree, VertexPosition, BindData[VertexIndex], BindValidityTolerance);
+					const FVector3f VertexPosition = (ItPositionBase + VertexIndex).GetAsVec3f();
+					const float MaskWeight = ItColorBase.ptr() ? GetVertexMaskWeight(ItColorBase + VertexIndex, ColorUsageMasks) : 1.0f;
+					
+					BindReshapePoint(ShapeMeshTree, VertexPosition, MaskWeight, BindData[VertexIndex], BindValidityTolerance);
 				}
 
 				// Remove data form the non rigid cluster so it ios not processed in the rigid parts binding step.
@@ -810,14 +756,21 @@ namespace mu
 				{
 					FReshapeVertexBindingData ClusterBinding;
 
-					// Copy cluster binding to every vertex of the cluster.
 					const int32 ClusterVertexCount = RigidCluster.Num();
 					if (ClusterVertexCount)
 					{
-						FindBindingForCluster(BaseMesh, ShapeMeshTree, RigidCluster, ClusterBinding, BindValidityTolerance );
+						FindBindingForCluster(BaseMesh, ShapeMeshTree, RigidCluster, ClusterBinding, BindValidityTolerance);
+						// Copy cluster binding to every vertex of the cluster modifing weight if the vertex color is used.
 						for (int32 I = 0; I < ClusterVertexCount; ++I)
 						{
 							const int32 VertexIndex = RigidCluster[I];
+
+							const float MaskWeight = ItColorBase.ptr() ? GetVertexMaskWeight(ItColorBase + VertexIndex, ColorUsageMasks) : 1.0f;
+
+							ClusterBinding.Weight = FMath::Min(ClusterBinding.Weight, MaskWeight);
+							// Invalidate binding for very small weights.
+							ClusterBinding.Triangle = FMath::IsNearlyZero(ClusterBinding.Weight) ? -1 : ClusterBinding.Triangle;
+
 							BindData[VertexIndex] = ClusterBinding;
 						}
 					}
@@ -861,8 +814,8 @@ namespace mu
 			}
 
 			FReshapeVertexBindingData BindData;
-
-			BindReshapePoint(ShapeMeshTree, Mesh->BonePoses[BoneIndex].BoneTransform.GetLocation(), BindData, BindValidityTolerance);
+			constexpr float MaskWeight = 1.0f;
+			BindReshapePoint(ShapeMeshTree, Mesh->BonePoses[BoneIndex].BoneTransform.GetLocation(), MaskWeight, BindData, BindValidityTolerance);
 
 			// Only add binding  if there is a chance of the bone moving.
 			if (BindData.Weight > UE_SMALL_NUMBER && BindData.Triangle >= 0)
@@ -877,6 +830,33 @@ namespace mu
 	}
 
 
+	inline FMeshBindColorChannelUsageMasks MakeColorChannelUsageMasks(FMeshBindColorChannelUsages Usages)
+	{
+		FMeshBindColorChannelUsageMasks Masks;
+	
+		// We assume the color will have the FColor layout
+		Masks.MaskWeight =
+			(Usages.B == EMeshBindColorChannelUsage::MaskWeight ? 0x000000FF : 0) |
+			(Usages.G == EMeshBindColorChannelUsage::MaskWeight ? 0x0000FF00 : 0) |
+			(Usages.R == EMeshBindColorChannelUsage::MaskWeight ? 0x00FF0000 : 0) |
+			(Usages.A == EMeshBindColorChannelUsage::MaskWeight ? 0xFF000000 : 0);
+
+		Masks.ClusterId =
+			(Usages.B == EMeshBindColorChannelUsage::ClusterId ? 0x000000FF : 0) |
+			(Usages.G == EMeshBindColorChannelUsage::ClusterId ? 0x0000FF00 : 0) |
+			(Usages.R == EMeshBindColorChannelUsage::ClusterId ? 0x00FF0000 : 0) |
+			(Usages.A == EMeshBindColorChannelUsage::ClusterId ? 0xFF000000 : 0);
+
+		// Maximum one weight channel.
+		check(Masks.MaskWeight == 0 ||
+			FMath::CountLeadingZeros(Masks.MaskWeight) + FMath::CountTrailingZeros(Masks.MaskWeight) == 32 - 8);
+		
+		// No overlaped channels.
+		check((Masks.ClusterId & Masks.MaskWeight) == 0);
+
+		return Masks;
+	}
+
 	//---------------------------------------------------------------------------------------------
     //! Generate the mesh-shape binding data
     //---------------------------------------------------------------------------------------------
@@ -884,7 +864,7 @@ namespace mu
 			Mesh* Result,
 			const Mesh* BaseMesh, const Mesh* ShapeMesh, 
 			const TArray<uint16>& BonesToDeform, const TArray<uint16>& PhysicsToDeform, 
-			EMeshBindShapeFlags BindFlags,
+			EMeshBindShapeFlags BindFlags, FMeshBindColorChannelUsages ColorChannelUsages,
 			bool& bOutSuccess)
     {
 		MUTABLE_CPUPROFILER_SCOPE(MeshBindShape);
@@ -898,8 +878,9 @@ namespace mu
 
 		const bool bReshapeVertices = EnumHasAnyFlags(BindFlags, EMeshBindShapeFlags::ReshapeVertices);
 		const bool bReshapeSkeleton = EnumHasAnyFlags(BindFlags, EMeshBindShapeFlags::ReshapeSkeleton);
-		const bool bReshapePhysics = EnumHasAnyFlags(BindFlags, EMeshBindShapeFlags::ReshapePhysicsVolumes);
-		const bool bEnableRigidParts = EnumHasAnyFlags(BindFlags, EMeshBindShapeFlags::EnableRigidParts);
+		const bool bReshapePhysics  = EnumHasAnyFlags(BindFlags, EMeshBindShapeFlags::ReshapePhysicsVolumes);
+
+		const FMeshBindColorChannelUsageMasks ColorUsagesMasks = MakeColorChannelUsageMasks(ColorChannelUsages);
 
 		// Early out if nothing will be modified and the vertices discarted. return null in this
 		// case indicating nothing has modified so the Base Mesh can be reused.
@@ -998,7 +979,7 @@ namespace mu
 		int32 BindingDataIndex = 0;
 		if (bReshapeVertices)
 		{
-			TArray<FReshapeVertexBindingData> VerticesBindData = BindVerticesReshape(BaseMesh, ShapeMeshTree, bEnableRigidParts);
+			TArray<FReshapeVertexBindingData> VerticesBindData = BindVerticesReshape(BaseMesh, ShapeMeshTree, ColorUsagesMasks);
 			
 			// Add the binding information to the mesh
 			// \TODO: Check that there is no other binding data.
diff --git a/Engine/Plugins/Experimental/Mutable/Source/MutableRuntime/Private/MuR/Operations.h b/Engine/Plugins/Experimental/Mutable/Source/MutableRuntime/Private/MuR/Operations.h
--- a/Engine/Plugins/Experimental/Mutable/Source/MutableRuntime/Private/MuR/Operations.h
+++ b/Engine/Plugins/Experimental/Mutable/Source/MutableRuntime/Private/MuR/Operations.h
@@ -368,6 +368,24 @@ namespace mu
 	};
 	ENUM_CLASS_FLAGS(EMeshBindShapeFlags);
 
+	enum class EMeshBindColorChannelUsage : uint8
+	{
+		None       = 0,
+		ClusterId  = 1,
+		MaskWeight = 2,
+	};
+
+	struct FMeshBindColorChannelUsages
+	{
+		EMeshBindColorChannelUsage R;
+		EMeshBindColorChannelUsage G;
+		EMeshBindColorChannelUsage B;
+		EMeshBindColorChannelUsage A;
+	};
+
+
+	static_assert(sizeof(FMeshBindColorChannelUsages) == sizeof(uint32));
+
     //---------------------------------------------------------------------------------------------
     //!
     //---------------------------------------------------------------------------------------------
@@ -953,6 +971,7 @@ namespace mu
 			ADDRESS shape;
 			uint32 flags;
 			uint32 bindingMethod;
+			uint32 ColorUsage;
 		};
 
 		struct MeshApplyShapeArgs
diff --git a/Engine/Plugins/Experimental/Mutable/Source/MutableRuntime/Public/MuR/Mesh.h b/Engine/Plugins/Experimental/Mutable/Source/MutableRuntime/Public/MuR/Mesh.h
--- a/Engine/Plugins/Experimental/Mutable/Source/MutableRuntime/Public/MuR/Mesh.h
+++ b/Engine/Plugins/Experimental/Mutable/Source/MutableRuntime/Public/MuR/Mesh.h
@@ -117,6 +117,13 @@ namespace mu
 		ClipDeformNormalProject = 3	
 	};
 
+	enum class EVertexColorUsage : uint32
+	{
+		None = 0,
+		ReshapeMaskWeight = 1,
+		ReshapeClusterId = 2
+	};
+
 	enum class EMeshCopyFlags : uint32
 	{
 		None = 0,
@@ -559,5 +566,6 @@ namespace mu
 	MUTABLE_DEFINE_ENUM_SERIALISABLE(EBoneUsageFlags)
 	MUTABLE_DEFINE_ENUM_SERIALISABLE(EMeshBufferType)
 	MUTABLE_DEFINE_ENUM_SERIALISABLE(EShapeBindingMethod)
+	MUTABLE_DEFINE_ENUM_SERIALISABLE(EVertexColorUsage)
 }
 
diff --git a/Engine/Plugins/Experimental/Mutable/Source/MutableRuntime/Public/MuR/Model.h b/Engine/Plugins/Experimental/Mutable/Source/MutableRuntime/Public/MuR/Model.h
--- a/Engine/Plugins/Experimental/Mutable/Source/MutableRuntime/Public/MuR/Model.h
+++ b/Engine/Plugins/Experimental/Mutable/Source/MutableRuntime/Public/MuR/Model.h
@@ -13,7 +13,7 @@
 //! data structures. Compiled models are not necessarily compatible when the runtime is updated,
 //! so this version number can be used externally to verify this. It is not used internally, and
 //! serializing models from different versions than this runtime will probably result in a crash.
-#define MUTABLE_COMPILED_MODEL_CODE_VERSION		uint32( 71 )
+#define MUTABLE_COMPILED_MODEL_CODE_VERSION		uint32( 72 )
 #define MUTABLE_PARAMETERS_VERSION              uint32( 2 )
 
 
diff --git a/Engine/Plugins/Experimental/Mutable/Source/MutableTools/Private/MuT/ASTOpMeshBindShape.cpp b/Engine/Plugins/Experimental/Mutable/Source/MutableTools/Private/MuT/ASTOpMeshBindShape.cpp
--- a/Engine/Plugins/Experimental/Mutable/Source/MutableTools/Private/MuT/ASTOpMeshBindShape.cpp
+++ b/Engine/Plugins/Experimental/Mutable/Source/MutableTools/Private/MuT/ASTOpMeshBindShape.cpp
@@ -24,7 +24,6 @@ namespace mu
 		: Mesh(this)
 		, Shape(this)
 		, bReshapeSkeleton(false)
-		, bEnableRigidParts(false)
 		, bReshapePhysicsVolumes(false)
 	{
 	}
@@ -43,7 +42,6 @@ namespace mu
 		{
 			const bool bSameFlags =
 				bReshapeSkeleton == Other->bReshapeSkeleton	&&
-				bEnableRigidParts == Other->bEnableRigidParts &&
 				bReshapePhysicsVolumes == Other->bReshapePhysicsVolumes &&
 				bReshapeVertices == Other->bReshapeVertices;
 
@@ -52,7 +50,11 @@ namespace mu
 				Shape == Other->Shape &&
 				BonesToDeform == Other->BonesToDeform &&
 				PhysicsToDeform == Other->PhysicsToDeform &&
-				BindingMethod == Other->BindingMethod;
+				BindingMethod == Other->BindingMethod &&
+				RChannelUsage == Other->RChannelUsage && 
+				GChannelUsage == Other->GChannelUsage && 
+				BChannelUsage == Other->BChannelUsage && 
+				AChannelUsage == Other->AChannelUsage;
 		}
 
 		return false;
@@ -64,11 +66,15 @@ namespace mu
 		uint64 res = std::hash<void*>()(Mesh.child().get());
 		hash_combine(res, Shape.child().get());
 		hash_combine(res, bool(bReshapeSkeleton));
-		hash_combine(res, bool(bEnableRigidParts));
 		hash_combine(res, bool(bReshapePhysicsVolumes));
 		hash_combine(res, bool(bReshapeVertices));
 		hash_combine(res, bool(BindingMethod));
 
+		hash_combine(res, static_cast<uint32>(RChannelUsage));
+		hash_combine(res, static_cast<uint32>(GChannelUsage));
+		hash_combine(res, static_cast<uint32>(BChannelUsage));
+		hash_combine(res, static_cast<uint32>(AChannelUsage));
+
 		for (const uint16 S : BonesToDeform)
 		{
 			hash_combine(res, S);
@@ -89,13 +95,17 @@ namespace mu
 		NewOp->Mesh = mapChild(Mesh.child());
 		NewOp->Shape = mapChild(Shape.child());
 		NewOp->bReshapeSkeleton	= bReshapeSkeleton;
-		NewOp->bEnableRigidParts = bEnableRigidParts;
 		NewOp->bReshapePhysicsVolumes = bReshapePhysicsVolumes;
 		NewOp->bReshapeVertices = bReshapeVertices;
 		NewOp->BonesToDeform = BonesToDeform;
 		NewOp->PhysicsToDeform = PhysicsToDeform;
 		NewOp->BindingMethod = BindingMethod;
 
+		NewOp->RChannelUsage = RChannelUsage;
+		NewOp->GChannelUsage = GChannelUsage;
+		NewOp->BChannelUsage = BChannelUsage;
+		NewOp->AChannelUsage = AChannelUsage;
+
 		return NewOp;
 	}
 
@@ -118,10 +128,31 @@ namespace mu
 			constexpr EMeshBindShapeFlags NoFlags = EMeshBindShapeFlags::None;
 			EMeshBindShapeFlags BindFlags = NoFlags;
 			EnumAddFlags(BindFlags, bReshapeSkeleton ? EMeshBindShapeFlags::ReshapeSkeleton : NoFlags);
-			EnumAddFlags(BindFlags, bEnableRigidParts ? EMeshBindShapeFlags::EnableRigidParts : NoFlags);
 			EnumAddFlags(BindFlags, bReshapePhysicsVolumes ? EMeshBindShapeFlags::ReshapePhysicsVolumes : NoFlags);
 			EnumAddFlags(BindFlags, bReshapeVertices ? EMeshBindShapeFlags::ReshapeVertices : NoFlags);
 
+			{
+				auto ConvertColorUsage = [](EVertexColorUsage Usage)
+				{
+					switch (Usage)
+					{
+					case EVertexColorUsage::None:			   return EMeshBindColorChannelUsage::None;
+					case EVertexColorUsage::ReshapeClusterId:  return EMeshBindColorChannelUsage::ClusterId;
+					case EVertexColorUsage::ReshapeMaskWeight: return EMeshBindColorChannelUsage::MaskWeight;
+					default: check(false); return EMeshBindColorChannelUsage::None;
+					};
+				};
+	
+				const FMeshBindColorChannelUsages ColorUsages = {
+					ConvertColorUsage(RChannelUsage),
+					ConvertColorUsage(GChannelUsage),
+					ConvertColorUsage(BChannelUsage),
+					ConvertColorUsage(AChannelUsage) };
+
+				FMemory::Memcpy(&Args.ColorUsage, &ColorUsages, sizeof(Args.ColorUsage));
+				static_assert(sizeof(Args.ColorUsage) == sizeof(ColorUsages));
+			}
+
 			Args.flags = static_cast<uint32>(BindFlags);
 
 			Args.bindingMethod = BindingMethod;
diff --git a/Engine/Plugins/Experimental/Mutable/Source/MutableTools/Private/MuT/ASTOpMeshBindShape.h b/Engine/Plugins/Experimental/Mutable/Source/MutableTools/Private/MuT/ASTOpMeshBindShape.h
--- a/Engine/Plugins/Experimental/Mutable/Source/MutableTools/Private/MuT/ASTOpMeshBindShape.h
+++ b/Engine/Plugins/Experimental/Mutable/Source/MutableTools/Private/MuT/ASTOpMeshBindShape.h
@@ -30,10 +30,13 @@ struct FProgram;
 		uint32 BindingMethod = 0;
 		
 		uint32 bReshapeSkeleton	      : 1;
-		uint32 bEnableRigidParts      : 1;
 		uint32 bReshapePhysicsVolumes : 1;
 		uint32 bReshapeVertices       : 1;
 
+		EVertexColorUsage RChannelUsage = EVertexColorUsage::None;
+		EVertexColorUsage GChannelUsage = EVertexColorUsage::None;
+		EVertexColorUsage BChannelUsage = EVertexColorUsage::None;
+		EVertexColorUsage AChannelUsage = EVertexColorUsage::None;
 	public:
 
 		ASTOpMeshBindShape();
diff --git a/Engine/Plugins/Experimental/Mutable/Source/MutableTools/Private/MuT/CodeGenerator_Mesh.cpp b/Engine/Plugins/Experimental/Mutable/Source/MutableTools/Private/MuT/CodeGenerator_Mesh.cpp
--- a/Engine/Plugins/Experimental/Mutable/Source/MutableTools/Private/MuT/CodeGenerator_Mesh.cpp
+++ b/Engine/Plugins/Experimental/Mutable/Source/MutableTools/Private/MuT/CodeGenerator_Mesh.cpp
@@ -1791,14 +1791,18 @@ class Node;
 		Ptr<ASTOpMeshBindShape> OpBind = new ASTOpMeshBindShape();
 		Ptr<ASTOpMeshApplyShape> OpApply = new ASTOpMeshApplyShape();
 
-		OpBind->bReshapeSkeleton = node.m_reshapeSkeleton;
-    	OpBind->bEnableRigidParts = node.m_enableRigidParts;
+		OpBind->bReshapeSkeleton = node.m_reshapeSkeleton;	
 		OpBind->BonesToDeform = node.BonesToDeform;
     	OpBind->bReshapePhysicsVolumes = node.m_reshapePhysicsVolumes;
 		OpBind->PhysicsToDeform = node.PhysicsToDeform;
 		OpBind->bReshapeVertices = node.m_reshapeVertices;
 		OpBind->BindingMethod = static_cast<uint32>(EShapeBindingMethod::ReshapeClosestProject);
 
+		OpBind->RChannelUsage = node.ColorRChannelUsage;
+		OpBind->GChannelUsage = node.ColorGChannelUsage;
+		OpBind->BChannelUsage = node.ColorBChannelUsage;
+		OpBind->AChannelUsage = node.ColorAChannelUsage;
+
 		OpApply->bReshapeVertices = OpBind->bReshapeVertices;
 		OpApply->bReshapeSkeleton = OpBind->bReshapeSkeleton;
 		OpApply->bReshapePhysicsVolumes = OpBind->bReshapePhysicsVolumes;
diff --git a/Engine/Plugins/Experimental/Mutable/Source/MutableTools/Private/MuT/NodeMeshReshape.cpp b/Engine/Plugins/Experimental/Mutable/Source/MutableTools/Private/MuT/NodeMeshReshape.cpp
--- a/Engine/Plugins/Experimental/Mutable/Source/MutableTools/Private/MuT/NodeMeshReshape.cpp
+++ b/Engine/Plugins/Experimental/Mutable/Source/MutableTools/Private/MuT/NodeMeshReshape.cpp
@@ -130,9 +130,12 @@ namespace mu
 	
 
 	//---------------------------------------------------------------------------------------------
-	void NodeMeshReshape::SetEnableRigidParts(bool bEnable)
-	{
-		m_pD->m_enableRigidParts = bEnable;
+	void NodeMeshReshape::SetColorUsages(EVertexColorUsage R, EVertexColorUsage G, EVertexColorUsage B, EVertexColorUsage A)
+	{	
+		m_pD->ColorRChannelUsage = R;
+		m_pD->ColorGChannelUsage = G;
+		m_pD->ColorBChannelUsage = B;
+		m_pD->ColorAChannelUsage = A;
 	}	
 
 	//---------------------------------------------------------------------------------------------
diff --git a/Engine/Plugins/Experimental/Mutable/Source/MutableTools/Private/MuT/NodeMeshReshapePrivate.h b/Engine/Plugins/Experimental/Mutable/Source/MutableTools/Private/MuT/NodeMeshReshapePrivate.h
--- a/Engine/Plugins/Experimental/Mutable/Source/MutableTools/Private/MuT/NodeMeshReshapePrivate.h
+++ b/Engine/Plugins/Experimental/Mutable/Source/MutableTools/Private/MuT/NodeMeshReshapePrivate.h
@@ -26,16 +26,19 @@ namespace mu
 		Ptr<NodeMesh> m_pTargetShape;
 		bool m_reshapeVertices = true;
 		bool m_reshapeSkeleton = false;
-		bool m_enableRigidParts = false;
 		bool m_reshapePhysicsVolumes = false;
-	
+
+		EVertexColorUsage ColorRChannelUsage = EVertexColorUsage::None;
+		EVertexColorUsage ColorGChannelUsage = EVertexColorUsage::None;
+		EVertexColorUsage ColorBChannelUsage = EVertexColorUsage::None;
+		EVertexColorUsage ColorAChannelUsage = EVertexColorUsage::None;
 
 		TArray<uint16> BonesToDeform;
 		TArray<uint16> PhysicsToDeform;
         //!
 		void Serialise( OutputArchive& arch ) const
 		{
-            uint32_t ver = 9;
+            uint32_t ver = 10;
 			arch << ver;
 
 			arch << m_pBaseMesh;
@@ -43,10 +46,14 @@ namespace mu
 			arch << m_pTargetShape;
 			arch << m_reshapeVertices;
 			arch << m_reshapeSkeleton;
-			arch << m_enableRigidParts;
 			arch << BonesToDeform;
 			arch << m_reshapePhysicsVolumes;
 			arch << PhysicsToDeform;
+
+			arch << ColorRChannelUsage;
+			arch << ColorGChannelUsage;
+			arch << ColorBChannelUsage;
+			arch << ColorAChannelUsage;
         }
 
 		//!
@@ -54,7 +61,7 @@ namespace mu
 		{
             uint32_t ver;
 			arch >> ver;
-			check(ver <= 9);
+			check(ver <= 10);
 
 			arch >> m_pBaseMesh;
 			arch >> m_pBaseShape;
@@ -66,7 +73,16 @@ namespace mu
 			}
 
 			arch >> m_reshapeSkeleton;
-			arch >> m_enableRigidParts;
+			if (ver <= 9)
+			{
+				bool bEnableRigidParts_DEPRECATED;
+				arch >> bEnableRigidParts_DEPRECATED;
+
+				if (bEnableRigidParts_DEPRECATED)
+				{
+					ColorRChannelUsage = ColorGChannelUsage = ColorBChannelUsage = ColorAChannelUsage = EVertexColorUsage::ReshapeClusterId;
+				}
+			}
 
 			if (ver <= 5)
 			{
@@ -118,6 +134,14 @@ namespace mu
 					PhysicsToDeform.Add(BoneIndex);
 				}
 			}
+
+			if (ver >= 10)
+			{
+				arch >> ColorRChannelUsage;
+				arch >> ColorGChannelUsage;
+				arch >> ColorBChannelUsage;
+				arch >> ColorAChannelUsage;
+			}
 		}
 
 		// NodeMesh::Private interface
diff --git a/Engine/Plugins/Experimental/Mutable/Source/MutableTools/Public/MuT/NodeMeshReshape.h b/Engine/Plugins/Experimental/Mutable/Source/MutableTools/Public/MuT/NodeMeshReshape.h
--- a/Engine/Plugins/Experimental/Mutable/Source/MutableTools/Public/MuT/NodeMeshReshape.h
+++ b/Engine/Plugins/Experimental/Mutable/Source/MutableTools/Public/MuT/NodeMeshReshape.h
@@ -5,7 +5,7 @@
 #include "MuR/Ptr.h"
 #include "MuT/Node.h"
 #include "MuT/NodeMesh.h"
-
+#include "MuR/Mesh.h"
 
 namespace mu
 {
@@ -58,8 +58,8 @@ namespace mu
 		/** Also deform the mesh skeleton. Disabled by default. */
 		void SetReshapeSkeleton(bool);
 		
-		/** Search for clusters of rigid parts and not deform them. */
-		void SetEnableRigidParts(bool);
+		/** Set vertex color channel usages for Reshape operations. */
+		void SetColorUsages(EVertexColorUsage R, EVertexColorUsage G, EVertexColorUsage B, EVertexColorUsage A);
 
 		/** Deform Mesh Physics Volumes */
 		void SetReshapePhysicsVolumes(bool);
