[Mutable] Fix another implicit type conversion that resulted in meshes being discarded with every update.

#rnx

diff --git a/Engine/Plugins/Experimental/Mutable/Source/CustomizableObject/Private/MuCO/CustomizableInstancePrivateData.h b/Engine/Plugins/Experimental/Mutable/Source/CustomizableObject/Private/MuCO/CustomizableInstancePrivateData.h
--- a/Engine/Plugins/Experimental/Mutable/Source/CustomizableObject/Private/MuCO/CustomizableInstancePrivateData.h
+++ b/Engine/Plugins/Experimental/Mutable/Source/CustomizableObject/Private/MuCO/CustomizableInstancePrivateData.h
@@ -12,7 +12,12 @@
 #include "UObject/Package.h"
 #include "CustomizableInstancePrivateData.generated.h"
 
-namespace mu { class PhysicsBody; }
+namespace mu 
+{
+	class PhysicsBody;
+	typedef uint64 FResourceID;
+}
+
 struct FMutableModelImageProperties;
 struct FMutableRefSkeletalMeshData;
 struct FMutableImageCacheKey;
@@ -137,8 +142,8 @@ struct FCustomizableInstanceComponentData
 
 	/** Array of generated MeshIds per each LOD, used to decide if the mesh should be updated or not.
 	 *  Size == NumLODsAvailable
-	 *  LODs without mesh will be set to -1 */
-	TArray<int32> LastMeshIdPerLOD;
+	 *  LODs without mesh will be set to the maximum value of FResourceID (Max_uint64). */
+	TArray<mu::FResourceID> LastMeshIdPerLOD;
 
 	bool operator==(const FCustomizableInstanceComponentData& Other) const { return ComponentIndex == Other.ComponentIndex; }
 };
@@ -286,8 +291,8 @@ private:
 
 	bool DoComponentsNeedUpdate(UCustomizableObjectInstance* CustomizableObjectInstance, const TSharedPtr<FMutableOperationData>& OperationData, TArray<bool>& OutComponentNeedsUpdate, bool& bOutEmptyMesh);
 
-	int32 GetLastMeshId(int32 ComponentIndex, int32 LODIndex) const;
-	void SetLastMeshId(int32 ComponentIndex, int32 LODIndex, int32 MeshId);
+	mu::FResourceID GetLastMeshId(int32 ComponentIndex, int32 LODIndex) const;
+	void SetLastMeshId(int32 ComponentIndex, int32 LODIndex, mu::FResourceID MeshId);
 
 public:
 
diff --git a/Engine/Plugins/Experimental/Mutable/Source/CustomizableObject/Private/MuCO/CustomizableObjectInstance.cpp b/Engine/Plugins/Experimental/Mutable/Source/CustomizableObject/Private/MuCO/CustomizableObjectInstance.cpp
--- a/Engine/Plugins/Experimental/Mutable/Source/CustomizableObject/Private/MuCO/CustomizableObjectInstance.cpp
+++ b/Engine/Plugins/Experimental/Mutable/Source/CustomizableObject/Private/MuCO/CustomizableObjectInstance.cpp
@@ -57,7 +57,7 @@ UTexture2D* UCustomizableInstancePrivateData::CreateTexture()
 }
 
 
-int32 UCustomizableInstancePrivateData::GetLastMeshId(int32 ComponentIndex, int32 LODIndex) const
+mu::FResourceID UCustomizableInstancePrivateData::GetLastMeshId(int32 ComponentIndex, int32 LODIndex) const
 {
 	const FCustomizableInstanceComponentData* ComponentData = GetComponentData(ComponentIndex);
 	check(ComponentData);
@@ -67,7 +67,7 @@ int32 UCustomizableInstancePrivateData::GetLastMeshId(int32 ComponentIndex, int3
 }
 
 
-void UCustomizableInstancePrivateData::SetLastMeshId(int32 ComponentIndex, int32 LODIndex, int32 MeshId)
+void UCustomizableInstancePrivateData::SetLastMeshId(int32 ComponentIndex, int32 LODIndex, mu::FResourceID MeshId)
 {
 	FCustomizableInstanceComponentData* ComponentData = GetComponentData(ComponentIndex);
 	check(ComponentData);
@@ -115,7 +115,7 @@ void UCustomizableInstancePrivateData::InvalidateGeneratedData()
 {
 	for (FCustomizableInstanceComponentData& ComponentData : ComponentsData)
 	{
-		ComponentData.LastMeshIdPerLOD.Init(INDEX_NONE, NumLODsAvailable);
+		ComponentData.LastMeshIdPerLOD.Init(MAX_uint64, NumLODsAvailable);
 	}
 
 	LastUpdateData.Clear();
@@ -203,7 +203,7 @@ void UCustomizableInstancePrivateData::PrepareForUpdate(const TSharedPtr<FMutabl
 
 			FCustomizableInstanceComponentData& NewComponentData = ComponentsData.AddDefaulted_GetRef();
 			NewComponentData.ComponentIndex = Component.Id;
-			NewComponentData.LastMeshIdPerLOD.Init(INDEX_NONE, NumLODsAvailable);
+			NewComponentData.LastMeshIdPerLOD.Init(MAX_uint64, NumLODsAvailable);
 		}
 	}
 
@@ -1794,7 +1794,7 @@ bool UCustomizableInstancePrivateData::DoComponentsNeedUpdate(UCustomizableObjec
 			{
 				// Check if an LOD is generated and it shouldn't
 				const bool LODCanBeGenerated = LODIndex >= OperationData->CurrentMinLOD && LODIndex <= OperationData->CurrentMaxLOD;
-				bMeshNeedsUpdate = !LODCanBeGenerated && ComponentData->LastMeshIdPerLOD[LODIndex] != INDEX_NONE;
+				bMeshNeedsUpdate = !LODCanBeGenerated && ComponentData->LastMeshIdPerLOD[LODIndex] != MAX_uint64;
 			}
 		}
 
@@ -1817,7 +1817,7 @@ bool UCustomizableInstancePrivateData::DoComponentsNeedUpdate(UCustomizableObjec
 		{
 			const FInstanceUpdateData::FComponent& Component = OperationData->InstanceUpdateData.Components[LOD.FirstComponent + ComponentIndex];
 
-			const int32 LastMeshID = GetLastMeshId(Component.Id, LODIndex);
+			const mu::FResourceID LastMeshID = GetLastMeshId(Component.Id, LODIndex);
 
 			if (Component.bGenerated)
 			{
@@ -1838,7 +1838,7 @@ bool UCustomizableInstancePrivateData::DoComponentsNeedUpdate(UCustomizableObjec
 				else if (Component.bReuseMesh)
 				{
 					bFoundNonEmptyMesh = true;
-					check(LastMeshID != INDEX_NONE);
+					check(LastMeshID != MAX_uint64);
 				}
 				else
 				{
@@ -1847,7 +1847,7 @@ bool UCustomizableInstancePrivateData::DoComponentsNeedUpdate(UCustomizableObjec
 			}
 			else
 			{
-				OutComponentNeedsUpdate[Component.Id] = OutComponentNeedsUpdate[Component.Id] || LastMeshID != INDEX_NONE;
+				OutComponentNeedsUpdate[Component.Id] = OutComponentNeedsUpdate[Component.Id] || LastMeshID != MAX_uint64;
 			}
 
 			bUpdateMeshes = bUpdateMeshes || OutComponentNeedsUpdate[Component.Id];
