[Mutable] Remove IsInGameThread() check in a method that can be called from any thread.

#rnx

diff --git a/Engine/Plugins/Experimental/Mutable/Source/CustomizableObject/Private/MuCO/UnrealMutableImageProvider.cpp b/Engine/Plugins/Experimental/Mutable/Source/CustomizableObject/Private/MuCO/UnrealMutableImageProvider.cpp
--- a/Engine/Plugins/Experimental/Mutable/Source/CustomizableObject/Private/MuCO/UnrealMutableImageProvider.cpp
+++ b/Engine/Plugins/Experimental/Mutable/Source/CustomizableObject/Private/MuCO/UnrealMutableImageProvider.cpp
@@ -554,7 +554,8 @@ void FUnrealMutableImageProvider::CacheImage(FName Id, bool bUser)
 
 void FUnrealMutableImageProvider::UnCacheImage(FName Id, bool bUser)
 {
-	check(IsInGameThread());
+	// TODO: Review GM
+	//check(IsInGameThread());
 
 	if (Id == FName())
 	{
