[mutable] Fix image projection seam correction for texture uv channels other than 0.
Fix bug in generated resources keys preventing the reusal of resources.
#rnx

diff --git a/Engine/Plugins/Experimental/Mutable/Source/CustomizableObject/Classes/MuCO/CustomizableObject.h b/Engine/Plugins/Experimental/Mutable/Source/CustomizableObject/Classes/MuCO/CustomizableObject.h
--- a/Engine/Plugins/Experimental/Mutable/Source/CustomizableObject/Classes/MuCO/CustomizableObject.h
+++ b/Engine/Plugins/Experimental/Mutable/Source/CustomizableObject/Classes/MuCO/CustomizableObject.h
@@ -1439,7 +1439,7 @@ private:
 	// This is a manual version number for the binary blobs in this asset.
 	// Increasing it invalidates all the previously compiled models.
 	// Warning: If while merging code both versions have changed, take the highest+1.
-	static const int32 CurrentSupportedVersion = 397;
+	static const int32 CurrentSupportedVersion = 398;
 
 public:
 
diff --git a/Engine/Plugins/Experimental/Mutable/Source/MutableRuntime/Private/MuR/CodeRunner.cpp b/Engine/Plugins/Experimental/Mutable/Source/MutableRuntime/Private/MuR/CodeRunner.cpp
--- a/Engine/Plugins/Experimental/Mutable/Source/MutableRuntime/Private/MuR/CodeRunner.cpp
+++ b/Engine/Plugins/Experimental/Mutable/Source/MutableRuntime/Private/MuR/CodeRunner.cpp
@@ -89,7 +89,6 @@ static FAutoConsoleVariableRef CVarUseProjectionVectorImpl (
 	ECVF_Default);
 }
 
-
 namespace mu
 {
 
@@ -773,7 +772,7 @@ namespace mu
 
 
     //---------------------------------------------------------------------------------------------
-    void CodeRunner::RunCode_InstanceAddResource(const FScheduledOp& item, const Model* InModel, const Parameters* InParams )
+    void CodeRunner::RunCode_InstanceAddResource(const FScheduledOp& item, const TSharedPtr<const Model>& InModel, const Parameters* InParams )
     {
 		MUTABLE_CPUPROFILER_SCOPE(RunCode_InstanceAddResource);
 
@@ -4351,7 +4350,7 @@ namespace mu
 					Ptr<Image> ResultImage = CreateImage(SizeX, SizeY, 1, EImageFormat::IF_L_UBYTE, EInitializationType::Black);
 					if (pMesh)
 					{
-						ImageRasterMesh(pMesh.get(), ResultImage.get(), args.blockIndex, CropMin, UncroppedSize);
+						ImageRasterMesh(pMesh.get(), ResultImage.get(), args.LayoutIndex, args.blockId, CropMin, UncroppedSize);
 						Release(pMesh);
 					}
 
@@ -4509,10 +4508,6 @@ namespace mu
 				scratch.vertices.SetNum( pMesh->GetVertexCount() );
 				scratch.culledVertex.SetNum( pMesh->GetVertexCount() );
 
-				// Layout is always 0 because the previous mesh project operations take care of
-				// moving the right layout channel to the 0.
-				int layout = 0;
-
 				ESamplingMethod SamplingMethod = Invoke([&]() -> ESamplingMethod
 				{
 					if (ForcedProjectionMode == 0)
@@ -4560,7 +4555,7 @@ namespace mu
 							args.bIsRGBFadingEnabled, args.bIsAlphaFadingEnabled,
 							SamplingMethod,
 							FadeStartRad, FadeEndRad, FMath::Frac(Data.RasterMesh.MipValue),
-							layout, args.blockIndex,
+							args.LayoutIndex, args.blockId,
 							CropMin, UncroppedSize,
 							&scratch, bUseProjectionVectorImpl);
 						break;
@@ -4571,7 +4566,7 @@ namespace mu
 							args.bIsRGBFadingEnabled, args.bIsAlphaFadingEnabled,
 							SamplingMethod,
 							FadeStartRad, FadeEndRad, FMath::Frac(Data.RasterMesh.MipValue),
-							layout, args.blockIndex,
+							args.LayoutIndex, args.blockId,
 							CropMin, UncroppedSize,
 							&scratch);
 						break;
@@ -4581,7 +4576,7 @@ namespace mu
 							Source.get(), Mask.get(),
 							args.bIsRGBFadingEnabled, args.bIsAlphaFadingEnabled,
 							FadeStartRad, FadeEndRad,
-							layout,
+							args.LayoutIndex,
 							Projector.projectionAngle,
 							CropMin, UncroppedSize,
 							&scratch);
@@ -5782,11 +5777,12 @@ namespace mu
 
 
     //---------------------------------------------------------------------------------------------
-    void CodeRunner::RunCode( const FScheduledOp& item, const Parameters* pParams, const Model* pModel, uint32 lodMask)
+    void CodeRunner::RunCode( const FScheduledOp& item, const Parameters* pParams, const TSharedPtr<const Model>& InModel, uint32 lodMask)
     {
 		//UE_LOG( LogMutableCore, Log, TEXT("Running :%5d , %d "), item.At, item.Stage );
 		check( item.Type == FScheduledOp::EType::Full );
 
+		const Model* pModel = InModel.Get();
 		OP_TYPE type = pModel->GetPrivate()->m_program.GetOpType(item.At);
 		//UE_LOG(LogMutableCore, Log, TEXT("Running :%5d , %d, of type %d "), item.At, item.Stage, type);
 
@@ -5829,7 +5825,7 @@ namespace mu
 
         case OP_TYPE::IN_ADDMESH:
         case OP_TYPE::IN_ADDIMAGE:
-            RunCode_InstanceAddResource(item, pModel, pParams);
+            RunCode_InstanceAddResource(item, InModel, pParams);
             break;
 
 		default:
@@ -6556,4 +6552,3 @@ namespace mu
 		}
 	}
 }
-
diff --git a/Engine/Plugins/Experimental/Mutable/Source/MutableRuntime/Private/MuR/CodeRunner.h b/Engine/Plugins/Experimental/Mutable/Source/MutableRuntime/Private/MuR/CodeRunner.h
--- a/Engine/Plugins/Experimental/Mutable/Source/MutableRuntime/Private/MuR/CodeRunner.h
+++ b/Engine/Plugins/Experimental/Mutable/Source/MutableRuntime/Private/MuR/CodeRunner.h
@@ -103,12 +103,12 @@ namespace  mu
 
         Ptr<RangeIndex> BuildCurrentOpRangeIndex( const FScheduledOp&, const Parameters*, const Model*, int32 ParameterIndex );
 
-        void RunCode( const FScheduledOp&, const Parameters*, const Model*, uint32 LodMask);
+        void RunCode( const FScheduledOp&, const Parameters*, const TSharedPtr<const Model>&, uint32 LodMask);
 
         void RunCode_Conditional(const FScheduledOp&, const Model* );
         void RunCode_Switch(const FScheduledOp&, const Model* );
         void RunCode_Instance(const FScheduledOp&, const Model*, uint32 LodMask );
-        void RunCode_InstanceAddResource(const FScheduledOp&, const Model*, const Parameters* );
+        void RunCode_InstanceAddResource(const FScheduledOp&, const TSharedPtr<const Model>& Model, const Parameters* );
         void RunCode_ConstantResource(const FScheduledOp&, const Model* );
         void RunCode_Mesh(const FScheduledOp&, const Model* );
         void RunCode_Image(const FScheduledOp&, const Parameters*, const Model* );
diff --git a/Engine/Plugins/Experimental/Mutable/Source/MutableRuntime/Private/MuR/CodeRunnerTasked.cpp b/Engine/Plugins/Experimental/Mutable/Source/MutableRuntime/Private/MuR/CodeRunnerTasked.cpp
--- a/Engine/Plugins/Experimental/Mutable/Source/MutableRuntime/Private/MuR/CodeRunnerTasked.cpp
+++ b/Engine/Plugins/Experimental/Mutable/Source/MutableRuntime/Private/MuR/CodeRunnerTasked.cpp
@@ -448,7 +448,7 @@ namespace mu
 				else
 				{
 					// Run immediately
-					RunCode(item, m_pParams, m_pModel.Get(), m_lodMask);
+					RunCode(item, m_pParams, m_pModel, m_lodMask);
 
 					if (ScheduledStagePerOp[item] == item.Stage + 1)
 					{
@@ -2083,7 +2083,7 @@ namespace mu
 		}
 
 		// Process the constant op normally, now that the rom is loaded.
-		Runner->RunCode(Op, Runner->m_pParams, Runner->m_pModel.Get(), Runner->m_lodMask);
+		Runner->RunCode(Op, Runner->m_pParams, Runner->m_pModel, Runner->m_lodMask);
 
 		FWorkingMemoryManager::FModelCacheEntry* ModelCache = Runner->m_pSystem->WorkingMemoryManager.FindModelCache(Runner->m_pModel.Get());
 
@@ -2162,7 +2162,7 @@ namespace mu
 		}
 
 		// Process the constant op normally, now that the data is loaded.
-		Runner->RunCode(Op, Runner->m_pParams, Runner->m_pModel.Get(), Runner->m_lodMask);
+		Runner->RunCode(Op, Runner->m_pParams, Runner->m_pModel, Runner->m_lodMask);
 	}
 
 
@@ -2336,7 +2336,7 @@ namespace mu
 		}
 		
 		// Process the constant op normally, now that the rom is loaded.	
-		Runner->RunCode(Op, Runner->m_pParams, Runner->m_pModel.Get(), Runner->m_lodMask);
+		Runner->RunCode(Op, Runner->m_pParams, Runner->m_pModel, Runner->m_lodMask);
 
 		for (int32 LODIndex = 0; LODIndex < LODIndexCount; ++LODIndex)
 		{
diff --git a/Engine/Plugins/Experimental/Mutable/Source/MutableRuntime/Private/MuR/OpImageRasterMesh.h b/Engine/Plugins/Experimental/Mutable/Source/MutableRuntime/Private/MuR/OpImageRasterMesh.h
--- a/Engine/Plugins/Experimental/Mutable/Source/MutableRuntime/Private/MuR/OpImageRasterMesh.h
+++ b/Engine/Plugins/Experimental/Mutable/Source/MutableRuntime/Private/MuR/OpImageRasterMesh.h
@@ -9,7 +9,6 @@
 #include "MuR/Raster.h"
 #include "MuR/ConvertData.h"
 
-
 namespace mu
 {
 
@@ -23,7 +22,7 @@ namespace mu
 	};
 
 
-	inline void ImageRasterMesh( const Mesh* pMesh, Image* pImage, int block,
+	inline void ImageRasterMesh( const Mesh* pMesh, Image* pImage, int32 LayoutIndex, int32 BlockId,
 		UE::Math::TIntVector2<uint16> CropMin, UE::Math::TIntVector2<uint16> UncroppedSize )
 	{
 		MUTABLE_CPUPROFILER_SCOPE(ImageRasterMesh)
@@ -60,8 +59,8 @@ namespace mu
 		}
 
 		// Get the indices
-		int faceCount = pMesh->GetFaceCount();
-		TArray<int> indices;
+		int32 faceCount = pMesh->GetFaceCount();
+		TArray<int32> indices;
 		indices.SetNumZeroed(faceCount * 3);
 
 		UntypedMeshBufferIteratorConst indIt( pMesh->GetIndexBuffers(), MBS_VERTEXINDEX, 0 );
@@ -74,9 +73,9 @@ namespace mu
 			++indIt;
 		}
 
-        UntypedMeshBufferIteratorConst bloIt( pMesh->GetVertexBuffers(), MBS_LAYOUTBLOCK, 0 );
+        UntypedMeshBufferIteratorConst bloIt( pMesh->GetVertexBuffers(), MBS_LAYOUTBLOCK, LayoutIndex );
 
-        if ( block<0 || bloIt.GetElementSize()==0 )
+        if (BlockId <0 || bloIt.GetElementSize()==0 )
 		{
 			// Raster all the faces
             WhitePixelProcessor pixelProc;
@@ -102,7 +101,7 @@ namespace mu
 			// Raster only the faces in the selected block
 
 			// Get the block per face
-			TArray<int> blocks;
+			TArray<uint16> blocks;
 			blocks.SetNumZeroed(vertexCount);
 
 			for ( int i=0; i<vertexCount; ++i )
@@ -117,11 +116,11 @@ namespace mu
             WhitePixelProcessor pixelProc;
 			//for (int f = 0; f < faceCount; ++f)
 			const auto& ProcessFace = [
-				vertices, indices, blocks, block, pImage, sizeX, sizeY, pixelProc
+				vertices, indices, blocks, BlockId, pImage, sizeX, sizeY, pixelProc
 			] (int32 f)
 			{
 				// TODO: Select faces outside for loop?
-				if (blocks[indices[f * 3 + 0]] == block)
+				if (blocks[indices[f * 3 + 0]] == BlockId)
 				{
 					Triangle(pImage->GetData(), pImage->GetDataSize(),
 						sizeX, sizeY,
diff --git a/Engine/Plugins/Experimental/Mutable/Source/MutableRuntime/Private/MuR/Operations.h b/Engine/Plugins/Experimental/Mutable/Source/MutableRuntime/Private/MuR/Operations.h
--- a/Engine/Plugins/Experimental/Mutable/Source/MutableRuntime/Private/MuR/Operations.h
+++ b/Engine/Plugins/Experimental/Mutable/Source/MutableRuntime/Private/MuR/Operations.h
@@ -785,7 +785,7 @@ namespace mu
 			//! like cylindrical projections.
 			ADDRESS projector;
 			
-			int32 blockIndex;
+			int32 blockId;
 			uint16 sizeX, sizeY;
 			uint16 SourceSizeX, SourceSizeY;
 			uint16 CropMinX, CropMinY;
@@ -797,6 +797,8 @@ namespace mu
 			uint8 SamplingMethod : 3;
 			// Currently only 2 min filter methods are contemplated, but reserve 3 bits for future uses. 
 			uint8 MinFilterMethod : 3;
+
+			uint8 LayoutIndex;
         };
 
         struct ImageMakeGrowMapArgs
diff --git a/Engine/Plugins/Experimental/Mutable/Source/MutableRuntime/Private/MuR/System.cpp b/Engine/Plugins/Experimental/Mutable/Source/MutableRuntime/Private/MuR/System.cpp
--- a/Engine/Plugins/Experimental/Mutable/Source/MutableRuntime/Private/MuR/System.cpp
+++ b/Engine/Plugins/Experimental/Mutable/Source/MutableRuntime/Private/MuR/System.cpp
@@ -1421,7 +1421,7 @@ namespace mu
 
 
 	//---------------------------------------------------------------------------------------------
-	FResourceID FWorkingMemoryManager::GetResourceKey(const Model* Model, const Parameters* Params, uint32 ParamListIndex, OP::ADDRESS RootAt)
+	FResourceID FWorkingMemoryManager::GetResourceKey(const TSharedPtr<const Model>& Model, const Parameters* Params, uint32 ParamListIndex, OP::ADDRESS RootAt)
 	{
 		MUTABLE_CPUPROFILER_SCOPE(GetResourceKey);
 
@@ -1645,31 +1645,35 @@ namespace mu
 		for (int32 CacheIndex = 0; CacheIndex < GeneratedResources.Num(); ++CacheIndex)
 		{
 			FGeneratedResourceData& Data = GeneratedResources[CacheIndex];
-			if (Data.Model.Pin().Get() == Model
+			bool bSameModel = Data.Model == Model;
+			bool bSameRoot = GetResourceIDRoot(Data.Id) == RootAt;
+			if (bSameModel
 				&&
-				GetResourceIDRoot(Data.Id) == RootAt
-				&&
-				Data.ParameterValuesBlob == NewKey.ParameterValuesBlob)
-			{
-				Data.LastRequestId = LastResourceResquestId;
-				return Data.Id;
-			}
-			else
+				bSameRoot
+				)
 			{
-				if (!Data.Model.Pin().Get() 
-					||
-					OldestRequestId > Data.LastRequestId)
+				bool bSameBlob = Data.ParameterValuesBlob == NewKey.ParameterValuesBlob;
+				if (bSameBlob)
 				{
-					OldestCachePosition = CacheIndex;
-					OldestRequestId = Data.Model.Pin() ? Data.LastRequestId : 0;
+					Data.LastRequestId = LastResourceResquestId;
+					return Data.Id;
 				}
 			}
+			
+			if (!Data.Model.Pin().Get() 
+				||
+				OldestRequestId > Data.LastRequestId)
+			{
+				OldestCachePosition = CacheIndex;
+				OldestRequestId = Data.Model.Pin() ? Data.LastRequestId : 0;
+			}
 		}
 
 		// Generate a new id
 		uint32 NewBlobId = ++LastResourceKeyId;
 		NewKey.Id = MakeResourceID( RootAt, NewBlobId);
 		NewKey.LastRequestId = LastResourceResquestId;
+		NewKey.Model = Model;
 
 		if (GeneratedResources.Num() >= MaxGeneratedResourceCacheSize)
 		{
diff --git a/Engine/Plugins/Experimental/Mutable/Source/MutableRuntime/Private/MuR/SystemPrivate.h b/Engine/Plugins/Experimental/Mutable/Source/MutableRuntime/Private/MuR/SystemPrivate.h
--- a/Engine/Plugins/Experimental/Mutable/Source/MutableRuntime/Private/MuR/SystemPrivate.h
+++ b/Engine/Plugins/Experimental/Mutable/Source/MutableRuntime/Private/MuR/SystemPrivate.h
@@ -1159,7 +1159,7 @@ namespace mu
 		TMemoryTrackedArray<FGeneratedResourceData> GeneratedResources;
 
 		/** */
-		FResourceID GetResourceKey(const Model* Model, const Parameters* Params, uint32 ParamListIndex, OP::ADDRESS RootAt);
+		FResourceID GetResourceKey(const TSharedPtr<const Model>&, const Parameters*, uint32 ParamListIndex, OP::ADDRESS RootAt);
 
 		//! @}
 
diff --git a/Engine/Plugins/Experimental/Mutable/Source/MutableRuntime/Public/MuR/Model.h b/Engine/Plugins/Experimental/Mutable/Source/MutableRuntime/Public/MuR/Model.h
--- a/Engine/Plugins/Experimental/Mutable/Source/MutableRuntime/Public/MuR/Model.h
+++ b/Engine/Plugins/Experimental/Mutable/Source/MutableRuntime/Public/MuR/Model.h
@@ -13,7 +13,7 @@
 //! data structures. Compiled models are not necessarily compatible when the runtime is updated,
 //! so this version number can be used externally to verify this. It is not used internally, and
 //! serializing models from different versions than this runtime will probably result in a crash.
-#define MUTABLE_COMPILED_MODEL_CODE_VERSION		uint32( 72 )
+#define MUTABLE_COMPILED_MODEL_CODE_VERSION		uint32( 73 )
 #define MUTABLE_PARAMETERS_VERSION              uint32( 2 )
 
 
diff --git a/Engine/Plugins/Experimental/Mutable/Source/MutableTools/Private/MuT/ASTOpImageLayer.cpp b/Engine/Plugins/Experimental/Mutable/Source/MutableTools/Private/MuT/ASTOpImageLayer.cpp
--- a/Engine/Plugins/Experimental/Mutable/Source/MutableTools/Private/MuT/ASTOpImageLayer.cpp
+++ b/Engine/Plugins/Experimental/Mutable/Source/MutableTools/Private/MuT/ASTOpImageLayer.cpp
@@ -320,7 +320,9 @@ namespace mu
 						&&
 						MaskRaster->angleFadeProperties.child() == BlendRaster->angleFadeProperties.child()
 						&&
-						MaskRaster->BlockIndex == BlendRaster->BlockIndex
+						MaskRaster->BlockId == BlendRaster->BlockId
+						&&
+						MaskRaster->LayoutIndex == BlendRaster->LayoutIndex
 						)
 					{
 						CurrentMask = MaskRaster->image.child();
diff --git a/Engine/Plugins/Experimental/Mutable/Source/MutableTools/Private/MuT/ASTOpImageRasterMesh.cpp b/Engine/Plugins/Experimental/Mutable/Source/MutableTools/Private/MuT/ASTOpImageRasterMesh.cpp
--- a/Engine/Plugins/Experimental/Mutable/Source/MutableTools/Private/MuT/ASTOpImageRasterMesh.cpp
+++ b/Engine/Plugins/Experimental/Mutable/Source/MutableTools/Private/MuT/ASTOpImageRasterMesh.cpp
@@ -26,7 +26,8 @@ namespace mu
 		, mask(this)
 		, projector(this)
 	{
-		BlockIndex = -1;
+		BlockId = -1;
+		LayoutIndex = -1;
 		SizeX = SizeY = 0;
 		SourceSizeX = SourceSizeY = 0;
 		CropMinX = CropMinY = 0;
@@ -57,7 +58,8 @@ namespace mu
 				angleFadeProperties == Other->angleFadeProperties &&
 				mask == Other->mask &&
 				projector == Other->projector &&
-				BlockIndex == Other->BlockIndex &&
+				BlockId == Other->BlockId &&
+				LayoutIndex == Other->LayoutIndex &&
 				SizeX == Other->SizeX &&
 				SizeY == Other->SizeY &&
 				SourceSizeX == Other->SourceSizeX &&
@@ -84,7 +86,7 @@ namespace mu
 		hash_combine(res, angleFadeProperties.child().get());
 		hash_combine(res, mask.child().get());
 		hash_combine(res, projector.child().get());
-		hash_combine(res, BlockIndex);
+		hash_combine(res, BlockId);
 		return res;
 	}
 
@@ -98,7 +100,8 @@ namespace mu
 		n->angleFadeProperties = mapChild(angleFadeProperties.child());
 		n->mask = mapChild(mask.child());
 		n->projector = mapChild(projector.child());
-		n->BlockIndex = BlockIndex;
+		n->BlockId = BlockId;
+		n->LayoutIndex = LayoutIndex;
 		n->SizeX = SizeX;
 		n->SizeY = SizeY;
 		n->CropMinX = CropMinX;
@@ -135,7 +138,8 @@ namespace mu
 			OP::ImageRasterMeshArgs args;
 			FMemory::Memzero(&args, sizeof(args));
 
-			args.blockIndex = BlockIndex;
+			args.blockId = BlockId;
+			args.LayoutIndex = LayoutIndex;
 			args.sizeX = SizeX;
 			args.sizeY = SizeY;
 			args.SourceSizeX = SourceSizeX;
@@ -247,10 +251,9 @@ namespace mu
 			if (!imageAt)
 			{
 				// We remove the project from the raster children
-				// \todo: maybe this clone is not necessary
-				const ASTOpFixed* typedSource = dynamic_cast<const ASTOpFixed*>(sourceAt.get());
+				const ASTOpFixed* MeshProjectOp = dynamic_cast<const ASTOpFixed*>(sourceAt.get());
 				Ptr<ASTOpImageRasterMesh> nop = mu::Clone<ASTOpImageRasterMesh>(this);
-				nop->mesh = typedSource->children[typedSource->op.args.MeshProject.mesh].child();
+				nop->mesh = MeshProjectOp->children[MeshProjectOp->op.args.MeshProject.mesh].child();
 				at = nop;
 			}
 			break;
diff --git a/Engine/Plugins/Experimental/Mutable/Source/MutableTools/Private/MuT/ASTOpImageRasterMesh.h b/Engine/Plugins/Experimental/Mutable/Source/MutableTools/Private/MuT/ASTOpImageRasterMesh.h
--- a/Engine/Plugins/Experimental/Mutable/Source/MutableTools/Private/MuT/ASTOpImageRasterMesh.h
+++ b/Engine/Plugins/Experimental/Mutable/Source/MutableTools/Private/MuT/ASTOpImageRasterMesh.h
@@ -24,7 +24,8 @@ namespace mu
 		ASTChild projector;
 
 		/** If layouts are used, this can indicate a single layout block that we want to raster. */
-		int32 BlockIndex;
+		int32 BlockId;
+		int8 LayoutIndex;
 
 		/** Size of the image to generate by rasterization of the mesh. */
 		uint16 SizeX, SizeY;
diff --git a/Engine/Plugins/Experimental/Mutable/Source/MutableTools/Private/MuT/ASTOpImageSwizzle.cpp b/Engine/Plugins/Experimental/Mutable/Source/MutableTools/Private/MuT/ASTOpImageSwizzle.cpp
--- a/Engine/Plugins/Experimental/Mutable/Source/MutableTools/Private/MuT/ASTOpImageSwizzle.cpp
+++ b/Engine/Plugins/Experimental/Mutable/Source/MutableTools/Private/MuT/ASTOpImageSwizzle.cpp
@@ -653,7 +653,8 @@ namespace mu
 							|| Typed->angleFadeProperties.child() != FirstRasterMesh->angleFadeProperties.child()
 							|| Typed->mask.child() != FirstRasterMesh->mask.child()
 							|| Typed->projector.child() != FirstRasterMesh->projector.child()
-							|| Typed->BlockIndex != FirstRasterMesh->BlockIndex
+							|| Typed->BlockId != FirstRasterMesh->BlockId
+							|| Typed->LayoutIndex != FirstRasterMesh->LayoutIndex
 							|| Typed->SizeX != FirstRasterMesh->SizeX
 							|| Typed->SizeY != FirstRasterMesh->SizeY
 							|| Typed->UncroppedSizeX != FirstRasterMesh->UncroppedSizeX
diff --git a/Engine/Plugins/Experimental/Mutable/Source/MutableTools/Private/MuT/CodeGenerator_Image.cpp b/Engine/Plugins/Experimental/Mutable/Source/MutableTools/Private/MuT/CodeGenerator_Image.cpp
--- a/Engine/Plugins/Experimental/Mutable/Source/MutableTools/Private/MuT/CodeGenerator_Image.cpp
+++ b/Engine/Plugins/Experimental/Mutable/Source/MutableTools/Private/MuT/CodeGenerator_Image.cpp
@@ -1287,7 +1287,7 @@ namespace mu
 			MeshOptions.bNormalizeUVs = true;		// We need normalized UVs for the projection
 
 			// We may need IDs at this point if there are modifiers for the mesh.
-			// \TODO: Detect this at mesh constant generation and enable there istead of here?
+			// \TODO: Detect this at mesh constant generation and enable there instead of here?
 			MeshOptions.bUniqueVertexIDs = true;
 			
 			FMeshGenerationResult MeshResult;
@@ -1311,7 +1311,8 @@ namespace mu
 				m_pErrorLog->GetPrivate()->Add("Layout or block index error.", ELMT_ERROR, node.m_errorContext);
 			}
 
-            pop->SetChild( pop->op.args.MeshProject.mesh, MeshResult.meshOp );
+
+			Ptr<ASTOp> CurrentMeshToProjectOp = MeshResult.meshOp;
 
             if (projectorResult.type == PROJECTOR_TYPE::WRAPPING)
             {
@@ -1330,8 +1331,8 @@ namespace mu
                         | OP::MeshFormatArgs::BT_FACE
                         | OP::MeshFormatArgs::BT_RESETBUFFERINDICES;
                 fop->Format = cop;
-                fop->Source = pop->children[pop->op.args.MeshProject.mesh].child();
-                pop->SetChild( pop->op.args.MeshProject.mesh, fop );
+                fop->Source = CurrentMeshToProjectOp;
+				CurrentMeshToProjectOp = fop;
             }
             else
             {
@@ -1339,12 +1340,12 @@ namespace mu
                 if ( GeneratedLayoutBlockId>=0 )
                 {
                     Ptr<ASTOpMeshExtractLayoutBlocks> eop = new ASTOpMeshExtractLayoutBlocks();
-                    eop->source = pop->children[pop->op.args.MeshProject.mesh].child();
+                    eop->source = CurrentMeshToProjectOp;
                     eop->layout = node.m_layout;
 
                     eop->blocks.Add(GeneratedLayoutBlockId);
 
-                    pop->SetChild( pop->op.args.MeshProject.mesh, eop );
+					CurrentMeshToProjectOp = eop;
                 }
 
                 // Reformat the mesh to a more efficient format for this operation
@@ -1362,9 +1363,12 @@ namespace mu
 					| OP::MeshFormatArgs::BT_FACE
 					| OP::MeshFormatArgs::BT_RESETBUFFERINDICES;
 				fop->Format = cop;
-                fop->Source = pop->children[pop->op.args.MeshProject.mesh].child();
-                pop->SetChild( pop->op.args.MeshProject.mesh, fop );
+                fop->Source = CurrentMeshToProjectOp;
+				CurrentMeshToProjectOp = fop;
             }
+
+			pop->SetChild(pop->op.args.MeshProject.mesh, CurrentMeshToProjectOp);
+
         }
         else
         {
@@ -1411,8 +1415,8 @@ namespace mu
         // Image size, from the current block being generated
         op->SizeX = Options.RectSize[0];
         op->SizeY = Options.RectSize[1];
-		// \TODO: Review naming of arg
-        op->BlockIndex = GeneratedLayoutBlockId;
+		op->BlockId = GeneratedLayoutBlockId;
+		op->LayoutIndex = node.m_layout;
 
 		op->bIsRGBFadingEnabled = node.bIsRGBFadingEnabled;
 		op->bIsAlphaFadingEnabled = node.bIsAlphaFadingEnabled;
@@ -1453,7 +1457,8 @@ namespace mu
         rasterop->mesh = op->mesh.child();
         rasterop->image = 0;
         rasterop->mask = 0;
-        rasterop->BlockIndex = op->BlockIndex;
+        rasterop->BlockId = op->BlockId;
+		rasterop->LayoutIndex = op->LayoutIndex;
 		rasterop->SizeX = op->SizeX;
 		rasterop->SizeY = op->SizeY;
 		rasterop->UncroppedSizeX = op->UncroppedSizeX;
