[mutable] Make image projection seam correction optional. By default it is on, for compatibility.
#rnx

diff --git a/Engine/Plugins/Experimental/Mutable/Source/CustomizableObjectEditor/Classes/MuCOE/Nodes/CustomizableObjectNodeTextureProject.h b/Engine/Plugins/Experimental/Mutable/Source/CustomizableObjectEditor/Classes/MuCOE/Nodes/CustomizableObjectNodeTextureProject.h
--- a/Engine/Plugins/Experimental/Mutable/Source/CustomizableObjectEditor/Classes/MuCOE/Nodes/CustomizableObjectNodeTextureProject.h
+++ b/Engine/Plugins/Experimental/Mutable/Source/CustomizableObjectEditor/Classes/MuCOE/Nodes/CustomizableObjectNodeTextureProject.h
@@ -41,7 +41,13 @@ public:
 	// Layout to use for the generated images.
 	UPROPERTY(EditAnywhere, Category = CustomizableObject)
 	uint8 Layout = 0;
-	
+
+	/** When this is enable, additional operations will happen to correct projections that go over a texture UV seam to prevent interpolation artifacts. 
+	* This is not necessary if the projection is guaranteed to not go over a seam.
+	*/
+	UPROPERTY(EditAnywhere, Category = CustomizableObject)
+	bool bEnableTextureSeamCorrection = true;
+
 	UPROPERTY(EditAnywhere, Category = CustomizableObject)
 	bool bEnableAngleFadeOutForRGB = true;
 
diff --git a/Engine/Plugins/Experimental/Mutable/Source/CustomizableObjectEditor/Private/MuCOE/GenerateMutableSource/GenerateMutableSourceGroupProjector.cpp b/Engine/Plugins/Experimental/Mutable/Source/CustomizableObjectEditor/Private/MuCOE/GenerateMutableSource/GenerateMutableSourceGroupProjector.cpp
--- a/Engine/Plugins/Experimental/Mutable/Source/CustomizableObjectEditor/Private/MuCOE/GenerateMutableSource/GenerateMutableSourceGroupProjector.cpp
+++ b/Engine/Plugins/Experimental/Mutable/Source/CustomizableObjectEditor/Private/MuCOE/GenerateMutableSource/GenerateMutableSourceGroupProjector.cpp
@@ -49,7 +49,7 @@ mu::NodeImagePtr GenerateMutableGroupProjection(const int32 NodeLOD, const int32
 
 	check(TypedNodeMat || TypedNodeExt);
 	
-	TArray<mu::NodeImageProjectPtr> ImageNodes;
+	TArray<mu::Ptr<mu::NodeImageProject>> ImageNodes;
 	TArray<FGroupProjectorTempData> ImageNodes_ProjectorTempData;
 
 	int32 TextureSize = 512;
@@ -106,7 +106,7 @@ mu::NodeImagePtr GenerateMutableGroupProjection(const int32 NodeLOD, const int32
 				continue;
 			}
 			
-			mu::NodeImageProjectPtr ImageNode = new mu::NodeImageProject();
+			mu::Ptr<mu::NodeImageProject> ImageNode = new mu::NodeImageProject();
 			bIsGroupProjectorImage = true;
 			ImageNode->SetLayout(ProjectorTempData.CustomizableObjectNodeGroupProjectorParameter->UVLayout);
 
diff --git a/Engine/Plugins/Experimental/Mutable/Source/CustomizableObjectEditor/Private/MuCOE/GenerateMutableSource/GenerateMutableSourceImage.cpp b/Engine/Plugins/Experimental/Mutable/Source/CustomizableObjectEditor/Private/MuCOE/GenerateMutableSource/GenerateMutableSourceImage.cpp
--- a/Engine/Plugins/Experimental/Mutable/Source/CustomizableObjectEditor/Private/MuCOE/GenerateMutableSource/GenerateMutableSourceImage.cpp
+++ b/Engine/Plugins/Experimental/Mutable/Source/CustomizableObjectEditor/Private/MuCOE/GenerateMutableSource/GenerateMutableSourceImage.cpp
@@ -561,7 +561,7 @@ mu::NodeImagePtr GenerateMutableSourceImage(const UEdGraphPin* Pin, FMutableGrap
 
 	else if (const UCustomizableObjectNodeTextureProject* TypedNodeProject = Cast<UCustomizableObjectNodeTextureProject>(Node))
 	{
-		mu::NodeImageProjectPtr ImageNode = new mu::NodeImageProject();
+		mu::Ptr<mu::NodeImageProject> ImageNode = new mu::NodeImageProject();
 		Result = ImageNode;
 
 		if (!FollowInputPin(*TypedNodeProject->MeshPin()))
@@ -592,6 +592,7 @@ mu::NodeImagePtr GenerateMutableSourceImage(const UEdGraphPin* Pin, FMutableGrap
 
 		ImageNode->SetImageSize(TextureSize);
 
+		ImageNode->SetEnableSeamCorrection(TypedNodeProject->bEnableTextureSeamCorrection);
 		ImageNode->SetAngleFadeChannels(TypedNodeProject->bEnableAngleFadeOutForRGB, TypedNodeProject->bEnableAngleFadeOutForAlpha);
 		ImageNode->SetSamplingMethod(Invoke(
 			[](ETextureProjectSamplingMethod Method) -> mu::ESamplingMethod 
diff --git a/Engine/Plugins/Experimental/Mutable/Source/MutableTools/Private/MuT/CodeGenerator_Image.cpp b/Engine/Plugins/Experimental/Mutable/Source/MutableTools/Private/MuT/CodeGenerator_Image.cpp
--- a/Engine/Plugins/Experimental/Mutable/Source/MutableTools/Private/MuT/CodeGenerator_Image.cpp
+++ b/Engine/Plugins/Experimental/Mutable/Source/MutableTools/Private/MuT/CodeGenerator_Image.cpp
@@ -1248,10 +1248,10 @@ namespace mu
 
         // Mesh project operation
         //------------------------------
-        Ptr<ASTOpFixed> pop = new ASTOpFixed();
-        pop->op.type = OP_TYPE::ME_PROJECT;
+        Ptr<ASTOpFixed> ProjectOp = new ASTOpFixed();
+		ProjectOp->op.type = OP_TYPE::ME_PROJECT;
 
-        Ptr<ASTOp> lastMeshOp = pop;
+        Ptr<ASTOp> LastMeshOp = ProjectOp;
 
         // Projector
         FProjectorGenerationResult projectorResult;
@@ -1266,7 +1266,7 @@ namespace mu
             GenerateMissingProjectorCode( projectorResult, node.m_errorContext );
         }
 
-        pop->SetChild( pop->op.args.MeshProject.projector, projectorResult.op );
+		ProjectOp->SetChild(ProjectOp->op.args.MeshProject.projector, projectorResult.op );
 
 		int32 LayoutBlockIndex = -1;
 		if (Options.LayoutToApply)
@@ -1296,7 +1296,7 @@ namespace mu
 			// Match the block id of the block we are generating with the id that resulted in the generated mesh
 			GeneratedLayoutBlockId = -1;
 			
-			mu::LayoutPtrConst Layout = MeshResult.GeneratedLayouts.IsValidIndex(node.m_layout) ? MeshResult.GeneratedLayouts[node.m_layout] : nullptr;
+			Ptr<const Layout> Layout = MeshResult.GeneratedLayouts.IsValidIndex(node.m_layout) ? MeshResult.GeneratedLayouts[node.m_layout] : nullptr;
 			if (Layout && Layout->m_blocks.IsValidIndex(LayoutBlockIndex))
 			{
 				GeneratedLayoutBlockId = Layout->m_blocks[LayoutBlockIndex].m_id;
@@ -1325,14 +1325,14 @@ namespace mu
 
                 cop->SetValue(FormatMeshResult, m_compilerOptions->OptimisationOptions.bUseDiskCache);
 
-                Ptr<ASTOpMeshFormat> fop = new ASTOpMeshFormat();
-                fop->Buffers = OP::MeshFormatArgs::BT_VERTEX
+                Ptr<ASTOpMeshFormat> FormatOp = new ASTOpMeshFormat();
+				FormatOp->Buffers = OP::MeshFormatArgs::BT_VERTEX
                         | OP::MeshFormatArgs::BT_INDEX
                         | OP::MeshFormatArgs::BT_FACE
                         | OP::MeshFormatArgs::BT_RESETBUFFERINDICES;
-                fop->Format = cop;
-                fop->Source = CurrentMeshToProjectOp;
-				CurrentMeshToProjectOp = fop;
+				FormatOp->Format = cop;
+				FormatOp->Source = CurrentMeshToProjectOp;
+				CurrentMeshToProjectOp = FormatOp;
             }
             else
             {
@@ -1357,36 +1357,36 @@ namespace mu
 
                 cop->SetValue(FormatMeshResult, m_compilerOptions->OptimisationOptions.bUseDiskCache);
 
-                Ptr<ASTOpMeshFormat> fop = new ASTOpMeshFormat();
-                fop->Buffers = OP::MeshFormatArgs::BT_VERTEX
+                Ptr<ASTOpMeshFormat> FormatOp = new ASTOpMeshFormat();
+				FormatOp->Buffers = OP::MeshFormatArgs::BT_VERTEX
 					| OP::MeshFormatArgs::BT_INDEX
 					| OP::MeshFormatArgs::BT_FACE
 					| OP::MeshFormatArgs::BT_RESETBUFFERINDICES;
-				fop->Format = cop;
-                fop->Source = CurrentMeshToProjectOp;
-				CurrentMeshToProjectOp = fop;
+				FormatOp->Format = cop;
+				FormatOp->Source = CurrentMeshToProjectOp;
+				CurrentMeshToProjectOp = FormatOp;
             }
 
-			pop->SetChild(pop->op.args.MeshProject.mesh, CurrentMeshToProjectOp);
+			ProjectOp->SetChild(ProjectOp->op.args.MeshProject.mesh, CurrentMeshToProjectOp);
 
         }
         else
         {
             // This argument is required
-            Ptr<const Mesh> pMesh = new Mesh();
+            Ptr<const Mesh> TempMesh = new Mesh();
             Ptr<ASTOpConstantResource> cop = new ASTOpConstantResource();
             cop->type = OP_TYPE::ME_CONSTANT;
-            cop->SetValue( pMesh, m_compilerOptions->OptimisationOptions.bUseDiskCache );
-            pop->SetChild( pop->op.args.MeshProject.mesh, cop );
+            cop->SetValue(TempMesh, m_compilerOptions->OptimisationOptions.bUseDiskCache );
+			ProjectOp->SetChild(ProjectOp->op.args.MeshProject.mesh, cop );
             m_pErrorLog->GetPrivate()->Add( "Projector mesh not set.", ELMT_ERROR, node.m_errorContext );
         }
 
 
         // Image raster operation
         //------------------------------
-        Ptr<ASTOpImageRasterMesh> op = new ASTOpImageRasterMesh();
-        op->mesh = lastMeshOp;
-        op->projector = projectorResult.op;
+        Ptr<ASTOpImageRasterMesh> ImageRasterOp = new ASTOpImageRasterMesh();
+		ImageRasterOp->mesh = LastMeshOp;
+		ImageRasterOp->projector = projectorResult.op;
 
         // Image
         if ( node.m_pImage )
@@ -1400,36 +1400,36 @@ namespace mu
 
 			FImageGenerationResult ImageResult;
 			GenerateImage(NewOptions, ImageResult, node.m_pImage);
-			op->image = ImageResult.op;
+			ImageRasterOp->image = ImageResult.op;
 
 			FImageDesc desc = CalculateImageDesc(*node.m_pImage->GetBasePrivate());
-			op->SourceSizeX = desc.m_size[0];
-			op->SourceSizeY = desc.m_size[1];
+			ImageRasterOp->SourceSizeX = desc.m_size[0];
+			ImageRasterOp->SourceSizeY = desc.m_size[1];
         }
         else
         {
             // This argument is required
-            op->image = GenerateMissingImageCode(TEXT("Projector image"), EImageFormat::IF_RGB_UBYTE, node.m_errorContext, Options);
+			ImageRasterOp->image = GenerateMissingImageCode(TEXT("Projector image"), EImageFormat::IF_RGB_UBYTE, node.m_errorContext, Options);
         }
 
         // Image size, from the current block being generated
-        op->SizeX = Options.RectSize[0];
-        op->SizeY = Options.RectSize[1];
-		op->BlockId = GeneratedLayoutBlockId;
-		op->LayoutIndex = node.m_layout;
+		ImageRasterOp->SizeX = Options.RectSize[0];
+		ImageRasterOp->SizeY = Options.RectSize[1];
+		ImageRasterOp->BlockId = GeneratedLayoutBlockId;
+		ImageRasterOp->LayoutIndex = node.m_layout;
 
-		op->bIsRGBFadingEnabled = node.bIsRGBFadingEnabled;
-		op->bIsAlphaFadingEnabled = node.bIsAlphaFadingEnabled;
-		op->SamplingMethod = node.SamplingMethod;
-		op->MinFilterMethod = node.MinFilterMethod;
+		ImageRasterOp->bIsRGBFadingEnabled = node.bIsRGBFadingEnabled;
+		ImageRasterOp->bIsAlphaFadingEnabled = node.bIsAlphaFadingEnabled;
+		ImageRasterOp->SamplingMethod = node.SamplingMethod;
+		ImageRasterOp->MinFilterMethod = node.MinFilterMethod;
 
 		// Fading angles are optional, and stored in a colour. If one exists, we generate both.
 		if (node.m_pAngleFadeStart || node.m_pAngleFadeEnd)
 		{
-			NodeScalarConstantPtr pDefaultFade = new NodeScalarConstant();
+			Ptr<NodeScalarConstant> pDefaultFade = new NodeScalarConstant();
 			pDefaultFade->SetValue(180.0f);
 
-			NodeColourFromScalarsPtr pPropsNode = new NodeColourFromScalars();
+			Ptr<NodeColourFromScalars> pPropsNode = new NodeColourFromScalars();
 
 			if (node.m_pAngleFadeStart) pPropsNode->SetX(node.m_pAngleFadeStart);
 			else pPropsNode->SetX(pDefaultFade);
@@ -1437,7 +1437,7 @@ namespace mu
 			if (node.m_pAngleFadeEnd) pPropsNode->SetY(node.m_pAngleFadeEnd);
 			else pPropsNode->SetY(pDefaultFade);
 
-			op->angleFadeProperties = Generate(pPropsNode);
+			ImageRasterOp->angleFadeProperties = Generate(pPropsNode);
 		}
 
         // Target mask
@@ -1448,45 +1448,52 @@ namespace mu
 			Ptr<ASTOp> mask = MaskResult.op;
 
             mask = GenerateImageFormat( mask, EImageFormat::IF_L_UBYTE );
-            op->mask = GenerateImageSize( mask, Options.RectSize);
+			ImageRasterOp->mask = GenerateImageSize( mask, Options.RectSize);
         }
 
         // Seam correction operations
         //------------------------------
-        Ptr<ASTOpImageRasterMesh> rasterop = new ASTOpImageRasterMesh();
-        rasterop->mesh = op->mesh.child();
-        rasterop->image = 0;
-        rasterop->mask = 0;
-        rasterop->BlockId = op->BlockId;
-		rasterop->LayoutIndex = op->LayoutIndex;
-		rasterop->SizeX = op->SizeX;
-		rasterop->SizeY = op->SizeY;
-		rasterop->UncroppedSizeX = op->UncroppedSizeX;
-		rasterop->UncroppedSizeY = op->UncroppedSizeY;
-		rasterop->CropMinX = op->CropMinX;
-		rasterop->CropMinY = op->CropMinY;
-		rasterop->SamplingMethod = ESamplingMethod::Point;
-		rasterop->MinFilterMethod = EMinFilterMethod::None;
-
-        Ptr<ASTOpImageMakeGrowMap> MakeGrowMapOp = new ASTOpImageMakeGrowMap();
-		MakeGrowMapOp->Mask = rasterop;
-		MakeGrowMapOp->Border = MUTABLE_GROW_BORDER_VALUE;
-		
-		// If we want to be able to generate progressive mips efficiently, we need mipmaps for the "displacement map".
-		if (m_compilerOptions->OptimisationOptions.bEnableProgressiveImages)
+		if (node.bEnableTextureSeamCorrection)
 		{
-			Ptr<ASTOpImageMipmap> MipMask = new ASTOpImageMipmap;
-			MipMask->Source = MakeGrowMapOp->Mask.child();
-			MipMask->bPreventSplitTail = true;
-			MakeGrowMapOp->Mask = MipMask;
-		}
+			Ptr<ASTOpImageRasterMesh> MaskRasterOp = new ASTOpImageRasterMesh();
+			MaskRasterOp->mesh = ImageRasterOp->mesh.child();
+			MaskRasterOp->image = 0;
+			MaskRasterOp->mask = 0;
+			MaskRasterOp->BlockId = ImageRasterOp->BlockId;
+			MaskRasterOp->LayoutIndex = ImageRasterOp->LayoutIndex;
+			MaskRasterOp->SizeX = ImageRasterOp->SizeX;
+			MaskRasterOp->SizeY = ImageRasterOp->SizeY;
+			MaskRasterOp->UncroppedSizeX = ImageRasterOp->UncroppedSizeX;
+			MaskRasterOp->UncroppedSizeY = ImageRasterOp->UncroppedSizeY;
+			MaskRasterOp->CropMinX = ImageRasterOp->CropMinX;
+			MaskRasterOp->CropMinY = ImageRasterOp->CropMinY;
+			MaskRasterOp->SamplingMethod = ESamplingMethod::Point;
+			MaskRasterOp->MinFilterMethod = EMinFilterMethod::None;
+
+			Ptr<ASTOpImageMakeGrowMap> MakeGrowMapOp = new ASTOpImageMakeGrowMap();
+			MakeGrowMapOp->Mask = MaskRasterOp;
+			MakeGrowMapOp->Border = MUTABLE_GROW_BORDER_VALUE;
+
+			// If we want to be able to generate progressive mips efficiently, we need mipmaps for the "displacement map".
+			if (m_compilerOptions->OptimisationOptions.bEnableProgressiveImages)
+			{
+				Ptr<ASTOpImageMipmap> MipMask = new ASTOpImageMipmap;
+				MipMask->Source = MakeGrowMapOp->Mask.child();
+				MipMask->bPreventSplitTail = true;
+				MakeGrowMapOp->Mask = MipMask;
+			}
 
-        Ptr<ASTOpFixed> disop = new ASTOpFixed();
-        disop->op.type = OP_TYPE::IM_DISPLACE;
-        disop->SetChild(disop->op.args.ImageDisplace.displacementMap, MakeGrowMapOp);
-        disop->SetChild(disop->op.args.ImageDisplace.source, op );
+			Ptr<ASTOpFixed> DisplaceOp = new ASTOpFixed();
+			DisplaceOp->op.type = OP_TYPE::IM_DISPLACE;
+			DisplaceOp->SetChild(DisplaceOp->op.args.ImageDisplace.displacementMap, MakeGrowMapOp);
+			DisplaceOp->SetChild(DisplaceOp->op.args.ImageDisplace.source, ImageRasterOp);
 
-        Result.op = disop;
+			Result.op = DisplaceOp;
+		}
+		else
+		{
+			Result.op = ImageRasterOp;
+		}		
     }
 
 
diff --git a/Engine/Plugins/Experimental/Mutable/Source/MutableTools/Private/MuT/NodeImageProject.cpp b/Engine/Plugins/Experimental/Mutable/Source/MutableTools/Private/MuT/NodeImageProject.cpp
--- a/Engine/Plugins/Experimental/Mutable/Source/MutableTools/Private/MuT/NodeImageProject.cpp
+++ b/Engine/Plugins/Experimental/Mutable/Source/MutableTools/Private/MuT/NodeImageProject.cpp
@@ -78,34 +78,34 @@ namespace mu
 
 
 	//---------------------------------------------------------------------------------------------
-	void NodeImageProject::SetInputNode( int i, NodePtr pNode )
+	void NodeImageProject::SetInputNode( int i, Ptr<Node> InNode)
 	{
 		check( i>=0 && i<GetInputCount() );
 
 		switch (i)
 		{
 		case 0:
-			m_pD->m_pMesh = dynamic_cast<NodeMesh*>(pNode.get());
+			m_pD->m_pMesh = dynamic_cast<NodeMesh*>(InNode.get());
 			break;
 
 		case 1:
-			m_pD->m_pProjector = dynamic_cast<NodeProjector*>(pNode.get());
+			m_pD->m_pProjector = dynamic_cast<NodeProjector*>(InNode.get());
 			break;
 
 		case 2:
-			m_pD->m_pAngleFadeStart = dynamic_cast<NodeScalar*>(pNode.get());
+			m_pD->m_pAngleFadeStart = dynamic_cast<NodeScalar*>(InNode.get());
 			break;
 
 		case 3:
-			m_pD->m_pAngleFadeEnd = dynamic_cast<NodeScalar*>(pNode.get());
+			m_pD->m_pAngleFadeEnd = dynamic_cast<NodeScalar*>(InNode.get());
 			break;
 
 		case 4:
-			m_pD->m_pMask = dynamic_cast<NodeImage*>(pNode.get());
+			m_pD->m_pMask = dynamic_cast<NodeImage*>(InNode.get());
 			break;
 
 		default:
-			m_pD->m_pImage = dynamic_cast<NodeImage*>(pNode.get());
+			m_pD->m_pImage = dynamic_cast<NodeImage*>(InNode.get());
 			break;
 		}
 	}
@@ -114,52 +114,45 @@ namespace mu
 	//---------------------------------------------------------------------------------------------
 	// Own Interface
 	//---------------------------------------------------------------------------------------------
-	NodeMeshPtr NodeImageProject::GetMesh() const
+	void NodeImageProject::SetMesh( Ptr<NodeMesh> InNode)
 	{
-		return m_pD->m_pMesh.get();
+		m_pD->m_pMesh = InNode;
 	}
 
 
 	//---------------------------------------------------------------------------------------------
-	void NodeImageProject::SetMesh( NodeMeshPtr pNode )
+	void NodeImageProject::SetProjector( Ptr<NodeProjector> InNode)
 	{
-		m_pD->m_pMesh = pNode;
+		m_pD->m_pProjector = InNode;
 	}
 
 
 	//---------------------------------------------------------------------------------------------
-	NodeProjectorPtr NodeImageProject::GetProjector() const
+	void NodeImageProject::SetAngleFadeChannels(bool bFadeRGB, bool bFadeA)
 	{
-		return m_pD->m_pProjector.get();
+		m_pD->bIsRGBFadingEnabled = bFadeRGB;
+		m_pD->bIsAlphaFadingEnabled = bFadeA;
 	}
 
 
 	//---------------------------------------------------------------------------------------------
-	void NodeImageProject::SetProjector( NodeProjectorPtr pNode )
+	void NodeImageProject::SetEnableSeamCorrection(bool bEnabled)
 	{
-		m_pD->m_pProjector = pNode;
+		m_pD->bEnableTextureSeamCorrection = bEnabled;
 	}
 
 
 	//---------------------------------------------------------------------------------------------
-	void NodeImageProject::SetAngleFadeChannels(bool bFadeRGB, bool bFadeA)
+	void NodeImageProject::SetAngleFadeStart(Ptr<NodeScalar> InNode)
 	{
-		m_pD->bIsRGBFadingEnabled = bFadeRGB;
-		m_pD->bIsAlphaFadingEnabled = bFadeA;
+		m_pD->m_pAngleFadeStart = InNode;
 	}
 
 
 	//---------------------------------------------------------------------------------------------
-	void NodeImageProject::SetAngleFadeStart(NodeScalarPtr pNode)
+	void NodeImageProject::SetAngleFadeEnd( Ptr<NodeScalar> InNode)
 	{
-		m_pD->m_pAngleFadeStart = pNode;
-	}
-
-
-	//---------------------------------------------------------------------------------------------
-	void NodeImageProject::SetAngleFadeEnd( NodeScalarPtr pNode )
-	{
-		m_pD->m_pAngleFadeEnd = pNode;
+		m_pD->m_pAngleFadeEnd = InNode;
 	}
 
 	//---------------------------------------------------------------------------------------------
@@ -175,59 +168,32 @@ namespace mu
 	}
 
 
-	//---------------------------------------------------------------------------------------------
-	NodeImagePtr NodeImageProject::GetTargetMask() const
-	{
-		return m_pD->m_pMask;
-	}
-
-
-	//---------------------------------------------------------------------------------------------
-	void NodeImageProject::SetTargetMask( NodeImagePtr pNode )
-	{
-		m_pD->m_pMask = pNode;
-	}
-
 
 	//---------------------------------------------------------------------------------------------
-	NodeImagePtr NodeImageProject::GetImage() const
+	void NodeImageProject::SetTargetMask( Ptr<NodeImage> InNode)
 	{
-		return m_pD->m_pImage;
+		m_pD->m_pMask = InNode;
 	}
 
 
 	//---------------------------------------------------------------------------------------------
-	void NodeImageProject::SetImage( NodeImagePtr pNode )
+	void NodeImageProject::SetImage( Ptr<NodeImage> InNode )
 	{
-		m_pD->m_pImage = pNode;
+		m_pD->m_pImage = InNode;
 	}
 
 
     //---------------------------------------------------------------------------------------------
-    uint8_t NodeImageProject::GetLayout() const
+    void NodeImageProject::SetLayout( uint8 LayoutIndex )
     {
-        return m_pD->m_layout;
+        m_pD->m_layout = LayoutIndex;
     }
 
 
-    //---------------------------------------------------------------------------------------------
-    void NodeImageProject::SetLayout( uint8_t l )
-    {
-        m_pD->m_layout = l;
-    }
-
-
-	//---------------------------------------------------------------------------------------------
-	const FUintVector2& NodeImageProject::GetImageSize() const
-	{
-		return m_pD->m_imageSize;
-	}
-
-
 	//---------------------------------------------------------------------------------------------
-	void NodeImageProject::SetImageSize(const FUintVector2& size)
+	void NodeImageProject::SetImageSize(const FUintVector2& Size)
 	{
-		m_pD->m_imageSize = size;
+		m_pD->m_imageSize = Size;
 	}
 
 }
diff --git a/Engine/Plugins/Experimental/Mutable/Source/MutableTools/Private/MuT/NodeImageProjectPrivate.h b/Engine/Plugins/Experimental/Mutable/Source/MutableTools/Private/MuT/NodeImageProjectPrivate.h
--- a/Engine/Plugins/Experimental/Mutable/Source/MutableTools/Private/MuT/NodeImageProjectPrivate.h
+++ b/Engine/Plugins/Experimental/Mutable/Source/MutableTools/Private/MuT/NodeImageProjectPrivate.h
@@ -24,23 +24,24 @@ namespace mu
 
 		static NODE_TYPE s_type;
 
-		NodeProjectorPtr m_pProjector;
-		NodeMeshPtr m_pMesh;
-		NodeScalarPtr m_pAngleFadeStart;
-		NodeScalarPtr m_pAngleFadeEnd;
-		NodeImagePtr m_pImage;
-		NodeImagePtr m_pMask;
+		Ptr<NodeProjector> m_pProjector;
+		Ptr<NodeMesh> m_pMesh;
+		Ptr<NodeScalar> m_pAngleFadeStart;
+		Ptr<NodeScalar> m_pAngleFadeEnd;
+		Ptr<NodeImage> m_pImage;
+		Ptr<NodeImage> m_pMask;
 		FUintVector2 m_imageSize;
 		uint8 m_layout = 0;
 		bool bIsRGBFadingEnabled = true;
 		bool bIsAlphaFadingEnabled = true;
+		bool bEnableTextureSeamCorrection = true;
 		ESamplingMethod SamplingMethod = ESamplingMethod::Point;
 		EMinFilterMethod MinFilterMethod = EMinFilterMethod::None;
 
 		//!
 		void Serialise( OutputArchive& arch ) const
 		{
-            uint32 ver = 4;
+            uint32 ver = 5;
 			arch << ver;
 
 			arch << m_pProjector;
@@ -53,6 +54,7 @@ namespace mu
 			arch << m_imageSize;
 			arch << bIsRGBFadingEnabled;
 			arch << bIsAlphaFadingEnabled;
+			arch << bEnableTextureSeamCorrection;
 			arch << SamplingMethod;
 			arch << MinFilterMethod;
 		}
@@ -62,7 +64,7 @@ namespace mu
 		{
             uint32 ver;
 			arch >> ver;
-            check(ver>=2 && ver<=4);
+            check(ver>=2 && ver<=5);
 
 			arch >> m_pProjector;
 			arch >> m_pMesh;
@@ -78,6 +80,11 @@ namespace mu
 				arch >> bIsAlphaFadingEnabled;
 			}
 
+			if (ver >= 5)
+			{
+				arch >> bEnableTextureSeamCorrection;
+			}
+
 			if (ver >= 4)
 			{
 				arch >> SamplingMethod;
diff --git a/Engine/Plugins/Experimental/Mutable/Source/MutableTools/Public/MuT/NodeImageProject.h b/Engine/Plugins/Experimental/Mutable/Source/MutableTools/Public/MuT/NodeImageProject.h
--- a/Engine/Plugins/Experimental/Mutable/Source/MutableTools/Public/MuT/NodeImageProject.h
+++ b/Engine/Plugins/Experimental/Mutable/Source/MutableTools/Public/MuT/NodeImageProject.h
@@ -15,25 +15,10 @@ namespace mu
 
 	// Forward definitions
 	class NodeScalar;
-	typedef Ptr<NodeScalar> NodeScalarPtr;
-	typedef Ptr<const NodeScalar> NodeScalarPtrConst;
-
 	class NodeColour;
-	typedef Ptr<NodeColour> NodeColourPtr;
-	typedef Ptr<const NodeColour> NodeColourPtrConst;
-
 	class NodeMesh;
-	typedef Ptr<NodeMesh> NodeMeshPtr;
-	typedef Ptr<const NodeMesh> NodeMeshPtrConst;
-
 	class NodeProjector;
-	typedef Ptr<NodeProjector> NodeProjectorPtr;
-	typedef Ptr<const NodeProjector> NodeProjectorPtrConst;
-
 	class NodeImageProject;
-	typedef Ptr<NodeImageProject> NodeImageProjectPtr;
-	typedef Ptr<const NodeImageProject> NodeImageProjectPtrConst;
-
 	class InputArchive;
 	class OutputArchive;
 
@@ -51,7 +36,7 @@ namespace mu
 
 		void SerialiseWrapper(OutputArchive& arch) const override;
 		static void Serialise( const NodeImageProject* pNode, OutputArchive& arch );
-		static NodeImageProjectPtr StaticUnserialise( InputArchive& arch );
+		static Ptr<NodeImageProject> StaticUnserialise( InputArchive& arch );
 
 
 		//-----------------------------------------------------------------------------------------
@@ -63,23 +48,23 @@ namespace mu
 
         virtual int GetInputCount() const override;
         virtual Node* GetInputNode( int i ) const override;
-        void SetInputNode( int i, NodePtr pNode ) override;
+        void SetInputNode( int i, Ptr<Node> pNode ) override;
 
 		//-----------------------------------------------------------------------------------------
 		// Own Interface
 		//-----------------------------------------------------------------------------------------
 
 		//! Get the node generating the mesh to be used for the projection.
-		NodeMeshPtr GetMesh() const;
-		void SetMesh( NodeMeshPtr );
+		void SetMesh( Ptr<NodeMesh> );
 
 		//! Get the node generating the projector to be used for the projection.
-		NodeProjectorPtr GetProjector() const;
-		void SetProjector( NodeProjectorPtr );
+		void SetProjector( Ptr<NodeProjector> );
 
 		//! Get the node generating a mask to filter out the projected images. It is optional.
-		NodeImagePtr GetTargetMask() const;
-		void SetTargetMask( NodeImagePtr );
+		void SetTargetMask( Ptr<NodeImage> );
+
+		/** Eanble or disable the additional operations to correct texture UV seam artifact correction. Default is enabled.*/
+		void SetEnableSeamCorrection(bool bEnabled);
 
 		//! Set the angle-based fading behaviour for this projector. It must be set separately for 
 		//! RGB and Alpha channels of the images to project. By default it is enabled for all channels.
@@ -88,10 +73,10 @@ namespace mu
 		void SetAngleFadeChannels( bool bFadeRGB, bool bFadeA );
 
         //! Set the node generating the fading start angle. Only relevant if fading is enabled with 
-		void SetAngleFadeStart( NodeScalarPtr );
+		void SetAngleFadeStart( Ptr<NodeScalar> );
 
 		//! Set the node generating the fading end angle
-		void SetAngleFadeEnd( NodeScalarPtr );
+		void SetAngleFadeEnd( Ptr<NodeScalar> );
 
 		//! Set sampling method.
 		void SetSamplingMethod(ESamplingMethod SamplingMethod);
@@ -100,16 +85,13 @@ namespace mu
 		void SetMinFilterMethod(EMinFilterMethod MinFilterMethod);
 
         //! Get the node generating the image to project.
-        NodeImagePtr GetImage() const;
-        void SetImage( NodeImagePtr );
+        void SetImage( Ptr<NodeImage> );
 
         //! UV layout of the mesh to use for the generated image. Defaults to 0.
-        uint8_t GetLayout() const;
-        void SetLayout( uint8_t  );
+        void SetLayout( uint8  );
 
 		//! Set the size of the image to generate with the projection. If set to 0 (default) a size
 		//! that matches how this node is used will try to be guessed. 
-		const FUintVector2& GetImageSize() const;
 		void SetImageSize( const FUintVector2& size );
 
 		//-----------------------------------------------------------------------------------------
