Adds support for user-defined extension data Selection and Variation nodes in Mutable

diff --git a/Engine/Plugins/Experimental/Mutable/Source/CustomizableObjectEditor/Classes/MuCOE/Nodes/CustomizableObjectNodeExtensionDataSwitch.h b/Engine/Plugins/Experimental/Mutable/Source/CustomizableObjectEditor/Classes/MuCOE/Nodes/CustomizableObjectNodeExtensionDataSwitch.h
new file mode 100644
--- /dev/null
+++ b/Engine/Plugins/Experimental/Mutable/Source/CustomizableObjectEditor/Classes/MuCOE/Nodes/CustomizableObjectNodeExtensionDataSwitch.h
@@ -0,0 +1,31 @@
+// Copyright Epic Games, Inc. All Rights Reserved.
+
+#pragma once
+
+#include "MuCOE/Nodes/CustomizableObjectNodeSwitchBase.h"
+#include "MuCOE/ICustomizableObjectExtensionNode.h"
+
+#include "CustomizableObjectNodeExtensionDataSwitch.generated.h"
+
+UCLASS(abstract)
+class CUSTOMIZABLEOBJECTEDITOR_API UCustomizableObjectNodeExtensionDataSwitch
+	: public UCustomizableObjectNodeSwitchBase
+	, public ICustomizableObjectExtensionNode
+{
+	GENERATED_BODY()
+
+public:
+
+	//~Begin UCustomizableObjectNode interface
+	virtual bool IsAffectedByLOD() const override;
+	//~End UCustomizableObjectNode interface
+
+	//~Begin UCustomizableObjectNodeSwitchBase interface
+	virtual bool ShouldAddToContextMenu(FText& OutCategory) const override;
+	virtual FString GetPinPrefix() const override;
+	//~End UCustomizableObjectNodeSwitchBase interface
+
+	//~Begin ICustomizableObjectExtensionNode interface
+	mu::NodeExtensionDataPtr GenerateMutableNode(class FExtensionDataCompilerInterface& InCompilerInterface) const override;
+	//~End ICustomizableObjectExtensionNode interface
+};
diff --git a/Engine/Plugins/Experimental/Mutable/Source/CustomizableObjectEditor/Classes/MuCOE/Nodes/CustomizableObjectNodeExtensionDataVariation.h b/Engine/Plugins/Experimental/Mutable/Source/CustomizableObjectEditor/Classes/MuCOE/Nodes/CustomizableObjectNodeExtensionDataVariation.h
new file mode 100644
--- /dev/null
+++ b/Engine/Plugins/Experimental/Mutable/Source/CustomizableObjectEditor/Classes/MuCOE/Nodes/CustomizableObjectNodeExtensionDataVariation.h
@@ -0,0 +1,61 @@
+// Copyright Epic Games, Inc. All Rights Reserved.
+
+#pragma once
+
+#include "MuCOE/Nodes/CustomizableObjectNode.h"
+#include "MuCOE/ICustomizableObjectExtensionNode.h"
+
+#include "CustomizableObjectNodeExtensionDataVariation.generated.h"
+
+USTRUCT()
+struct CUSTOMIZABLEOBJECTEDITOR_API FCustomizableObjectExtensionDataVariation
+{
+	GENERATED_BODY()
+
+	UPROPERTY(EditAnywhere, Category = "CustomizableObject")
+	FString Tag;
+};
+
+UCLASS(abstract)
+class CUSTOMIZABLEOBJECTEDITOR_API UCustomizableObjectNodeExtensionDataVariation
+	: public UCustomizableObjectNode
+	, public ICustomizableObjectExtensionNode
+{
+	GENERATED_BODY()
+
+public:
+
+	//~Begin UObject interface
+	virtual void PostEditChangeProperty(FPropertyChangedEvent& InPropertyChangedEvent) override;
+	//~End UObject interface
+
+	//~Begin UEdGraphNode interface
+	virtual FLinearColor GetNodeTitleColor() const override;
+	//~End UEdGraphNode interface
+
+	//~Begin UCustomizableObjectNode interface
+	virtual bool IsAffectedByLOD() const override;
+	virtual bool ShouldAddToContextMenu(FText& OutCategory) const override;
+	virtual void AllocateDefaultPins(UCustomizableObjectNodeRemapPins* InRemapPins) override;
+	//~End UCustomizableObjectNode interface
+
+	//~Begin ICustomizableObjectExtensionNode interface
+	virtual mu::NodeExtensionDataPtr GenerateMutableNode(FExtensionDataCompilerInterface& InCompilerInterface) const override;
+	//~End ICustomizableObjectExtensionNode interface
+
+	virtual FName GetCategory() const PURE_VIRTUAL(GetCategory(), return FName();)
+	virtual FName GetOutputPinName() const PURE_VIRTUAL(GetOutputPinName(), return FName();)
+
+	FName GetDefaultPinName() const;
+	FName GetVariationPinName(int32 InIndex) const;
+	class UEdGraphPin* GetDefaultPin() const;
+	class UEdGraphPin* GetVariationPin(int32 InIndex) const;
+
+	int32 GetNumVariations() const;
+
+public:
+
+	UPROPERTY(EditAnywhere, Category = "CustomizableObject")
+	TArray<FCustomizableObjectExtensionDataVariation> Variations;
+
+};
\ No newline at end of file
diff --git a/Engine/Plugins/Experimental/Mutable/Source/CustomizableObjectEditor/Private/MuCOE/EdGraphSchema_CustomizableObject.cpp b/Engine/Plugins/Experimental/Mutable/Source/CustomizableObjectEditor/Private/MuCOE/EdGraphSchema_CustomizableObject.cpp
--- a/Engine/Plugins/Experimental/Mutable/Source/CustomizableObjectEditor/Private/MuCOE/EdGraphSchema_CustomizableObject.cpp
+++ b/Engine/Plugins/Experimental/Mutable/Source/CustomizableObjectEditor/Private/MuCOE/EdGraphSchema_CustomizableObject.cpp
@@ -586,9 +586,9 @@ void UEdGraphSchema_CustomizableObject::GetGraphContextActions(FGraphContextMenu
 	for (TObjectIterator<UCustomizableObjectNode> It(RF_NoFlags); It; ++It)
 	{
 		const UCustomizableObjectNode* Node = *It;
-		if (!Node->HasAllFlags(RF_ClassDefaultObject))
+		if (!Node->HasAllFlags(RF_ClassDefaultObject) || Node->GetClass()->HasAnyClassFlags(CLASS_Abstract))
 		{
-			// Only interested in CDOs
+			// Only interested in non-abstract CDOs
 			continue;
 		}
 
diff --git a/Engine/Plugins/Experimental/Mutable/Source/CustomizableObjectEditor/Private/MuCOE/ExtensionDataCompilerInterface.cpp b/Engine/Plugins/Experimental/Mutable/Source/CustomizableObjectEditor/Private/MuCOE/ExtensionDataCompilerInterface.cpp
--- a/Engine/Plugins/Experimental/Mutable/Source/CustomizableObjectEditor/Private/MuCOE/ExtensionDataCompilerInterface.cpp
+++ b/Engine/Plugins/Experimental/Mutable/Source/CustomizableObjectEditor/Private/MuCOE/ExtensionDataCompilerInterface.cpp
@@ -5,6 +5,7 @@
 #include "InstancedStruct.h"
 #include "MuCO/CustomizableObjectStreamedExtensionData.h"
 #include "MuCOE/GenerateMutableSource/GenerateMutableSource.h"
+#include "MuCOE/CustomizableObjectCompiler.h"
 #include "MuR/ExtensionData.h"
 
 FExtensionDataCompilerInterface::FExtensionDataCompilerInterface(FMutableGraphGenerationContext& InGenerationContext)
@@ -61,3 +62,19 @@ UObject* FExtensionDataCompilerInterface::GetOuterForAlwaysLoadedObjects()
 	check(GenerationContext.Object);
 	return GenerationContext.Object;
 }
+
+void FExtensionDataCompilerInterface::AddGeneratedNode(const UCustomizableObjectNode* InNode)
+{
+	check(InNode);
+
+	// A const_cast here is required because the new node needs to be added in the GeneratedNodes list so mutable can
+	// discover new parameters that can potentially be attached to the extension node, however, this
+	// function is called as ICustomizableObjectExtensionNode::GenerateMutableNode(this), so we need to cast the const away here.
+	// Decided to do the case here so the use of AddGeneratedNode is as clean as possible
+	GenerationContext.GeneratedNodes.Add(const_cast<UCustomizableObjectNode*>(InNode));
+}
+
+void FExtensionDataCompilerInterface::CompilerLog(const FText& InLogText, const UCustomizableObjectNode* InNode)
+{
+	GenerationContext.Compiler->CompilerLog(InLogText, InNode);
+}
diff --git a/Engine/Plugins/Experimental/Mutable/Source/CustomizableObjectEditor/Private/MuCOE/Nodes/CustomizableObjectNodeExtensionDataSwitch.cpp b/Engine/Plugins/Experimental/Mutable/Source/CustomizableObjectEditor/Private/MuCOE/Nodes/CustomizableObjectNodeExtensionDataSwitch.cpp
new file mode 100644
--- /dev/null
+++ b/Engine/Plugins/Experimental/Mutable/Source/CustomizableObjectEditor/Private/MuCOE/Nodes/CustomizableObjectNodeExtensionDataSwitch.cpp
@@ -0,0 +1,62 @@
+// Copyright Epic Games, Inc. All Rights Reserved.
+
+#include "MuCOE/Nodes/CustomizableObjectNodeExtensionDataSwitch.h"
+#include "MuCOE/GenerateMutableSource/GenerateMutableSource.h"
+#include "MuCOE/GenerateMutableSource/GenerateMutableSourceFloat.h"
+#include "MuCOE/ExtensionDataCompilerInterface.h"
+#include "MuCOE/GraphTraversal.h"
+#include "MuT/NodeExtensionDataSwitch.h"
+
+#include UE_INLINE_GENERATED_CPP_BY_NAME(CustomizableObjectNodeExtensionDataSwitch)
+
+
+bool UCustomizableObjectNodeExtensionDataSwitch::IsAffectedByLOD() const
+{
+	return false;
+}
+
+bool UCustomizableObjectNodeExtensionDataSwitch::ShouldAddToContextMenu(FText& OutCategory) const
+{
+	OutCategory = FText::FromName(GetCategory());
+	return true;
+}
+
+FString UCustomizableObjectNodeExtensionDataSwitch::GetPinPrefix() const
+{
+	return FString::Format(TEXT("{0} "), { GetOutputPinName() });
+}
+
+mu::NodeExtensionDataPtr UCustomizableObjectNodeExtensionDataSwitch::GenerateMutableNode(FExtensionDataCompilerInterface& InCompilerInterface) const
+{
+	const int32 NumParameters = FollowInputPinArray(*SwitchParameter()).Num();
+
+	const UEdGraphPin* EnumPin = FollowInputPin(*SwitchParameter());
+	mu::NodeScalarPtr SwitchParam = EnumPin ? GenerateMutableSourceFloat(EnumPin, InCompilerInterface.GenerationContext) : nullptr;
+
+	const int32 NumSwitchOptions = GetNumElements();
+
+	mu::NodeExtensionDataSwitchPtr SwitchNode = new mu::NodeExtensionDataSwitch;
+	SwitchNode->SetParameter(SwitchParam);
+	SwitchNode->SetOptionCount(NumSwitchOptions);
+
+	for (int32 OptionIndex = 0; OptionIndex < NumSwitchOptions; ++OptionIndex)
+	{
+		if (const UEdGraphPin* ExtensionDataPin = GetElementPin(OptionIndex))
+		{
+			if (const UEdGraphPin* ConnectedPin = FollowInputPin(*ExtensionDataPin))
+			{
+				if (const ICustomizableObjectExtensionNode* ExtensionNode = Cast<ICustomizableObjectExtensionNode>(ConnectedPin->GetOwningNode()))
+				{
+					if (mu::NodeExtensionDataPtr ExtensionMutableNode = ExtensionNode->GenerateMutableNode(InCompilerInterface))
+					{
+						SwitchNode->SetOption(OptionIndex, ExtensionMutableNode);
+					}
+				}
+			}
+		}
+	}
+
+	InCompilerInterface.AddGeneratedNode(this);
+
+	return SwitchNode;
+}
diff --git a/Engine/Plugins/Experimental/Mutable/Source/CustomizableObjectEditor/Private/MuCOE/Nodes/CustomizableObjectNodeExtensionDataVariation.cpp b/Engine/Plugins/Experimental/Mutable/Source/CustomizableObjectEditor/Private/MuCOE/Nodes/CustomizableObjectNodeExtensionDataVariation.cpp
new file mode 100644
--- /dev/null
+++ b/Engine/Plugins/Experimental/Mutable/Source/CustomizableObjectEditor/Private/MuCOE/Nodes/CustomizableObjectNodeExtensionDataVariation.cpp
@@ -0,0 +1,132 @@
+// Copyright Epic Games, Inc. All Rights Reserved.
+
+#include "MuCOE/Nodes/CustomizableObjectNodeExtensionDataVariation.h"
+
+#include "MuCOE/ExtensionDataCompilerInterface.h"
+#include "MuCOE/EdGraphSchema_CustomizableObject.h"
+#include "MuCOE/GraphTraversal.h"
+#include "MuT/NodeExtensionDataVariation.h"
+
+#include UE_INLINE_GENERATED_CPP_BY_NAME(CustomizableObjectNodeExtensionDataVariation)
+
+#define LOCTEXT_NAMESPACE "CustomizableObjectEditor"
+
+void UCustomizableObjectNodeExtensionDataVariation::PostEditChangeProperty(FPropertyChangedEvent& InPropertyChangedEvent)
+{
+	if (FProperty* PropertyThatChanged = InPropertyChangedEvent.Property)
+	{
+		ReconstructNode();
+	}
+
+	Super::PostEditChangeProperty(InPropertyChangedEvent);
+}
+
+FLinearColor UCustomizableObjectNodeExtensionDataVariation::GetNodeTitleColor() const
+{
+	const UEdGraphSchema_CustomizableObject* Schema = GetDefault<UEdGraphSchema_CustomizableObject>();
+	return Schema->GetPinTypeColor(GetCategory());
+}
+
+bool UCustomizableObjectNodeExtensionDataVariation::IsAffectedByLOD() const
+{
+	return false;
+}
+
+bool UCustomizableObjectNodeExtensionDataVariation::ShouldAddToContextMenu(FText& OutCategory) const
+{
+	OutCategory = FText::FromName(GetCategory());
+	return true;
+}
+
+void UCustomizableObjectNodeExtensionDataVariation::AllocateDefaultPins(UCustomizableObjectNodeRemapPins* InRemapPins)
+{
+	const UEdGraphSchema_CustomizableObject* Schema = GetDefault<UEdGraphSchema_CustomizableObject>();
+
+	UEdGraphPin* OutputPin = CustomCreatePin(EGPD_Output, GetCategory(), GetOutputPinName());
+	OutputPin->bDefaultValueIsIgnored = true;
+
+	for (int32 VariationIndex = Variations.Num() - 1; VariationIndex >= 0; --VariationIndex)
+	{
+		constexpr bool bIsArray = false;
+		UEdGraphPin* VariationPin = CustomCreatePin(EGPD_Input, GetCategory(), GetVariationPinName(VariationIndex), bIsArray);
+		VariationPin->bDefaultValueIsIgnored = true;
+		VariationPin->PinFriendlyName = FText::Format(LOCTEXT("VariationPinFriendlyName", "Variation {0} [{1}]"), { VariationIndex, FText::FromString(Variations[VariationIndex].Tag) });
+	}
+
+	UEdGraphPin* DefaultVariation = CustomCreatePin(EGPD_Input, GetCategory(), GetDefaultPinName(), false);
+	DefaultVariation->bDefaultValueIsIgnored = true;
+}
+
+mu::NodeExtensionDataPtr UCustomizableObjectNodeExtensionDataVariation::GenerateMutableNode(FExtensionDataCompilerInterface& InCompilerInterface) const
+{
+	mu::NodeExtensionDataVariationPtr VariationNode = new mu::NodeExtensionDataVariation;
+
+	if (const UEdGraphPin* ConnectedPin = FollowInputPin(*GetDefaultPin()))
+	{
+		if (const ICustomizableObjectExtensionNode* ExtensionNode = Cast<ICustomizableObjectExtensionNode>(ConnectedPin->GetOwningNode()))
+		{
+			if (mu::NodeExtensionDataPtr DefaultValueExtensionData = ExtensionNode->GenerateMutableNode(InCompilerInterface))
+			{
+				VariationNode->SetDefaultValue(DefaultValueExtensionData);
+			}
+			else
+			{
+				InCompilerInterface.CompilerLog(LOCTEXT("ExtensionDataVariationFailed", "Extension Data Variation failed"), this);
+			}
+		}
+	}
+	else
+	{
+		InCompilerInterface.CompilerLog(LOCTEXT("ExtensionDataVariationMissinfDef", "Extension Data Variation node requires a default value"), this);
+	}
+
+	VariationNode->SetVariationCount(Variations.Num());
+	for (int32 VariationIndex = 0; VariationIndex < Variations.Num(); ++VariationIndex)
+	{
+		if (const UEdGraphPin* VariationPin = GetVariationPin(VariationIndex))
+		{
+			VariationNode->SetVariationTag(VariationIndex, StringCast<ANSICHAR>(*Variations[VariationIndex].Tag).Get());
+			if (const UEdGraphPin* ConnectedPin = FollowInputPin(*VariationPin))
+			{
+				if (const ICustomizableObjectExtensionNode* ExtensionNode = Cast<ICustomizableObjectExtensionNode>(ConnectedPin->GetOwningNode()))
+				{
+					if (mu::NodeExtensionDataPtr VariationExtensionData = ExtensionNode->GenerateMutableNode(InCompilerInterface))
+					{
+						VariationNode->SetVariationValue(VariationIndex, VariationExtensionData);
+					}
+				}
+			}
+		}
+	}
+
+	InCompilerInterface.AddGeneratedNode(this);
+
+	return VariationNode;
+}
+
+FName UCustomizableObjectNodeExtensionDataVariation::GetDefaultPinName() const
+{
+	return TEXT("Default");
+}
+
+FName UCustomizableObjectNodeExtensionDataVariation::GetVariationPinName(int32 InIndex) const
+{
+	return FName{ *FString::Format(TEXT("Variation {0}"), { InIndex }) };
+}
+
+UEdGraphPin* UCustomizableObjectNodeExtensionDataVariation::GetDefaultPin() const
+{
+	return FindPin(GetDefaultPinName());
+}
+
+UEdGraphPin* UCustomizableObjectNodeExtensionDataVariation::GetVariationPin(int32 InIndex) const
+{
+	return FindPin(GetVariationPinName(InIndex));
+}
+
+int32 UCustomizableObjectNodeExtensionDataVariation::GetNumVariations() const
+{
+	return Variations.Num();
+}
+
+#undef LOCTEXT_NAMESPACE
diff --git a/Engine/Plugins/Experimental/Mutable/Source/CustomizableObjectEditor/Private/MuCOE/SMutableCodeViewer.cpp b/Engine/Plugins/Experimental/Mutable/Source/CustomizableObjectEditor/Private/MuCOE/SMutableCodeViewer.cpp
--- a/Engine/Plugins/Experimental/Mutable/Source/CustomizableObjectEditor/Private/MuCOE/SMutableCodeViewer.cpp
+++ b/Engine/Plugins/Experimental/Mutable/Source/CustomizableObjectEditor/Private/MuCOE/SMutableCodeViewer.cpp
@@ -1463,6 +1463,7 @@ void SMutableCodeViewer::GenerateElementRecursive(const int32& InStateIndex, mu:
 	case mu::OP_TYPE::SC_SWITCH:
 	case mu::OP_TYPE::NU_SWITCH:
 	case mu::OP_TYPE::IN_SWITCH:
+	case mu::OP_TYPE::ED_SWITCH:
 	{
 		const uint8* OpData = InProgram.GetOpArgsPointer(InParentAddress);
 
diff --git a/Engine/Plugins/Experimental/Mutable/Source/CustomizableObjectEditor/Public/MuCOE/ExtensionDataCompilerInterface.h b/Engine/Plugins/Experimental/Mutable/Source/CustomizableObjectEditor/Public/MuCOE/ExtensionDataCompilerInterface.h
--- a/Engine/Plugins/Experimental/Mutable/Source/CustomizableObjectEditor/Public/MuCOE/ExtensionDataCompilerInterface.h
+++ b/Engine/Plugins/Experimental/Mutable/Source/CustomizableObjectEditor/Public/MuCOE/ExtensionDataCompilerInterface.h
@@ -61,7 +61,15 @@ public:
 	/** The Outer to use for objects owned by always-loaded Extension Data constants. */
 	UObject* GetOuterForAlwaysLoadedObjects();
 
-private:
-	FMutableGraphGenerationContext& GenerationContext;
-};
+	/**
+	* Adds a node to the Generation Context list of generated nodes. This function is meant to be called
+	* from classes that implement ICustomizableObjectExtensionNode::GenerateMutableNode for any generated nodes
+	* so they are registered against the Mutable compiler
+	*/
+	void AddGeneratedNode(const class UCustomizableObjectNode* InNode);
 
+	/** Adds a compiler log message to be displayed at the end of the compilation process */
+	void CompilerLog(const FText& InLogText, const class UCustomizableObjectNode* InNode);
+
+	FMutableGraphGenerationContext& GenerationContext;
+};
\ No newline at end of file
diff --git a/Engine/Plugins/Experimental/Mutable/Source/MutableRuntime/Private/MuR/CodeRunner.cpp b/Engine/Plugins/Experimental/Mutable/Source/MutableRuntime/Private/MuR/CodeRunner.cpp
--- a/Engine/Plugins/Experimental/Mutable/Source/MutableRuntime/Private/MuR/CodeRunner.cpp
+++ b/Engine/Plugins/Experimental/Mutable/Source/MutableRuntime/Private/MuR/CodeRunner.cpp
@@ -247,6 +247,7 @@ namespace mu
             case DT_IMAGE:      StoreImage( cat, LoadImage(rat) ); break;
             case DT_LAYOUT:     StoreLayout( cat, LoadLayout(rat) ); break;
             case DT_INSTANCE:   StoreInstance( cat, LoadInstance(rat) ); break;
+			case DT_EXTENSION_DATA: StoreExtensionData( cat, LoadExtensionData(rat) ); break;
             default:
                 // Not implemented
                 check( false );
@@ -302,6 +303,7 @@ namespace mu
 				case DT_IMAGE:      StoreImage(item, nullptr); break;
 				case DT_LAYOUT:     StoreLayout(item, nullptr); break;
 				case DT_INSTANCE:   StoreInstance(item, nullptr); break;
+				case DT_EXTENSION_DATA: StoreExtensionData(item, new ExtensionData); break;
 				default:
 					// Not implemented
 					check(false);
@@ -359,6 +361,7 @@ namespace mu
             case DT_IMAGE:      StoreImage( cat, LoadImage(rat) ); break;
             case DT_LAYOUT:     StoreLayout( cat, LoadLayout(rat) ); break;
             case DT_INSTANCE:   StoreInstance( cat, LoadInstance(rat) ); break;
+			case DT_EXTENSION_DATA: StoreExtensionData( cat, LoadExtensionData(rat) ); break;
             default:
                 // Not implemented
                 check( false );
@@ -5804,6 +5807,7 @@ namespace mu
         case OP_TYPE::ME_CONDITIONAL:
         case OP_TYPE::LA_CONDITIONAL:
         case OP_TYPE::IN_CONDITIONAL:
+		case OP_TYPE::ED_CONDITIONAL:
             RunCode_Conditional(item, pModel);
             break;
 
@@ -5820,6 +5824,7 @@ namespace mu
         case OP_TYPE::ME_SWITCH:
         case OP_TYPE::LA_SWITCH:
         case OP_TYPE::IN_SWITCH:
+		case OP_TYPE::ED_SWITCH:
             RunCode_Switch(item, pModel);
             break;
 
diff --git a/Engine/Plugins/Experimental/Mutable/Source/MutableRuntime/Private/MuR/CodeRunnerTasked.cpp b/Engine/Plugins/Experimental/Mutable/Source/MutableRuntime/Private/MuR/CodeRunnerTasked.cpp
--- a/Engine/Plugins/Experimental/Mutable/Source/MutableRuntime/Private/MuR/CodeRunnerTasked.cpp
+++ b/Engine/Plugins/Experimental/Mutable/Source/MutableRuntime/Private/MuR/CodeRunnerTasked.cpp
@@ -114,6 +114,7 @@ namespace mu
 		{   0,   0,   0,   0 },	// ME_CONDITIONAL
 		{   0,   0,   0,   0 },	// LA_CONDITIONAL
 		{   0,   0,   0,   0 },	// IN_CONDITIONAL
+		{   0,   0,   0,   0 },	// ED_CONDITIONAL
 		    
 		{   0,   0,   0,   0 },	// NU_SWITCH
 		{   0,   0,   0,   0 },	// SC_SWITCH
@@ -122,6 +123,7 @@ namespace mu
 		{   0,   0,   0,   0 },	// ME_SWITCH
 		{   0,   0,   0,   0 },	// LA_SWITCH
 		{   0,   0,   0,   0 },	// IN_SWITCH
+		{   0,   0,   0,   0 },	// ED_SWITCH
 		    
 		{   0,   0,   0,   0 },	// BO_LESS
 		{   0,   0,   0,   0 },	// BO_EQUAL_SC_CONST
diff --git a/Engine/Plugins/Experimental/Mutable/Source/MutableRuntime/Private/MuR/CodeVisitor.h b/Engine/Plugins/Experimental/Mutable/Source/MutableRuntime/Private/MuR/CodeVisitor.h
--- a/Engine/Plugins/Experimental/Mutable/Source/MutableRuntime/Private/MuR/CodeVisitor.h
+++ b/Engine/Plugins/Experimental/Mutable/Source/MutableRuntime/Private/MuR/CodeVisitor.h
@@ -449,6 +449,7 @@ class Parameters;
             case OP_TYPE::ME_CONDITIONAL:
             case OP_TYPE::LA_CONDITIONAL:
             case OP_TYPE::IN_CONDITIONAL:
+			case OP_TYPE::ED_CONDITIONAL:
             {
                 auto args = program.GetOpArgs<OP::ConditionalArgs>(at);
 
@@ -482,6 +483,7 @@ class Parameters;
             case OP_TYPE::ME_SWITCH:
             case OP_TYPE::LA_SWITCH:
             case OP_TYPE::IN_SWITCH:
+			case OP_TYPE::ED_SWITCH:
             {
                 recurse = false;
 
diff --git a/Engine/Plugins/Experimental/Mutable/Source/MutableRuntime/Private/MuR/Operations.cpp b/Engine/Plugins/Experimental/Mutable/Source/MutableRuntime/Private/MuR/Operations.cpp
--- a/Engine/Plugins/Experimental/Mutable/Source/MutableRuntime/Private/MuR/Operations.cpp
+++ b/Engine/Plugins/Experimental/Mutable/Source/MutableRuntime/Private/MuR/Operations.cpp
@@ -50,6 +50,7 @@ namespace mu
 		{ DT_MESH,			true,	{ 0, 0, 0, 0, 0, 0, 0, 0, 0, 0, 0, 0, 0, 0, 0, 0, 0, 0, 0, 0, 0, 0, 0, 0, 0, 0, 0, 0, 0 }		},	// ME_CONDITIONAL
 		{ DT_LAYOUT,		false,  { 0, 0, 0, 0, 0, 0, 0, 0, 0, 0, 0, 0, 0, 0, 0, 0, 0, 0, 0, 0, 0, 0, 0, 0, 0, 0, 0, 0, 0 }		},	// LA_CONDITIONAL
 		{ DT_INSTANCE,		false,	{ 0, 0, 0, 0, 0, 0, 0, 0, 0, 0, 0, 0, 0, 0, 0, 0, 0, 0, 0, 0, 0, 0, 0, 0, 0, 0, 0, 0, 0 }		},	// IN_CONDITIONAL
+		{ DT_EXTENSION_DATA,false,	{ 0, 0, 0, 0, 0, 0, 0, 0, 0, 0, 0, 0, 0, 0, 0, 0, 0, 0, 0, 0, 0, 0, 0, 0, 0, 0, 0, 0, 0 }		},	// ED_CONDITIONAL
 
 		{ DT_INT,			false,	{ 0, 0, 0, 0, 0, 0, 0, 0, 0, 0, 0, 0, 0, 0, 0, 0, 0, 0, 0, 0, 0, 0, 0, 0, 0, 0, 0, 0, 0 }		},	// NU_SWITCH
 		{ DT_SCALAR,		false,	{ 0, 0, 0, 0, 0, 0, 0, 0, 0, 0, 0, 0, 0, 0, 0, 0, 0, 0, 0, 0, 0, 0, 0, 0, 0, 0, 0, 0, 0 }		},	// SC_SWITCH
@@ -58,6 +59,7 @@ namespace mu
 		{ DT_MESH,			true,	{ 0, 0, 0, 0, 0, 0, 0, 0, 0, 0, 0, 0, 0, 0, 0, 0, 0, 0, 0, 0, 0, 0, 0, 0, 0, 0, 0, 0, 0 }		},	// ME_SWITCH
 		{ DT_LAYOUT,		false,	{ 0, 0, 0, 0, 0, 0, 0, 0, 0, 0, 0, 0, 0, 0, 0, 0, 0, 0, 0, 0, 0, 0, 0, 0, 0, 0, 0, 0, 0 }		},	// LA_SWITCH
 		{ DT_INSTANCE,		false,	{ 0, 0, 0, 0, 0, 0, 0, 0, 0, 0, 0, 0, 0, 0, 0, 0, 0, 0, 0, 0, 0, 0, 0, 0, 0, 0, 0, 0, 0 }		},	// IN_SWITCH
+		{ DT_EXTENSION_DATA,false,	{ 0, 0, 0, 0, 0, 0, 0, 0, 0, 0, 0, 0, 0, 0, 0, 0, 0, 0, 0, 0, 0, 0, 0, 0, 0, 0, 0, 0, 0 }		},	// ED_SWITCH
 
 		{ DT_BOOL,			false,  { 0, 0, 0, 0, 0, 0, 0, 0, 0, 0, 0, 0, 0, 0, 0, 0, 0, 0, 0, 0, 0, 0, 0, 0, 0, 0, 0, 0, 0 }		},	// BO_LESS
 		{ DT_BOOL,			false,  { 0, 0, 0, 0, 0, 0, 0, 0, 0, 0, 0, 0, 0, 0, 0, 0, 0, 0, 0, 0, 0, 0, 0, 0, 0, 0, 0, 0, 0 }		},	// BO_EQUAL_SC_CONST
@@ -415,6 +417,7 @@ namespace mu
         case OP_TYPE::ME_CONDITIONAL:
         case OP_TYPE::LA_CONDITIONAL:
         case OP_TYPE::IN_CONDITIONAL:
+		case OP_TYPE::ED_CONDITIONAL:
         {
 			OP::ConditionalArgs args = program.GetOpArgs<OP::ConditionalArgs>(at);
             f(args.condition );
@@ -430,6 +433,7 @@ namespace mu
         case OP_TYPE::ME_SWITCH:
         case OP_TYPE::LA_SWITCH:
         case OP_TYPE::IN_SWITCH:
+		case OP_TYPE::ED_SWITCH:
         {
 			const uint8_t* data = program.GetOpArgsPointer(at);
 			
diff --git a/Engine/Plugins/Experimental/Mutable/Source/MutableRuntime/Private/MuR/Operations.h b/Engine/Plugins/Experimental/Mutable/Source/MutableRuntime/Private/MuR/Operations.h
--- a/Engine/Plugins/Experimental/Mutable/Source/MutableRuntime/Private/MuR/Operations.h
+++ b/Engine/Plugins/Experimental/Mutable/Source/MutableRuntime/Private/MuR/Operations.h
@@ -62,6 +62,7 @@ namespace mu
         ME_CONDITIONAL,
         LA_CONDITIONAL,
         IN_CONDITIONAL,
+		ED_CONDITIONAL,
 
         //! Select one of several values depending on an int input
         NU_SWITCH,
@@ -71,6 +72,7 @@ namespace mu
         ME_SWITCH,
         LA_SWITCH,
         IN_SWITCH,
+		ED_SWITCH,
 
         //-----------------------------------------------------------------------------------------
         // Boolean operations
@@ -1255,6 +1257,8 @@ namespace mu
         case DT_COLOUR: return OP_TYPE::CO_SWITCH;
         case DT_SCALAR: return OP_TYPE::SC_SWITCH;
         case DT_INT: return OP_TYPE::NU_SWITCH;
+		case DT_EXTENSION_DATA: return OP_TYPE::ED_SWITCH;
+
         default:
 			check(false);
 			break;
diff --git a/Engine/Plugins/Experimental/Mutable/Source/MutableTools/Private/MuT/ASTOpConditional.cpp b/Engine/Plugins/Experimental/Mutable/Source/MutableTools/Private/MuT/ASTOpConditional.cpp
--- a/Engine/Plugins/Experimental/Mutable/Source/MutableTools/Private/MuT/ASTOpConditional.cpp
+++ b/Engine/Plugins/Experimental/Mutable/Source/MutableTools/Private/MuT/ASTOpConditional.cpp
@@ -82,6 +82,7 @@ namespace mu
 		case OP_TYPE::ME_CONDITIONAL:
 		case OP_TYPE::LA_CONDITIONAL:
 		case OP_TYPE::IN_CONDITIONAL:
+		case OP_TYPE::ED_CONDITIONAL:
 			break;
 		default:
 			// Unexpected type
diff --git a/Engine/Plugins/Experimental/Mutable/Source/MutableTools/Private/MuT/ASTOpSwitch.cpp b/Engine/Plugins/Experimental/Mutable/Source/MutableTools/Private/MuT/ASTOpSwitch.cpp
--- a/Engine/Plugins/Experimental/Mutable/Source/MutableTools/Private/MuT/ASTOpSwitch.cpp
+++ b/Engine/Plugins/Experimental/Mutable/Source/MutableTools/Private/MuT/ASTOpSwitch.cpp
@@ -75,6 +75,7 @@ namespace mu
 		case OP_TYPE::ME_SWITCH:
 		case OP_TYPE::LA_SWITCH:
 		case OP_TYPE::IN_SWITCH:
+		case OP_TYPE::ED_SWITCH:
 			break;
 		default:
 			// Unexpected type
diff --git a/Engine/Plugins/Experimental/Mutable/Source/MutableTools/Private/MuT/CodeGenerator.h b/Engine/Plugins/Experimental/Mutable/Source/MutableTools/Private/MuT/CodeGenerator.h
--- a/Engine/Plugins/Experimental/Mutable/Source/MutableTools/Private/MuT/CodeGenerator.h
+++ b/Engine/Plugins/Experimental/Mutable/Source/MutableTools/Private/MuT/CodeGenerator.h
@@ -615,6 +615,9 @@ namespace mu
 
 		void GenerateExtensionData(FExtensionDataGenerationResult& OutResult, const NodeExtensionDataPtrConst& InUntypedNode);
 		void GenerateExtensionData_Constant(FExtensionDataGenerationResult& OutResult, const class NodeExtensionDataConstant* Constant);
+		void GenerateExtensionData_Switch(FExtensionDataGenerationResult& OutResult, const class NodeExtensionDataSwitch* Switch);
+		void GenerateExtensionData_Variation(FExtensionDataGenerationResult& OutResult, const class NodeExtensionDataVariation* Variation);
+		Ptr<ASTOp> GenerateMissingExtensionDataCode(const TCHAR* StrWhere, const void* ErrorContext);
 
         //-----------------------------------------------------------------------------------------
         // Projectors
diff --git a/Engine/Plugins/Experimental/Mutable/Source/MutableTools/Private/MuT/CodeGenerator_Color.cpp b/Engine/Plugins/Experimental/Mutable/Source/MutableTools/Private/MuT/CodeGenerator_Color.cpp
--- a/Engine/Plugins/Experimental/Mutable/Source/MutableTools/Private/MuT/CodeGenerator_Color.cpp
+++ b/Engine/Plugins/Experimental/Mutable/Source/MutableTools/Private/MuT/CodeGenerator_Color.cpp
@@ -262,7 +262,7 @@ namespace mu
 
 			if (tagIndex < 0)
 			{
-				FString Msg = FString::Printf(TEXT("Unknown tag found in image variation [%s]."), *FString(tag.c_str()));
+				FString Msg = FString::Printf(TEXT("Unknown tag found in color variation [%s]."), *FString(tag.c_str()));
 				m_pErrorLog->GetPrivate()->Add(Msg, ELMT_WARNING, node.m_errorContext);
 				continue;
 			}
diff --git a/Engine/Plugins/Experimental/Mutable/Source/MutableTools/Private/MuT/CodeGenerator_ExtensionData.cpp b/Engine/Plugins/Experimental/Mutable/Source/MutableTools/Private/MuT/CodeGenerator_ExtensionData.cpp
--- a/Engine/Plugins/Experimental/Mutable/Source/MutableTools/Private/MuT/CodeGenerator_ExtensionData.cpp
+++ b/Engine/Plugins/Experimental/Mutable/Source/MutableTools/Private/MuT/CodeGenerator_ExtensionData.cpp
@@ -12,13 +12,20 @@
 #include "MuT/AST.h"
 #include "MuT/ASTOpConditional.h"
 #include "MuT/ASTOpConstantExtensionData.h"
+#include "MuT/ASTOpSwitch.h"
+#include "MuT/ASTOpConditional.h"
 #include "MuT/CodeGenerator.h"
 #include "MuT/CompilerPrivate.h"
 #include "MuT/ErrorLog.h"
 #include "MuT/ErrorLogPrivate.h"
+#include "MuT/NodeScalar.h"
 #include "MuT/NodeExtensionData.h"
 #include "MuT/NodeExtensionDataConstant.h"
 #include "MuT/NodeExtensionDataConstantPrivate.h"
+#include "MuT/NodeExtensionDataSwitch.h"
+#include "MuT/NodeExtensionDataSwitchPrivate.h"
+#include "MuT/NodeExtensionDataVariation.h"
+#include "MuT/NodeExtensionDataVariationPrivate.h"
 
 
 namespace mu
@@ -51,7 +58,9 @@ class Node;
 		// Generate for each different type of node
 		switch (Node->GetExtensionDataNodeType())
 		{
-			case NodeExtensionData::EType::Constant: GenerateExtensionData_Constant(OutResult, static_cast<const NodeExtensionDataConstant*>(Node)); break;
+			case NodeExtensionData::EType::Constant:  GenerateExtensionData_Constant(OutResult, static_cast<const NodeExtensionDataConstant*>(Node)); break;
+			case NodeExtensionData::EType::Switch:    GenerateExtensionData_Switch(OutResult, static_cast<const NodeExtensionDataSwitch*>(Node)); break;
+			case NodeExtensionData::EType::Variation: GenerateExtensionData_Variation(OutResult, static_cast<const NodeExtensionDataVariation*>(Node)); break;
 			case NodeExtensionData::EType::None: check(false);
 		}
 
@@ -80,5 +89,134 @@ class Node;
 		Op->Value = Data;
 	}
 
+	//---------------------------------------------------------------------------------------------
+	void CodeGenerator::GenerateExtensionData_Switch(FExtensionDataGenerationResult& OutResult, const class NodeExtensionDataSwitch* Switch)
+	{
+		NodeExtensionDataSwitch::Private& Node = *Switch->GetPrivate();
+
+		MUTABLE_CPUPROFILER_SCOPE(NodeExtensionDataSwitch);
+
+		if (Node.Options.IsEmpty())
+		{
+			Ptr<ASTOp> MissingOp = GenerateMissingExtensionDataCode(TEXT("Switch option"), Node.m_errorContext);
+			OutResult.Op = MissingOp;
+			return;
+		}
+
+		Ptr<ASTOpSwitch> Op = new ASTOpSwitch;
+		Op->type = OP_TYPE::ED_SWITCH;
+
+		// Variable name
+		if (Node.Parameter)
+		{
+			Op->variable = Generate(Node.Parameter.get());
+		}
+		else
+		{
+			// This argument is required
+			Op->variable = GenerateMissingScalarCode(TEXT("Switch variable"), 0.0f, Node.m_errorContext);
+		}
+
+		// Options
+		for (int32 OptionIndex = 0; OptionIndex < Node.Options.Num(); ++OptionIndex)
+		{
+			Ptr<ASTOp> Branch;
+			if (Node.Options[OptionIndex])
+			{
+				NodeExtensionDataPtr Option = Node.Options[OptionIndex];
+
+				FExtensionDataGenerationResult OptionResult;
+				GenerateExtensionData(OptionResult, Option);
+
+				Branch = OptionResult.Op;
+			}
+			else
+			{
+				// This argument is required
+				Branch = GenerateMissingExtensionDataCode(TEXT("Switch option"), Node.m_errorContext);
+			}
+			Op->cases.Emplace(static_cast<int16_t>(OptionIndex), Op, Branch);
+		}
+
+		OutResult.Op = Op;
+	}
+
+	//---------------------------------------------------------------------------------------------
+	void CodeGenerator::GenerateExtensionData_Variation(FExtensionDataGenerationResult& OutResult, const class NodeExtensionDataVariation* Variation)
+	{
+		const NodeExtensionDataVariation::Private& Node = *Variation->GetPrivate();
+
+		Ptr<ASTOp> CurrentOp;
+
+		// Default case
+		if (Node.DefaultValue)
+		{
+			FExtensionDataGenerationResult DefaultResult;
+			GenerateExtensionData(DefaultResult, Node.DefaultValue);
+
+			CurrentOp = DefaultResult.Op;
+		}
+		else
+		{
+			// This argument is required
+			CurrentOp = GenerateMissingExtensionDataCode(TEXT("Variation default"), Node.m_errorContext);
+		}
+
+		// Process variations in reverse order, since conditionals are built bottom-up.
+		for (int32 VariationIndex = Node.Variations.Num() - 1; VariationIndex >= 0; --VariationIndex)
+		{
+			const string& Tag = Node.Variations[VariationIndex].Tag;
+			const int32 TagIndex = m_firstPass.m_tags.IndexOfByPredicate([Tag](const FirstPassGenerator::TAG& CandidateTag)
+			{
+				return CandidateTag.tag == Tag;
+			});
+
+			if (TagIndex == INDEX_NONE)
+			{
+				const FString Msg = FString::Printf(TEXT("Unknown tag found in Extension Data variation [%s]"), *FString{ Tag.c_str() });
+				m_pErrorLog->GetPrivate()->Add(Msg, ELMT_WARNING, Node.m_errorContext);
+				continue;
+			}
+
+			Ptr<ASTOp> VariationOp;
+			if (NodeExtensionDataPtr VariationValue = Node.Variations[VariationIndex].Value)
+			{
+				FExtensionDataGenerationResult VariationResult;
+				GenerateExtensionData(VariationResult, VariationValue);
+
+				VariationOp = VariationResult.Op;
+			}
+			else
+			{
+				// This argument is required
+				VariationOp = GenerateMissingExtensionDataCode(TEXT("Variation option"), Node.m_errorContext);
+			}
+
+			Ptr<ASTOpConditional> Conditional = new ASTOpConditional;
+			Conditional->type = OP_TYPE::ED_CONDITIONAL;
+			Conditional->no = CurrentOp;
+			Conditional->yes = VariationOp;
+			Conditional->condition = m_firstPass.m_tags[TagIndex].genericCondition;
+
+			CurrentOp = Conditional;
+		}
+
+		OutResult.Op = CurrentOp;
+	}
+
+	Ptr<ASTOp> CodeGenerator::GenerateMissingExtensionDataCode(const TCHAR* StrWhere, const void* ErrorContext)
+	{
+		// Log a warning
+		const FString Msg = FString::Printf(TEXT("Required connection not found: %s"), StrWhere);
+		m_pErrorLog->GetPrivate()->Add(Msg, ELMT_ERROR, ErrorContext);
+
+		// Create a constant extension data
+		NodeExtensionDataConstantPtrConst Node = new NodeExtensionDataConstant;
+
+		FExtensionDataGenerationResult Result;
+		GenerateExtensionData_Constant(Result, Node.get());
+
+		return Result.Op;
+	}
 }
 
diff --git a/Engine/Plugins/Experimental/Mutable/Source/MutableTools/Private/MuT/DataPacker.cpp b/Engine/Plugins/Experimental/Mutable/Source/MutableTools/Private/MuT/DataPacker.cpp
--- a/Engine/Plugins/Experimental/Mutable/Source/MutableTools/Private/MuT/DataPacker.cpp
+++ b/Engine/Plugins/Experimental/Mutable/Source/MutableTools/Private/MuT/DataPacker.cpp
@@ -173,6 +173,7 @@ namespace mu
                 case OP_TYPE::ME_CONDITIONAL:
                 case OP_TYPE::LA_CONDITIONAL:
                 case OP_TYPE::IN_CONDITIONAL:
+				case OP_TYPE::ED_CONDITIONAL:
                 {
 					OP::ConditionalArgs args = program.GetOpArgs<OP::ConditionalArgs>(at);
 
@@ -203,6 +204,7 @@ namespace mu
                 case OP_TYPE::IM_SWITCH:
                 case OP_TYPE::ME_SWITCH:
                 case OP_TYPE::LA_SWITCH:
+				case OP_TYPE::ED_SWITCH:
                 {
 					const uint8_t* data = program.GetOpArgsPointer(at);
 					
diff --git a/Engine/Plugins/Experimental/Mutable/Source/MutableTools/Private/MuT/ErrorLog.cpp b/Engine/Plugins/Experimental/Mutable/Source/MutableTools/Private/MuT/ErrorLog.cpp
--- a/Engine/Plugins/Experimental/Mutable/Source/MutableTools/Private/MuT/ErrorLog.cpp
+++ b/Engine/Plugins/Experimental/Mutable/Source/MutableTools/Private/MuT/ErrorLog.cpp
@@ -234,6 +234,7 @@ namespace mu
 		TEXT("ME_CONDITIONAL   "),
 		TEXT("LA_CONDITIONAL   "),
 		TEXT("IN_CONDITIONAL   "),
+		TEXT("ED_CONDITIONAL   "),
 		
 		TEXT("NU_SWITCH        "),
 		TEXT("SC_SWITCH        "),
@@ -242,6 +243,7 @@ namespace mu
 		TEXT("ME_SWITCH        "),
 		TEXT("LA_SWITCH        "),
 		TEXT("IN_SWITCH        "),
+		TEXT("ED_SWITCH        "),
 
 		TEXT("BO_LESS          "),
 		TEXT("BO_EQUAL_SC_CONST"),
diff --git a/Engine/Plugins/Experimental/Mutable/Source/MutableTools/Private/MuT/NodeExtensionData.cpp b/Engine/Plugins/Experimental/Mutable/Source/MutableTools/Private/MuT/NodeExtensionData.cpp
--- a/Engine/Plugins/Experimental/Mutable/Source/MutableTools/Private/MuT/NodeExtensionData.cpp
+++ b/Engine/Plugins/Experimental/Mutable/Source/MutableTools/Private/MuT/NodeExtensionData.cpp
@@ -5,6 +5,8 @@
 #include "MuR/Serialisation.h"
 #include "MuR/SerialisationPrivate.h"
 #include "MuT/NodeExtensionDataConstant.h"
+#include "MuT/NodeExtensionDataSwitch.h"
+#include "MuT/NodeExtensionDataVariation.h"
 
 namespace mu
 {
@@ -43,7 +45,8 @@ namespace mu
 		switch (EType(Id))
 		{
 			case EType::Constant:		return NodeExtensionDataConstant::StaticUnserialise(Archive); break;
-			
+			case EType::Switch:			return NodeExtensionDataSwitch::StaticUnserialise(Archive); break;
+			case EType::Variation:		return NodeExtensionDataVariation::StaticUnserialise(Archive); break;
 			default: check(false);
 		}
 
diff --git a/Engine/Plugins/Experimental/Mutable/Source/MutableTools/Private/MuT/NodeExtensionDataSwitch.cpp b/Engine/Plugins/Experimental/Mutable/Source/MutableTools/Private/MuT/NodeExtensionDataSwitch.cpp
new file mode 100644
--- /dev/null
+++ b/Engine/Plugins/Experimental/Mutable/Source/MutableTools/Private/MuT/NodeExtensionDataSwitch.cpp
@@ -0,0 +1,104 @@
+// Copyright Epic Games, Inc. All Rights Reserved.
+
+#include "MuT/NodeExtensionDataSwitch.h"
+
+#include "Misc/AssertionMacros.h"
+#include "MuT/NodeExtensionDataSwitchPrivate.h"
+// #include "MuT/NodePrivate.h"
+#include "MuT/NodeScalar.h"
+
+namespace mu
+{
+	//---------------------------------------------------------------------------------------------
+	// Static initialisation
+	//---------------------------------------------------------------------------------------------
+	NODE_TYPE NodeExtensionDataSwitch::Private::s_type =
+		NODE_TYPE("ExtensionDataSwitch", NodeExtensionData::GetStaticType());
+
+	//---------------------------------------------------------------------------------------------
+	//!
+	//---------------------------------------------------------------------------------------------
+
+	MUTABLE_IMPLEMENT_NODE(NodeExtensionDataSwitch, EType::Switch, Node, Node::EType::ExtensionData);
+
+	//---------------------------------------------------------------------------------------------
+	// Node Interface
+	//---------------------------------------------------------------------------------------------
+
+	//---------------------------------------------------------------------------------------------
+	int NodeExtensionDataSwitch::GetInputCount() const
+	{
+		return 1 + m_pD->Options.Num();
+	}
+
+	//---------------------------------------------------------------------------------------------
+	Node* NodeExtensionDataSwitch::GetInputNode(int i) const
+	{
+		check(i >=0 && i < GetInputCount());
+
+
+		Node* Result = 0;
+		switch (i)
+		{
+			case 0:
+				Result = m_pD->Parameter.get();
+				break;
+
+			default:
+				Result = m_pD->Options[i - 1].get();
+				break;
+		}
+
+		return Result;
+	}
+
+	//---------------------------------------------------------------------------------------------
+	void NodeExtensionDataSwitch::SetInputNode(int i, NodePtr pNode)
+	{
+		check(i >= 0 && i < GetInputCount());
+
+		switch (i)
+		{
+			case 0:
+				m_pD->Parameter = dynamic_cast<NodeScalar*>(pNode.get());
+				break;
+
+			default:
+				m_pD->Options[i - 1] = dynamic_cast<NodeExtensionData*>(pNode.get());
+				break;
+		}
+	}
+
+	//---------------------------------------------------------------------------------------------
+	NodeScalarPtr NodeExtensionDataSwitch::GetParameter() const
+	{
+		return m_pD->Parameter.get();
+	}
+
+	//---------------------------------------------------------------------------------------------
+	void NodeExtensionDataSwitch::SetParameter(NodeScalarPtr pNode)
+	{
+		m_pD->Parameter = pNode;
+	}
+
+	//---------------------------------------------------------------------------------------------
+	void NodeExtensionDataSwitch::SetOptionCount(int InNumOptions)
+	{
+		check(InNumOptions >= 0);
+		m_pD->Options.SetNum(InNumOptions);
+	}
+
+	//---------------------------------------------------------------------------------------------
+	NodeExtensionDataPtr NodeExtensionDataSwitch::GetOption(int t) const
+	{
+		check(m_pD->Options.IsValidIndex(t));
+		return m_pD->Options[t].get();
+	}
+
+	//---------------------------------------------------------------------------------------------
+	void NodeExtensionDataSwitch::SetOption(int t, NodeExtensionDataPtr pNode)
+	{
+		check(m_pD->Options.IsValidIndex(t));
+		m_pD->Options[t] = pNode;
+	}
+}
\ No newline at end of file
diff --git a/Engine/Plugins/Experimental/Mutable/Source/MutableTools/Private/MuT/NodeExtensionDataSwitchPrivate.h b/Engine/Plugins/Experimental/Mutable/Source/MutableTools/Private/MuT/NodeExtensionDataSwitchPrivate.h
new file mode 100644
--- /dev/null
+++ b/Engine/Plugins/Experimental/Mutable/Source/MutableTools/Private/MuT/NodeExtensionDataSwitchPrivate.h
@@ -0,0 +1,43 @@
+// Copyright Epic Games, Inc. All Rights Reserved.
+
+#pragma once
+
+#include "MuR/ExtensionData.h"
+#include "MuT/NodeExtensionDataSwitch.h"
+#include "MuT/NodePrivate.h"
+
+namespace mu
+{
+	class NodeExtensionDataSwitch::Private : public Node::Private
+	{
+	public:
+
+		MUTABLE_DEFINE_CONST_VISITABLE()
+
+	public:
+
+		static NODE_TYPE s_type;
+
+		NodeScalarPtr Parameter;
+		TArray<NodeExtensionDataPtr> Options;
+
+		void Serialise(OutputArchive& Archive) const
+		{
+			uint32_t Version = 0;
+			Archive << Version;
+
+			Archive << Parameter;
+			Archive << Options;
+		}
+
+		void Unserialise(InputArchive& Archive)
+		{
+			uint32 Version;
+			Archive >> Version;
+			check(Version == 0);
+
+			Archive >> Parameter;
+			Archive >> Options;
+		}
+	};
+}
\ No newline at end of file
diff --git a/Engine/Plugins/Experimental/Mutable/Source/MutableTools/Private/MuT/NodeExtensionDataVariation.cpp b/Engine/Plugins/Experimental/Mutable/Source/MutableTools/Private/MuT/NodeExtensionDataVariation.cpp
new file mode 100644
--- /dev/null
+++ b/Engine/Plugins/Experimental/Mutable/Source/MutableTools/Private/MuT/NodeExtensionDataVariation.cpp
@@ -0,0 +1,111 @@
+// Copyright Epic Games, Inc. All Rights Reserved.
+
+#include "MuT/NodeExtensionDataVariation.h"
+
+#include "Misc/AssertionMacros.h"
+#include "MuT/NodeExtensionDataVariationPrivate.h"
+
+namespace mu
+{
+	//---------------------------------------------------------------------------------------------
+	// Static initialisation
+	//---------------------------------------------------------------------------------------------
+	NODE_TYPE NodeExtensionDataVariation::Private::s_type =
+		NODE_TYPE("ExtensionDataVariation", NodeExtensionData::GetStaticType());
+
+	//---------------------------------------------------------------------------------------------
+	//!
+	//---------------------------------------------------------------------------------------------
+	MUTABLE_IMPLEMENT_NODE(NodeExtensionDataVariation, EType::Variation, Node, Node::EType::ExtensionData);
+
+	//---------------------------------------------------------------------------------------------
+	// Node Interface
+	//---------------------------------------------------------------------------------------------
+
+	//---------------------------------------------------------------------------------------------
+	int NodeExtensionDataVariation::GetInputCount() const
+	{
+		return 1 + m_pD->Variations.Num();
+	}
+
+	//---------------------------------------------------------------------------------------------
+	Node* NodeExtensionDataVariation::GetInputNode(int i) const
+	{
+		check(i >= 0 && i < GetInputCount());
+
+		if (i == 0)
+		{
+			return m_pD->DefaultValue.get();
+		}
+
+		i -= 1;
+
+		if (m_pD->Variations.IsValidIndex(i))
+		{
+			return m_pD->Variations[i].Value.get();
+		}
+
+		return nullptr;
+	}
+
+	//---------------------------------------------------------------------------------------------
+	void NodeExtensionDataVariation::SetInputNode(int i, NodePtr pNode)
+	{
+		check(i >= 0 && i < GetInputCount());
+
+		if (i == 0)
+		{
+			m_pD->DefaultValue = dynamic_cast<NodeExtensionData*>(pNode.get());
+			return;
+		}
+
+		i -= 1;
+
+		if (m_pD->Variations.IsValidIndex(i))
+		{
+			m_pD->Variations[i].Value = dynamic_cast<NodeExtensionData*>(pNode.get());
+			return;
+		}
+	}
+
+	//---------------------------------------------------------------------------------------------
+	// Own Interface
+	//---------------------------------------------------------------------------------------------
+
+	//---------------------------------------------------------------------------------------------
+	void NodeExtensionDataVariation::SetDefaultValue(NodeExtensionDataPtr InValue)
+	{
+		m_pD->DefaultValue = InValue;
+	}
+
+	//---------------------------------------------------------------------------------------------
+	void NodeExtensionDataVariation::SetVariationCount(int InCount)
+	{
+		check(InCount >= 0);
+		m_pD->Variations.SetNum(InCount);
+	}
+
+	//---------------------------------------------------------------------------------------------
+	int NodeExtensionDataVariation::GetVariationCount() const
+	{
+		return m_pD->Variations.Num();
+	}
+
+	//---------------------------------------------------------------------------------------------
+	void NodeExtensionDataVariation::SetVariationTag(int InIndex, const char* InStrTag)
+	{
+		check(m_pD->Variations.IsValidIndex(InIndex));
+		check(InStrTag);
+
+		m_pD->Variations[InIndex].Tag = InStrTag;
+	}
+
+	//---------------------------------------------------------------------------------------------
+	void NodeExtensionDataVariation::SetVariationValue(int InIndex, NodeExtensionDataPtr InValue)
+	{
+		check(m_pD->Variations.IsValidIndex(InIndex));
+
+		m_pD->Variations[InIndex].Value = InValue;
+	}
+
+}
diff --git a/Engine/Plugins/Experimental/Mutable/Source/MutableTools/Private/MuT/NodeExtensionDataVariationPrivate.h b/Engine/Plugins/Experimental/Mutable/Source/MutableTools/Private/MuT/NodeExtensionDataVariationPrivate.h
new file mode 100644
--- /dev/null
+++ b/Engine/Plugins/Experimental/Mutable/Source/MutableTools/Private/MuT/NodeExtensionDataVariationPrivate.h
@@ -0,0 +1,71 @@
+// Copyright Epic Games, Inc. All Rights Reserved.
+
+#pragma once
+
+#include "MuT/NodeExtensionDataVariation.h"
+#include "MuT/NodePrivate.h"
+
+namespace mu
+{
+	class NodeExtensionDataVariation::Private : public Node::Private
+	{
+	public:
+
+		MUTABLE_DEFINE_CONST_VISITABLE()
+
+	public:
+
+		static NODE_TYPE s_type;
+
+		NodeExtensionDataPtr DefaultValue;
+
+		struct VARIATION
+		{
+			NodeExtensionDataPtr Value;
+			string Tag;
+
+			//!
+			void Serialise(OutputArchive& InArch) const
+			{
+				uint32_t Ver = 1;
+				InArch << Ver;
+
+				InArch << Tag;
+				InArch << Value;
+			}
+
+			void Unserialise(InputArchive& InArch)
+			{
+				uint32_t Ver = 0;
+				InArch >> Ver;
+				check(Ver == 1);
+
+				InArch >> Tag;
+				InArch >> Value;
+			}
+		};
+
+		TArray<VARIATION> Variations;
+
+		//!
+		void Serialise(OutputArchive& InArch) const
+		{
+			uint32_t Ver = 1;
+			InArch << Ver;
+
+			InArch << DefaultValue;
+			InArch << Variations;
+		}
+
+		//!
+		void Unserialise(InputArchive& InArch)
+		{
+			uint32_t Ver;
+			InArch >> Ver;
+			check(Ver == 1);
+
+			InArch >> DefaultValue;
+			InArch >> Variations;
+		}
+	};
+}
\ No newline at end of file
diff --git a/Engine/Plugins/Experimental/Mutable/Source/MutableTools/Private/MuT/NodeObjectNew.cpp b/Engine/Plugins/Experimental/Mutable/Source/MutableTools/Private/MuT/NodeObjectNew.cpp
--- a/Engine/Plugins/Experimental/Mutable/Source/MutableTools/Private/MuT/NodeObjectNew.cpp
+++ b/Engine/Plugins/Experimental/Mutable/Source/MutableTools/Private/MuT/NodeObjectNew.cpp
@@ -63,7 +63,17 @@ namespace mu
 		else
 		{
 			i -= (int)m_pD->m_lods.Num();
-			pResult = m_pD->m_children[i].get();
+			
+			if ( i<(int) m_pD->m_children.Num() )
+			{
+				pResult = m_pD->m_children[i].get();
+			}
+			else
+			{
+				i -= (int)m_pD->m_children.Num();
+
+				pResult = m_pD->m_extensionDataNodes[i].Node.get();
+			}
 		}
 
 		return pResult;
@@ -82,7 +92,15 @@ namespace mu
 		else
 		{
 			i -= (int)m_pD->m_lods.Num();
-			m_pD->m_children[i] = dynamic_cast<NodeObject*>( pNode.get() );
+
+			if (i < (int) m_pD->m_children.Num())
+			{
+				m_pD->m_children[i] = dynamic_cast<NodeObject*>( pNode.get() );
+			}
+			else
+			{
+				checkf(false, TEXT("Trying to set an input node for an extension input data. Please use AddExtensionDataNode for this"));
+			}
 		}
 	}
 
diff --git a/Engine/Plugins/Experimental/Mutable/Source/MutableTools/Public/MuT/NodeExtensionData.h b/Engine/Plugins/Experimental/Mutable/Source/MutableTools/Public/MuT/NodeExtensionData.h
--- a/Engine/Plugins/Experimental/Mutable/Source/MutableTools/Public/MuT/NodeExtensionData.h
+++ b/Engine/Plugins/Experimental/Mutable/Source/MutableTools/Public/MuT/NodeExtensionData.h
@@ -26,7 +26,8 @@ public:
 	enum class EType : uint8
 	{
 		Constant = 0,
-
+		Switch,
+		Variation,
 		None
 	};
 
diff --git a/Engine/Plugins/Experimental/Mutable/Source/MutableTools/Public/MuT/NodeExtensionDataSwitch.h b/Engine/Plugins/Experimental/Mutable/Source/MutableTools/Public/MuT/NodeExtensionDataSwitch.h
new file mode 100644
--- /dev/null
+++ b/Engine/Plugins/Experimental/Mutable/Source/MutableTools/Public/MuT/NodeExtensionDataSwitch.h
@@ -0,0 +1,71 @@
+// Copyright Epic Games, Inc. All Rights Reserved.
+
+#pragma once
+
+#include "MuR/ExtensionData.h"
+#include "MuR/Ptr.h"
+#include "MuR/RefCounted.h"
+#include "MuT/NodeExtensionData.h"
+
+namespace mu
+{
+	// Forward declarations
+	class NodeScalar;
+	using NodeScalarPtr = Ptr<NodeScalar>;
+	using NodeScalarPtrConst = Ptr<const NodeScalar>;
+
+	class NodeExtensionDataSwitch;
+	using NodeExtensionDataSwitchPtr = Ptr<NodeExtensionDataSwitch>;
+	using NodeExtensionDataSwitchPtrConst = Ptr<const NodeExtensionDataSwitch>;
+
+	class MUTABLETOOLS_API NodeExtensionDataSwitch : public NodeExtensionData
+	{
+	public:
+
+		//-----------------------------------------------------------------------------------------
+		// Life cycle
+		//-----------------------------------------------------------------------------------------
+		NodeExtensionDataSwitch();
+
+		void SerialiseWrapper(OutputArchive& arch) const override;
+		static void Serialise(const NodeExtensionDataSwitch* pNode, OutputArchive& arch);
+		static NodeExtensionDataSwitchPtr StaticUnserialise(InputArchive& arch);
+
+		//-----------------------------------------------------------------------------------------
+		// Node Interface
+		//-----------------------------------------------------------------------------------------
+		const NODE_TYPE* GetType() const override;
+		static const NODE_TYPE* GetStaticType();
+
+		int GetInputCount() const override;
+		Node* GetInputNode(int i) const override;
+		void SetInputNode(int i, NodePtr pNode) override;
+
+		//-----------------------------------------------------------------------------------------
+		// Own Interface
+		//-----------------------------------------------------------------------------------------
+		NodeScalarPtr GetParameter() const;
+		void SetParameter(NodeScalarPtr InParameter);
+
+		void SetOptionCount(int);
+
+		NodeExtensionDataPtr GetOption(int t) const;
+		void SetOption(int t, NodeExtensionDataPtr);
+
+		//-----------------------------------------------------------------------------------------
+		// Interface pattern
+		//-----------------------------------------------------------------------------------------
+		class Private;
+		Private* GetPrivate() const;
+		Node::Private* GetBasePrivate() const override;
+
+	protected:
+
+		//! Forbidden. Manage with the Ptr<> template.
+		~NodeExtensionDataSwitch();
+
+	private:
+
+		Private* m_pD;
+	};
+}
\ No newline at end of file
diff --git a/Engine/Plugins/Experimental/Mutable/Source/MutableTools/Public/MuT/NodeExtensionDataVariation.h b/Engine/Plugins/Experimental/Mutable/Source/MutableTools/Public/MuT/NodeExtensionDataVariation.h
new file mode 100644
--- /dev/null
+++ b/Engine/Plugins/Experimental/Mutable/Source/MutableTools/Public/MuT/NodeExtensionDataVariation.h
@@ -0,0 +1,68 @@
+// Copyright Epic Games, Inc. All Rights Reserved.
+
+#pragma once
+
+#include "MuR/ExtensionData.h"
+#include "MuR/Ptr.h"
+#include "MuR/RefCounted.h"
+#include "MuT/NodeExtensionData.h"
+
+namespace mu
+{
+	// Forward declarations
+	class NodeExtensionDataVariation;
+	using NodeExtensionDataVariationPtr = Ptr<NodeExtensionDataVariation>;
+	using NodeExtensionDataVariationPtrConst = Ptr<const NodeExtensionDataVariation>;
+
+	class MUTABLETOOLS_API NodeExtensionDataVariation : public NodeExtensionData
+	{
+	public:
+
+		//-----------------------------------------------------------------------------------------
+		// Life cycle
+		//-----------------------------------------------------------------------------------------
+		NodeExtensionDataVariation();
+
+		void SerialiseWrapper(OutputArchive& arch) const override;
+		static void Serialise(const NodeExtensionDataVariation* pNode, OutputArchive& arch);
+		static NodeExtensionDataVariationPtr StaticUnserialise(InputArchive& arch);
+
+		//-----------------------------------------------------------------------------------------
+		// Node Interface
+		//-----------------------------------------------------------------------------------------
+		const NODE_TYPE* GetType() const override;
+		static const NODE_TYPE* GetStaticType();
+
+		int GetInputCount() const override;
+		Node* GetInputNode(int i) const override;
+		void SetInputNode(int i, NodePtr pNode) override;
+
+		//-----------------------------------------------------------------------------------------
+		// Own Interface
+		//-----------------------------------------------------------------------------------------
+		void SetDefaultValue(NodeExtensionDataPtr InValue);
+
+		void SetVariationCount(int InCount);
+		int GetVariationCount() const;
+
+		void SetVariationTag(int InIndex, const char* InStrTag);
+
+		void SetVariationValue(int InIndex, NodeExtensionDataPtr InValue);
+
+		//-----------------------------------------------------------------------------------------
+		// Interface pattern
+		//-----------------------------------------------------------------------------------------
+		class Private;
+		Private* GetPrivate() const;
+		Node::Private* GetBasePrivate() const override;
+
+	protected:
+
+		//! Forbidden. Manage with the Ptr<> template.
+		~NodeExtensionDataVariation();
+
+	private:
+
+		Private* m_pD;
+	};
+}
\ No newline at end of file
