[mutable] Minor ASTC compression code cleanup and generalization to prepare for more formats.
#rnx

diff --git a/Engine/Plugins/Experimental/Mutable/Source/MutableRuntime/Private/MuR/BlockCompression/Miro/Miro.cpp b/Engine/Plugins/Experimental/Mutable/Source/MutableRuntime/Private/MuR/BlockCompression/Miro/Miro.cpp
--- a/Engine/Plugins/Experimental/Mutable/Source/MutableRuntime/Private/MuR/BlockCompression/Miro/Miro.cpp
+++ b/Engine/Plugins/Experimental/Mutable/Source/MutableRuntime/Private/MuR/BlockCompression/Miro/Miro.cpp
@@ -19,7 +19,7 @@
 #endif
 
 
-union MIRO_ALIGN miro_pixel_block
+union MIRO_ALIGN FMiroPixelBlock4x4
 {
 	uint8 bytes[4 * 4 * 4];
 };
@@ -44,8 +44,8 @@ namespace impl
 	//-------------------------------------------------------------------------------------------------
 	inline bool BC1BlockHasAlpha(uint8 const* bytes)
 	{
-		int a = ((int)bytes[0] << 0) + ((int)bytes[1] << 8);
-		int b = ((int)bytes[2] << 0) + ((int)bytes[3] << 8);
+		int32 a = (int32(bytes[0]) << 0) + (int32(bytes[1]) << 8);
+		int32 b = (int32(bytes[2]) << 0) + (int32(bytes[3]) << 8);
 		return a <= b;
 	}
 
@@ -2113,25 +2113,23 @@ namespace impl
 #if MIRO_INCLUDE_ASTC
 
 	//-------------------------------------------------------------------------------------------------
-	// Based on ASTCRT, under the BSD 3-Clause License
+	// Based on ASTCRT, under the BSD 3-Clause License.
+	// Modified for additional format support (alpha), and integration with UE platform libs.
 	// See: [REDACTED_PROFILE_URL]/astcrt
 	//-------------------------------------------------------------------------------------------------
-	inline void DCHECK(bool x)
+	inline constexpr void DCHECK(bool x)
 	{
 		(void)x;
-		// if (!x)
-		// {
-		//     check(x);
-		// }
+		 //if (!x)
+		 //{
+		 //    check(x);
+		 //}
 	}
 
 
 	namespace astcrt
 	{
 		constexpr int APPROX_COLOR_EPSILON = 50;
-		constexpr size_t BLOCK_WIDTH = 4;
-		constexpr size_t BLOCK_HEIGHT = 4;
-		constexpr size_t BLOCK_TEXEL_COUNT = BLOCK_WIDTH * BLOCK_HEIGHT;
 		constexpr size_t BLOCK_BYTES = 16;
 
 		constexpr size_t MAXIMUM_ENCODED_WEIGHT_BITS = 96;
@@ -2746,7 +2744,7 @@ namespace impl
 			{{102,103,70,71,38},{110,111,78,79,46},{118,119,86,87,54},{126,127,94,95,62},{39,47,55,63,31}}
 		};
 
-		const int8_t color_endpoint_range_table[2][12][16] =
+		constexpr int8 color_endpoint_range_table[2][12][16] =
 		{
 			{
 				{-1,-1,-1,-1,-1,-1,-1,-1,-1,-1,-1,-1,-1,-1,-1,-1},
@@ -3573,10 +3571,8 @@ namespace impl
 			}
 		}
 
-		inline void integer_sequence_encode(const uint8* numbers,
-			size_t count,
-			range_t range,
-			uint8* output) {
+		inline void integer_sequence_encode(const uint8* numbers, size_t count, range_t range, uint8* output) 
+		{
 			integer_sequence_encode(numbers, count, range, bitwriter(output));
 		}
 
@@ -3584,7 +3580,8 @@ namespace impl
 		 * Compute the number of bits required to store a number of items in a specific
 		 * range using the binary integer sequence encoding.
 		 */
-		inline size_t compute_ise_bitcount(size_t items, range_t range) {
+		inline size_t compute_ise_bitcount(size_t items, range_t range) 
+		{
 			size_t bits = bits_trits_quints_table[range][0];
 			size_t trits = bits_trits_quints_table[range][1];
 			size_t quints = bits_trits_quints_table[range][2];
@@ -3660,9 +3657,6 @@ namespace impl
 		}
 
 
-
-
-
 		int color_channel_sum(vec3i_t color)
 		{
 			return color.r + color.g + color.b;
@@ -3782,7 +3776,7 @@ namespace impl
 			}
 		}
 
-
+		template<int32 BLOCK_WIDTH, int32 BLOCK_HEIGHT>
 		inline void symbolic_to_physical
 		(
 			color_endpoint_mode_t color_endpoint_mode,
@@ -3805,7 +3799,7 @@ namespace impl
 			DCHECK(endpoint_quant < RANGE_MAX);
 			DCHECK(color_endpoint_mode < CEM_MAX);
 			DCHECK(partition_count == 1);
-			DCHECK(compute_ise_bitcount(BLOCK_TEXEL_COUNT, weight_quant) < MAXIMUM_ENCODED_WEIGHT_BITS);
+			DCHECK(compute_ise_bitcount(BLOCK_WIDTH*BLOCK_HEIGHT, weight_quant) < MAXIMUM_ENCODED_WEIGHT_BITS);
 
 			//        if (s_debuglog)
 			//        {
@@ -3827,8 +3821,9 @@ namespace impl
 			size_t r = 0;
 			if (weight_quant < RANGE_32 + 1) r = r_table[weight_quant];
 
-			// Use the first row of Table 11 in the ASTC specification. Beware that
-			// this has to be changed if another block-size is used.
+			// Use the first row of Table C.2.8 in the ASTC specification. Beware that this has to be changed if another block-size is used.
+			// Actually weight-grid and block sizes don't need to match. We are defining the weight-grid size here.
+			// \TODO: change to support different block sizes (jordi)
 			size_t a = m - 2;
 			size_t b = n - 4;
 
@@ -3865,7 +3860,7 @@ namespace impl
 
 			if (dual_plane)
 			{
-				size_t bits_for_weights = compute_ise_bitcount(2 * BLOCK_TEXEL_COUNT, weight_quant);
+				size_t bits_for_weights = compute_ise_bitcount(2 * BLOCK_WIDTH * BLOCK_HEIGHT, weight_quant);
 
 				size_t secondPlaneChannel = 3;
 				orbits8_ptr(pb->data, 128 - bits_for_weights - 2, secondPlaneChannel, 2);
@@ -3911,21 +3906,6 @@ namespace impl
 			return size_t(FMath::Max(0, FMath::Min(int(k), int(t) - int(m))) * 1024) / k;
 		}
 
-		template<typename T> T clamp(T a, T b, T x)
-		{
-			if (x < a)
-			{
-				return a;
-			}
-
-			if (x > b)
-			{
-				return b;
-			}
-
-			return x;
-		}
-
 
 		/**
 		 * Project a texel to a line and quantize the result in 3 dimensions.
@@ -3933,15 +3913,16 @@ namespace impl
 		size_t project(vec3i_t k, int kk, vec3i_t m, vec3i_t t)
 		{
 			DCHECK(kk > 0);
-			return static_cast<size_t>(clamp(0, 1024, dot(t - m, k) * 1024 / kk));
+			return static_cast<size_t>(FMath::Clamp(dot(t - m, k) * 1024 / kk, 0, 1024));
 		}
 
 		size_t project(vec4i_t k, int kk, vec4i_t m, vec4i_t t)
 		{
 			DCHECK(kk > 0);
-			return static_cast<size_t>(clamp(0, 1024, dot(t - m, k) * 1024 / kk));
+			return static_cast<size_t>(FMath::Clamp(dot(t - m, k) * 1024 / kk, 0, 1024));
 		}
 
+		template<int32 BLOCK_TEXEL_COUNT>
 		void calculate_quantized_weights_luminance(const uint8 texels[BLOCK_TEXEL_COUNT],
 			range_t quant,
 			uint8 l0,
@@ -3971,6 +3952,7 @@ namespace impl
 		}
 
 
+		template<int32 BLOCK_TEXEL_COUNT>
 		void calculate_quantized_weights_channel(const unorm8_t texels[BLOCK_TEXEL_COUNT],
 			range_t quant,
 			uint8 l0,
@@ -4001,6 +3983,7 @@ namespace impl
 		}
 
 
+		template<int32 BLOCK_TEXEL_COUNT>
 		void calculate_quantized_weights_rgb(const unorm8_t texels[BLOCK_TEXEL_COUNT],
 			range_t quant,
 			vec3i_t e0,
@@ -4027,6 +4010,7 @@ namespace impl
 			}
 		}
 
+		template<int32 BLOCK_TEXEL_COUNT>
 		void calculate_quantized_weights_rgba(const unorm8_t texels[BLOCK_TEXEL_COUNT],
 			range_t quant,
 			vec4i_t e0,
@@ -4054,20 +4038,16 @@ namespace impl
 		}
 
 
-		range_t endpoint_quantization(size_t partitions,
-			range_t weight_quant,
-			color_endpoint_mode_t endpoint_mode)
+		constexpr range_t endpoint_quantization(size_t partitions, range_t weight_quant, color_endpoint_mode_t endpoint_mode)
 		{
-			int8_t ce_range = color_endpoint_range_table[partitions - 1][weight_quant][endpoint_mode];
+			int8 ce_range = color_endpoint_range_table[partitions - 1][weight_quant][endpoint_mode];
 			DCHECK(ce_range >= 0 && ce_range <= RANGE_MAX);
 			return static_cast<range_t>(ce_range);
 		}
 
 
 
-		/**
-		 * Write void extent block bits for LDR mode and unused extent coordinates.
-		 */
+		/** Write void extent block bits for LDR mode and unused extent coordinates. */
 		void encode_void_extent(vec3i_t color, PhysicalBlock* physical_block)
 		{
 			void_extent_to_physical(unorm8_to_unorm16(to_unorm8(color)), physical_block);
@@ -4080,7 +4060,8 @@ namespace impl
 		}
 
 
-		void encode_luminance(const uint8 texels[BLOCK_TEXEL_COUNT],
+		template<int32 BLOCK_WIDTH, int32 BLOCK_HEIGHT>
+		void encode_luminance(const uint8 texels[BLOCK_WIDTH*BLOCK_HEIGHT],
 			PhysicalBlock* physical_block)
 		{
 			size_t partition_count = 1;
@@ -4091,7 +4072,7 @@ namespace impl
 
 			uint8 l0 = 255;
 			uint8 l1 = 0;
-			for (size_t i = 0; i < BLOCK_TEXEL_COUNT; ++i)
+			for (size_t i = 0; i < BLOCK_WIDTH * BLOCK_HEIGHT; ++i)
 			{
 				l0 = FMath::Min(l0, texels[i]);
 				l1 = FMath::Max(l1, texels[i]);
@@ -4101,8 +4082,8 @@ namespace impl
 			uint8 endpoint_quantized[2];
 			encode_luminance_direct(endpoint_quant, l0, l1, endpoint_quantized, endpoint_unquantized);
 
-			uint8 weights_quantized[BLOCK_TEXEL_COUNT];
-			calculate_quantized_weights_luminance(
+			uint8 weights_quantized[BLOCK_WIDTH * BLOCK_HEIGHT];
+			calculate_quantized_weights_luminance<BLOCK_WIDTH*BLOCK_HEIGHT>(
 				texels, weight_quant, endpoint_unquantized[0], endpoint_unquantized[1],
 				weights_quantized);
 
@@ -4110,49 +4091,46 @@ namespace impl
 			integer_sequence_encode(endpoint_quantized, 2, RANGE_256, endpoint_ise);
 
 			uint8 weights_ise[MAXIMUM_ENCODED_WEIGHT_BYTES + 1] = { 0 };
-			integer_sequence_encode(weights_quantized, BLOCK_TEXEL_COUNT, RANGE_32, weights_ise);
+			integer_sequence_encode(weights_quantized, BLOCK_WIDTH * BLOCK_HEIGHT, RANGE_32, weights_ise);
 
-			symbolic_to_physical(color_endpoint_mode, endpoint_quant, weight_quant,
+			symbolic_to_physical<BLOCK_WIDTH,BLOCK_HEIGHT>(color_endpoint_mode, endpoint_quant, weight_quant,
 				partition_count, endpoint_ise,
 				weights_ise, physical_block);
 		}
 
 
-		void encode_rgb_single_partition(const unorm8_t texels[BLOCK_TEXEL_COUNT],
+		template<int32 BLOCK_WIDTH, int32 BLOCK_HEIGHT>
+		void encode_rgb_single_partition(const unorm8_t texels[BLOCK_WIDTH * BLOCK_HEIGHT],
 			vec3f_t e0,
 			vec3f_t e1,
 			PhysicalBlock* physical_block)
 		{
-			size_t partition_count = 1;
+			constexpr size_t partition_count = 1;
 
-			color_endpoint_mode_t color_endpoint_mode = CEM_LDR_RGB_DIRECT;
-			range_t weight_quant = RANGE_12;
-			range_t endpoint_quant = endpoint_quantization(partition_count, weight_quant, color_endpoint_mode);
+			constexpr color_endpoint_mode_t color_endpoint_mode = CEM_LDR_RGB_DIRECT;
+			constexpr range_t weight_quant = RANGE_12;
+			constexpr range_t endpoint_quant = endpoint_quantization(partition_count, weight_quant, color_endpoint_mode);
 
 			vec3i_t endpoint_unquantized[2];
 			uint8 endpoint_quantized[6];
-			encode_rgb_direct(endpoint_quant, round(e0), round(e1), endpoint_quantized,
-				endpoint_unquantized);
+			encode_rgb_direct(endpoint_quant, round(e0), round(e1), endpoint_quantized, endpoint_unquantized);
 
-			uint8 weights_quantized[BLOCK_TEXEL_COUNT];
-			calculate_quantized_weights_rgb(texels, weight_quant, endpoint_unquantized[0],
-				endpoint_unquantized[1], weights_quantized);
+			uint8 weights_quantized[BLOCK_WIDTH * BLOCK_HEIGHT];
+			calculate_quantized_weights_rgb<BLOCK_WIDTH*BLOCK_HEIGHT>(texels, weight_quant, endpoint_unquantized[0], endpoint_unquantized[1], weights_quantized);
 
 			uint8 endpoint_ise[MAXIMUM_ENCODED_COLOR_ENDPOINT_BYTES] = { 0 };
 			integer_sequence_encode(endpoint_quantized, 6, endpoint_quant, endpoint_ise);
 
 			uint8 weights_ise[MAXIMUM_ENCODED_WEIGHT_BYTES + 1] = { 0 };
-			integer_sequence_encode(weights_quantized, BLOCK_TEXEL_COUNT, weight_quant,
-				weights_ise);
+			integer_sequence_encode(weights_quantized, BLOCK_WIDTH * BLOCK_HEIGHT, weight_quant, weights_ise);
 
-			symbolic_to_physical(color_endpoint_mode, endpoint_quant, weight_quant,
-				partition_count, endpoint_ise,
-				weights_ise, physical_block);
+			symbolic_to_physical<BLOCK_WIDTH, BLOCK_HEIGHT>(color_endpoint_mode, endpoint_quant, weight_quant, partition_count, endpoint_ise, weights_ise, physical_block);
 		}
 
 
 
-		void encode_rgba_single_partition_dual_plane(const unorm8_t texels[BLOCK_TEXEL_COUNT],
+		template<int32 BLOCK_WIDTH, int32 BLOCK_HEIGHT>
+		void encode_rgba_single_partition_dual_plane(const unorm8_t texels[BLOCK_WIDTH * BLOCK_HEIGHT],
 			vec4f_t e0,
 			vec4f_t e1,
 			PhysicalBlock* physical_block)
@@ -4189,15 +4167,15 @@ namespace impl
 			//                  endpoint_quantized[4], endpoint_quantized[5], endpoint_quantized[6], endpoint_quantized[7] );
 			//      }
 
-			uint8 weights_quantized[BLOCK_TEXEL_COUNT];
-			calculate_quantized_weights_rgb(texels, weight_quant,
+			uint8 weights_quantized[BLOCK_WIDTH * BLOCK_HEIGHT];
+			calculate_quantized_weights_rgb<BLOCK_WIDTH*BLOCK_HEIGHT>(texels, weight_quant,
 				endpoint_unquantized[0].rgb(),
 				endpoint_unquantized[1].rgb(),
 				weights_quantized);
 
 
-			uint8 alpha_weights_quantized[BLOCK_TEXEL_COUNT];
-			calculate_quantized_weights_channel(texels, weight_quant,
+			uint8 alpha_weights_quantized[BLOCK_WIDTH * BLOCK_HEIGHT];
+			calculate_quantized_weights_channel<BLOCK_WIDTH*BLOCK_HEIGHT>(texels, weight_quant,
 				(uint8)endpoint_unquantized[0].a,
 				(uint8)endpoint_unquantized[1].a,
 				alpha_weights_quantized,
@@ -4221,8 +4199,8 @@ namespace impl
 			//          i+=4;
 			//      }
 
-			uint8 final_weights_quantized[BLOCK_TEXEL_COUNT * 2];
-			for (size_t i = 0; i < BLOCK_TEXEL_COUNT; i++)
+			uint8 final_weights_quantized[BLOCK_WIDTH * BLOCK_HEIGHT * 2];
+			for (size_t i = 0; i < BLOCK_WIDTH * BLOCK_HEIGHT; i++)
 			{
 				final_weights_quantized[2 * i] = weights_quantized[i];
 				final_weights_quantized[2 * i + 1] = alpha_weights_quantized[i];
@@ -4232,16 +4210,17 @@ namespace impl
 			integer_sequence_encode(endpoint_quantized, 8, endpoint_quant, endpoint_ise);
 
 			uint8 weights_ise[MAXIMUM_ENCODED_WEIGHT_BYTES + 1] = { 0 };
-			integer_sequence_encode(final_weights_quantized, BLOCK_TEXEL_COUNT * 2, weight_quant, weights_ise);
+			integer_sequence_encode(final_weights_quantized, BLOCK_WIDTH * BLOCK_HEIGHT * 2, weight_quant, weights_ise);
 
-			symbolic_to_physical(color_endpoint_mode, endpoint_quant, weight_quant,
+			symbolic_to_physical<BLOCK_WIDTH, BLOCK_HEIGHT>(color_endpoint_mode, endpoint_quant, weight_quant,
 				partition_count, endpoint_ise,
 				weights_ise, physical_block, true);
 		}
 
 
 
-		void encode_rgba_single_partition(const unorm8_t texels[BLOCK_TEXEL_COUNT],
+		template<int32 BLOCK_WIDTH, int32 BLOCK_HEIGHT>
+		void encode_rgba_single_partition(const unorm8_t texels[BLOCK_WIDTH * BLOCK_HEIGHT],
 			vec4f_t e0,
 			vec4f_t e1,
 			PhysicalBlock* physical_block)
@@ -4258,8 +4237,8 @@ namespace impl
 				endpoint_unquantized);
 
 			// Calculate weights ignoring alpha (TODO)
-			uint8 weights_quantized[BLOCK_TEXEL_COUNT];
-			calculate_quantized_weights_rgba(texels, weight_quant,
+			uint8 weights_quantized[BLOCK_WIDTH * BLOCK_HEIGHT];
+			calculate_quantized_weights_rgba<BLOCK_WIDTH*BLOCK_HEIGHT>(texels, weight_quant,
 				endpoint_unquantized[0],
 				endpoint_unquantized[1],
 				weights_quantized);
@@ -4269,17 +4248,17 @@ namespace impl
 			integer_sequence_encode(endpoint_quantized, 8, endpoint_quant, endpoint_ise);
 
 			uint8 weights_ise[MAXIMUM_ENCODED_WEIGHT_BYTES + 1] = { 0 };
-			integer_sequence_encode(weights_quantized, BLOCK_TEXEL_COUNT, weight_quant,
-				weights_ise);
+			integer_sequence_encode(weights_quantized, BLOCK_WIDTH * BLOCK_HEIGHT, weight_quant, weights_ise);
 
-			symbolic_to_physical(color_endpoint_mode, endpoint_quant, weight_quant,
+			symbolic_to_physical<BLOCK_WIDTH, BLOCK_HEIGHT>(color_endpoint_mode, endpoint_quant, weight_quant,
 				partition_count, endpoint_ise,
 				weights_ise, physical_block, false);
 		}
 
 
 
-		void encode_rg_single_partition_dual_plane(const unorm8_t texels[BLOCK_TEXEL_COUNT],
+		template<int32 BLOCK_WIDTH, int32 BLOCK_HEIGHT>
+		void encode_rg_single_partition_dual_plane(const unorm8_t texels[BLOCK_WIDTH*BLOCK_HEIGHT],
 			vec4f_t e0,
 			vec4f_t e1,
 			PhysicalBlock* physical_block)
@@ -4317,16 +4296,16 @@ namespace impl
 			//                  endpoint_quantized[4], endpoint_quantized[5], endpoint_quantized[6], endpoint_quantized[7] );
 			//        }
 
-			uint8 r_weights_quantized[BLOCK_TEXEL_COUNT];
-			calculate_quantized_weights_channel(texels, weight_quant,
+			uint8 r_weights_quantized[BLOCK_WIDTH * BLOCK_HEIGHT];
+			calculate_quantized_weights_channel<BLOCK_WIDTH*BLOCK_HEIGHT>(texels, weight_quant,
 				(uint8)endpoint_unquantized[0].r,
 				(uint8)endpoint_unquantized[1].r,
 				r_weights_quantized,
 				2);
 
 
-			uint8 g_weights_quantized[BLOCK_TEXEL_COUNT];
-			calculate_quantized_weights_channel(texels, weight_quant,
+			uint8 g_weights_quantized[BLOCK_WIDTH * BLOCK_HEIGHT];
+			calculate_quantized_weights_channel<BLOCK_WIDTH*BLOCK_HEIGHT>(texels, weight_quant,
 				(uint8)endpoint_unquantized[0].g,
 				(uint8)endpoint_unquantized[1].g,
 				g_weights_quantized,
@@ -4362,8 +4341,8 @@ namespace impl
 			////            }
 			//        }
 
-			uint8 final_weights_quantized[BLOCK_TEXEL_COUNT * 2];
-			for (size_t i = 0; i < BLOCK_TEXEL_COUNT; i++)
+			uint8 final_weights_quantized[BLOCK_WIDTH * BLOCK_HEIGHT * 2];
+			for (size_t i = 0; i < BLOCK_WIDTH * BLOCK_HEIGHT; i++)
 			{
 				final_weights_quantized[2 * i] = r_weights_quantized[i];
 				final_weights_quantized[2 * i + 1] = g_weights_quantized[i];
@@ -4373,17 +4352,15 @@ namespace impl
 			integer_sequence_encode(endpoint_quantized, 4, endpoint_quant, endpoint_ise);
 
 			uint8 weights_ise[MAXIMUM_ENCODED_WEIGHT_BYTES + 1] = { 0 };
-			integer_sequence_encode(final_weights_quantized, BLOCK_TEXEL_COUNT * 2, weight_quant, weights_ise);
+			integer_sequence_encode(final_weights_quantized, BLOCK_WIDTH * BLOCK_HEIGHT * 2, weight_quant, weights_ise);
 
-			symbolic_to_physical(color_endpoint_mode, endpoint_quant, weight_quant,
+			symbolic_to_physical<BLOCK_WIDTH,BLOCK_HEIGHT>(color_endpoint_mode, endpoint_quant, weight_quant,
 				partition_count, endpoint_ise,
 				weights_ise, physical_block, true);
 		}
 
 
-		bool is_solid(const unorm8_t texels[BLOCK_TEXEL_COUNT],
-			size_t count,
-			unorm8_t* color)
+		bool is_solid(const unorm8_t* texels, size_t count, unorm8_t* color)
 		{
 			for (size_t i = 0; i < count; ++i)
 			{
@@ -4399,9 +4376,7 @@ namespace impl
 		}
 
 
-		bool is_solid_rgba(const unorm8_t texels[BLOCK_TEXEL_COUNT],
-			size_t count,
-			unorm8_t* color)
+		bool is_solid_rgba(const unorm8_t* texels, size_t count, unorm8_t* color)
 		{
 			for (size_t i = 0; i < count; ++i)
 			{
@@ -4417,9 +4392,7 @@ namespace impl
 		}
 
 
-		bool is_greyscale(const unorm8_t texels[BLOCK_TEXEL_COUNT],
-			size_t count,
-			uint8 luminances[BLOCK_TEXEL_COUNT])
+		bool is_greyscale(const unorm8_t* texels, size_t count, uint8* luminances)
 		{
 			for (size_t i = 0; i < count; ++i)
 			{
@@ -4442,15 +4415,9 @@ namespace impl
 		public:
 			mat3x3f_t() {}
 
-			mat3x3f_t(float m00,
-				float m01,
-				float m02,
-				float m10,
-				float m11,
-				float m12,
-				float m20,
-				float m21,
-				float m22)
+			mat3x3f_t(float m00, float m01, float m02,
+				float m10, float m11, float m12,
+				float m20, float m21, float m22)
 			{
 				m[0] = vec3f_t(m00, m01, m02);
 				m[1] = vec3f_t(m10, m11, m12);
@@ -4475,8 +4442,7 @@ namespace impl
 			return tmp;
 		}
 
-
-		vec3f_t mean(const unorm8_t texels[BLOCK_TEXEL_COUNT], size_t count)
+		vec3f_t mean(const unorm8_t* texels, size_t count)
 		{
 			vec3i_t sum(0, 0, 0);
 			for (size_t i = 0; i < count; ++i)
@@ -4487,10 +4453,10 @@ namespace impl
 			return to_vec3f(sum) / static_cast<float>(count);
 		}
 
-		void subtract(const unorm8_t texels[BLOCK_TEXEL_COUNT],
+		void subtract(const unorm8_t* texels,
 			size_t count,
 			vec3f_t v,
-			vec3f_t output[BLOCK_TEXEL_COUNT])
+			vec3f_t* output)
 		{
 			for (size_t i = 0; i < count; ++i)
 			{
@@ -4498,7 +4464,7 @@ namespace impl
 			}
 		}
 
-		mat3x3f_t covariance(const vec3f_t m[BLOCK_TEXEL_COUNT], size_t count)
+		mat3x3f_t covariance(const vec3f_t* m, size_t count)
 		{
 			mat3x3f_t cov;
 			for (size_t i = 0; i < 3; ++i)
@@ -4529,22 +4495,14 @@ namespace impl
 		}
 
 
-
-		inline bool approx_equal(float x, float y, float epsilon)
-		{
-			return fabs(x - y) < epsilon;
-		}
-
-
-		void find_min_max(const unorm8_t texels[BLOCK_TEXEL_COUNT],
+		void find_min_max(const unorm8_t* texels,
 			size_t count,
 			vec3f_t line_k,
 			vec3f_t line_m,
 			vec3f_t& e0,
 			vec3f_t& e1)
 		{
-			DCHECK(count <= BLOCK_TEXEL_COUNT);
-			DCHECK(approx_equal(quadrance(line_k), 1.0, 0.0001f));
+			DCHECK(FMath::IsNearlyEqual(quadrance(line_k), 1.0, 0.0001f));
 
 			float a, b;
 			{
@@ -4563,24 +4521,13 @@ namespace impl
 			e1 = clamp_rgb(line_k * b + line_m);
 		}
 
-		void find_min_max_block(const unorm8_t texels[BLOCK_TEXEL_COUNT],
-			vec3f_t line_k,
-			vec3f_t line_m,
-			vec3f_t& e0,
-			vec3f_t& e1)
-		{
-			find_min_max(texels, BLOCK_TEXEL_COUNT, line_k, line_m, e0, e1);
-		}
-
 
+		template<int32 BLOCK_TEXEL_COUNT>
 		void principal_component_analysis(const unorm8_t texels[BLOCK_TEXEL_COUNT],
 			size_t count,
 			vec3f_t& line_k,
-			vec3f_t& line_m) {
-			// Since we are working with fixed sized blocks count we can cap count. This
-			// avoids dynamic allocation.
-			DCHECK(count <= BLOCK_TEXEL_COUNT);
-
+			vec3f_t& line_m) 
+		{
 			line_m = mean(texels, count);
 
 			vec3f_t n[BLOCK_TEXEL_COUNT];
@@ -4592,22 +4539,12 @@ namespace impl
 		}
 
 
-		inline void principal_component_analysis_block(
-			const unorm8_t texels[BLOCK_TEXEL_COUNT],
-			vec3f_t& line_k,
-			vec3f_t& line_m)
-		{
-			principal_component_analysis(texels, BLOCK_TEXEL_COUNT, line_k, line_m);
-		}
-
-
-
-		void compress_block(const unorm8_t texels[BLOCK_TEXEL_COUNT],
-			PhysicalBlock* physical_block)
+		template<int32 BLOCK_WIDTH, int32 BLOCK_HEIGHT>
+		void compress_block(const unorm8_t texels[BLOCK_WIDTH*BLOCK_HEIGHT], PhysicalBlock* physical_block)
 		{
 			{
 				unorm8_t color;
-				if (is_solid(texels, BLOCK_TEXEL_COUNT, &color))
+				if (is_solid(texels, BLOCK_WIDTH * BLOCK_HEIGHT, &color))
 				{
 					encode_void_extent(to_vec3i(color), physical_block);
 					return;
@@ -4615,27 +4552,27 @@ namespace impl
 			}
 
 			{
-				uint8 luminances[BLOCK_TEXEL_COUNT];
-				if (is_greyscale(texels, BLOCK_TEXEL_COUNT, luminances))
+				uint8 luminances[BLOCK_WIDTH * BLOCK_HEIGHT];
+				if (is_greyscale(texels, BLOCK_WIDTH * BLOCK_HEIGHT, luminances))
 				{
-					encode_luminance(luminances, physical_block);
+					encode_luminance<BLOCK_WIDTH,BLOCK_HEIGHT>(luminances, physical_block);
 					return;
 				}
 			}
 
 			vec3f_t k, m;
-			principal_component_analysis_block(texels, k, m);
+			principal_component_analysis<BLOCK_WIDTH*BLOCK_HEIGHT>(texels, BLOCK_WIDTH * BLOCK_HEIGHT, k, m);
 			vec3f_t e0, e1;
-			find_min_max_block(texels, k, m, e0, e1);
-			encode_rgb_single_partition(texels, e0, e1, physical_block);
+			find_min_max(texels, BLOCK_WIDTH * BLOCK_HEIGHT, k, m, e0, e1);
+			encode_rgb_single_partition<BLOCK_WIDTH,BLOCK_HEIGHT>(texels, e0, e1, physical_block);
 		}
 
 
-		void compress_block_rgba(const unorm8_t texels[BLOCK_TEXEL_COUNT],
-			PhysicalBlock* physical_block)
+		template<int32 BLOCK_WIDTH, int32 BLOCK_HEIGHT>
+		void compress_block_rgba(const unorm8_t texels[BLOCK_WIDTH * BLOCK_HEIGHT], PhysicalBlock* physical_block)
 		{
 			unorm8_t color;
-			bool isSolidRGBA = is_solid_rgba(texels, BLOCK_TEXEL_COUNT, &color);
+			bool isSolidRGBA = is_solid_rgba(texels, BLOCK_WIDTH * BLOCK_HEIGHT, &color);
 			if (isSolidRGBA)
 			{
 				encode_void_extent(to_vec4i(color), physical_block);
@@ -4643,7 +4580,7 @@ namespace impl
 			}
 
 			vec3f_t e0, e1;
-			bool isSolid = is_solid(texels, BLOCK_TEXEL_COUNT, &color);
+			bool isSolid = is_solid(texels, BLOCK_WIDTH * BLOCK_HEIGHT, &color);
 			if (isSolid)
 			{
 				e0 = e1 = to_vec3f(texels[0]);
@@ -4651,14 +4588,14 @@ namespace impl
 			else
 			{
 				vec3f_t k, m;
-				principal_component_analysis_block(texels, k, m);
-				find_min_max_block(texels, k, m, e0, e1);
+				principal_component_analysis<BLOCK_WIDTH* BLOCK_HEIGHT>(texels, BLOCK_WIDTH * BLOCK_HEIGHT, k, m);
+				find_min_max(texels, BLOCK_WIDTH * BLOCK_HEIGHT, k, m, e0, e1);
 			}
 
-
+			// Find alpha min and max
 			uint8 l0 = 255;
 			uint8 l1 = 0;
-			for (size_t i = 0; i < BLOCK_TEXEL_COUNT; ++i)
+			for (size_t i = 0; i < BLOCK_WIDTH * BLOCK_HEIGHT; ++i)
 			{
 				l0 = FMath::Min(l0, texels[i].components[3]);
 				l1 = FMath::Max(l1, texels[i].components[3]);
@@ -4675,15 +4612,15 @@ namespace impl
 			//        }
 
 					//encode_rgba_single_partition(texels, re0, re1, physical_block);
-			encode_rgba_single_partition_dual_plane(texels, re0, re1, physical_block);
+			encode_rgba_single_partition_dual_plane<BLOCK_WIDTH, BLOCK_HEIGHT>(texels, re0, re1, physical_block);
 		}
 
 
-		void compress_block_rg(const unorm8_t texels[BLOCK_TEXEL_COUNT],
-			PhysicalBlock* physical_block)
+		template<int32 BLOCK_WIDTH, int32 BLOCK_HEIGHT>
+		void compress_block_rg(const unorm8_t texels[BLOCK_WIDTH * BLOCK_HEIGHT], PhysicalBlock* physical_block)
 		{
 			unorm8_t color;
-			bool isSolidRGBA = is_solid_rgba(texels, BLOCK_TEXEL_COUNT, &color);
+			bool isSolidRGBA = is_solid_rgba(texels, BLOCK_WIDTH * BLOCK_HEIGHT, &color);
 			if (isSolidRGBA)
 			{
 				encode_void_extent(vec4i_t(color.components[2], color.components[2], color.components[2], color.components[1]), physical_block);
@@ -4694,7 +4631,7 @@ namespace impl
 			uint8 r1 = 0;
 			uint8 g0 = 255;
 			uint8 g1 = 0;
-			for (size_t i = 0; i < BLOCK_TEXEL_COUNT; ++i)
+			for (size_t i = 0; i < BLOCK_WIDTH * BLOCK_HEIGHT; ++i)
 			{
 				r0 = FMath::Min(r0, texels[i].components[2]);
 				r1 = FMath::Max(r1, texels[i].components[2]);
@@ -4705,7 +4642,7 @@ namespace impl
 			vec4f_t re0(r0, g0, r0, g0);
 			vec4f_t re1(r1, g1, r1, g1);
 
-			encode_rg_single_partition_dual_plane(texels, re0, re1, physical_block);
+			encode_rg_single_partition_dual_plane<BLOCK_WIDTH, BLOCK_HEIGHT>(texels, re0, re1, physical_block);
 		}
 
 	}
@@ -6301,20 +6238,20 @@ namespace miro
 #if MIRO_INCLUDE_BC
 
 	//---------------------------------------------------------------------------------------------
-	void RGB_to_BC1(uint32 sx, uint32 sy, const uint8* from, uint8* to, int q)
+	void RGB_to_BC1(uint32 sx, uint32 sy, const uint8* from, uint8* to, int32 Quality)
 	{
 		uint32 bx = FMath::DivideAndRoundUp(sx, 4u);
 		uint32 by = FMath::DivideAndRoundUp(sy, 4u);
 		//for ( uint32 y = 0; y < by; ++y )
 		ParallelFor(by,
 			[
-				bx, sx, sy, from, to, q
+				bx, sx, sy, from, to, Quality
 			] (uint32 y)
 			{
 				uint8* rowTo = to + 8 * bx * y;
 				for (uint32 x = 0; x < bx; ++x)
 				{
-					miro_pixel_block block;
+					FMiroPixelBlock4x4 block;
 
 					int b = 0;
 					for (uint32 j = 0; j < 4; ++j)
@@ -6336,8 +6273,8 @@ namespace miro
 						}
 					}
 
-					int quality = q < 2 ? STB_DXT_NORMAL : STB_DXT_HIGHQUAL;
-					stb::stb_compress_dxt_block(rowTo, block.bytes, 0, quality);
+					int32 STBQuality = Quality < 2 ? STB_DXT_NORMAL : STB_DXT_HIGHQUAL;
+					stb::stb_compress_dxt_block(rowTo, block.bytes, 0, STBQuality);
 
 					rowTo += 8;
 				}
@@ -6346,14 +6283,14 @@ namespace miro
 
 
 	//---------------------------------------------------------------------------------------------
-	void RGBA_to_BC1(uint32 sx, uint32 sy, const uint8* from, uint8* to, int q)
+	void RGBA_to_BC1(uint32 sx, uint32 sy, const uint8* from, uint8* to, int32 Quality)
 	{
 		uint32 bx = FMath::DivideAndRoundUp(sx, 4u);
 		uint32 by = FMath::DivideAndRoundUp(sy, 4u);
 		//for ( uint32 y = 0; y < by; ++y )
 		ParallelFor(by,
 			[
-				bx, sx, sy, from, to, q
+				bx, sx, sy, from, to, Quality
 			] (uint32 y)
 			{
 				uint8* rowTo = to + 8 * bx * y;
@@ -6381,8 +6318,8 @@ namespace miro
 						}
 					}
 
-					int quality = q < 2 ? STB_DXT_NORMAL : STB_DXT_HIGHQUAL;
-					stb::stb_compress_dxt_block(rowTo, block, 0, quality);
+					int32 STBQuality = Quality < 2 ? STB_DXT_NORMAL : STB_DXT_HIGHQUAL;
+					stb::stb_compress_dxt_block(rowTo, block, 0, STBQuality);
 
 					rowTo += 8;
 				}
@@ -6480,7 +6417,7 @@ namespace miro
 
 
 	//---------------------------------------------------------------------------------------------
-	void RGB_to_BC2(uint32 sx, uint32 sy, const uint8* from, uint8* to, int q)
+	void RGB_to_BC2(uint32 sx, uint32 sy, const uint8* from, uint8* to, int32 Quality)
 	{
 		uint32 bx = FMath::DivideAndRoundUp(sx, 4u);
 		uint32 by = FMath::DivideAndRoundUp(sy, 4u);
@@ -6488,7 +6425,7 @@ namespace miro
 		//for ( uint32 y=0; y<by; ++y )
 		ParallelFor(by,
 			[
-				bx, sx, sy, from, to, q
+				bx, sx, sy, from, to, Quality
 			] (uint32 y)
 			{
 				uint8* rowTo = to + 16 * bx * y;
@@ -6517,8 +6454,8 @@ namespace miro
 						}
 					}
 
-					int quality = q < 2 ? STB_DXT_NORMAL : STB_DXT_HIGHQUAL;
-					stb::stb_compress_dxt_block(rowTo, block, 2, quality);
+					int32 STBQuality = Quality < 2 ? STB_DXT_NORMAL : STB_DXT_HIGHQUAL;
+					stb::stb_compress_dxt_block(rowTo, block, 2, STBQuality);
 
 					rowTo += 16;
 				}
@@ -6527,7 +6464,7 @@ namespace miro
 
 
 	//---------------------------------------------------------------------------------------------
-	void RGBA_to_BC2(uint32 sx, uint32 sy, const uint8* from, uint8* to, int q)
+	void RGBA_to_BC2(uint32 sx, uint32 sy, const uint8* from, uint8* to, int32 Quality)
 	{
 		uint32 bx = FMath::DivideAndRoundUp(sx, 4u);
 		uint32 by = FMath::DivideAndRoundUp(sy, 4u);
@@ -6535,7 +6472,7 @@ namespace miro
 		//for (uint32 y = 0; y < by; ++y)
 		ParallelFor(by,
 			[
-				bx, sx, sy, from, to, q
+				bx, sx, sy, from, to, Quality
 			] (uint32 y)
 			{
 				uint8* rowTo = to + 16 * bx * y;
@@ -6563,8 +6500,8 @@ namespace miro
 						}
 					}
 
-					int quality = q < 2 ? STB_DXT_NORMAL : STB_DXT_HIGHQUAL;
-					stb::stb_compress_dxt_block(rowTo, block, 2, quality);
+					int32 STBQuality = Quality < 2 ? STB_DXT_NORMAL : STB_DXT_HIGHQUAL;
+					stb::stb_compress_dxt_block(rowTo, block, 2, STBQuality);
 
 					rowTo += 16;
 				}
@@ -6668,7 +6605,7 @@ namespace miro
 
 
 	//---------------------------------------------------------------------------------------------
-	void RGB_to_BC3(uint32 sx, uint32 sy, const uint8* from, uint8* to, int q)
+	void RGB_to_BC3(uint32 sx, uint32 sy, const uint8* from, uint8* to, int32 Quality)
 	{
 		uint32 bx = FMath::DivideAndRoundUp(sx, 4u);
 		uint32 by = FMath::DivideAndRoundUp(sy, 4u);
@@ -6676,7 +6613,7 @@ namespace miro
 		//for ( uint32 y=0; y<by; ++y )
 		ParallelFor(by,
 			[
-				bx, sx, sy, from, to, q
+				bx, sx, sy, from, to, Quality
 			] (uint32 y)
 			{
 				uint8* rowTo = to + 16 * bx * y;
@@ -6705,8 +6642,8 @@ namespace miro
 						}
 					}
 
-					int quality = q < 2 ? STB_DXT_NORMAL : STB_DXT_HIGHQUAL;
-					stb::stb_compress_dxt_block(rowTo, block, 1, quality);
+					int32 STBQuality = Quality < 2 ? STB_DXT_NORMAL : STB_DXT_HIGHQUAL;
+					stb::stb_compress_dxt_block(rowTo, block, 1, STBQuality);
 
 					rowTo += 16;
 				}
@@ -6715,7 +6652,7 @@ namespace miro
 
 
 	//---------------------------------------------------------------------------------------------
-	void RGBA_to_BC3(uint32 sx, uint32 sy, const uint8* from, uint8* to, int q)
+	void RGBA_to_BC3(uint32 sx, uint32 sy, const uint8* from, uint8* to, int32 Quality)
 	{
 		uint32 bx = FMath::DivideAndRoundUp(sx, 4u);
 		uint32 by = FMath::DivideAndRoundUp(sy, 4u);
@@ -6723,7 +6660,7 @@ namespace miro
 		//for ( uint32 y=0; y<by; ++y )
 		ParallelFor(by,
 			[
-				bx, sx, sy, from, to, q
+				bx, sx, sy, from, to, Quality
 			] (uint32 y)
 			{
 				uint8* rowTo = to + 16 * bx * y;
@@ -6752,8 +6689,8 @@ namespace miro
 						}
 					}
 
-					int quality = q < 2 ? STB_DXT_NORMAL : STB_DXT_HIGHQUAL;
-					stb::stb_compress_dxt_block(rowTo, block, 1, quality);
+					int32 STBQuality = Quality < 2 ? STB_DXT_NORMAL : STB_DXT_HIGHQUAL;
+					stb::stb_compress_dxt_block(rowTo, block, 1, STBQuality);
 
 					rowTo += 16;
 				}
@@ -6777,7 +6714,7 @@ namespace miro
 
 				for (uint32 bx = 0; bx < nx; ++bx)
 				{
-					miro_pixel_block block;
+					FMiroPixelBlock4x4 block;
 
 					//uint8 const* colourBlock = from + 8;
 					//uint8 const* alphaBlock = from;
@@ -6854,13 +6791,13 @@ namespace miro
 
 
 	//---------------------------------------------------------------------------------------------
-	void BC1_to_BC3(uint32 sx, uint32 sy, const uint8* from, uint8* to, int q)
+	void BC1_to_BC3(uint32 sx, uint32 sy, const uint8* from, uint8* to, int32 Quality)
 	{
 		uint32 bx = FMath::DivideAndRoundUp(sx, 4u);
 		uint32 by = FMath::DivideAndRoundUp(sy, 4u);
 		uint32 bcount = bx * by;
 
-		const auto& ProcessBlock = [from, to, q](uint32 b)
+		const auto& ProcessBlock = [from, to, Quality](uint32 b)
 		{
 			const uint8* blockFrom = from + 8 * b;
 			uint8* blockTo = to + 16 * b;
@@ -6872,8 +6809,8 @@ namespace miro
 				uint8 block[4 * 4 * 4];
 				bcdec::bcdec_bc1(blockFrom, block, 4 * 4);
 
-				int quality = q < 2 ? STB_DXT_NORMAL : STB_DXT_HIGHQUAL;
-				stb::stb_compress_dxt_block(blockTo, block, 1, quality);
+				int quality = Quality < 2 ? STB_DXT_NORMAL : STB_DXT_HIGHQUAL;
+				stb::stb_compress_dxt_block(blockTo, block, 1, Quality);
 			}
 			else
 			{
@@ -6903,10 +6840,10 @@ namespace miro
 
 
 	//---------------------------------------------------------------------------------------------
-	void L_to_BC4(uint32 sx, uint32 sy, const uint8* from, uint8* to, int q)
+	void L_to_BC4(uint32 sx, uint32 sy, const uint8* from, uint8* to, int32 Quality)
 	{
 		// We don't support quality setting for this format yet
-		(void)q;
+		(void)Quality;
 
 		uint32 bx = FMath::DivideAndRoundUp(sx, 4u);
 		uint32 by = FMath::DivideAndRoundUp(sy, 4u);
@@ -6914,7 +6851,7 @@ namespace miro
 		//for ( uint32 y=0; y<by; ++y )
 		ParallelFor(by,
 			[
-				bx, sx, sy, from, to, q
+				bx, sx, sy, from, to
 			] (uint32 y)
 			{
 				uint8* rowTo = to + 8 * bx * y;
@@ -7077,10 +7014,10 @@ namespace miro
 
 
 	//---------------------------------------------------------------------------------------------
-	void RGBA_to_BC5(uint32 sx, uint32 sy, const uint8* from, uint8* to, int q)
+	void RGBA_to_BC5(uint32 sx, uint32 sy, const uint8* from, uint8* to, int32 Quality)
 	{
 		// We don't support quality setting for this format yet
-		(void)q;
+		(void)Quality;
 
 		uint32 bx = FMath::DivideAndRoundUp(sx, 4u);
 		uint32 by = FMath::DivideAndRoundUp(sy, 4u);
@@ -7088,7 +7025,7 @@ namespace miro
 		//for ( uint32 y=0; y<by; ++y )
 		ParallelFor(by,
 			[
-				bx, sx, sy, from, to, q
+				bx, sx, sy, from, to
 			] (uint32 y)
 			{
 				uint8* rowTo = to + 16 * bx * y;
@@ -7121,10 +7058,10 @@ namespace miro
 
 
 	//---------------------------------------------------------------------------------------------
-	void RGB_to_BC5(uint32 sx, uint32 sy, const uint8* from, uint8* to, int q)
+	void RGB_to_BC5(uint32 sx, uint32 sy, const uint8* from, uint8* to, int32 Quality)
 	{
 		// We don't support quality setting for this format yet
-		(void)q;
+		(void)Quality;
 
 		uint32 bx = FMath::DivideAndRoundUp(sx, 4u);
 		uint32 by = FMath::DivideAndRoundUp(sy, 4u);
@@ -7132,7 +7069,7 @@ namespace miro
 		//for ( uint32 y=0; y<by; ++y )
 		ParallelFor(by,
 			[
-				bx, sx, sy, from, to, q
+				bx, sx, sy, from, to
 			] (uint32 y)
 			{
 				uint8* rowTo = to + 16 * bx * y;
@@ -7259,9 +7196,9 @@ namespace miro
 #if MIRO_INCLUDE_ASTC
 
 	template<uint32 BLOCK_SIZE>
-	void Generic_RGB_to_ASTCRGBL(uint32 sx, uint32 sy, const uint8* from, uint8* to, int q)
+	void Generic_RGB_to_ASTCRGBL(uint32 sx, uint32 sy, const uint8* from, uint8* to, int32 Quality)
 	{
-		(void)q;
+		(void)Quality;
 		astcrt::PhysicalBlock physical_block_zero = { {0} };
 		astcrt::PhysicalBlock* dst_re = reinterpret_cast<astcrt::PhysicalBlock*>(to);
 
@@ -7271,7 +7208,7 @@ namespace miro
 		//for ( uint32 y=0; y<by; ++y )
 		ParallelFor(by,
 			[
-				bx, sx, sy, from, to, q, &physical_block_zero
+				bx, sx, sy, from, to, &physical_block_zero
 			] (uint32 y)
 			{
 				astcrt::PhysicalBlock* rowTo = reinterpret_cast<astcrt::PhysicalBlock*>(to + sizeof(astcrt::PhysicalBlock) * bx * y);
@@ -7300,7 +7237,7 @@ namespace miro
 					}
 
 					*rowTo = physical_block_zero;
-					astcrt::compress_block((astcrt::unorm8_t*)block, rowTo);
+					astcrt::compress_block<BLOCK_SIZE, BLOCK_SIZE>((astcrt::unorm8_t*)block, rowTo);
 
 					++rowTo;
 				}
@@ -7310,9 +7247,9 @@ namespace miro
 
 	//---------------------------------------------------------------------------------------------
 	template<uint32 BLOCK_SIZE>
-	void Generic_RGBA_to_ASTCRGBL(uint32 sx, uint32 sy, const uint8* from, uint8* to, int q)
+	void Generic_RGBA_to_ASTCRGBL(uint32 sx, uint32 sy, const uint8* from, uint8* to, int32 Quality)
 	{
-		(void)q;
+		(void)Quality;
 		astcrt::PhysicalBlock physical_block_zero = { {0} };
 		astcrt::PhysicalBlock* dst_re = reinterpret_cast<astcrt::PhysicalBlock*>(to);
 
@@ -7322,7 +7259,7 @@ namespace miro
 		//for ( uint32 y=0; y<by; ++y )
 		ParallelFor(by,
 			[
-				bx, sx, sy, from, to, q, &physical_block_zero
+				bx, sx, sy, from, to, &physical_block_zero
 			] (uint32 y)
 			{
 				astcrt::PhysicalBlock* rowTo = reinterpret_cast<astcrt::PhysicalBlock*>(to + sizeof(astcrt::PhysicalBlock) * bx * y);
@@ -7352,7 +7289,7 @@ namespace miro
 					}
 
 					*rowTo = physical_block_zero;
-					astcrt::compress_block((astcrt::unorm8_t*)block, rowTo);
+					astcrt::compress_block<BLOCK_SIZE, BLOCK_SIZE>((astcrt::unorm8_t*)block, rowTo);
 
 					++rowTo;
 				}
@@ -7441,9 +7378,9 @@ namespace miro
 
 	//---------------------------------------------------------------------------------------------
 	template<uint32 BLOCK_SIZE>
-	void Generic_RGBA_to_ASTCRGBAL(uint32 sx, uint32 sy, const uint8* from, uint8* to, int q)
+	void Generic_RGBA_to_ASTCRGBAL(uint32 sx, uint32 sy, const uint8* from, uint8* to, int32 Quality)
 	{
-		(void)q;
+		(void)Quality;
 		astcrt::PhysicalBlock physical_block_zero = { {0} };
 		astcrt::PhysicalBlock* dst_re = reinterpret_cast<astcrt::PhysicalBlock*>(to);
 
@@ -7453,7 +7390,7 @@ namespace miro
 		//for ( uint32 y=0; y<by; ++y )
 		ParallelFor(by,
 			[
-				bx, sx, sy, from, to, q, &physical_block_zero
+				bx, sx, sy, from, to, &physical_block_zero
 			] (uint32 y)
 			{
 				astcrt::PhysicalBlock* rowTo = reinterpret_cast<astcrt::PhysicalBlock*>(to + sizeof(astcrt::PhysicalBlock) * bx * y);
@@ -7483,7 +7420,7 @@ namespace miro
 					}
 
 					*rowTo = physical_block_zero;
-					astcrt::compress_block_rgba((astcrt::unorm8_t*)block, rowTo);
+					astcrt::compress_block_rgba<BLOCK_SIZE, BLOCK_SIZE>((astcrt::unorm8_t*)block, rowTo);
 
 					++rowTo;
 				}
@@ -7493,9 +7430,9 @@ namespace miro
 
 	//---------------------------------------------------------------------------------------------
 	template<uint32 BLOCK_SIZE>
-	void Generic_RGB_to_ASTCRGBAL(uint32 sx, uint32 sy, const uint8* from, uint8* to, int q)
+	void Generic_RGB_to_ASTCRGBAL(uint32 sx, uint32 sy, const uint8* from, uint8* to, int32 Quality)
 	{
-		(void)q;
+		(void)Quality;
 		astcrt::PhysicalBlock physical_block_zero = { {0} };
 		astcrt::PhysicalBlock* dst_re = reinterpret_cast<astcrt::PhysicalBlock*>(to);
 
@@ -7505,7 +7442,7 @@ namespace miro
 		//for ( uint32 y=0; y<by; ++y )
 		ParallelFor(by,
 			[
-				bx, sx, sy, from, to, q, &physical_block_zero
+				bx, sx, sy, from, to, &physical_block_zero
 			] (uint32 y)
 			{
 				astcrt::PhysicalBlock* rowTo = reinterpret_cast<astcrt::PhysicalBlock*>(to + sizeof(astcrt::PhysicalBlock) * bx * y);
@@ -7535,7 +7472,7 @@ namespace miro
 					}
 
 					*rowTo = physical_block_zero;
-					astcrt::compress_block_rgba((astcrt::unorm8_t*)block, rowTo);
+					astcrt::compress_block_rgba<BLOCK_SIZE, BLOCK_SIZE>((astcrt::unorm8_t*)block, rowTo);
 
 					++rowTo;
 				}
@@ -7637,10 +7574,10 @@ namespace miro
 
 	//---------------------------------------------------------------------------------------------
 	template<uint32 BLOCK_SIZE>
-	void Generic_L_to_ASTCRGBL(uint32 sx, uint32 sy, const uint8_t* from, uint8_t* to, int q)
+	void Generic_L_to_ASTCRGBL(uint32 sx, uint32 sy, const uint8_t* from, uint8_t* to, int32 Quality)
 	{
 		// Almost the same as RGB_to_ASTC4x4RGBL
-		(void)q;
+		(void)Quality;
 		astcrt::PhysicalBlock physical_block_zero = { {0} };
 		astcrt::PhysicalBlock* dst_re = reinterpret_cast<astcrt::PhysicalBlock*>(to);
 
@@ -7650,7 +7587,7 @@ namespace miro
 		//for ( uint32 y=0; y<by; ++y )
 		ParallelFor(by,
 			[
-				bx, sx, sy, from, to, q, &physical_block_zero
+				bx, sx, sy, from, to, &physical_block_zero
 			] (uint32 y)
 			{
 				astcrt::PhysicalBlock* rowTo = reinterpret_cast<astcrt::PhysicalBlock*>(to + sizeof(astcrt::PhysicalBlock) * bx * y);
@@ -7679,7 +7616,7 @@ namespace miro
 					}
 
 					*rowTo = physical_block_zero;
-					astcrt::compress_block((astcrt::unorm8_t*)block, rowTo);
+					astcrt::compress_block<BLOCK_SIZE, BLOCK_SIZE>((astcrt::unorm8_t*)block, rowTo);
 
 					++rowTo;
 				}
@@ -7688,9 +7625,9 @@ namespace miro
 
 	//---------------------------------------------------------------------------------------------
 	template<uint32 BLOCK_SIZE>
-	void Generic_RGB_to_ASTCRGL(uint32 sx, uint32 sy, const uint8* from, uint8* to, int q)
+	void Generic_RGB_to_ASTCRGL(uint32 sx, uint32 sy, const uint8* from, uint8* to, int32 Quality)
 	{
-		(void)q;
+		(void)Quality;
 		astcrt::PhysicalBlock physical_block_zero = { {0} };
 		astcrt::PhysicalBlock* dst_re = reinterpret_cast<astcrt::PhysicalBlock*>(to);
 
@@ -7700,7 +7637,7 @@ namespace miro
 		//for ( uint32 y=0; y<by; ++y )
 		ParallelFor(by,
 			[
-				bx, sx, sy, from, to, q, &physical_block_zero
+				bx, sx, sy, from, to, &physical_block_zero
 			] (uint32 y)
 			{
 				astcrt::PhysicalBlock* rowTo = reinterpret_cast<astcrt::PhysicalBlock*>(to + sizeof(astcrt::PhysicalBlock) * bx * y);
@@ -7730,7 +7667,7 @@ namespace miro
 					}
 
 					*rowTo = physical_block_zero;
-					astcrt::compress_block_rg((astcrt::unorm8_t*)block, rowTo);
+					astcrt::compress_block_rg<BLOCK_SIZE, BLOCK_SIZE>((astcrt::unorm8_t*)block, rowTo);
 					++rowTo;
 				}
 			});
@@ -7739,9 +7676,9 @@ namespace miro
 
 	//---------------------------------------------------------------------------------------------
 	template<uint32 BLOCK_SIZE>
-	void Generic_RGBA_to_ASTCRGL(uint32 sx, uint32 sy, const uint8* from, uint8* to, int q)
+	void Generic_RGBA_to_ASTCRGL(uint32 sx, uint32 sy, const uint8* from, uint8* to, int32 Quality)
 	{
-		(void)q;
+		(void)Quality;
 		astcrt::PhysicalBlock physical_block_zero = { {0} };
 		astcrt::PhysicalBlock* dst_re = reinterpret_cast<astcrt::PhysicalBlock*>(to);
 
@@ -7751,7 +7688,7 @@ namespace miro
 		//for ( uint32 y=0; y<by; ++y )
 		ParallelFor(by,
 			[
-				bx, sx, sy, from, to, q, &physical_block_zero
+				bx, sx, sy, from, to, &physical_block_zero
 			] (uint32 y)
 			{
 				astcrt::PhysicalBlock* rowTo = reinterpret_cast<astcrt::PhysicalBlock*>(to + sizeof(astcrt::PhysicalBlock) * bx * y);
@@ -7781,7 +7718,7 @@ namespace miro
 					}
 
 					*rowTo = physical_block_zero;
-					astcrt::compress_block_rg((astcrt::unorm8_t*)block, rowTo);
+					astcrt::compress_block_rg<BLOCK_SIZE, BLOCK_SIZE>((astcrt::unorm8_t*)block, rowTo);
 
 					++rowTo;
 				}
@@ -7795,35 +7732,40 @@ namespace miro
 	{
 		init_astc_decompress();
 
-		for (uint32 y = 0; y < sy; y += BLOCK_SIZE)
-		{
-			for (uint32 x = 0; x < sx; x += BLOCK_SIZE)
-			{
-				bool bIsSRGB = false;
-				uint8 Block[BLOCK_SIZE * BLOCK_SIZE * 4];
-				bool bSuccess = astcdec::decompress(Block, from, bIsSRGB, BLOCK_SIZE, BLOCK_SIZE);
-				check(bSuccess);
+		const int32 NumBlocksY = FMath::DivideAndRoundUp(sy, BLOCK_SIZE);
+		const int32 NumBlocksX = FMath::DivideAndRoundUp(sx, BLOCK_SIZE);
 
-				for (uint32 py = 0; py < BLOCK_SIZE; py++)
+		//for (uint32 by = 0; by < NumBlocksY; ++by)
+		ParallelFor(NumBlocksY, [NumBlocksX, sx, sy, from, to](int32 by)
+			{
+				for (int32 bx = 0; bx < NumBlocksX; ++bx)
 				{
-					for (uint32 px = 0; px < BLOCK_SIZE; px++)
-					{
-						uint32 xi = x + px;
-						uint32 yi = y + py;
+					constexpr int32 CompressedBlockSize = 16;
+					const uint8* SrcBlockPtr = from + (by * NumBlocksX + bx) * CompressedBlockSize;
 
-						if (xi < sx && yi < sy)
+					bool bIsSRGB = false;
+					uint8 Block[BLOCK_SIZE * BLOCK_SIZE * 4];
+					bool bSuccess = astcdec::decompress(Block, SrcBlockPtr, bIsSRGB, BLOCK_SIZE, BLOCK_SIZE);
+					check(bSuccess);
+
+					for (uint32 BlockY = 0; BlockY < BLOCK_SIZE; BlockY++)
+					{
+						for (uint32 BlockX = 0; BlockX < BLOCK_SIZE; BlockX++)
 						{
-							uint8* toPixel = to + (yi * sx + xi) * 3;
-							toPixel[0] = Block[py * BLOCK_SIZE * 4 + px * 4 + 0];
-							toPixel[1] = Block[py * BLOCK_SIZE * 4 + px * 4 + 3];
-							toPixel[2] = 0;
+							uint32 xi = bx * BLOCK_SIZE + BlockX;
+							uint32 yi = by * BLOCK_SIZE + BlockY;
+
+							if (xi < sx && yi < sy)
+							{
+								uint8* ToPixel = to + (yi * sx + xi) * 3;
+								ToPixel[0] = Block[BlockY * BLOCK_SIZE * 4 + BlockX * 4 + 0];
+								ToPixel[1] = Block[BlockY * BLOCK_SIZE * 4 + BlockX * 4 + 3];
+								ToPixel[2] = 0;
+							}
 						}
 					}
 				}
-
-				from += 16;
-			}
-		}
+			});
 	}
 
 
@@ -7833,36 +7775,41 @@ namespace miro
 	{
 		init_astc_decompress();
 
-		for (uint32 y = 0; y < sy; y += BLOCK_SIZE)
-		{
-			for (uint32 x = 0; x < sx; x += BLOCK_SIZE)
-			{
-				bool bIsSRGB = false;
-				uint8 Block[BLOCK_SIZE * BLOCK_SIZE * 4];
-				bool bSuccess = astcdec::decompress(Block, from, bIsSRGB, BLOCK_SIZE, BLOCK_SIZE);
-				check(bSuccess);
+		const int32 NumBlocksY = FMath::DivideAndRoundUp(sy, BLOCK_SIZE);
+		const int32 NumBlocksX = FMath::DivideAndRoundUp(sx, BLOCK_SIZE);
 
-				for (uint32 py = 0; py < BLOCK_SIZE; py++)
+		//for (uint32 by = 0; by < NumBlocksY; ++by)
+		ParallelFor(NumBlocksY, [NumBlocksX, sx, sy, from, to](int32 by)
+			{
+				for (int32 bx = 0; bx < NumBlocksX; ++bx)
 				{
-					for (uint32 px = 0; px < BLOCK_SIZE; px++)
-					{
-						uint32 xi = x + px;
-						uint32 yi = y + py;
+					constexpr int32 CompressedBlockSize = 16;
+					const uint8* SrcBlockPtr = from + (by * NumBlocksX + bx) * CompressedBlockSize;
 
-						if (xi < sx && yi < sy)
+					bool bIsSRGB = false;
+					uint8 Block[BLOCK_SIZE * BLOCK_SIZE * 4];
+					bool bSuccess = astcdec::decompress(Block, SrcBlockPtr, bIsSRGB, BLOCK_SIZE, BLOCK_SIZE);
+					check(bSuccess);
+
+					for (uint32 BlockY = 0; BlockY < BLOCK_SIZE; BlockY++)
+					{
+						for (uint32 BlockX = 0; BlockX < BLOCK_SIZE; BlockX++)
 						{
-							uint8* toPixel = to + (yi * sx + xi) * 4;
-							toPixel[0] = Block[py * BLOCK_SIZE * 4 + px * 4 + 0];
-							toPixel[1] = Block[py * BLOCK_SIZE * 4 + px * 4 + 3];
-							toPixel[2] = 0;
-							toPixel[3] = 0;
+							uint32 xi = bx * BLOCK_SIZE + BlockX;
+							uint32 yi = by * BLOCK_SIZE + BlockY;
+
+							if (xi < sx && yi < sy)
+							{
+								uint8* ToPixel = to + (yi * sx + xi) * 4;
+								ToPixel[0] = Block[BlockY * BLOCK_SIZE * 4 + BlockX * 4 + 0];
+								ToPixel[1] = Block[BlockY * BLOCK_SIZE * 4 + BlockX * 4 + 3];
+								ToPixel[2] = 0;
+								ToPixel[3] = 0;
+							}
 						}
 					}
 				}
-
-				from += 16;
-			}
-		}
+			});
 	}
 
 
@@ -7871,24 +7818,24 @@ namespace miro
 	//---------------------------------------------------------------------------------------------
 	void ASTC4x4RGBL_to_ASTC4x4RGBAL(uint32 sx, uint32 sy, const uint8* from, uint8* to)
 	{
-		int blockCount = FMath::DivideAndRoundUp(sy, 4u) * FMath::DivideAndRoundUp(sx, 4u);
-		FMemory::Memcpy(to, from, blockCount * 16);
+		int32 BlockCount = FMath::DivideAndRoundUp(sy, 4u) * FMath::DivideAndRoundUp(sx, 4u);
+		FMemory::Memcpy(to, from, BlockCount * 16);
 	}
 
 	void ASTC4x4RGBAL_to_ASTC4x4RGBL(uint32 sx, uint32 sy, const uint8* from, uint8* to)
 	{
-		int blockCount = FMath::DivideAndRoundUp(sy, 4u) * FMath::DivideAndRoundUp(sx, 4u);
-		FMemory::Memcpy(to, from, blockCount * 16);
+		int32 BlockCount = FMath::DivideAndRoundUp(sy, 4u) * FMath::DivideAndRoundUp(sx, 4u);
+		FMemory::Memcpy(to, from, BlockCount * 16);
 	}
 
-	void RGBA_to_ASTC4x4RGBAL(uint32 sx, uint32 sy, const uint8* from, uint8* to, int q)
+	void RGBA_to_ASTC4x4RGBAL(uint32 sx, uint32 sy, const uint8* from, uint8* to, int32 Quality)
 	{
-		Generic_RGBA_to_ASTCRGBAL<4>(sx, sy, from, to, q);
+		Generic_RGBA_to_ASTCRGBAL<4>(sx, sy, from, to, Quality);
 	}
 
-	void RGB_to_ASTC4x4RGBAL(uint32 sx, uint32 sy, const uint8* from, uint8* to, int q)
+	void RGB_to_ASTC4x4RGBAL(uint32 sx, uint32 sy, const uint8* from, uint8* to, int32 Quality)
 	{
-		Generic_RGB_to_ASTCRGBAL<4>(sx, sy, from, to, q);
+		Generic_RGB_to_ASTCRGBAL<4>(sx, sy, from, to, Quality);
 	}
 
 	void ASTC4x4RGBAL_to_RGBA(uint32 sx, uint32 sy, const uint8* from, uint8* to)
@@ -7901,14 +7848,14 @@ namespace miro
 		Generic_ASTCRGBAL_to_RGB<4>(sx, sy, from, to);
 	}
 
-	void RGB_to_ASTC4x4RGBL(uint32 sx, uint32 sy, const uint8* from, uint8* to, int q)
+	void RGB_to_ASTC4x4RGBL(uint32 sx, uint32 sy, const uint8* from, uint8* to, int32 Quality)
 	{
-		Generic_RGB_to_ASTCRGBL<4>(sx, sy, from, to, q);
+		Generic_RGB_to_ASTCRGBL<4>(sx, sy, from, to, Quality);
 	}
 
-	void RGBA_to_ASTC4x4RGBL(uint32 sx, uint32 sy, const uint8* from, uint8* to, int q)
+	void RGBA_to_ASTC4x4RGBL(uint32 sx, uint32 sy, const uint8* from, uint8* to, int32 Quality)
 	{
-		Generic_RGBA_to_ASTCRGBL<4>(sx, sy, from, to, q);
+		Generic_RGBA_to_ASTCRGBL<4>(sx, sy, from, to, Quality);
 	}
 
 	void ASTC4x4RGBL_to_RGB(uint32 sx, uint32 sy, const uint8* from, uint8* to)
@@ -7921,19 +7868,19 @@ namespace miro
 		Generic_ASTCRGBL_to_RGBA<4>(sx, sy, from, to);
 	}
 	
-	void L_to_ASTC4x4RGBL(uint32 sx, uint32 sy, const uint8* from, uint8* to, int q)
+	void L_to_ASTC4x4RGBL(uint32 sx, uint32 sy, const uint8* from, uint8* to, int32 Quality)
 	{
-		Generic_L_to_ASTCRGBL<4>(sx, sy, from, to,q);
+		Generic_L_to_ASTCRGBL<4>(sx, sy, from, to, Quality);
 	}
 
-	void RGB_to_ASTC4x4RGL(uint32 sx, uint32 sy, const uint8* from, uint8* to, int q)
+	void RGB_to_ASTC4x4RGL(uint32 sx, uint32 sy, const uint8* from, uint8* to, int32 Quality)
 	{
-		Generic_RGB_to_ASTCRGL<4>(sx, sy, from, to,q);
+		Generic_RGB_to_ASTCRGL<4>(sx, sy, from, to, Quality);
 	}
 	
-	void RGBA_to_ASTC4x4RGL(uint32 sx, uint32 sy, const uint8* from, uint8* to, int q)
+	void RGBA_to_ASTC4x4RGL(uint32 sx, uint32 sy, const uint8* from, uint8* to, int32 Quality)
 	{
-		Generic_RGBA_to_ASTCRGL<4>(sx, sy, from, to,q);
+		Generic_RGBA_to_ASTCRGL<4>(sx, sy, from, to, Quality);
 	}
 	
 	void ASTC4x4RGL_to_RGB(uint32 sx, uint32 sy, const uint8* from, uint8* to)
@@ -7947,6 +7894,40 @@ namespace miro
 	}
 
 
+	//---------------------------------------------------------------------------------------------
+	// 6x6
+	//---------------------------------------------------------------------------------------------
+	void ASTC6x6RGBAL_to_RGBA(uint32 sx, uint32 sy, const uint8* from, uint8* to)
+	{
+		Generic_ASTCRGBAL_to_RGBA<6>(sx, sy, from, to);
+	}
+
+	void ASTC6x6RGBAL_to_RGB(uint32 sx, uint32 sy, const uint8* from, uint8* to)
+	{
+		Generic_ASTCRGBAL_to_RGB<6>(sx, sy, from, to);
+	}
+
+	void ASTC6x6RGBL_to_RGB(uint32 sx, uint32 sy, const uint8* from, uint8* to)
+	{
+		Generic_ASTCRGBL_to_RGB<6>(sx, sy, from, to);
+	}
+
+	void ASTC6x6RGBL_to_RGBA(uint32 sx, uint32 sy, const uint8* from, uint8* to)
+	{
+		Generic_ASTCRGBL_to_RGBA<6>(sx, sy, from, to);
+	}
+
+	void ASTC6x6RGL_to_RGB(uint32 sx, uint32 sy, const uint8* from, uint8* to)
+	{
+		Generic_ASTCRGL_to_RGB<6>(sx, sy, from, to);
+	}
+
+	void ASTC6x6RGL_to_RGBA(uint32 sx, uint32 sy, const uint8* from, uint8* to)
+	{
+		Generic_ASTCRGL_to_RGBA<6>(sx, sy, from, to);
+	}
+
+
 	//---------------------------------------------------------------------------------------------
 	// 8x8
 	//---------------------------------------------------------------------------------------------
diff --git a/Engine/Plugins/Experimental/Mutable/Source/MutableRuntime/Private/MuR/BlockCompression/Miro/Miro.h b/Engine/Plugins/Experimental/Mutable/Source/MutableRuntime/Private/MuR/BlockCompression/Miro/Miro.h
--- a/Engine/Plugins/Experimental/Mutable/Source/MutableRuntime/Private/MuR/BlockCompression/Miro/Miro.h
+++ b/Engine/Plugins/Experimental/Mutable/Source/MutableRuntime/Private/MuR/BlockCompression/Miro/Miro.h
@@ -41,35 +41,35 @@ namespace miro
 #if MIRO_INCLUDE_BC
 
     //! BC1 support
-    extern void RGB_to_BC1( uint32 sx, uint32 sy, const uint8* from, uint8* to, int q );
-    extern void RGBA_to_BC1( uint32 sx, uint32 sy, const uint8* from, uint8* to, int q );
-    extern void BC1_to_RGB( uint32 sx, uint32 sy, const uint8* from, uint8* to );
-    extern void BC1_to_RGBA( uint32 sx, uint32 sy, const uint8* from, uint8* to );
+    extern void RGB_to_BC1( uint32 sx, uint32 sy, const uint8* From, uint8* To, int32 Quality );
+    extern void RGBA_to_BC1( uint32 sx, uint32 sy, const uint8* From, uint8* To, int32 Quality );
+    extern void BC1_to_RGB( uint32 sx, uint32 sy, const uint8* From, uint8* To );
+    extern void BC1_to_RGBA( uint32 sx, uint32 sy, const uint8* From, uint8* To );
 
     //! BC2 support
-    extern void RGB_to_BC2( uint32 sx, uint32 sy, const uint8* from, uint8* to, int q );
-    extern void RGBA_to_BC2( uint32 sx, uint32 sy, const uint8* from, uint8* to, int q );
-    extern void BC2_to_RGB( uint32 sx, uint32 sy, const uint8* from, uint8* to );
-    extern void BC2_to_RGBA( uint32 sx, uint32 sy, const uint8* from, uint8* to );
+    extern void RGB_to_BC2( uint32 sx, uint32 sy, const uint8* From, uint8* To, int32 Quality );
+    extern void RGBA_to_BC2( uint32 sx, uint32 sy, const uint8* From, uint8* To, int32 Quality );
+    extern void BC2_to_RGB( uint32 sx, uint32 sy, const uint8* From, uint8* To );
+    extern void BC2_to_RGBA( uint32 sx, uint32 sy, const uint8* From, uint8* To );
 
     //! BC3 support
-    extern void RGB_to_BC3( uint32 sx, uint32 sy, const uint8* from, uint8* to, int q );
-    extern void RGBA_to_BC3( uint32 sx, uint32 sy, const uint8* from, uint8* to, int q );
-    extern void BC3_to_RGB( uint32 sx, uint32 sy, const uint8* from, uint8* to );
-    extern void BC3_to_RGBA( uint32 sx, uint32 sy, const uint8* from, uint8* to );
-    extern void BC1_to_BC3( uint32 sx, uint32 sy, const uint8* from, uint8* to, int q );
+    extern void RGB_to_BC3( uint32 sx, uint32 sy, const uint8* From, uint8* To, int32 Quality );
+    extern void RGBA_to_BC3( uint32 sx, uint32 sy, const uint8* From, uint8* To, int32 Quality );
+    extern void BC3_to_RGB( uint32 sx, uint32 sy, const uint8* From, uint8* To );
+    extern void BC3_to_RGBA( uint32 sx, uint32 sy, const uint8* From, uint8* To );
+    extern void BC1_to_BC3( uint32 sx, uint32 sy, const uint8* From, uint8* To, int32 Quality );
 
     //! BC4 support
-    extern void L_to_BC4( uint32 sx, uint32 sy, const uint8* from, uint8* to, int q );
-    extern void BC4_to_L( uint32 sx, uint32 sy, const uint8* from, uint8* to );
-    extern void BC4_to_RGB( uint32 sx, uint32 sy, const uint8* from, uint8* to );
-    extern void BC4_to_RGBA( uint32 sx, uint32 sy, const uint8* from, uint8* to );
+    extern void L_to_BC4( uint32 sx, uint32 sy, const uint8* From, uint8* To, int32 Quality );
+    extern void BC4_to_L( uint32 sx, uint32 sy, const uint8* From, uint8* To );
+    extern void BC4_to_RGB( uint32 sx, uint32 sy, const uint8* From, uint8* To );
+    extern void BC4_to_RGBA( uint32 sx, uint32 sy, const uint8* From, uint8* To );
 
     //! BC5 support
-    extern void RGBA_to_BC5( uint32 sx, uint32 sy, const uint8* from, uint8* to, int q );
-    extern void RGB_to_BC5( uint32 sx, uint32 sy, const uint8* from, uint8* to, int q );
-    extern void BC5_to_RGBA( uint32 sx, uint32 sy, const uint8* from, uint8* to );
-    extern void BC5_to_RGB( uint32 sx, uint32 sy, const uint8* from, uint8* to );
+    extern void RGBA_to_BC5( uint32 sx, uint32 sy, const uint8* From, uint8* To, int32 Quality );
+    extern void RGB_to_BC5( uint32 sx, uint32 sy, const uint8* From, uint8* To, int32 Quality );
+    extern void BC5_to_RGBA( uint32 sx, uint32 sy, const uint8* From, uint8* To );
+    extern void BC5_to_RGB( uint32 sx, uint32 sy, const uint8* From, uint8* To );
 
 #endif
 
@@ -77,51 +77,63 @@ namespace miro
 #if MIRO_INCLUDE_ASTC
 
     //! ASTC 4x4 RGBA low support
-    extern void RGBA_to_ASTC4x4RGBAL( uint32 sx, uint32 sy, const uint8* from, uint8* to, int q );
-    extern void RGB_to_ASTC4x4RGBAL( uint32 sx, uint32 sy, const uint8* from, uint8* to, int q );
-    extern void ASTC4x4RGBAL_to_RGBA( uint32 sx, uint32 sy, const uint8* from, uint8* to );
-    extern void ASTC4x4RGBAL_to_RGB( uint32 sx, uint32 sy, const uint8* from, uint8* to );
+    extern void RGBA_to_ASTC4x4RGBAL( uint32 sx, uint32 sy, const uint8* From, uint8* To, int32 Quality );
+    extern void RGB_to_ASTC4x4RGBAL( uint32 sx, uint32 sy, const uint8* From, uint8* To, int32 Quality );
+    extern void ASTC4x4RGBAL_to_RGBA( uint32 sx, uint32 sy, const uint8* From, uint8* To );
+    extern void ASTC4x4RGBAL_to_RGB( uint32 sx, uint32 sy, const uint8* From, uint8* To );
 
     //! ASTC 4x4 RGB low support
-    extern void RGB_to_ASTC4x4RGBL( uint32 sx, uint32 sy, const uint8* from, uint8* to, int q );
-    extern void RGBA_to_ASTC4x4RGBL( uint32 sx, uint32 sy, const uint8* from, uint8* to, int q );
-    extern void ASTC4x4RGBL_to_RGB( uint32 sx, uint32 sy, const uint8* from, uint8* to );
-    extern void ASTC4x4RGBL_to_RGBA( uint32 sx, uint32 sy, const uint8* from, uint8* to );
-	extern void L_to_ASTC4x4RGBL(uint32 sx, uint32 sy, const uint8* from, uint8* to, int q);
+    extern void RGB_to_ASTC4x4RGBL( uint32 sx, uint32 sy, const uint8* From, uint8* To, int32 Quality );
+    extern void RGBA_to_ASTC4x4RGBL( uint32 sx, uint32 sy, const uint8* From, uint8* To, int32 Quality );
+    extern void ASTC4x4RGBL_to_RGB( uint32 sx, uint32 sy, const uint8* From, uint8* To );
+    extern void ASTC4x4RGBL_to_RGBA( uint32 sx, uint32 sy, const uint8* From, uint8* To );
+	extern void L_to_ASTC4x4RGBL(uint32 sx, uint32 sy, const uint8* From, uint8* To, int32 Quality);
     
 	//! Convenience transforms
-    extern void ASTC4x4RGBL_to_ASTC4x4RGBAL( uint32 sx, uint32 sy, const uint8* from, uint8* to );
-    extern void ASTC4x4RGBAL_to_ASTC4x4RGBL( uint32 sx, uint32 sy, const uint8* from, uint8* to );
+    extern void ASTC4x4RGBL_to_ASTC4x4RGBAL( uint32 sx, uint32 sy, const uint8* From, uint8* To );
+    extern void ASTC4x4RGBAL_to_ASTC4x4RGBL( uint32 sx, uint32 sy, const uint8* From, uint8* To );
 
     //! ASTC 4x4 RG low support
-    extern void RGB_to_ASTC4x4RGL( uint32 sx, uint32 sy, const uint8* from, uint8* to, int q );
-    extern void RGBA_to_ASTC4x4RGL( uint32 sx, uint32 sy, const uint8* from, uint8* to, int q );
-    extern void ASTC4x4RGL_to_RGB( uint32 sx, uint32 sy, const uint8* from, uint8* to );
-    extern void ASTC4x4RGL_to_RGBA( uint32 sx, uint32 sy, const uint8* from, uint8* to );
+    extern void RGB_to_ASTC4x4RGL( uint32 sx, uint32 sy, const uint8* From, uint8* To, int32 Quality );
+    extern void RGBA_to_ASTC4x4RGL( uint32 sx, uint32 sy, const uint8* From, uint8* To, int32 Quality );
+    extern void ASTC4x4RGL_to_RGB( uint32 sx, uint32 sy, const uint8* From, uint8* To );
+    extern void ASTC4x4RGL_to_RGBA( uint32 sx, uint32 sy, const uint8* From, uint8* To );
+
+	//! ASTC 6x6 RGBA low support
+	extern void ASTC6x6RGBAL_to_RGBA(uint32 sx, uint32 sy, const uint8* From, uint8* To);
+	extern void ASTC6x6RGBAL_to_RGB(uint32 sx, uint32 sy, const uint8* From, uint8* To);
+
+	//! ASTC 6x6 RGB low support
+	extern void ASTC6x6RGBL_to_RGB(uint32 sx, uint32 sy, const uint8* From, uint8* To);
+	extern void ASTC6x6RGBL_to_RGBA(uint32 sx, uint32 sy, const uint8* From, uint8* To);
+
+	//! ASTC 6x6 RG low support
+	extern void ASTC6x6RGL_to_RGB(uint32 sx, uint32 sy, const uint8* From, uint8* To);
+	extern void ASTC6x6RGL_to_RGBA(uint32 sx, uint32 sy, const uint8* From, uint8* To);
 
 	//! ASTC 8x8 RGBA low support
-	extern void ASTC8x8RGBAL_to_RGBA(uint32 sx, uint32 sy, const uint8* from, uint8* to);
-	extern void ASTC8x8RGBAL_to_RGB(uint32 sx, uint32 sy, const uint8* from, uint8* to);
+	extern void ASTC8x8RGBAL_to_RGBA(uint32 sx, uint32 sy, const uint8* From, uint8* To);
+	extern void ASTC8x8RGBAL_to_RGB(uint32 sx, uint32 sy, const uint8* From, uint8* To);
 
 	//! ASTC 8x8 RGB low support
-	extern void ASTC8x8RGBL_to_RGB(uint32 sx, uint32 sy, const uint8* from, uint8* to);
-	extern void ASTC8x8RGBL_to_RGBA(uint32 sx, uint32 sy, const uint8* from, uint8* to);
+	extern void ASTC8x8RGBL_to_RGB(uint32 sx, uint32 sy, const uint8* From, uint8* To);
+	extern void ASTC8x8RGBL_to_RGBA(uint32 sx, uint32 sy, const uint8* From, uint8* To);
 
 	//! ASTC 8x8 RG low support
-	extern void ASTC8x8RGL_to_RGB(uint32 sx, uint32 sy, const uint8* from, uint8* to);
-	extern void ASTC8x8RGL_to_RGBA(uint32 sx, uint32 sy, const uint8* from, uint8* to);
+	extern void ASTC8x8RGL_to_RGB(uint32 sx, uint32 sy, const uint8* From, uint8* To);
+	extern void ASTC8x8RGL_to_RGBA(uint32 sx, uint32 sy, const uint8* From, uint8* To);
 
 	//! ASTC 12x12 RGBA low support
-	extern void ASTC12x12RGBAL_to_RGBA(uint32 sx, uint32 sy, const uint8* from, uint8* to);
-	extern void ASTC12x12RGBAL_to_RGB(uint32 sx, uint32 sy, const uint8* from, uint8* to);
+	extern void ASTC12x12RGBAL_to_RGBA(uint32 sx, uint32 sy, const uint8* From, uint8* To);
+	extern void ASTC12x12RGBAL_to_RGB(uint32 sx, uint32 sy, const uint8* From, uint8* To);
 
 	//! ASTC 12x12 RGB low support
-	extern void ASTC12x12RGBL_to_RGB(uint32 sx, uint32 sy, const uint8* from, uint8* to);
-	extern void ASTC12x12RGBL_to_RGBA(uint32 sx, uint32 sy, const uint8* from, uint8* to);
+	extern void ASTC12x12RGBL_to_RGB(uint32 sx, uint32 sy, const uint8* From, uint8* To);
+	extern void ASTC12x12RGBL_to_RGBA(uint32 sx, uint32 sy, const uint8* From, uint8* To);
 
 	//! ASTC 12x12 RG low support
-	extern void ASTC12x12RGL_to_RGB(uint32 sx, uint32 sy, const uint8* from, uint8* to);
-	extern void ASTC12x12RGL_to_RGBA(uint32 sx, uint32 sy, const uint8* from, uint8* to);
+	extern void ASTC12x12RGL_to_RGB(uint32 sx, uint32 sy, const uint8* From, uint8* To);
+	extern void ASTC12x12RGL_to_RGBA(uint32 sx, uint32 sy, const uint8* From, uint8* To);
 
 #endif
 
