[mutable] Change code optimizations to only swap pixelformat and resize in some cases.
Add additional code optimizations to pixelformat to compensate for that.
#rnx

diff --git a/Engine/Plugins/Experimental/Mutable/Source/MutableTools/Private/MuT/ASTOpImagePixelFormat.cpp b/Engine/Plugins/Experimental/Mutable/Source/MutableTools/Private/MuT/ASTOpImagePixelFormat.cpp
--- a/Engine/Plugins/Experimental/Mutable/Source/MutableTools/Private/MuT/ASTOpImagePixelFormat.cpp
+++ b/Engine/Plugins/Experimental/Mutable/Source/MutableTools/Private/MuT/ASTOpImagePixelFormat.cpp
@@ -100,11 +100,29 @@ namespace mu
 	}
 
 
-	mu::Ptr<ASTOp> ASTOpImagePixelFormat::OptimiseSink(const FModelOptimizationOptions& options, FOptimizeSinkContext& context) const
+	Ptr<ASTOp> ASTOpImagePixelFormat::OptimiseSemantic(const FModelOptimizationOptions& options, int32 Pass) const
 	{
-		mu::Ptr<ASTOp> at;
+		mu::Ptr<ASTOp> NewOp;
 
-		mu::Ptr<ASTOp> sourceAt = Source.child();
+		// Skip this operation if the source op format is already the one we want.
+		if (Source)
+		{
+			FImageDesc SourceDesc = Source->GetImageDesc();
+			if (SourceDesc.m_format==Format)
+			{
+				NewOp = Source.child();
+			}
+		}
+
+		return NewOp;
+	}
+
+
+	Ptr<ASTOp> ASTOpImagePixelFormat::OptimiseSink(const FModelOptimizationOptions& options, FOptimizeSinkContext& context) const
+	{
+		Ptr<ASTOp> at;
+
+		Ptr<ASTOp> sourceAt = Source.child();
 
 		EImageFormat format = Format;
 		bool bIsCompressedFormat = IsCompressedFormat(Format);
diff --git a/Engine/Plugins/Experimental/Mutable/Source/MutableTools/Private/MuT/ASTOpImagePixelFormat.h b/Engine/Plugins/Experimental/Mutable/Source/MutableTools/Private/MuT/ASTOpImagePixelFormat.h
--- a/Engine/Plugins/Experimental/Mutable/Source/MutableTools/Private/MuT/ASTOpImagePixelFormat.h
+++ b/Engine/Plugins/Experimental/Mutable/Source/MutableTools/Private/MuT/ASTOpImagePixelFormat.h
@@ -34,6 +34,7 @@ template <class SCALAR> class vec4;
 		Ptr<ASTOp> Clone(MapChildFuncRef mapChild) const override;
 		void ForEachChild(const TFunctionRef<void(ASTChild&)>) override;
 		void Link(FProgram& program, FLinkerOptions* Options) override;
+		Ptr<ASTOp> OptimiseSemantic(const FModelOptimizationOptions& options, int32 Pass) const override;
 		Ptr<ASTOp> OptimiseSink(const FModelOptimizationOptions& options, FOptimizeSinkContext& context) const override;
 		FImageDesc GetImageDesc(bool returnBestOption, FGetImageDescContext* context) const override;
 		void GetLayoutBlockSize(int* pBlockX, int* pBlockY) override;
diff --git a/Engine/Plugins/Experimental/Mutable/Source/MutableTools/Private/MuT/CodeGenerator_Image.cpp b/Engine/Plugins/Experimental/Mutable/Source/MutableTools/Private/MuT/CodeGenerator_Image.cpp
--- a/Engine/Plugins/Experimental/Mutable/Source/MutableTools/Private/MuT/CodeGenerator_Image.cpp
+++ b/Engine/Plugins/Experimental/Mutable/Source/MutableTools/Private/MuT/CodeGenerator_Image.cpp
@@ -1057,8 +1057,12 @@ namespace mu
         if ( node.m_pBase )
         {
 			FImageGenerationOptions NewOptions = Options;
-			NewOptions.RectSize[0] = FMath::RoundToInt32(NewOptions.RectSize[0] / node.m_sizeX);
-			NewOptions.RectSize[1] = FMath::RoundToInt32(NewOptions.RectSize[1] / node.m_sizeY);
+			
+			if (node.m_relative)
+			{
+				NewOptions.RectSize[0] = FMath::RoundToInt32(NewOptions.RectSize[0] / node.m_sizeX);
+				NewOptions.RectSize[1] = FMath::RoundToInt32(NewOptions.RectSize[1] / node.m_sizeY);
+			}
 
 			FImageGenerationResult BaseResult;
 			GenerateImage(NewOptions, BaseResult, node.m_pBase);
@@ -1396,7 +1400,6 @@ namespace mu
 			NewOptions.ImageLayoutStrategy = CompilerOptions::TextureLayoutStrategy::None;
 			NewOptions.LayoutToApply = nullptr;
 			NewOptions.LayoutBlockId = -1;
-			NewOptions.RectSize = {};
 
 			FImageGenerationResult ImageResult;
 			GenerateImage(NewOptions, ImageResult, node.m_pImage);
diff --git a/Engine/Plugins/Experimental/Mutable/Source/MutableTools/Private/MuT/CodeOptimiser_Semantic.cpp b/Engine/Plugins/Experimental/Mutable/Source/MutableTools/Private/MuT/CodeOptimiser_Semantic.cpp
--- a/Engine/Plugins/Experimental/Mutable/Source/MutableTools/Private/MuT/CodeOptimiser_Semantic.cpp
+++ b/Engine/Plugins/Experimental/Mutable/Source/MutableTools/Private/MuT/CodeOptimiser_Semantic.cpp
@@ -1522,17 +1522,26 @@ namespace mu
 			}
 
 			case OP_TYPE::IM_PIXELFORMAT:
-			{
+			{				
 				// \todo: only if shrinking?
-				Ptr<ASTOpImagePixelFormat> NewOp = mu::Clone<ASTOpImagePixelFormat>(sourceAt);
-				Ptr<ASTOp> BaseAt = NewOp->Source.child();
+				
+				// Only sink the resize if we know that the pixelformat source image is uncompressed.
+				Ptr<ASTOpImagePixelFormat> SourceTyped = dynamic_cast<ASTOpImagePixelFormat*>(sourceAt.get());
+				FImageDesc PixelFormatSourceDesc = SourceTyped ->Source->GetImageDesc();
+				if (PixelFormatSourceDesc.m_format!=EImageFormat::IF_NONE
+					&&
+					!mu::IsCompressedFormat(PixelFormatSourceDesc.m_format) )
+				{
+					Ptr<ASTOpImagePixelFormat> NewOp = mu::Clone<ASTOpImagePixelFormat>(sourceAt);
+					Ptr<ASTOp> BaseAt = NewOp->Source.child();
 
-				Ptr<ASTOpFixed> NewBase = mu::Clone<ASTOpFixed>(this);
-				NewBase->SetChild(NewBase->op.args.ImageResize.source, BaseAt);
+					Ptr<ASTOpFixed> NewBase = mu::Clone<ASTOpFixed>(this);
+					NewBase->SetChild(NewBase->op.args.ImageResize.source, BaseAt);
 
-				NewOp->Source = NewBase;
+					NewOp->Source = NewBase;
 
-				at = NewOp;
+					at = NewOp;
+				}
 				break;
 			}
 
