[Mutable] Generate requested LODs will generate all LODs from the lowest requested LOD to the max LOD without skipping unused LODs.

#rnx

diff --git a/Engine/Plugins/Experimental/Mutable/Source/CustomizableObject/Private/MuCO/CustomizableObjectInstance.cpp b/Engine/Plugins/Experimental/Mutable/Source/CustomizableObject/Private/MuCO/CustomizableObjectInstance.cpp
--- a/Engine/Plugins/Experimental/Mutable/Source/CustomizableObject/Private/MuCO/CustomizableObjectInstance.cpp
+++ b/Engine/Plugins/Experimental/Mutable/Source/CustomizableObject/Private/MuCO/CustomizableObjectInstance.cpp
@@ -6309,7 +6309,7 @@ void UCustomizableObjectInstance::SetRequestedLODs(int32 InMinLOD, int32 InMaxLO
 	if (!bIgnoreMinMaxLOD)
 	{
 		InMinLOD = 0;
-		InMaxLOD = MAX_int32;
+		InMaxLOD = MAX_MESH_LOD_COUNT;
 	}
 #endif
 
@@ -6319,7 +6319,7 @@ void UCustomizableObjectInstance::SetRequestedLODs(int32 InMinLOD, int32 InMaxLO
 	InMinLOD = FMath::Min(FMath::Max(InMinLOD, static_cast<int32>(PrivateData->FirstLODAvailable)), PrivateData->FirstLODAvailable + PrivateData->NumMaxLODsToStream);
 
 	// Clamp Max LOD
-	InMaxLOD = FMath::Max(FMath::Min(InMaxLOD, PrivateData->NumLODsAvailable - 1), InMinLOD);
+	InMaxLOD = FMath::Max(FMath::Min3(InMaxLOD, PrivateData->NumLODsAvailable - 1, (int32)MAX_MESH_LOD_COUNT), InMinLOD);
 
 	const bool bMinMaxLODChanged = Descriptor.MaxLOD != InMaxLOD || Descriptor.MinLOD != InMinLOD;
 
@@ -6334,8 +6334,6 @@ void UCustomizableObjectInstance::SetRequestedLODs(int32 InMinLOD, int32 InMaxLO
 	bool bUpdateRequestedLODs = false;
 	if (UCustomizableObjectSystem::GetInstance()->IsOnlyGenerateRequestedLODsEnabled())
 	{
-		// TODO Pere: With Requested LODs We could skip updates by increasing the lower limit of LODs mutable can generate if only MinLOD have changed.
-
 		const int32 ComponentCount = GetNumComponents();
 		MutableUpdateCandidate.RequestedLODLevels.SetNum(ComponentCount);
 
@@ -6350,41 +6348,32 @@ void UCustomizableObjectInstance::SetRequestedLODs(int32 InMinLOD, int32 InMaxLO
 			bUpdateRequestedLODs = bIgnoreGeneratedLODs;
 			for (int32 ComponentIndex = 0; ComponentIndex < ComponentCount; ++ComponentIndex)
 			{
-				// Generate at least the MaxLOD
-				int32 RequestedLODs = 1 << MutableUpdateCandidate.MaxLOD;
-
+				// Find the first requested LOD. We'll generate [FirstRequestedLOD ... MaxLOD].
+				int32 FirstRequestedLOD = MutableUpdateCandidate.MaxLOD;
 				for (int32 LODIndex = 0; LODIndex < MutableUpdateCandidate.MaxLOD; ++LODIndex)
 				{
-					// Add previously generated LODs if the descriptor have not changed
-					if (!bIgnoreGeneratedLODs && GeneratedLODsPerComponent[ComponentIndex] & (1 << LODIndex))
+					if ((!bIgnoreGeneratedLODs && GeneratedLODsPerComponent[ComponentIndex] & (1 << LODIndex))
+						||
+						InRequestedLODsPerComponent[ComponentIndex] & (1 << LODIndex))
 					{
-						RequestedLODs |= (1 << LODIndex);
-					}
-
-					// Add new requested LODs that fall within the range
-					if (InRequestedLODsPerComponent[ComponentIndex] & (1 << LODIndex))
-					{
-						if (LODIndex >= MutableUpdateCandidate.MinLOD)
-						{
-							RequestedLODs |= (1 << LODIndex);
-						}
-						else
-						{
-							// RequestedLOD is out of range, generate the MinLOD instead;
-							RequestedLODs |= (1 << MutableUpdateCandidate.MinLOD);
-						}
+						// First RequestedLOD that fall within the range
+						FirstRequestedLOD = LODIndex >= MutableUpdateCandidate.MinLOD ? LODIndex : MutableUpdateCandidate.MinLOD;
+						break;
 					}
 				}
 
-				bool RequestedLODsDifferent = !Descriptor.RequestedLODLevels.IsValidIndex(ComponentIndex) || 
-					                          RequestedLODs != Descriptor.RequestedLODLevels[ComponentIndex];
-				bUpdateRequestedLODs |= RequestedLODsDifferent;
+				// Generate at least the MaxLOD
+				int32 RequestedLODs = 1 << MutableUpdateCandidate.MaxLOD;
 
-				// Save new RequestedLODs
-				if (ComponentIndex >= MutableUpdateCandidate.RequestedLODLevels.Num())
+				// Mark all LODs up until MAX_MESH_LOD_COUNT since MaxLOD can be set to Max_int32
+				for (int32 LODIndex = FirstRequestedLOD; LODIndex < MAX_MESH_LOD_COUNT; ++LODIndex)
 				{
-					MutableUpdateCandidate.RequestedLODLevels.SetNumZeroed(ComponentIndex+1);
+					RequestedLODs |= (1 << LODIndex);
 				}
+				
+				bUpdateRequestedLODs |= (RequestedLODs != Descriptor.RequestedLODLevels[ComponentIndex]);
+
+				// Save new RequestedLODs
 				MutableUpdateCandidate.RequestedLODLevels[ComponentIndex] = RequestedLODs;
 			}
 		}
