[Mutable] Additional logging to provide more context when an external image is missing.

#rnx

diff --git a/Engine/Plugins/Experimental/Mutable/Source/CustomizableObject/Private/MuCO/UnrealMutableImageProvider.cpp b/Engine/Plugins/Experimental/Mutable/Source/CustomizableObject/Private/MuCO/UnrealMutableImageProvider.cpp
--- a/Engine/Plugins/Experimental/Mutable/Source/CustomizableObject/Private/MuCO/UnrealMutableImageProvider.cpp
+++ b/Engine/Plugins/Experimental/Mutable/Source/CustomizableObject/Private/MuCO/UnrealMutableImageProvider.cpp
@@ -164,7 +164,7 @@ mu::EImageFormat GetMutablePixelFormat(EPixelFormat InTextureFormat)
 		if (!GlobalExternalImages.Contains(Id))
 		{
 			// Null case, no image was provided
-			UE_LOG(LogMutable, Warning, TEXT("Failed to get external image. GlobalExternalImage not found."));
+			UE_LOG(LogMutable, Warning, TEXT("Failed to get external image [%s]. GlobalExternalImage not found."), *Id.ToString());
 
 			ResultCallback(CreateDummy());
 			return Invoke(TrivialReturn);
@@ -219,7 +219,7 @@ mu::EImageFormat GetMutablePixelFormat(EPixelFormat InTextureFormat)
 			// Check if it's a format we support
 			if (MutImageFormat == mu::EImageFormat::IF_NONE)
 			{
-				UE_LOG(LogMutable, Warning, TEXT("Failed to get external image. Unexpected image format. EImageFormat [%s]."), GetPixelFormatString(Format));
+				UE_LOG(LogMutable, Warning, TEXT("Failed to get external image [%s]. Unexpected image format. EImageFormat [%s]."), *Id.ToString(), GetPixelFormatString(Format));
 				ResultCallback(CreateDummy());
 				return Invoke(TrivialReturn);
 			}
@@ -239,8 +239,8 @@ mu::EImageFormat GetMutablePixelFormat(EPixelFormat InTextureFormat)
 
 			if (BulkDataSize != MutImageDataSize)
 			{
-				UE_LOG(LogMutable, Warning, TEXT("Failed to get external image. Bulk data size is different than the expected size. BulkData size [%d]. Mutable image data size [%d]."),
-					BulkDataSize, MutImageDataSize);
+				UE_LOG(LogMutable, Warning, TEXT("Failed to get external image [%s]. Bulk data size is different than the expected size. BulkData size [%d]. Mutable image data size [%d]."),
+					*Id.ToString(),	BulkDataSize, MutImageDataSize);
 
 				ResultCallback(CreateDummy());
 				return Invoke(TrivialReturn);
@@ -531,7 +531,12 @@ void FUnrealMutableImageProvider::CacheImage(FName Id, bool bUser)
 
 			if (!pResult && !UnrealDeferredTexture)
 			{
-				UE_LOG(LogMutable, Warning, TEXT("Failed to cache external image %s."), *Id.ToString());
+				UE_LOG(LogMutable, Warning, TEXT("Failed to cache external image %s. Missing result or source texture. Result [%d]. Unreal texture [%d]. Num providers [%d]"),
+					*Id.ToString(),
+					pResult ? 1 : 0,
+					UnrealDeferredTexture ? 1 : 0,
+					ImageProviders.Num());
+
 				pResult = CreateDummy();
 			}
 			
