[mutable] Fix performance regression due to projected image size calculation being wrong.
#rnx

diff --git a/Engine/Plugins/Experimental/Mutable/Source/MutableTools/Private/MuT/CodeGenerator_Image.cpp b/Engine/Plugins/Experimental/Mutable/Source/MutableTools/Private/MuT/CodeGenerator_Image.cpp
--- a/Engine/Plugins/Experimental/Mutable/Source/MutableTools/Private/MuT/CodeGenerator_Image.cpp
+++ b/Engine/Plugins/Experimental/Mutable/Source/MutableTools/Private/MuT/CodeGenerator_Image.cpp
@@ -1400,6 +1400,7 @@ namespace mu
 			NewOptions.ImageLayoutStrategy = CompilerOptions::TextureLayoutStrategy::None;
 			NewOptions.LayoutToApply = nullptr;
 			NewOptions.LayoutBlockId = -1;
+			NewOptions.RectSize = {};
 
 			FImageGenerationResult ImageResult;
 			GenerateImage(NewOptions, ImageResult, node.m_pImage);
