Make FReferencedSkeletons::Skeleton a uproperty so it can be discovered by GC

#localization none
#tests compile and launched client on a platform

diff --git a/Engine/Plugins/Experimental/Mutable/Source/CustomizableObject/Private/MuCO/CustomizableInstancePrivateData.h b/Engine/Plugins/Experimental/Mutable/Source/CustomizableObject/Private/MuCO/CustomizableInstancePrivateData.h
--- a/Engine/Plugins/Experimental/Mutable/Source/CustomizableObject/Private/MuCO/CustomizableInstancePrivateData.h
+++ b/Engine/Plugins/Experimental/Mutable/Source/CustomizableObject/Private/MuCO/CustomizableInstancePrivateData.h
@@ -103,6 +103,7 @@ struct FReferencedSkeletons
 	GENERATED_USTRUCT_BODY();
 
 	// Merged skeleton if found in the cache
+	UPROPERTY()
 	TObjectPtr<USkeleton> Skeleton;
 
 	UPROPERTY()
