[Mutable] Fix missing texture parameters and skipped updates.
Texture parameters of the ongoing update were being uncached by incoming petitions if they had different parameters.
Asking for an update of the current instance being updated would sometimes skip it and mark it as successful.

#rnx

diff --git a/Engine/Plugins/Experimental/Mutable/Source/CustomizableObject/Private/MuCO/CustomizableObjectSystem.cpp b/Engine/Plugins/Experimental/Mutable/Source/CustomizableObject/Private/MuCO/CustomizableObjectSystem.cpp
--- a/Engine/Plugins/Experimental/Mutable/Source/CustomizableObject/Private/MuCO/CustomizableObjectSystem.cpp
+++ b/Engine/Plugins/Experimental/Mutable/Source/CustomizableObject/Private/MuCO/CustomizableObjectSystem.cpp
@@ -1166,6 +1166,7 @@ void UCustomizableObjectSystem::PurgePendingReleaseSkeletalMesh()
 	}
 }
 
+
 void UCustomizableObjectSystem::AddPendingReleaseSkeletalMesh(USkeletalMesh* SkeletalMesh)
 {
 	check(SkeletalMesh != nullptr);
@@ -2926,6 +2927,91 @@ UDefaultImageProvider& UCustomizableObjectSystem::GetOrCreateDefaultImageProvide
 	return *DefaultImageProvider;
 }
 
+static bool bRevertCacheTextureParameters = false;
+static FAutoConsoleVariableRef CVarRevertCacheTextureParameters(
+	TEXT("mutable.RevertCacheTextureParameters"), bRevertCacheTextureParameters,
+	TEXT("If true, FMutableOperation will not cache/uncache texture parameters. If false, FMutableOperation will add an additional reference to the TextureParameters being used in the update."));
+
+void CacheTexturesParameters(const TArray<FName>& TextureParameters)
+{
+	if (bRevertCacheTextureParameters)
+	{
+		return;
+	}
+
+	if (!TextureParameters.IsEmpty() && UCustomizableObjectSystem::IsCreated())
+	{
+		FUnrealMutableImageProvider* ImageProvider = UCustomizableObjectSystem::GetInstance()->GetPrivateChecked()->GetImageProviderChecked();
+		check(ImageProvider);
+
+		for (const FName& TextureParameter : TextureParameters)
+		{
+			ImageProvider->CacheImage(TextureParameter, false);
+		}
+	}
+}
+
+
+void UnCacheTexturesParameters(const TArray<FName>& TextureParameters)
+{
+	if (bRevertCacheTextureParameters)
+	{
+		return;
+	}
+
+	if (!TextureParameters.IsEmpty() && UCustomizableObjectSystem::IsCreated())
+	{
+		FUnrealMutableImageProvider* ImageProvider = UCustomizableObjectSystem::GetInstance()->GetPrivateChecked()->GetImageProviderChecked();
+		check(ImageProvider);
+
+		for (const FName& TextureParameter : TextureParameters)
+		{
+			ImageProvider->UnCacheImage(TextureParameter, false);
+		}
+	}
+}
+
+
+FMutableOperation::FMutableOperation(const FMutableOperation& Other)
+{
+	bNeverStream = Other.bNeverStream;
+	MipsToSkip = Other.MipsToSkip;
+	CustomizableObjectInstance = Other.CustomizableObjectInstance;
+	InstanceDescriptorRuntimeHash = Other.InstanceDescriptorRuntimeHash;
+	bStarted = Other.bStarted;
+	bBuildParameterRelevancy = Other.bBuildParameterRelevancy;
+	Parameters = Other.Parameters;
+	TextureParameters = Other.TextureParameters;
+	UpdateCallback = Other.UpdateCallback;
+
+	CacheTexturesParameters(TextureParameters);
+}
+
+
+FMutableOperation& FMutableOperation::operator=(const FMutableOperation& Other)
+{
+	bNeverStream = Other.bNeverStream;
+	MipsToSkip = Other.MipsToSkip;
+	CustomizableObjectInstance = Other.CustomizableObjectInstance;
+	InstanceDescriptorRuntimeHash = Other.InstanceDescriptorRuntimeHash;
+	bStarted = Other.bStarted;
+	bBuildParameterRelevancy = Other.bBuildParameterRelevancy;
+	Parameters = Other.Parameters;
+	TextureParameters = Other.TextureParameters;
+	UpdateCallback = Other.UpdateCallback;
+
+	CacheTexturesParameters(TextureParameters);
+
+	return *this;
+}
+
+
+FMutableOperation::~FMutableOperation()
+{
+	// Uncache Texture Parameters
+	UnCacheTexturesParameters(TextureParameters);
+}
+
 
 FMutableOperation FMutableOperation::CreateInstanceUpdate(UCustomizableObjectInstance* InCustomizableObjectInstance, bool bInNeverStream, int32 InMipsToSkip, const FInstanceUpdateDelegate* UpdateCallback)
 {
@@ -2941,6 +3027,7 @@ FMutableOperation FMutableOperation::CreateInstanceUpdate(UCustomizableObjectIns
 	Op.bStarted = false;
 	Op.bBuildParameterRelevancy = InCustomizableObjectInstance->GetBuildParameterRelevancy();
 	Op.Parameters = InCustomizableObjectInstance->GetDescriptor().GetParameters();
+	Op.TextureParameters = InCustomizableObjectInstance->GetPrivate()->UpdateTextureParameters;
 
 	if (UpdateCallback)
 	{
@@ -2955,6 +3042,8 @@ FMutableOperation FMutableOperation::CreateInstanceUpdate(UCustomizableObjectIns
 		Op.CustomizableObjectInstance = nullptr;
 	}
 
+	CacheTexturesParameters(Op.TextureParameters);
+
 	return Op;
 }
 
diff --git a/Engine/Plugins/Experimental/Mutable/Source/CustomizableObject/Private/MuCO/CustomizableObjectSystemPrivate.h b/Engine/Plugins/Experimental/Mutable/Source/CustomizableObject/Private/MuCO/CustomizableObjectSystemPrivate.h
--- a/Engine/Plugins/Experimental/Mutable/Source/CustomizableObject/Private/MuCO/CustomizableObjectSystemPrivate.h
+++ b/Engine/Plugins/Experimental/Mutable/Source/CustomizableObject/Private/MuCO/CustomizableObjectSystemPrivate.h
@@ -35,6 +35,8 @@ class FMutableOperation
 	/** Instance parameters at the time of the operation request. */
 	mu::ParametersPtr Parameters; 
 	
+	TArray<FName> TextureParameters;
+
 	bool bBuildParameterRelevancy = false;
 	bool bMeshNeedsUpdate = false;
 
@@ -43,6 +45,12 @@ class FMutableOperation
 
 public:
 
+	FMutableOperation(const FMutableOperation&);
+	FMutableOperation(FMutableOperation&&) = default;
+	FMutableOperation& operator=(const FMutableOperation&);
+	FMutableOperation& operator=(FMutableOperation&&) = default;
+
+	~FMutableOperation();
 
 	static FMutableOperation CreateInstanceUpdate(UCustomizableObjectInstance* COInstance, bool bInNeverStream, int32 MipsToSkip, const FInstanceUpdateDelegate* UpdateCallback);
 
diff --git a/Engine/Plugins/Experimental/Mutable/Source/CustomizableObject/Private/MuCO/UnrealMutableImageProvider.cpp b/Engine/Plugins/Experimental/Mutable/Source/CustomizableObject/Private/MuCO/UnrealMutableImageProvider.cpp
--- a/Engine/Plugins/Experimental/Mutable/Source/CustomizableObject/Private/MuCO/UnrealMutableImageProvider.cpp
+++ b/Engine/Plugins/Experimental/Mutable/Source/CustomizableObject/Private/MuCO/UnrealMutableImageProvider.cpp
@@ -537,7 +537,7 @@ void FUnrealMutableImageProvider::CacheImage(FName Id, bool bUser)
 					UnrealDeferredTexture ? 1 : 0,
 					ImageProviders.Num());
 
-				pResult = CreateDummy();
+				return;
 			}
 			
 			FUnrealMutableImageInfo ImageInfo(pResult, UnrealDeferredTexture);
