[mutable] Fixed projections permanently appearing as default textures when quickly changing between two options.
The problem was that an update was cancelled because it was a subset of the current descriptor of the instance, but didn't take into consideration that the current running operation was modifiying the descriptor of the instance.

#rnx

#localization none
#tests PIE Frontend

diff --git a/Engine/Plugins/Experimental/Mutable/Source/CustomizableObject/Private/MuCO/CustomizableObjectSystem.cpp b/Engine/Plugins/Experimental/Mutable/Source/CustomizableObject/Private/MuCO/CustomizableObjectSystem.cpp
--- a/Engine/Plugins/Experimental/Mutable/Source/CustomizableObject/Private/MuCO/CustomizableObjectSystem.cpp
+++ b/Engine/Plugins/Experimental/Mutable/Source/CustomizableObject/Private/MuCO/CustomizableObjectSystem.cpp
@@ -891,7 +891,9 @@ void FCustomizableObjectSystemPrivate::InitUpdateSkeletalMesh(UCustomizableObjec
 	Instance.BeginUpdateDelegate.Broadcast(&Instance);
 	Instance.BeginUpdateNativeDelegate.Broadcast(&Instance);
 
-	if (UpdateDescriptorHash.IsSubset(Instance.GetDescriptorRuntimeHash()))
+	if (UpdateDescriptorHash.IsSubset(Instance.GetDescriptorRuntimeHash()) &&
+		!(CurrentMutableOperation &&
+			&Instance == CurrentMutableOperation->CustomizableObjectInstance)) // This condition is necessary because even if the descriptor is a subset, it will be replaced by the CurrentMutableOperation
 	{
 		Instance.SkeletalMeshStatus = ESkeletalMeshState::Correct; // TODO FutureGMT MTBL-1033 should not be here. Move to UCustomizableObjectInstance::Updated
 		UpdateSkeletalMesh(Instance, Instance.GetDescriptorRuntimeHash(), EUpdateResult::Success, UpdateCallback);
