[Mutable] Allow forcing the generation of the mipmaps tail of textures marked as low-priority textures in the customizable object.

#rnx

diff --git a/Engine/Plugins/Experimental/Mutable/Source/CustomizableObject/Classes/MuCO/CustomizableObjectInstance.h b/Engine/Plugins/Experimental/Mutable/Source/CustomizableObject/Classes/MuCO/CustomizableObjectInstance.h
--- a/Engine/Plugins/Experimental/Mutable/Source/CustomizableObject/Classes/MuCO/CustomizableObjectInstance.h
+++ b/Engine/Plugins/Experimental/Mutable/Source/CustomizableObject/Classes/MuCO/CustomizableObjectInstance.h
@@ -614,6 +614,9 @@ public:
 	// It will only work if the textures aren't compressed, so set the instance to a Mutable state with texture compression disabled
 	// WARNING! If texture reuse is enabled, do NOT keep external references to the textures of the instance. The instance owns the textures.
 	void SetReuseInstanceTextures(bool bTextureReuseEnabled);
+	
+	// If enabled, low-priority textures will generate resident mipmaps too.
+	void SetForceGenerateResidentMips(bool bForceGenerateResidentMips);
 
 	// Adds/removes a texture channel coverage query to the instance update process. Every material updated by mutable with a channel named TextureName will be checked
 	// It only works properly when the Instance is in a state with texture compression disabled.
diff --git a/Engine/Plugins/Experimental/Mutable/Source/CustomizableObject/Private/MuCO/CustomizableInstancePrivateData.h b/Engine/Plugins/Experimental/Mutable/Source/CustomizableObject/Private/MuCO/CustomizableInstancePrivateData.h
--- a/Engine/Plugins/Experimental/Mutable/Source/CustomizableObject/Private/MuCO/CustomizableInstancePrivateData.h
+++ b/Engine/Plugins/Experimental/Mutable/Source/CustomizableObject/Private/MuCO/CustomizableInstancePrivateData.h
@@ -50,6 +50,7 @@ enum ECOInstanceFlags
 	
 	// Generation
 	ForceGenerateAllLODs			= 1 << 12,	// If set, Requested LOD Levels will be ignored and all LODs in between the current min/max lod will be generated
+	ForceGenerateMipTail			= 1 << 13,	// If set, SkipGenerateResidentMips will be ignored and the mip tail will be generated
 };
 
 ENUM_CLASS_FLAGS(ECOInstanceFlags);
diff --git a/Engine/Plugins/Experimental/Mutable/Source/CustomizableObject/Private/MuCO/CustomizableObjectInstance.cpp b/Engine/Plugins/Experimental/Mutable/Source/CustomizableObject/Private/MuCO/CustomizableObjectInstance.cpp
--- a/Engine/Plugins/Experimental/Mutable/Source/CustomizableObject/Private/MuCO/CustomizableObjectInstance.cpp
+++ b/Engine/Plugins/Experimental/Mutable/Source/CustomizableObject/Private/MuCO/CustomizableObjectInstance.cpp
@@ -6166,6 +6166,12 @@ void UCustomizableObjectInstance::SetReuseInstanceTextures(bool bTextureReuseEna
 }
 
 
+void UCustomizableObjectInstance::SetForceGenerateResidentMips(bool bForceGenerateResidentMips)
+{
+	bForceGenerateResidentMips ? GetPrivate()->SetCOInstanceFlags(ForceGenerateMipTail) : GetPrivate()->ClearCOInstanceFlags(ForceGenerateMipTail);
+}
+
+
 void UCustomizableObjectInstance::AddQueryTextureCoverage(const FString& TextureName, const FString* MaskOutChannelName)
 {
 	FTextureCoverageQueryData& QueryData = GetPrivate()->TextureCoverageQueries.FindOrAdd(TextureName);
diff --git a/Engine/Plugins/Experimental/Mutable/Source/CustomizableObject/Private/MuCO/CustomizableObjectSystem.cpp b/Engine/Plugins/Experimental/Mutable/Source/CustomizableObject/Private/MuCO/CustomizableObjectSystem.cpp
--- a/Engine/Plugins/Experimental/Mutable/Source/CustomizableObject/Private/MuCO/CustomizableObjectSystem.cpp
+++ b/Engine/Plugins/Experimental/Mutable/Source/CustomizableObject/Private/MuCO/CustomizableObjectSystem.cpp
@@ -2437,7 +2437,11 @@ namespace impl
 		CurrentOperationData->State = CandidateInstance->GetState();
 		CurrentOperationData->UpdateResult = EUpdateResult::Success;
 		CurrentOperationData->UpdateCallback = Operation->UpdateCallback;
-		CustomizableObject->GetLowPriorityTextureNames(CurrentOperationData->LowPriorityTextures);
+
+		if (!CandidateInstancePrivateData->HasCOInstanceFlags(ForceGenerateMipTail))
+		{
+			CustomizableObject->GetLowPriorityTextureNames(CurrentOperationData->LowPriorityTextures);
+		}
 
 		if (System->IsOnlyGenerateRequestedLODsEnabled() && System->CurrentInstanceLODManagement->IsOnlyGenerateRequestedLODLevelsEnabled() && 
 			!Operation->bForceGenerateAllLODs)
