#rnx

diff --git a/Engine/Plugins/Experimental/Mutable/Source/CustomizableObject/Classes/MuCO/CustomizableInstanceLODManagement.h b/Engine/Plugins/Experimental/Mutable/Source/CustomizableObject/Classes/MuCO/CustomizableInstanceLODManagement.h
--- a/Engine/Plugins/Experimental/Mutable/Source/CustomizableObject/Classes/MuCO/CustomizableInstanceLODManagement.h
+++ b/Engine/Plugins/Experimental/Mutable/Source/CustomizableObject/Classes/MuCO/CustomizableInstanceLODManagement.h
@@ -8,7 +8,7 @@
 typedef TMap<const class UCustomizableObjectInstance*, class FMutableUpdateCandidate> FMutableInstanceUpdateMap;
 
 // This is an abstract base class, override it to create a new Instance LOD management system and register with UCustomizableObjectSystem::SetInstanceLODManagement
-UCLASS(Blueprintable, BlueprintType)
+UCLASS(Abstract, Blueprintable, BlueprintType)
 class CUSTOMIZABLEOBJECT_API UCustomizableInstanceLODManagementBase : public UObject
 {
 	GENERATED_BODY()
@@ -17,13 +17,13 @@ public:
 	virtual ~UCustomizableInstanceLODManagementBase() {};
 
 	// WARNING! The following methods must be overriden in derived classes
-	virtual void UpdateInstanceDistsAndLODs(FMutableInstanceUpdateMap& InOutRequestedUpdates) { check(0); };
-	virtual int32 GetNumGeneratedInstancesLimitFullLODs() const { check(0); return 0; };
-	virtual int32 GetNumGeneratedInstancesLimitLOD1() const { check(0); return 0; };
-	virtual int32 GetNumGeneratedInstancesLimitLOD2() const { check(0); return 0; };
-	virtual float GetOnlyUpdateCloseCustomizableObjectsDist() const { check(0); return 0.f; };
-	virtual bool IsOnlyUpdateCloseCustomizableObjectsEnabled() const { check(0); return false; };
-	virtual bool IsOnlyGenerateRequestedLODLevelsEnabled() const { check(0); return false; };
+	virtual void UpdateInstanceDistsAndLODs(FMutableInstanceUpdateMap& InOutRequestedUpdates) PURE_VIRTUAL(UCustomizableInstanceLODManagementBase::UpdateInstanceDistsAndLODs, return;);
+	virtual int32 GetNumGeneratedInstancesLimitFullLODs() const PURE_VIRTUAL(UCustomizableInstanceLODManagementBase::GetNumGeneratedInstancesLimitFullLODs, return 0;);
+	virtual int32 GetNumGeneratedInstancesLimitLOD1() const PURE_VIRTUAL(UCustomizableInstanceLODManagementBase::GetNumGeneratedInstancesLimitLOD1, return 0;);
+	virtual int32 GetNumGeneratedInstancesLimitLOD2() const PURE_VIRTUAL(UCustomizableInstanceLODManagementBase::GetNumGeneratedInstancesLimitLOD2, return 0;);
+	virtual float GetOnlyUpdateCloseCustomizableObjectsDist() const PURE_VIRTUAL(UCustomizableInstanceLODManagementBase::GetOnlyUpdateCloseCustomizableObjectsDist, return 0.f;);
+	virtual bool IsOnlyUpdateCloseCustomizableObjectsEnabled() const PURE_VIRTUAL(UCustomizableInstanceLODManagementBase::IsOnlyUpdateCloseCustomizableObjectsEnabled, return false;);
+	virtual bool IsOnlyGenerateRequestedLODLevelsEnabled() const PURE_VIRTUAL(UCustomizableInstanceLODManagementBase::IsOnlyGenerateRequestedLODLevelsEnabled, return false;);
 };
 
 
