#rnx

diff --git a/Engine/Plugins/Experimental/Mutable/Source/CustomizableObject/Private/MuCO/CustomizableObject.cpp b/Engine/Plugins/Experimental/Mutable/Source/CustomizableObject/Private/MuCO/CustomizableObject.cpp
--- a/Engine/Plugins/Experimental/Mutable/Source/CustomizableObject/Private/MuCO/CustomizableObject.cpp
+++ b/Engine/Plugins/Experimental/Mutable/Source/CustomizableObject/Private/MuCO/CustomizableObject.cpp
@@ -1642,7 +1642,7 @@ FCustomizableObjectProjector UCustomizableObject::GetProjectorParameterDefaultVa
 	const int32 ParameterIndex = FindParameter(InParameterName);
 	if (ParameterIndex == INDEX_NONE)
 	{
-		checkNoEntry();
+		UE_LOG(LogMutable, Error, TEXT("Tried to access the default value of the unexistent projector [%s] in the CustomizableObject [%s]."), *InParameterName, *GetName());
 		return FCustomizableObjectProjectorParameterValue::DEFAULT_PARAMETER_VALUE;
 	}
 
