#rnx

diff --git a/Engine/Plugins/Experimental/Mutable/Source/CustomizableObject/Private/MuCO/CustomizableObjectSystem.cpp b/Engine/Plugins/Experimental/Mutable/Source/CustomizableObject/Private/MuCO/CustomizableObjectSystem.cpp
--- a/Engine/Plugins/Experimental/Mutable/Source/CustomizableObject/Private/MuCO/CustomizableObjectSystem.cpp
+++ b/Engine/Plugins/Experimental/Mutable/Source/CustomizableObject/Private/MuCO/CustomizableObjectSystem.cpp
@@ -687,6 +687,55 @@ void FCustomizableObjectSystemPrivate::AddReferencedObjects(FReferenceCollector&
 }
 
 
+TAutoConsoleVariable<bool> CVarCleanupTextureCache(
+	TEXT("mutable.EnableCleanupCache"),
+	true,
+	TEXT("If enabled stale textures and meshes in mutable's cache will be removed."),
+	ECVF_Scalability);
+
+
+void FCustomizableObjectSystemPrivate::CleanupCache()
+{
+	check(IsInGameThread());
+
+	bool bCleanupEnabled = CVarCleanupTextureCache.GetValueOnGameThread();
+
+	for (int32 ModelIndex = 0; ModelIndex < ModelResourcesCache.Num();)
+	{
+		if (!ModelResourcesCache[ModelIndex].Object.IsValid(false, true))
+		{
+			// The whole object has been destroyed. Remove everything.
+			ModelResourcesCache.RemoveAtSwap(ModelIndex);
+		}
+		else
+		{
+			if (bCleanupEnabled)
+			{
+				// Remove stale textures
+				for (auto Iterator = ModelResourcesCache[ModelIndex].Images.CreateIterator(); Iterator; ++Iterator)
+				{
+					if (Iterator->Value.IsStale())
+					{
+						Iterator.RemoveCurrent();
+					}
+				}
+
+				// Remove stale meshes
+				for (auto Iterator = ModelResourcesCache[ModelIndex].Meshes.CreateIterator(); Iterator; ++Iterator)
+				{
+					if (Iterator->Value.IsStale())
+					{
+						Iterator.RemoveCurrent();
+					}
+				}
+			}
+
+			++ModelIndex;
+		}
+	}
+}
+
+
 int32 FCustomizableObjectSystemPrivate::EnableMutableProgressiveMipStreaming = 1;
 
 // Warning! If this is enabled, do not get references to the textures generated by Mutable! They are owned by Mutable and could become invalid at any moment
diff --git a/Engine/Plugins/Experimental/Mutable/Source/CustomizableObject/Private/MuCO/CustomizableObjectSystemPrivate.h b/Engine/Plugins/Experimental/Mutable/Source/CustomizableObject/Private/MuCO/CustomizableObjectSystemPrivate.h
--- a/Engine/Plugins/Experimental/Mutable/Source/CustomizableObject/Private/MuCO/CustomizableObjectSystemPrivate.h
+++ b/Engine/Plugins/Experimental/Mutable/Source/CustomizableObject/Private/MuCO/CustomizableObjectSystemPrivate.h
@@ -676,23 +676,8 @@ public:
 
 	// Remove references to cached objects that have been deleted in the unreal
 	// side, and cannot be cached anyway.
-	// This should only happen in the mutable thread
-	void CleanupCache()
-	{
-		for (int ModelIndex=0; ModelIndex<ModelResourcesCache.Num();)
-		{
-			if (!ModelResourcesCache[ModelIndex].Object.IsValid(false,true))
-			{
-				// The whole object has been destroyed. Remove everything.
-				ModelResourcesCache.RemoveAtSwap(ModelIndex);
-			}
-			else
-			{
-				// \todo: Free invalid references to specific resources
-				++ModelIndex;
-			}
-		}
-	}
+	// This should only happen in the game thread
+	void CleanupCache();
 
 
 	// This should only happen in the game thread
