[Mutable] Fix crash caused by the CustomizableObjectSystem triggering garbage collection when a World is in the middle of its tick.

#rnx

diff --git a/Engine/Plugins/Experimental/Mutable/Source/CustomizableObject/Private/MuCO/CustomizableObjectSystem.cpp b/Engine/Plugins/Experimental/Mutable/Source/CustomizableObject/Private/MuCO/CustomizableObjectSystem.cpp
--- a/Engine/Plugins/Experimental/Mutable/Source/CustomizableObject/Private/MuCO/CustomizableObjectSystem.cpp
+++ b/Engine/Plugins/Experimental/Mutable/Source/CustomizableObject/Private/MuCO/CustomizableObjectSystem.cpp
@@ -3281,7 +3281,10 @@ void UCustomizableObjectSystem::OnPreBeginPIE(const bool bIsSimulatingInEditor)
 
 void UCustomizableObjectSystem::StartNextRecompile()
 {
-	CollectGarbage(GARBAGE_COLLECTION_KEEPFLAGS);
+	if (GEngine)
+	{
+		GEngine->ForceGarbageCollection();
+	}
 
 	FAssetData Itr = ObjectsToRecompile.Pop();
 	UCustomizableObject* CustomizableObject = Cast<UCustomizableObject>(Itr.GetAsset());
@@ -3388,7 +3391,10 @@ void UCustomizableObjectSystem::TickRecompileCustomizableObjects()
 			FSlateNotificationManager::Get().CancelProgressNotification(RecompileNotificationHandle);
 			RecompileNotificationHandle.Reset();
 
-			CollectGarbage(GARBAGE_COLLECTION_KEEPFLAGS);
+			if (GEngine)
+			{
+				GEngine->ForceGarbageCollection();
+			}
 		}
 	}
 }
