#rnx

diff --git a/Engine/Plugins/Experimental/Mutable/Source/CustomizableObjectEditor/Private/MuCOE/MutableUtils.cpp b/Engine/Plugins/Experimental/Mutable/Source/CustomizableObjectEditor/Private/MuCOE/MutableUtils.cpp
--- a/Engine/Plugins/Experimental/Mutable/Source/CustomizableObjectEditor/Private/MuCOE/MutableUtils.cpp
+++ b/Engine/Plugins/Experimental/Mutable/Source/CustomizableObjectEditor/Private/MuCOE/MutableUtils.cpp
@@ -130,7 +130,7 @@ bool IsMeshClosed(const USkeletalMesh* Mesh, int LOD, int MaterialIndex)
 	const FSkeletalMeshLODModel& LODResource = Mesh->GetImportedModel()->LODModels[LOD];
 	const FSkelMeshSection& Section = LODResource.Sections[MaterialIndex];
 
-	const FSoftSkinVertex* VertexBuffer = Section.SoftVertices.GetData();
+	const TArray<FSoftSkinVertex>& VertexBuffer = Section.SoftVertices;
 	const uint32 VertexStart = Section.BaseVertexIndex;
 	const uint32 VertexCount = Section.NumVertices;
 
@@ -138,20 +138,26 @@ bool IsMeshClosed(const USkeletalMesh* Mesh, int LOD, int MaterialIndex)
 	TArray<uint32> CollapsedVertexMap;
 	CollapsedVertexMap.SetNum(VertexCount);
 
+	auto GetCollapsedVertex = [&CollapsedVertexMap, VertexStart](uint32 VertIndex) -> uint32&
+	{
+
+		check(CollapsedVertexMap.IsValidIndex(VertIndex - VertexStart));
+		return CollapsedVertexMap[VertIndex - VertexStart];
+	};
 
 	for (uint32 V = VertexStart; V < VertexCount; ++V)
 	{
-		CollapsedVertexMap[V] = V;
+		GetCollapsedVertex(V) = V;
 
 		for (uint32 CandidateIndex = VertexStart; CandidateIndex < V; ++CandidateIndex)
 		{
-			uint32 Collapsed = CollapsedVertexMap[CandidateIndex];
+			uint32 Collapsed = GetCollapsedVertex(CandidateIndex);
 
 			FVector3f Vec = VertexBuffer[V].Position - VertexBuffer[Collapsed].Position;
 
 			if (FVector3f::DotProduct(Vec, Vec) <= VertexCollapseEpsilonSquared)
 			{
-				CollapsedVertexMap[V] = Collapsed;
+				GetCollapsedVertex(V) = Collapsed;
 				break;
 			}
 		}
@@ -160,16 +166,19 @@ bool IsMeshClosed(const USkeletalMesh* Mesh, int LOD, int MaterialIndex)
 	// Count faces per edge
 	TMap<TPair<uint32, uint32>, uint32> FaceCountPerEdge;
 
-	const uint32* IndexBuffer = LODResource.IndexBuffer.GetData();
+	const TArray<uint32>& IndexBuffer = LODResource.IndexBuffer;//.GetData();
 	const uint32 IndexStart = Section.BaseIndex;
 	const uint32 IndexCount = Section.NumTriangles * 3;
 
 	for (uint32 I = IndexStart; I < IndexCount; I += 3)
 	{
 		uint32 Face[3];
-		Face[0] = CollapsedVertexMap[IndexBuffer[I + 0]];
-		Face[1] = CollapsedVertexMap[IndexBuffer[I + 1]];
-		Face[2] = CollapsedVertexMap[IndexBuffer[I + 2]];
+		uint32 Index0 = IndexBuffer[I + 0];
+		Face[0] = GetCollapsedVertex(Index0);
+		uint32 Index1 = IndexBuffer[I + 1];
+		Face[1] = GetCollapsedVertex(Index1);
+		uint32 Index2 = IndexBuffer[I + 2];
+		Face[2] = GetCollapsedVertex(Index2);
 
 		for (int E = 0; E < 3; ++E)
 		{
@@ -220,19 +229,26 @@ bool IsMeshClosed(const UStaticMesh* Mesh, int LOD, int MaterialIndex)
 	TArray<uint32> CollapsedVertexMap;
 	CollapsedVertexMap.SetNum(VertexCount);
 
+	auto GetCollapsedVertex = [&CollapsedVertexMap, VertexStart](uint32 VertIndex) -> uint32&
+	{
+
+		check(CollapsedVertexMap.IsValidIndex(VertIndex - VertexStart));
+		return CollapsedVertexMap[VertIndex - VertexStart];
+	};
+
 	for (uint32 V = VertexStart; V < VertexCount; ++V)
 	{
-		CollapsedVertexMap[V] = V;
+		GetCollapsedVertex(V) = V;
 
 		for (uint32 CandidateIndex = VertexStart; CandidateIndex < V; ++CandidateIndex)
 		{
-			uint32 Collapsed = CollapsedVertexMap[CandidateIndex];
+			uint32 Collapsed = GetCollapsedVertex(CandidateIndex);
 
 			FVector3f Vec = LODResource.VertexBuffers.PositionVertexBuffer.VertexPosition(V) - LODResource.VertexBuffers.PositionVertexBuffer.VertexPosition(Collapsed);
 
 			if (FVector3f::DotProduct(Vec, Vec) <= VertexCollapseEpsilonSquared)
 			{
-				CollapsedVertexMap[V] = Collapsed;
+				GetCollapsedVertex(V) = Collapsed;
 				break;
 			}
 		}
@@ -248,9 +264,9 @@ bool IsMeshClosed(const UStaticMesh* Mesh, int LOD, int MaterialIndex)
 	for (uint32 I = IndexStart; I < IndexCount; I += 3)
 	{
 		uint32 Face[3];
-		Face[0] = CollapsedVertexMap[IndexBuffer[I + 0]];
-		Face[1] = CollapsedVertexMap[IndexBuffer[I + 1]];
-		Face[2] = CollapsedVertexMap[IndexBuffer[I + 2]];
+		Face[0] = GetCollapsedVertex(IndexBuffer[I + 0]);
+		Face[1] = GetCollapsedVertex(IndexBuffer[I + 1]);
+		Face[2] = GetCollapsedVertex(IndexBuffer[I + 2]);
 
 		for (int E = 0; E < 3; ++E)
 		{
