[mutable] Fixed validation commandlet not being able to acces children COs
- This is required to be able to locate the referenced assets (or children) of the target CO during compilation.

diff --git a/Engine/Plugins/Experimental/Mutable/Source/MutableValidation/Private/MuV/CustomizableObjectValidationCommandlet.cpp b/Engine/Plugins/Experimental/Mutable/Source/MutableValidation/Private/MuV/CustomizableObjectValidationCommandlet.cpp
--- a/Engine/Plugins/Experimental/Mutable/Source/MutableValidation/Private/MuV/CustomizableObjectValidationCommandlet.cpp
+++ b/Engine/Plugins/Experimental/Mutable/Source/MutableValidation/Private/MuV/CustomizableObjectValidationCommandlet.cpp
@@ -2,6 +2,7 @@
 
 #include "MuV/CustomizableObjectValidationCommandlet.h"
 
+#include "AssetRegistry/AssetRegistryModule.h"
 #include "Containers/Array.h"
 #include "MuCO/CustomizableObject.h"
 #include "MuCO/CustomizableObjectSystem.h"
@@ -43,6 +44,13 @@ int32 UCustomizableObjectValidationCommandlet::Main(const FString& Params)
 		return 1;
 	}
 	
+	// Perform a blocking search to ensure all assets used by mutable are reachable using the AssetRegistry
+	{
+		UE_LOG(LogMutable,Display,TEXT("Searching all assets (this will take some time)..."));
+		FAssetRegistryModule& AssetRegistryModule = FModuleManager::LoadModuleChecked<FAssetRegistryModule>(AssetRegistryConstants::ModuleName);
+		AssetRegistryModule.Get().SearchAllAssets(true /* bSynchronousSearch */);
+	}
+	
 	// Compile the Customizable Object ------------------------------------------------------------------------------ //
 	bool bWasCoCompilationSuccesfull = false;
 	{
