[Mutable] Image transform optimizations. Better mip selection, vectorized triliniear sampling, relevancy processing and mutable source optimizations.

- Added CVarGlobalImageTransformLodBias to control global image transform lod bias.

- Added CVarUseImageTransformVectorImpl enable/disable the vectorized implementation of image transform. enabled by default as it gives better performance.

#rnx

diff --git a/Engine/Plugins/Experimental/Mutable/Source/CustomizableObject/Classes/MuCO/CustomizableObject.h b/Engine/Plugins/Experimental/Mutable/Source/CustomizableObject/Classes/MuCO/CustomizableObject.h
--- a/Engine/Plugins/Experimental/Mutable/Source/CustomizableObject/Classes/MuCO/CustomizableObject.h
+++ b/Engine/Plugins/Experimental/Mutable/Source/CustomizableObject/Classes/MuCO/CustomizableObject.h
@@ -1439,7 +1439,7 @@ private:
 	// This is a manual version number for the binary blobs in this asset.
 	// Increasing it invalidates all the previously compiled models.
 	// Warning: If while merging code both versions have changed, take the highest+1.
-	static const int32 CurrentSupportedVersion = 399;
+	static const int32 CurrentSupportedVersion = 400;
 
 public:
 
diff --git a/Engine/Plugins/Experimental/Mutable/Source/MutableRuntime/Private/MuR/CodeRunner.cpp b/Engine/Plugins/Experimental/Mutable/Source/MutableRuntime/Private/MuR/CodeRunner.cpp
--- a/Engine/Plugins/Experimental/Mutable/Source/MutableRuntime/Private/MuR/CodeRunner.cpp
+++ b/Engine/Plugins/Experimental/Mutable/Source/MutableRuntime/Private/MuR/CodeRunner.cpp
@@ -2,6 +2,7 @@
 
 #include "MuR/CodeRunner.h"
 
+#include "GenericPlatform/GenericPlatformMath.h"
 #include "HAL/UnrealMemory.h"
 #include "Logging/LogCategory.h"
 #include "Logging/LogMacros.h"
@@ -78,7 +79,7 @@ float GlobalProjectionLodBias = 0.0f;
 static FAutoConsoleVariableRef CVarGlobalProjectionLodBias (
 	TEXT("mutable.GlobalProjectionLodBias"),
 	GlobalProjectionLodBias,
-	TEXT("Lod bias applied to the lod resulting form the best mip computation, only used if a min filter method different than None is used."),
+	TEXT("Lod bias applied to the lod resulting form the best mip computation for ImageProject operations, only used if a min filter method different than None is used."),
 	ECVF_Default);
 
 bool bUseProjectionVectorImpl = true;
@@ -87,6 +88,20 @@ static FAutoConsoleVariableRef CVarUseProjectionVectorImpl (
 	bUseProjectionVectorImpl,
 	TEXT("If set to true, enables the vectorized implementation of the projection pixel processing."),
 	ECVF_Default);
+
+float GlobalImageTransformLodBias = 0.0f;
+static FAutoConsoleVariableRef CVarGlobalImageTransformLodBias (
+	TEXT("mutable.GlobalImageTransformLodBias"),
+	GlobalImageTransformLodBias,
+		TEXT("Lod bias applied to the lod resulting form the best mip computation for ImageTransform operations"),
+	ECVF_Default);
+
+bool bUseImageTransformVectorImpl = true;
+static FAutoConsoleVariableRef CVarUseImageTransformVectorImpl (
+	TEXT("mutable.UseImageTransformVectorImpl"),
+	bUseImageTransformVectorImpl,
+	TEXT("If set to true, enables the vectorized implementation of the image transform pixel processing."),
+	ECVF_Default);
 }
 
 namespace mu
@@ -4715,50 +4730,147 @@ namespace mu
             {
             case 0:
 			{
-				const TArray<FScheduledOp, TInlineAllocator<6>> Deps = {
-						FScheduledOp(Args.base, item),
-						FScheduledOp(Args.offsetX, item),
-						FScheduledOp(Args.offsetY, item),
-						FScheduledOp(Args.scaleX, item),
-						FScheduledOp(Args.scaleY, item),
-						FScheduledOp(Args.rotation, item) };
+				const TArray<FScheduledOp, TFixedAllocator<2>> Deps = 
+				{
+					FScheduledOp(Args.scaleX, item),
+					FScheduledOp(Args.scaleY, item),
+				};
 
                 AddOp(FScheduledOp(item.At, item, 1), Deps);
 
 				break;
 			}
-            case 1:
-            {
+			case 1:
+			{
             	MUTABLE_CPUPROFILER_SCOPE(IM_TRANSFORM_1)
-            		
-                Ptr<const Image> pBaseImage = LoadImage(FCacheAddress(Args.base, item));
-                
-                const FVector2f Offset = FVector2f(
-                        Args.offsetX ? LoadScalar(FCacheAddress(Args.offsetX, item)) : 0.0f,
-                        Args.offsetY ? LoadScalar(FCacheAddress(Args.offsetY, item)) : 0.0f);
 
                 const FVector2f Scale = FVector2f(
                         Args.scaleX ? LoadScalar(FCacheAddress(Args.scaleX, item)) : 1.0f,
                         Args.scaleY ? LoadScalar(FCacheAddress(Args.scaleY, item)) : 1.0f);
+	
+				using FUint16Vector2 = UE::Math::TIntVector2<uint16>;
+				const FUint16Vector2 DestSizeI = Invoke([&]() 
+				{
+					int32 MipsToDrop = item.ExecutionOptions;
+					
+					FUint16Vector2 Size = FUint16Vector2(
+							Args.SizeX > 0 ? Args.SizeX : Args.SourceSizeX, 
+							Args.SizeY > 0 ? Args.SizeY : Args.SourceSizeY); 
+
+					while (MipsToDrop && !(Size.X % 2) && !(Size.Y % 2))
+					{
+						Size.X = FMath::Max(uint16(1), FMath::DivideAndRoundUp(Size.X, uint16(2)));
+						Size.Y = FMath::Max(uint16(1), FMath::DivideAndRoundUp(Size.Y, uint16(2)));
+						--MipsToDrop;
+					}
+
+					return Size;
+				});
+
+				const FVector2f DestSize = FVector2f(DestSizeI.X, DestSizeI.Y);
+				const FVector2f SourceSize = FVector2f(Args.SourceSizeX, Args.SourceSizeY);
+
+				const float DestAspectRatio = DestSize.X / DestSize.Y;
+				const float SrcAspecRatio = SourceSize.X / SourceSize.Y;
+			
+				const bool bKeepSourceAspectRatio = true;
+				
+				const FTransform2f Transform = FTransform2f(FVector2f(-0.5f)).
+					Concatenate(FTransform2f(FScale2f(Scale))).
+					Concatenate(FTransform2f(FVector2f(0.5f)));
+
+				FBox2f NormalizedCropRect(ForceInit);
+				NormalizedCropRect += Transform.TransformPoint(FVector2f(0.0f, 0.0f));
+				NormalizedCropRect += Transform.TransformPoint(FVector2f(1.0f, 0.0f));
+				NormalizedCropRect += Transform.TransformPoint(FVector2f(0.0f, 1.0f));
+				NormalizedCropRect += Transform.TransformPoint(FVector2f(1.0f, 1.0f));
+
+				const FVector2f ScaledSourceSize = NormalizedCropRect.GetSize() * DestSize;
+
+				const float BestMip = 
+					FMath::Log2(FMath::Max(1.0f, FMath::Square(SourceSize.GetMin()))) * 0.5f - 
+				    FMath::Log2(FMath::Max(1.0f, FMath::Square(ScaledSourceSize.GetMin()))) * 0.5f;
+
+				FScheduledOpData HeapData;
+				HeapData.ImageTransform.SizeX = DestSizeI.X;
+				HeapData.ImageTransform.SizeY = DestSizeI.Y;
+				FPlatformMath::StoreHalf(&HeapData.ImageTransform.ScaleXEncodedHalf, Scale.X),
+				FPlatformMath::StoreHalf(&HeapData.ImageTransform.ScaleYEncodedHalf, Scale.Y),
+				HeapData.ImageTransform.MipValue = BestMip + GlobalImageTransformLodBias;
+
+				const int32 HeapDataAddress = m_heapData.Add(HeapData);
+
+				const uint8 Mip = static_cast<uint8>(FMath::Max(0, FMath::FloorToInt(HeapData.ImageTransform.MipValue)));
+				const TArray<FScheduledOp, TFixedAllocator<4>> Deps = 
+				{
+					FScheduledOp::FromOpAndOptions(Args.base, item, Mip),
+					FScheduledOp(Args.offsetX,  item),
+					FScheduledOp(Args.offsetY,  item),
+					FScheduledOp(Args.rotation, item) 
+				};
+				
+                AddOp(FScheduledOp(item.At, item, 2, HeapDataAddress), Deps);
+
+				break;
+			}
+            case 2:
+            {
+				MUTABLE_CPUPROFILER_SCOPE(IM_TRANSFORM_2);
+			
+				const FScheduledOpData HeapData = m_heapData[item.CustomState];
+
+				const uint8 Mip = static_cast<uint8>(FMath::Max(0, FMath::FloorToInt(HeapData.ImageTransform.MipValue)));
+				Ptr<const Image> Source = LoadImage(FCacheAddress(Args.base, item.ExecutionIndex, Mip));
+
+				const FVector2f Offset = FVector2f(
+                        Args.offsetX ? LoadScalar(FCacheAddress(Args.offsetX, item)) : 0.0f,
+                        Args.offsetY ? LoadScalar(FCacheAddress(Args.offsetY, item)) : 0.0f);
+
+                const FVector2f Scale = FVector2f(
+						FPlatformMath::LoadHalf(&HeapData.ImageTransform.ScaleXEncodedHalf),
+						FPlatformMath::LoadHalf(&HeapData.ImageTransform.ScaleYEncodedHalf));
 
 				// Map Range 0-1 to a full rotation
                 const float Rotation = LoadScalar(FCacheAddress(Args.rotation, item)) * UE_TWO_PI;
-			
-				EImageFormat BaseFormat = pBaseImage->GetFormat();
-				int32 BaseLODs = pBaseImage->GetLODCount();
-				FImageSize BaseSize = pBaseImage->GetSize();
-				FImageSize SampleImageSize =  FImageSize(
-					static_cast<uint16>(FMath::Clamp(FMath::FloorToInt(float(BaseSize.X) * FMath::Abs(Scale.X)), 2, BaseSize.X)),
-					static_cast<uint16>(FMath::Clamp(FMath::FloorToInt(float(BaseSize.Y) * FMath::Abs(Scale.Y)), 2, BaseSize.Y)));
 	
-				Ptr<Image> pSampleImage = CreateImage(SampleImageSize[0], SampleImageSize[1], BaseLODs, BaseFormat, EInitializationType::NotInitialized);
-				ImOp.ImageResizeLinear( pSampleImage.get(), 0, pBaseImage.get());
-				Release(pBaseImage);
+				EImageFormat SourceFormat = Source->GetFormat();
+				EImageFormat Format = GetUncompressedFormat(SourceFormat);
+
+				if (Format != SourceFormat)
+				{
+					MUTABLE_CPUPROFILER_SCOPE(RunCode_ImageTransform_FormatFixup);	
+					Ptr<Image> Formatted = CreateImage(Source->GetSizeX(), Source->GetSizeY(), Source->GetLODCount(), Format, EInitializationType::NotInitialized);
+					bool bSuccess = false;
+					ImOp.ImagePixelFormat(bSuccess, m_pSettings->ImageCompressionQuality, Formatted.get(), Source.get());
+					check(bSuccess); 
+
+					Release(Source);
+					Source = Formatted;
+				}
+
+				if (Source->GetLODCount() < 2 && Source->GetSizeX() > 1 && Source->GetSizeY() > 1)
+				{
+					MUTABLE_CPUPROFILER_SCOPE(RunCode_ImageTransform_BilinearMipGen);
 
-				Ptr<Image> Result = CreateImage(BaseSize.X, BaseSize.Y, 1, BaseFormat, EInitializationType::NotInitialized);
-				ImageTransform(Result.get(), pSampleImage.get(), Offset, Scale, Rotation, static_cast<EAddressMode>(Args.AddressMode));
+					Ptr<Image> NewImage = CreateImage(Source->GetSizeX(), Source->GetSizeY(), 2, Source->GetFormat(), EInitializationType::NotInitialized);
 
-				Release(pSampleImage);
+					check(NewImage->GetDataSize() >= Source->GetDataSize());
+					FMemory::Memcpy(NewImage->GetData(), Source->GetData(), Source->GetDataSize());
+
+					ImageMipmapInPlace(0, NewImage.get(), FMipmapGenerationSettings{});
+
+					Release(Source);
+					Source = NewImage;
+				}
+
+				const EAddressMode AddressMode = static_cast<EAddressMode>(Args.AddressMode);
+				Ptr<Image> Result = CreateImage(
+						HeapData.ImageTransform.SizeX, HeapData.ImageTransform.SizeY, 1, Format, EInitializationType::Black);
+
+				const float MipFactor = FMath::Frac(FMath::Max(0.0f, HeapData.ImageTransform.MipValue));
+				ImageTransform(Result.get(), Source.get(), Offset, Scale, Rotation, MipFactor, AddressMode, bUseImageTransformVectorImpl);
+
+				Release(Source);
 				StoreImage(item, Result);
 
                 break;
diff --git a/Engine/Plugins/Experimental/Mutable/Source/MutableRuntime/Private/MuR/CodeRunner.h b/Engine/Plugins/Experimental/Mutable/Source/MutableRuntime/Private/MuR/CodeRunner.h
--- a/Engine/Plugins/Experimental/Mutable/Source/MutableRuntime/Private/MuR/CodeRunner.h
+++ b/Engine/Plugins/Experimental/Mutable/Source/MutableRuntime/Private/MuR/CodeRunner.h
@@ -96,10 +96,22 @@ namespace  mu
 					uint8 Mip;
 					float MipValue;
 				} RasterMesh;
+
+				struct
+				{
+					uint16 SizeX;
+					uint16 SizeY;
+					uint16 ScaleXEncodedHalf;
+					uint16 ScaleYEncodedHalf;
+					float MipValue;
+				} ImageTransform;
 			};
             Ptr<RefCounted> Resource;
         };
 
+		// Assertion to know when FScheduledOpData size changes. It is ok to modifiy if needed. 
+		static_assert(sizeof(FScheduledOpData) == 4*4 + sizeof(Ptr<RefCounted>), "FScheduledOpData size changed.");
+
 
         Ptr<RangeIndex> BuildCurrentOpRangeIndex( const FScheduledOp&, const Parameters*, const Model*, int32 ParameterIndex );
 
diff --git a/Engine/Plugins/Experimental/Mutable/Source/MutableRuntime/Private/MuR/OpImageLayer.h b/Engine/Plugins/Experimental/Mutable/Source/MutableRuntime/Private/MuR/OpImageLayer.h
--- a/Engine/Plugins/Experimental/Mutable/Source/MutableRuntime/Private/MuR/OpImageLayer.h
+++ b/Engine/Plugins/Experimental/Mutable/Source/MutableRuntime/Private/MuR/OpImageLayer.h
@@ -1729,6 +1729,9 @@ namespace mu
 			&& 
 			(bOnlyFirstLOD || pBase->GetLODCount()==1 ) )
 		{
+			check(pBlend->RelevancyMaxY < pBase->GetSizeY());
+			check(pBlend->RelevancyMaxY >= pBlend->RelevancyMinY);
+
 			int32 NumRelevantRows = pBlend->RelevancyMaxY - pBlend->RelevancyMinY + 1;
 			PixelCount = SizeX * NumRelevantRows;
 			int32 DataOffset = pBlend->RelevancyMinY * SizeX * 4;
@@ -1834,6 +1837,9 @@ namespace mu
 			&&
 			(bOnlyFirstLOD || pBase->GetLODCount() == 1))
 		{
+			check(pBlend->RelevancyMaxY < pBase->GetSizeY());
+			check(pBlend->RelevancyMaxY >= pBlend->RelevancyMinY);
+
 			int32 NumRelevantRows = pBlend->RelevancyMaxY - pBlend->RelevancyMinY + 1;
 			PixelCount = SizeX * NumRelevantRows;
 			int32 DataOffset = pBlend->RelevancyMinY * SizeX;
diff --git a/Engine/Plugins/Experimental/Mutable/Source/MutableRuntime/Private/MuR/OpImageTransform.h b/Engine/Plugins/Experimental/Mutable/Source/MutableRuntime/Private/MuR/OpImageTransform.h
--- a/Engine/Plugins/Experimental/Mutable/Source/MutableRuntime/Private/MuR/OpImageTransform.h
+++ b/Engine/Plugins/Experimental/Mutable/Source/MutableRuntime/Private/MuR/OpImageTransform.h
@@ -7,72 +7,205 @@
 
 #include "MuR/OpImageTransformPrivate.h"
 
+#include "Math/IntRect.h"
+#include "Math/Box2D.h"
+
 namespace mu
 {
-	void ImageTransform(Image* pDestImage, const Image* pImage, FVector2f Offset, FVector2f Scale, float RotationRad, EAddressMode AddressMode)
+	void ImageTransform(
+			Image* pDestImage, const Image* pImage, 
+			FVector2f Offset, FVector2f Scale, float RotationRad, 
+			float MipFactor,
+			EAddressMode AddressMode, 
+			bool bUseVectorImplCvar)
     {
-		MUTABLE_CPUPROFILER_SCOPE(ImageTransform)
+		MUTABLE_CPUPROFILER_SCOPE(ImageTransform);
+
+		if (!pImage->GetSizeX() || !pImage->GetSizeY())
+		{
+			return;
+		}
+
+		Scale.X = FMath::IsNearlyZero(Scale.X, UE_KINDA_SMALL_NUMBER) ? UE_KINDA_SMALL_NUMBER : Scale.X;
+		Scale.Y = FMath::IsNearlyZero(Scale.Y, UE_KINDA_SMALL_NUMBER) ? UE_KINDA_SMALL_NUMBER : Scale.Y;
+
+		const FTransform2f Transform = FTransform2f(FVector2f(-0.5f)).
+				Concatenate(FTransform2f(FScale2f(Scale))).
+				Concatenate(FTransform2f(FQuat2f(RotationRad))).
+				Concatenate(FTransform2f(Offset + FVector2f(0.5f)));
+
+		const FIntVector2 SrcSize  = FIntVector2(pImage->GetSizeX(), pImage->GetSizeY());
+		const FIntVector2 DestSize = FIntVector2(pDestImage->GetSizeX(), pDestImage->GetSizeY());
+
+		const FIntRect DestCropRect = Invoke([&]()
+		{
+			if (AddressMode != EAddressMode::ClampToBlack)
+			{
+				return FIntRect(0, 0, DestSize.X, DestSize.Y);
+			}
+	
+			FBox2f NormalizedCropRect(ForceInit);
+			NormalizedCropRect += Transform.TransformPoint(FVector2f(0.0f, 0.0f));
+			NormalizedCropRect += Transform.TransformPoint(FVector2f(1.0f, 0.0f));
+			NormalizedCropRect += Transform.TransformPoint(FVector2f(0.0f, 1.0f));
+			NormalizedCropRect += Transform.TransformPoint(FVector2f(1.0f, 1.0f));
+
+			const FIntRect CropRect(
+				FInt32Point(FMath::FloorToInt(NormalizedCropRect.Min.X * (float)DestSize.X), 
+							FMath::FloorToInt(NormalizedCropRect.Min.Y * (float)DestSize.Y)),
+				FInt32Point(FMath::CeilToInt (NormalizedCropRect.Max.X * (float)DestSize.X), 
+							FMath::CeilToInt (NormalizedCropRect.Max.Y * (float)DestSize.Y)));
+
+			return FIntRect(
+				FMath::Max(CropRect.Min.X, 0), FMath::Max(CropRect.Min.Y, 0),
+				FMath::Min(CropRect.Max.X, DestSize.X), FMath::Min(CropRect.Max.Y, DestSize.Y));
+		});
+
+		if (AddressMode == EAddressMode::ClampToBlack)
+		{
+			pDestImage->m_flags |= Image::IF_HAS_RELEVANCY_MAP;
+			pDestImage->RelevancyMinY = DestCropRect.Min.Y;
+			pDestImage->RelevancyMaxY = DestCropRect.Max.Y - 1;
+		}
+
+		uint8* DestData       = pDestImage->GetData();
+		const uint8* Src0Data = pImage->GetMipData(0);
+		const uint8* Src1Data = pImage->GetLODCount() > 1 ? pImage->GetMipData(1) : Src0Data;
+
+		const FIntVector2 Src0Size = FIntVector2(pImage->GetSizeX(), pImage->GetSizeY());
+		const FIntVector2 Src1Size = Src1Data != Src0Data ? pImage->CalculateMipSize(1) : Src0Size;
 
 		switch (pDestImage->GetFormat())
 		{
 		
 		case EImageFormat::IF_L_UBYTE:
 		{
-			switch (AddressMode)
+
+			if (bUseVectorImplCvar)
 			{
-			case EAddressMode::ClampToBlack:
-				ImageTransformImpl<1, EAddressMode::ClampToBlack>(pDestImage, pImage, Offset, Scale, RotationRad);
-				break;
-			case EAddressMode::ClampToEdge:
-				ImageTransformImpl<1, EAddressMode::ClampToEdge>(pDestImage, pImage, Offset, Scale, RotationRad);
-				break;
-			case EAddressMode::Wrap:
-				ImageTransformImpl<1, EAddressMode::Wrap>(pDestImage, pImage, Offset, Scale, RotationRad);
-				break;
-			default:
-				check(false);
-				break;
+				constexpr bool bUseVectorImpl = true;
+				switch (AddressMode)
+				{
+				case EAddressMode::ClampToBlack:
+					ImageTransformImpl<1, EAddressMode::ClampToBlack, bUseVectorImpl>(DestData, DestSize, DestCropRect, Src0Data, Src0Size, Src1Data, Src1Size, MipFactor, Transform);
+					break;
+				case EAddressMode::ClampToEdge:
+					ImageTransformImpl<1, EAddressMode::ClampToEdge, bUseVectorImpl>(DestData, DestSize, DestCropRect, Src0Data, Src0Size, Src1Data, Src1Size, MipFactor, Transform);
+					break;
+				case EAddressMode::Wrap:
+					ImageTransformImpl<1, EAddressMode::Wrap, bUseVectorImpl>(DestData, DestSize, DestCropRect, Src0Data, Src0Size, Src1Data, Src1Size, MipFactor, Transform);
+					break;
+				default:
+					check(false);
+					break;
+				}
+			}
+			else
+			{
+				constexpr bool bUseVectorImpl = false;
+				switch (AddressMode)
+				{
+				case EAddressMode::ClampToBlack:
+					ImageTransformImpl<1, EAddressMode::ClampToBlack, bUseVectorImpl>(DestData, DestSize, DestCropRect, Src0Data, Src0Size, Src1Data, Src1Size, MipFactor, Transform);
+					break;
+				case EAddressMode::ClampToEdge:
+					ImageTransformImpl<1, EAddressMode::ClampToEdge, bUseVectorImpl>(DestData, DestSize, DestCropRect, Src0Data, Src0Size, Src1Data, Src1Size, MipFactor, Transform);
+					break;
+				case EAddressMode::Wrap:
+					ImageTransformImpl<1, EAddressMode::Wrap, bUseVectorImpl>(DestData, DestSize, DestCropRect, Src0Data, Src0Size, Src1Data, Src1Size, MipFactor, Transform);
+					break;
+				default:
+					check(false);
+					break;
+				}
 			}
 			break;
 		}
 		
 		case EImageFormat::IF_RGB_UBYTE:
 		{
-			switch (AddressMode)
+			if (bUseVectorImplCvar)
+			{
+				constexpr bool bUseVectorImpl = true;
+				switch (AddressMode)
+				{
+				case EAddressMode::ClampToBlack:
+					ImageTransformImpl<3, EAddressMode::ClampToBlack, bUseVectorImpl>(DestData, DestSize, DestCropRect, Src0Data, Src0Size, Src1Data, Src1Size, MipFactor, Transform);
+					break;
+				case EAddressMode::ClampToEdge:
+					ImageTransformImpl<3, EAddressMode::ClampToEdge, bUseVectorImpl>(DestData, DestSize, DestCropRect, Src0Data, Src0Size, Src1Data, Src1Size, MipFactor, Transform);
+					break;
+				case EAddressMode::Wrap:
+					ImageTransformImpl<3, EAddressMode::Wrap, bUseVectorImpl>(DestData, DestSize, DestCropRect, Src0Data, Src0Size, Src1Data, Src1Size, MipFactor, Transform);
+					break;
+				default:
+					check(false);
+					break;
+				}
+			}
+			else
 			{
-			case EAddressMode::ClampToBlack:
-				ImageTransformImpl<3, EAddressMode::ClampToBlack>(pDestImage, pImage, Offset, Scale, RotationRad);
-				break;
-			case EAddressMode::ClampToEdge:
-				ImageTransformImpl<3, EAddressMode::ClampToEdge>(pDestImage, pImage, Offset, Scale, RotationRad);
-				break;
-			case EAddressMode::Wrap:
-				ImageTransformImpl<3, EAddressMode::Wrap>(pDestImage, pImage, Offset, Scale, RotationRad);
-				break;
-			default:
-				check(false);
-				break;
+				constexpr bool bUseVectorImpl = false;
+				switch (AddressMode)
+				{
+				case EAddressMode::ClampToBlack:
+					ImageTransformImpl<3, EAddressMode::ClampToBlack, bUseVectorImpl>(DestData, DestSize, DestCropRect, Src0Data, Src0Size, Src1Data, Src1Size, MipFactor, Transform);
+					break;
+				case EAddressMode::ClampToEdge:
+					ImageTransformImpl<3, EAddressMode::ClampToEdge, bUseVectorImpl>(DestData, DestSize, DestCropRect, Src0Data, Src0Size, Src1Data, Src1Size, MipFactor, Transform);
+					break;
+				case EAddressMode::Wrap:
+					ImageTransformImpl<3, EAddressMode::Wrap, bUseVectorImpl>(DestData, DestSize, DestCropRect, Src0Data, Src0Size, Src1Data, Src1Size, MipFactor, Transform);
+					break;
+				default:
+					check(false);
+					break;
+				}
 			}
+		
 			break;
 		}
 
 		case EImageFormat::IF_BGRA_UBYTE:
 		case EImageFormat::IF_RGBA_UBYTE:
 		{
-			switch (AddressMode)
+			if (bUseVectorImplCvar)
+			{
+				constexpr bool bUseVectorImpl = true;
+				switch (AddressMode)
+				{
+				case EAddressMode::ClampToBlack:
+					ImageTransformImpl<4, EAddressMode::ClampToBlack, bUseVectorImpl>(DestData, DestSize, DestCropRect, Src0Data, Src0Size, Src1Data, Src1Size, MipFactor, Transform);
+					break;
+				case EAddressMode::ClampToEdge:
+					ImageTransformImpl<4, EAddressMode::ClampToEdge, bUseVectorImpl>(DestData, DestSize, DestCropRect, Src0Data, Src0Size, Src1Data, Src1Size, MipFactor, Transform);
+					break;
+				case EAddressMode::Wrap:
+					ImageTransformImpl<4, EAddressMode::Wrap, bUseVectorImpl>(DestData, DestSize, DestCropRect, Src0Data, Src0Size, Src1Data, Src1Size, MipFactor, Transform);
+					break;
+				default:
+					check(false);
+					break;
+				}
+			}
+			else
 			{
-			case EAddressMode::ClampToBlack:
-				ImageTransformImpl<4, EAddressMode::ClampToBlack>(pDestImage, pImage, Offset, Scale, RotationRad);
-				break;
-			case EAddressMode::ClampToEdge:
-				ImageTransformImpl<4, EAddressMode::ClampToEdge>(pDestImage, pImage, Offset, Scale, RotationRad);
-				break;
-			case EAddressMode::Wrap:
-				ImageTransformImpl<4, EAddressMode::Wrap>(pDestImage, pImage, Offset, Scale, RotationRad);
-				break;
-			default:
-				check(false);
-				break;
+				constexpr bool bUseVectorImpl = false;
+				switch (AddressMode)
+				{
+				case EAddressMode::ClampToBlack:
+					ImageTransformImpl<4, EAddressMode::ClampToBlack, bUseVectorImpl>(DestData, DestSize, DestCropRect, Src0Data, Src0Size, Src1Data, Src1Size, MipFactor, Transform);
+					break;
+				case EAddressMode::ClampToEdge:
+					ImageTransformImpl<4, EAddressMode::ClampToEdge, bUseVectorImpl>(DestData, DestSize, DestCropRect, Src0Data, Src0Size, Src1Data, Src1Size, MipFactor, Transform);
+					break;
+				case EAddressMode::Wrap:
+					ImageTransformImpl<4, EAddressMode::Wrap, bUseVectorImpl>(DestData, DestSize, DestCropRect, Src0Data, Src0Size, Src1Data, Src1Size, MipFactor, Transform);
+					break;
+				default:
+					check(false);
+					break;
+				}
 			}
 			break;
 		}
diff --git a/Engine/Plugins/Experimental/Mutable/Source/MutableRuntime/Private/MuR/OpImageTransformPrivate.cpp b/Engine/Plugins/Experimental/Mutable/Source/MutableRuntime/Private/MuR/OpImageTransformPrivate.cpp
--- a/Engine/Plugins/Experimental/Mutable/Source/MutableRuntime/Private/MuR/OpImageTransformPrivate.cpp
+++ b/Engine/Plugins/Experimental/Mutable/Source/MutableRuntime/Private/MuR/OpImageTransformPrivate.cpp
@@ -1,8 +1,12 @@
 // Copyright Epic Games, Inc. All Rights Reserved.
 
 #include "OpImageTransformPrivate.h"
+
+#include "Async/ParallelFor.h"
+
 #include "Math/Vector.h"
 #include "Math/IntVector.h"
+#include "Math/TransformCalculus.h"
 
 namespace mu
 {
@@ -11,119 +15,137 @@ namespace mu
     //! Image transform, translate, scale, rotate 
     //---------------------------------------------------------------------------------------------
 
-    template<int32 NC, EAddressMode AddressMode>
-	void ImageTransformImpl(Image* pDest, const Image* pBase, FVector2f Offset, FVector2f Scale, float RotationRad)
+    template<int32 NumChannels, EAddressMode AddressMode>
+	void ImageTransformNonVectorImpl(
+		uint8* DestData, FIntVector2 DestSize, FIntRect DestCropRect, 
+		const uint8* Src0Data, FIntVector2 Src0Size, const uint8* Src1Data, FIntVector2 Src1Size,
+		float MipFactor,
+		const FTransform2f& Transform)
     {
-		using FUInt16Vector2 = UE::Math::TIntVector2<uint16>;
+		using FUint16Vector2 = UE::Math::TIntVector2<uint16>;
 		struct FPixelData
 		{
-			alignas(8) uint16 Data[NC];
+			alignas(8) uint16 Data[NumChannels];
 		};
 
-		FUInt16Vector2 BaseSize = FUInt16Vector2(pBase->GetSizeX(), pBase->GetSizeY());
-		FUInt16Vector2 DestSize = FUInt16Vector2(pDest->GetSizeX(), pDest->GetSizeY());
 
-		const FVector2f BaseSizeF = FVector2f(BaseSize.X, BaseSize.Y);
+		const FVector2f BaseSizeF = FVector2f(Src0Size.X, Src1Size.Y);
 		const FVector2f DestSizeF = FVector2f(DestSize.X, DestSize.Y);
 
-		float SinRot, CosRot;
-		FMath::SinCos(&SinRot, &CosRot, RotationRad);
-
 		const FVector2f DestNormFactor = FVector2f(1.0f) / DestSizeF;
 		const FVector2f BaseNormFactor = FVector2f(1.0f) / BaseSizeF;
 
-		Scale.X = FMath::IsNearlyZero(Scale.X) ? UE_SMALL_NUMBER : Scale.X;
-		Scale.Y = FMath::IsNearlyZero(Scale.Y) ? UE_SMALL_NUMBER : Scale.Y;
+		check(DestCropRect.Min.X >= 0 && DestCropRect.Max.X <= DestSize.X);
+		check(DestCropRect.Min.Y >= 0 && DestCropRect.Max.Y <= DestSize.Y);
+
+		const FTransform2f InverseTransform = Transform.Inverse();
+
+		constexpr int32 BatchSize = 1 << 12;
 
-		const FVector2f ScaleReciprocal = FVector2f(1.0f) / Scale; 
-		const FVector2f ScaledOffset = Offset * ScaleReciprocal;
+		const int32 NumBatchRows = FMath::DivideAndRoundUp(BatchSize, DestCropRect.Width());
+		const int32 NumBatches   = FMath::DivideAndRoundUp(DestCropRect.Height(), NumBatchRows);
 
-		for (int32 Y = 0; Y < DestSize.Y; ++Y)
+		auto ProcessBatch = [=](int32 BatchId)
 		{
-			for (int32 X = 0; X < DestSize.X; ++X)
-			{
-				FVector2f DestUv = (FVector2f(X, Y) + 0.5f) * DestNormFactor - 0.5f;	
+			const int32 RowOffset = NumBatchRows * BatchId;
+			const FIntRect BatchRect = FIntRect(
+				DestCropRect.Min.X, DestCropRect.Min.Y + RowOffset,
+				DestCropRect.Max.X, FMath::Min(DestCropRect.Min.Y + RowOffset + NumBatchRows, DestCropRect.Max.Y));
 
-				DestUv -= Offset;
-				DestUv = FVector2f(CosRot * DestUv.X + SinRot * DestUv.Y,
-								   CosRot * DestUv.Y - SinRot * DestUv.X);
-				DestUv *= ScaleReciprocal;
-				DestUv = DestUv + 0.5f;
+			const FVector2f Size0F(Src0Size.X, Src0Size.Y);
+			const FVector2f Size1F(Src1Size.X, Src1Size.Y);
 
-				// TODO: Make the bloack edge antialiazed.
-				if constexpr (AddressMode == EAddressMode::ClampToBlack)
+			for (int32 Y = BatchRect.Min.Y; Y < BatchRect.Max.Y; ++Y)
+			{
+				for (int32 X = BatchRect.Min.X; X < BatchRect.Max.X; ++X)
 				{
-					if ((DestUv.X < 0.0f) | (DestUv.X > 1.0f) | (DestUv.Y < 0.0f) | (DestUv.Y > 1.0f))
+					const FVector2f DestUv = InverseTransform.TransformPoint((FVector2f(X, Y) + 0.5f) * DestNormFactor);
+
+					// TODO: Make the black edge antialiazed?
+					if constexpr (AddressMode == EAddressMode::ClampToBlack)
 					{
-						for (uint32 Channel = 0; Channel < NC; ++Channel)
+						if ((DestUv.X < 0.0f) | (DestUv.X > 1.0f) | (DestUv.Y < 0.0f) | (DestUv.Y > 1.0f))
 						{
-							pDest->GetData()[(Y*DestSize.X + X)*NC + Channel] = 0;
+							if constexpr (NumChannels == 4)
+							{
+								*reinterpret_cast<uint32*>(&DestData[(Y * DestSize.X + X) * NumChannels]) = 0;
+							}
+							else
+							{
+								for (uint32 Channel = 0; Channel < NumChannels; ++Channel)
+								{
+									DestData[(Y * DestSize.X + X) * NumChannels + Channel] = 0;
+								}
+							}
+							continue;
 						}
-
-						continue;
 					}
-				}
 
-				const FVector2f Uv = Invoke([&]() -> FVector2f
-				{
-					if constexpr (AddressMode == EAddressMode::Wrap)
-					{
-						return FVector2f(FMath::Frac(DestUv.X), FMath::Frac(DestUv.Y));
-					}
-					else if (AddressMode == EAddressMode::ClampToBlack || AddressMode == EAddressMode::ClampToEdge)
+					const FVector2f Uv = Invoke([&]() -> FVector2f
 					{
-						return FVector2f(FMath::Clamp(DestUv.X, 0.0f, 1.0f), FMath::Clamp(DestUv.Y, 0.0f, 1.0f));
-					}
-					else
-					{
-						return DestUv;
-					}
-				});
+						if constexpr (AddressMode == EAddressMode::Wrap)
+						{
+							return FVector2f(FMath::Frac(DestUv.X), FMath::Frac(DestUv.Y));
+						}
+						else if (AddressMode == EAddressMode::ClampToBlack || AddressMode == EAddressMode::ClampToEdge)
+						{
+							return FVector2f(FMath::Clamp(DestUv.X, 0.0f, 1.0f), FMath::Clamp(DestUv.Y, 0.0f, 1.0f));
+						}
+						else
+						{
+							return DestUv;
+						}
+					});
 
-				auto SampleImageBilinear = [&](FVector2f Uv, FUInt16Vector2 Size, const uint8* DataPtr)
-				{
-					auto ComputeInterpolator = [](float T) -> uint16
-					{
-						return static_cast<uint16>(255.0f * T);
-					};
+					const FVector2f Coords0F = FVector2f(
+						FMath::Clamp(Uv.X * Size0F.X - 0.5f, 0.0f, Size0F.X - 1.0f),
+						FMath::Clamp(Uv.Y * Size0F.Y - 0.5f, 0.0f, Size0F.Y - 1.0f));
 
-					const FVector2f SizeF(Size.X, Size.Y);
+					const FUint16Vector2 Frac0 = FUint16Vector2(FMath::Frac(Coords0F.X) * 255.0f, FMath::Frac(Coords0F.Y) * 255.0f);
 
-					const FVector2f CoordsF = FVector2f(
-						FMath::Clamp(Uv.X * SizeF.X - 0.5f, 0.0f, SizeF.X - 1.0f),
-						FMath::Clamp(Uv.Y * SizeF.Y - 0.5f, 0.0f, SizeF.Y - 1.0f));
+					const FIntVector2 Coords0 = FIntVector2(Coords0F.X, Coords0F.Y);
+					const FIntVector2 Coords0PlusOne = Invoke([&]() -> FIntVector2
+					{
+						if constexpr (AddressMode == EAddressMode::Wrap)
+						{
+							return FIntVector2(Coords0.X + 1 < Src0Size.X ? Coords0.X + 1 : 0,
+								Coords0.Y + 1 < Src0Size.Y ? Coords0.Y + 1 : 0);
+						}
+						else if (AddressMode == EAddressMode::ClampToBlack || AddressMode == EAddressMode::ClampToEdge)
+						{
+							return FIntVector2(FMath::Min(Src0Size.X - 1, Coords0.X + 1),
+								FMath::Min(Src0Size.Y - 1, Coords0.Y + 1));
+						}
+					});
 
-					const FUInt16Vector2 Frac = FUInt16Vector2(
-						ComputeInterpolator(FMath::Frac(CoordsF.X)),
-						ComputeInterpolator(FMath::Frac(CoordsF.Y)));
+					const FVector2f Coords1F = FVector2f(
+						FMath::Clamp(Uv.X * Size1F.X - 0.5f, 0.0f, Size1F.X - 1.0f),
+						FMath::Clamp(Uv.Y * Size1F.Y - 0.5f, 0.0f, Size1F.Y - 1.0f));
 
-					const FIntVector2 Coords = FIntVector2(CoordsF.X, CoordsF.Y);
-					const FIntVector2 CoordsPlusOne = Invoke([&]() -> FIntVector2
-					{	
+					const FUint16Vector2 Frac1 = FUint16Vector2(FMath::Frac(Coords1F.X) * 255.0f, FMath::Frac(Coords1F.Y) * 255.0f);
+					
+					const FIntVector2 Coords1 = FIntVector2(Coords1F.X, Coords1F.Y);
+					const FIntVector2 Coords1PlusOne = Invoke([&]() -> FIntVector2
+					{
 						if constexpr (AddressMode == EAddressMode::Wrap)
 						{
-							return FIntVector2(Coords.X + 1 < Size.X ? Coords.X + 1 : 0, 
-											   Coords.Y + 1 < Size.Y ? Coords.Y + 1 : 0);
+							return FIntVector2(Coords1.X + 1 < Src1Size.X ? Coords1.X + 1 : 0,
+								Coords1.Y + 1 < Src1Size.Y ? Coords1.Y + 1 : 0);
 						}
 						else if (AddressMode == EAddressMode::ClampToBlack || AddressMode == EAddressMode::ClampToEdge)
 						{
-							return FIntVector2(FMath::Min(Size.X - 1, Coords.X + 1), 
-											   FMath::Min(Size.Y - 1, Coords.Y + 1));
+							return FIntVector2(FMath::Min(Src1Size.X - 1, Coords1.X + 1),
+								FMath::Min(Src1Size.Y - 1, Coords1.Y + 1));
 						}
 					});
-						
-					uint8 const * const Pixel00Ptr = DataPtr + (Coords.Y        * Size.X + Coords.X)        * NC;
-					uint8 const * const Pixel10Ptr = DataPtr + (Coords.Y        * Size.X + CoordsPlusOne.X) * NC;
-					uint8 const * const Pixel01Ptr = DataPtr + (CoordsPlusOne.Y * Size.X + Coords.X)        * NC;
-					uint8 const * const Pixel11Ptr = DataPtr + (CoordsPlusOne.Y * Size.X + CoordsPlusOne.X) * NC;
 
 					auto LoadPixel = [](const uint8* Ptr) -> FPixelData
 					{
 						FPixelData Result;
-						if constexpr (NC == 4)
+						if constexpr (NumChannels == 4)
 						{
 							const uint32 PackedData = *reinterpret_cast<const uint32*>(Ptr);
-							
+
 							Result.Data[0] = static_cast<uint16>((PackedData >> (8 * 0)) & 0xFF);
 							Result.Data[1] = static_cast<uint16>((PackedData >> (8 * 1)) & 0xFF);
 							Result.Data[2] = static_cast<uint16>((PackedData >> (8 * 2)) & 0xFF);
@@ -131,7 +153,7 @@ namespace mu
 						}
 						else
 						{
-							for (int32 C = 0; C < NC; ++C)
+							for (int32 C = 0; C < NumChannels; ++C)
 							{
 								Result.Data[C] = static_cast<uint16>(Ptr[C]);
 							}
@@ -140,43 +162,292 @@ namespace mu
 						return Result;
 					};
 
-					FPixelData PixelData00 = LoadPixel(Pixel00Ptr);
-					FPixelData PixelData10 = LoadPixel(Pixel10Ptr);
-					FPixelData PixelData01 = LoadPixel(Pixel01Ptr);
-					FPixelData PixelData11 = LoadPixel(Pixel11Ptr);
+					FPixelData PixelData000 = LoadPixel(Src0Data + (Coords0.Y * Src0Size.X + Coords0.X) * NumChannels);
+					FPixelData PixelData010 = LoadPixel(Src0Data + (Coords0.Y * Src0Size.X + Coords0PlusOne.X) * NumChannels);
+					FPixelData PixelData001 = LoadPixel(Src0Data + (Coords0PlusOne.Y * Src0Size.X + Coords0.X) * NumChannels);
+					FPixelData PixelData011 = LoadPixel(Src0Data + (Coords0PlusOne.Y * Src0Size.X + Coords0PlusOne.X) * NumChannels);
+                                                                                                                                  
+					FPixelData PixelData100 = LoadPixel(Src1Data + (Coords1.Y * Src1Size.X + Coords1.X) * NumChannels);
+					FPixelData PixelData110 = LoadPixel(Src1Data + (Coords1.Y * Src1Size.X + Coords1PlusOne.X) * NumChannels);
+					FPixelData PixelData101 = LoadPixel(Src1Data + (Coords1PlusOne.Y * Src1Size.X + Coords1.X) * NumChannels);
+					FPixelData PixelData111 = LoadPixel(Src1Data + (Coords1PlusOne.Y * Src1Size.X + Coords1PlusOne.X) * NumChannels);
 
-					FPixelData FilteredPixelData;
-				
-					for (int32 C = 0; C < NC; ++C)
+					FPixelData ResultPixelData;
+					for (int32 C = 0; C < NumChannels; ++C)
 					{
-						const uint16 LerpY0 = ((PixelData10.Data[C] * Frac.X) + PixelData00.Data[C] * (255 - Frac.X)) / 255;
-						const uint16 LerpY1 = ((PixelData11.Data[C] * Frac.X) + PixelData01.Data[C] * (255 - Frac.X)) / 255;
-						FilteredPixelData.Data[C] = ((LerpY1 * Frac.Y) + LerpY0*(255 - Frac.Y)) / 255;
-					} 
+						const uint16 LerpY00 = ((PixelData010.Data[C] * Frac0.X) + PixelData000.Data[C] * (255 - Frac0.X)) / 255;
+						const uint16 LerpY01 = ((PixelData011.Data[C] * Frac0.X) + PixelData001.Data[C] * (255 - Frac0.X)) / 255;
+
+						const uint16 LerpY10 = ((PixelData110.Data[C] * Frac1.X) + PixelData100.Data[C] * (255 - Frac1.X)) / 255;
+						const uint16 LerpY11 = ((PixelData111.Data[C] * Frac1.X) + PixelData101.Data[C] * (255 - Frac1.X)) / 255;
+						
+						const uint16 Mip0 = ((LerpY01 * Frac0.Y) + LerpY00 * (255 - Frac0.Y)) / 255;
+						const uint16 Mip1 = ((LerpY11 * Frac1.Y) + LerpY10 * (255 - Frac1.Y)) / 255;
+
+						ResultPixelData.Data[C] = ((Mip0 * MipFactor) + Mip1 * (255 - MipFactor)) / 255;
+					}
 
-					return FilteredPixelData;
-				};
+					static_assert(NumChannels <= 4);
+
+					if constexpr (NumChannels == 4)
+					{
+						const uint32 PackedPixel = 
+							(static_cast<uint32>((ResultPixelData.Data[0]) & 0xFF) << (8 * 0)) | 
+							(static_cast<uint32>((ResultPixelData.Data[1]) & 0xFF) << (8 * 1)) | 
+							(static_cast<uint32>((ResultPixelData.Data[2]) & 0xFF) << (8 * 2)) | 
+							(static_cast<uint32>((ResultPixelData.Data[3]) & 0xFF) << (8 * 3));
+						*reinterpret_cast<uint32*>(&DestData[(Y * DestSize.X + X) * NumChannels]) = PackedPixel;
+					}
+					else
+					{
+						for (uint32 Channel = 0; Channel < NumChannels; ++Channel)
+						{
+							DestData[(Y * DestSize.X + X) * NumChannels + Channel] = 
+								static_cast<uint8>(ResultPixelData.Data[Channel]);
+						}
+					}
+				}
+			}
+		};
+
+		if (NumBatches == 1)
+		{
+			ProcessBatch(0);
+		}
+		else if (NumBatches > 1)
+		{
+			ParallelFor(NumBatches, ProcessBatch);
+		}
+    }
+
+    template<int32 NumChannels, EAddressMode AddressMode>
+	void ImageTransformVectorImpl(
+		uint8* DestData, FIntVector2 DestSize, FIntRect DestCropRect, 
+		const uint8* Src0Data, FIntVector2 Src0Size,
+		const uint8* Src1Data, FIntVector2 Src1Size,
+		float MipFactor,
+		const FTransform2f& Transform)
+    {
+		const FVector2f DestSizeF = FVector2f(DestSize.X, DestSize.Y);
+		const FVector2f DestNormFactor = FVector2f(1.0f) / DestSizeF;
 
-				FPixelData PixelData = SampleImageBilinear(Uv, BaseSize, pBase->GetData());
-				static_assert(NC <= 4);
+		check(DestCropRect.Min.X >= 0 && DestCropRect.Max.X <= DestSize.X);
+		check(DestCropRect.Min.Y >= 0 && DestCropRect.Max.Y <= DestSize.Y);
 
-				uint8 * const DestData = pDest->GetData();
-				for (uint32 Channel = 0; Channel < NC; ++Channel)
+		const FTransform2f InverseTransform = Transform.Inverse();
+
+		using namespace GlobalVectorConstants;
+
+		const VectorRegister4Int SizeI = 
+			MakeVectorRegisterInt(Src0Size.X, Src0Size.Y, Src1Size.X, Src1Size.Y);
+
+		const VectorRegister4Float SizeF = VectorIntToFloat(SizeI);
+
+		constexpr int32 BatchSize = 1 << 12;
+
+		const int32 NumBatchRows = FMath::DivideAndRoundUp(BatchSize, DestCropRect.Width());
+		const int32 NumBatches   = FMath::DivideAndRoundUp(DestCropRect.Height(), NumBatchRows);
+
+		auto ProcessBatch = [=](int32 BatchId)
+		{
+			const int32 RowOffset = NumBatchRows * BatchId;
+			const FIntRect BatchRect = FIntRect(
+				DestCropRect.Min.X, DestCropRect.Min.Y + RowOffset,
+				DestCropRect.Max.X, FMath::Min(DestCropRect.Min.Y + RowOffset + NumBatchRows, DestCropRect.Max.Y));
+			
+			const VectorRegister4Float VectorMipFactor = VectorSetFloat1(MipFactor); 
+
+			for (int32 Y = BatchRect.Min.Y; Y < BatchRect.Max.Y; ++Y)
+			{
+				for (int32 X = BatchRect.Min.X; X < BatchRect.Max.X; ++X)
 				{
-					DestData[(Y*DestSize.X + X)*NC + Channel] = PixelData.Data[Channel];
+					const FVector2f DestUv = InverseTransform.TransformPoint((FVector2f(X, Y) + 0.5f) * DestNormFactor);
+
+					// TODO: Make the black edge antialiazed.
+					if constexpr (AddressMode == EAddressMode::ClampToBlack)
+					{
+						if ((DestUv.X < 0.0f) | (DestUv.X > 1.0f) | (DestUv.Y < 0.0f) | (DestUv.Y > 1.0f))
+						{
+							if constexpr (NumChannels == 4)
+							{
+								*reinterpret_cast<uint32*>(&DestData[(Y * DestSize.X + X) * NumChannels]) = 0;
+							}
+							else
+							{
+								for (uint32 Channel = 0; Channel < NumChannels; ++Channel)
+								{
+									DestData[(Y * DestSize.X + X) * NumChannels + Channel] = 0;
+								}
+							}
+
+							continue;
+						}
+					}
+
+					const FVector2f Uv = Invoke([&]() -> FVector2f
+					{
+						if constexpr (AddressMode == EAddressMode::Wrap)
+						{
+							return FVector2f(FMath::Frac(DestUv.X), FMath::Frac(DestUv.Y));
+						}
+						else if (AddressMode == EAddressMode::ClampToBlack || AddressMode == EAddressMode::ClampToEdge)
+						{
+							return FVector2f(FMath::Clamp(DestUv.X, 0.0f, 1.0f), FMath::Clamp(DestUv.Y, 0.0f, 1.0f));
+						}
+						else
+						{
+							return DestUv;
+						}
+					});
+
+					const VectorRegister4Float CoordsF = VectorMin(VectorMax(
+						VectorMultiplyAdd(MakeVectorRegister(Uv.X, Uv.Y, Uv.X, Uv.Y), SizeF, FloatMinusOneHalf),
+						FloatZero), 
+						VectorSubtract(SizeF, FloatOne));
+
+					const VectorRegister4Float Frac = VectorSubtract(CoordsF, VectorFloor(CoordsF));
+					
+					const VectorRegister4Int CoordsI = VectorFloatToInt(CoordsF);
+					const VectorRegister4Int CoordsIPlusOne = VectorIntMin(
+							VectorIntAdd(CoordsI, IntOne), VectorIntSubtract(SizeI, IntOne));
+
+					alignas(VectorRegister4Int) int32 CoordsIData[4];
+					VectorIntStoreAligned(CoordsI, CoordsIData);
+
+					alignas(VectorRegister4Int) int32 CoordsIPlusOneData[4];
+					VectorIntStoreAligned(CoordsIPlusOne, CoordsIPlusOneData);
+
+					auto LoadPixel = [](const uint8* DataPtr) -> VectorRegister4Float
+					{
+						if constexpr (NumChannels == 4)
+						{
+							const uint32 PackedData = *reinterpret_cast<const uint32*>(DataPtr);
+							return VectorIntToFloat(MakeVectorRegisterInt(
+										(PackedData >> (8 * 0)) & 0xFF, 
+										(PackedData >> (8 * 1)) & 0xFF,
+										(PackedData >> (8 * 2)) & 0xFF,
+										(PackedData >> (8 * 3)) & 0xFF));
+						}
+						else
+						{
+							alignas(VectorRegister4Int) int32 PixelData[4] = {0};
+							for (int32 C = 0; C < NumChannels; ++C)
+							{
+								PixelData[C] = static_cast<int32>(DataPtr[C]);
+							}
+
+							return VectorIntToFloat(VectorIntLoadAligned(&PixelData[0]));
+						}
+					};
+
+					const VectorRegister4Float Pixel000 = LoadPixel(Src0Data + (CoordsIData[1]        * Src0Size.X + CoordsIData[0])        * NumChannels);
+					const VectorRegister4Float Pixel010 = LoadPixel(Src0Data + (CoordsIData[1]        * Src0Size.X + CoordsIPlusOneData[0]) * NumChannels);
+					const VectorRegister4Float Pixel001 = LoadPixel(Src0Data + (CoordsIPlusOneData[1] * Src0Size.X + CoordsIData[0])        * NumChannels);
+					const VectorRegister4Float Pixel011 = LoadPixel(Src0Data + (CoordsIPlusOneData[1] * Src0Size.X + CoordsIPlusOneData[0]) * NumChannels);
+                                                                                                                                                         
+					const VectorRegister4Float Pixel100 = LoadPixel(Src1Data + (CoordsIData[3]        * Src1Size.X + CoordsIData[2])        * NumChannels);
+					const VectorRegister4Float Pixel110 = LoadPixel(Src1Data + (CoordsIData[3]        * Src1Size.X + CoordsIPlusOneData[2]) * NumChannels);
+					const VectorRegister4Float Pixel101 = LoadPixel(Src1Data + (CoordsIPlusOneData[3] * Src1Size.X + CoordsIData[2])        * NumChannels);
+					const VectorRegister4Float Pixel111 = LoadPixel(Src1Data + (CoordsIPlusOneData[3] * Src1Size.X + CoordsIPlusOneData[2]) * NumChannels);
+					
+					const VectorRegister4Float Frac0X = VectorReplicate(Frac, 0);
+					const VectorRegister4Float Frac0Y = VectorReplicate(Frac, 1);
+
+					// Bilerp image 0
+					VectorRegister4Float Sample0 = VectorMultiplyAdd(Frac0X, VectorSubtract(Pixel010, Pixel000), Pixel000);
+					Sample0 = VectorMultiplyAdd(
+						Frac0Y, 
+						VectorSubtract(
+							VectorMultiplyAdd(Frac0X, VectorSubtract(Pixel011, Pixel001), Pixel001), 
+							Sample0), 
+						Sample0);
+
+					const VectorRegister4Float Frac1X = VectorReplicate(Frac, 2);
+					const VectorRegister4Float Frac1Y = VectorReplicate(Frac, 3);
+
+					// Bilerp image 1
+					VectorRegister4Float Sample1 = VectorMultiplyAdd(Frac1X, VectorSubtract(Pixel110, Pixel100), Pixel100);
+					Sample1 = VectorMultiplyAdd(
+						Frac1Y, 
+						VectorSubtract(
+							VectorMultiplyAdd(Frac1X, VectorSubtract(Pixel111, Pixel101), Pixel101), 
+							Sample1), 
+						Sample1);
+
+					const VectorRegister4Float Result =
+						VectorMultiplyAdd(VectorMipFactor, VectorSubtract(Sample1, Sample0), Sample0);
+
+					if constexpr (NumChannels == 4)
+					{
+						VectorStoreByte4(Result, &DestData[(Y * DestSize.X + X) * NumChannels]);
+						VectorResetFloatRegisters();
+					}
+					else
+					{
+						const AlignedFloat4 ResultData(VectorMin(Result, Float255));
+						for (int32 C = 0; C < NumChannels; ++C)
+						{
+							DestData[(Y * DestSize.X + X) * NumChannels + C] = static_cast<uint8>(ResultData[C]);		
+						}				
+					}
 				}
 			}
+		};
+
+		if (NumBatches == 1)
+		{
+			ProcessBatch(0);
+		}
+		else if (NumBatches > 1)
+		{
+			ParallelFor(NumBatches, ProcessBatch);
 		}
     }
 
+    template<int32 NumChannels, EAddressMode AddressMode, bool bUseVectorImpl>
+	void ImageTransformImpl(
+		uint8* DestData, FIntVector2 DestSize, FIntRect DestCropRect, 
+		const uint8* Src0Data, FIntVector2 Src0Size,
+		const uint8* Src1Data, FIntVector2 Src1Size,
+		float MipFactor,
+		const FTransform2f& Transform)
+	{
+		if constexpr (bUseVectorImpl)
+		{
+			ImageTransformVectorImpl<NumChannels, AddressMode>(
+				DestData, DestSize, DestCropRect,
+				Src0Data, Src0Size, Src1Data, Src1Size,
+				MipFactor,
+				Transform);
+		}
+		else
+		{
+			ImageTransformNonVectorImpl<NumChannels, AddressMode>(
+				DestData, DestSize, DestCropRect,
+				Src0Data, Src0Size, Src1Data, Src1Size,
+				MipFactor,
+				Transform);
+		}
+	}
+
 	// Force instantiation of used templates. A link error will occur if any of them is missing.
-	template void ImageTransformImpl<1, EAddressMode::ClampToBlack>(Image*, const Image*, FVector2f, FVector2f, float);
-	template void ImageTransformImpl<1, EAddressMode::ClampToEdge>(Image*, const Image*, FVector2f, FVector2f, float);
-	template void ImageTransformImpl<1, EAddressMode::Wrap>(Image*, const Image*, FVector2f, FVector2f, float);
-	template void ImageTransformImpl<3, EAddressMode::ClampToBlack>(Image*, const Image*, FVector2f, FVector2f, float);
-	template void ImageTransformImpl<3, EAddressMode::ClampToEdge>(Image*, const Image*, FVector2f, FVector2f, float);
-	template void ImageTransformImpl<3, EAddressMode::Wrap>(Image*, const Image*, FVector2f, FVector2f, float);
-	template void ImageTransformImpl<4, EAddressMode::ClampToBlack>(Image*, const Image*, FVector2f, FVector2f, float);
-	template void ImageTransformImpl<4, EAddressMode::ClampToEdge>(Image*, const Image*, FVector2f, FVector2f, float);
-	template void ImageTransformImpl<4, EAddressMode::Wrap>(Image*, const Image*, FVector2f, FVector2f, float);
+	template void ImageTransformImpl<1, EAddressMode::ClampToBlack, true>(uint8* DestData, FIntVector2 DestSize, FIntRect DestCropRect, const uint8* Src0Data, FIntVector2 Src0Size, const uint8* Src1Data, FIntVector2 Src1Size, float MipFactor, const FTransform2f& Transform);
+	template void ImageTransformImpl<1, EAddressMode::ClampToEdge, true> (uint8* DestData, FIntVector2 DestSize, FIntRect DestCropRect, const uint8* Src0Data, FIntVector2 Src0Size, const uint8* Src1Data, FIntVector2 Src1Size, float MipFactor, const FTransform2f& Transform);
+	template void ImageTransformImpl<1, EAddressMode::Wrap, true>        (uint8* DestData, FIntVector2 DestSize, FIntRect DestCropRect, const uint8* Src0Data, FIntVector2 Src0Size, const uint8* Src1Data, FIntVector2 Src1Size, float MipFactor, const FTransform2f& Transform);
+	template void ImageTransformImpl<3, EAddressMode::ClampToBlack, true>(uint8* DestData, FIntVector2 DestSize, FIntRect DestCropRect, const uint8* Src0Data, FIntVector2 Src0Size, const uint8* Src1Data, FIntVector2 Src1Size, float MipFactor, const FTransform2f& Transform);
+	template void ImageTransformImpl<3, EAddressMode::ClampToEdge, true> (uint8* DestData, FIntVector2 DestSize, FIntRect DestCropRect, const uint8* Src0Data, FIntVector2 Src0Size, const uint8* Src1Data, FIntVector2 Src1Size, float MipFactor, const FTransform2f& Transform);
+	template void ImageTransformImpl<3, EAddressMode::Wrap, true>        (uint8* DestData, FIntVector2 DestSize, FIntRect DestCropRect, const uint8* Src0Data, FIntVector2 Src0Size, const uint8* Src1Data, FIntVector2 Src1Size, float MipFactor, const FTransform2f& Transform);
+	template void ImageTransformImpl<4, EAddressMode::ClampToBlack, true>(uint8* DestData, FIntVector2 DestSize, FIntRect DestCropRect, const uint8* Src0Data, FIntVector2 Src0Size, const uint8* Src1Data, FIntVector2 Src1Size, float MipFactor, const FTransform2f& Transform);
+	template void ImageTransformImpl<4, EAddressMode::ClampToEdge, true> (uint8* DestData, FIntVector2 DestSize, FIntRect DestCropRect, const uint8* Src0Data, FIntVector2 Src0Size, const uint8* Src1Data, FIntVector2 Src1Size, float MipFactor, const FTransform2f& Transform);
+	template void ImageTransformImpl<4, EAddressMode::Wrap, true>        (uint8* DestData, FIntVector2 DestSize, FIntRect DestCropRect, const uint8* Src0Data, FIntVector2 Src0Size, const uint8* Src1Data, FIntVector2 Src1Size, float MipFactor, const FTransform2f& Transform);
+
+	template void ImageTransformImpl<1, EAddressMode::ClampToBlack, false>(uint8* DestData, FIntVector2 DestSize, FIntRect DestCropRect, const uint8* Src0Data, FIntVector2 Src0Size, const uint8* Src1Data, FIntVector2 Src1Size, float MipFactor, const FTransform2f& Transform);
+	template void ImageTransformImpl<1, EAddressMode::ClampToEdge, false> (uint8* DestData, FIntVector2 DestSize, FIntRect DestCropRect, const uint8* Src0Data, FIntVector2 Src0Size, const uint8* Src1Data, FIntVector2 Src1Size, float MipFactor, const FTransform2f& Transform);
+	template void ImageTransformImpl<1, EAddressMode::Wrap, false>        (uint8* DestData, FIntVector2 DestSize, FIntRect DestCropRect, const uint8* Src0Data, FIntVector2 Src0Size, const uint8* Src1Data, FIntVector2 Src1Size, float MipFactor, const FTransform2f& Transform);
+	template void ImageTransformImpl<3, EAddressMode::ClampToBlack, false>(uint8* DestData, FIntVector2 DestSize, FIntRect DestCropRect, const uint8* Src0Data, FIntVector2 Src0Size, const uint8* Src1Data, FIntVector2 Src1Size, float MipFactor, const FTransform2f& Transform);
+	template void ImageTransformImpl<3, EAddressMode::ClampToEdge, false> (uint8* DestData, FIntVector2 DestSize, FIntRect DestCropRect, const uint8* Src0Data, FIntVector2 Src0Size, const uint8* Src1Data, FIntVector2 Src1Size, float MipFactor, const FTransform2f& Transform);
+	template void ImageTransformImpl<3, EAddressMode::Wrap, false>        (uint8* DestData, FIntVector2 DestSize, FIntRect DestCropRect, const uint8* Src0Data, FIntVector2 Src0Size, const uint8* Src1Data, FIntVector2 Src1Size, float MipFactor, const FTransform2f& Transform);
+	template void ImageTransformImpl<4, EAddressMode::ClampToBlack, false>(uint8* DestData, FIntVector2 DestSize, FIntRect DestCropRect, const uint8* Src0Data, FIntVector2 Src0Size, const uint8* Src1Data, FIntVector2 Src1Size, float MipFactor, const FTransform2f& Transform);
+	template void ImageTransformImpl<4, EAddressMode::ClampToEdge, false> (uint8* DestData, FIntVector2 DestSize, FIntRect DestCropRect, const uint8* Src0Data, FIntVector2 Src0Size, const uint8* Src1Data, FIntVector2 Src1Size, float MipFactor, const FTransform2f& Transform);
+	template void ImageTransformImpl<4, EAddressMode::Wrap, false>        (uint8* DestData, FIntVector2 DestSize, FIntRect DestCropRect, const uint8* Src0Data, FIntVector2 Src0Size, const uint8* Src1Data, FIntVector2 Src1Size, float MipFactor, const FTransform2f& Transform);
 }
diff --git a/Engine/Plugins/Experimental/Mutable/Source/MutableRuntime/Private/MuR/OpImageTransformPrivate.h b/Engine/Plugins/Experimental/Mutable/Source/MutableRuntime/Private/MuR/OpImageTransformPrivate.h
--- a/Engine/Plugins/Experimental/Mutable/Source/MutableRuntime/Private/MuR/OpImageTransformPrivate.h
+++ b/Engine/Plugins/Experimental/Mutable/Source/MutableRuntime/Private/MuR/OpImageTransformPrivate.h
@@ -6,6 +6,7 @@
 
 #include "Math/Vector.h"
 #include "Math/IntVector.h"
+#include "Math/TransformCalculus2D.h"
 
 namespace mu
 {
@@ -13,17 +14,28 @@ namespace mu
     //---------------------------------------------------------------------------------------------
     //! Image transform, translate, scale, rotate 
     //---------------------------------------------------------------------------------------------
-	template<int32, EAddressMode> 
-	void ImageTransformImpl(Image*, const Image*, FVector2f, FVector2f, float);
 
+	template<int32 NumChannels, EAddressMode, bool bUseVectorImpl> 
+	void ImageTransformImpl(uint8* DestData, FIntVector2 DestSize, FIntRect DestCropRect, const uint8* Src0Data, FIntVector2 Src0Size, const uint8* Src1Data, FIntVector2 Src1Size, float MipFactor, const FTransform2f& Transform);
+	
 	// extern template could be used if for some reason we need to provide the full declaration here. 
-	//extern template FORCENOINLINE void ImageTransformImpl<1, EAddressMode::ClampToBlack>(Image*, const Image*, FVector2f, FVector2f, float);
-	//extern template FORCENOINLINE void ImageTransformImpl<1, EAddressMode::ClampToEdge>(Image*, const Image*, FVector2f, FVector2f, float);
-	//extern template FORCENOINLINE void ImageTransformImpl<1, EAddressMode::Wrap>(Image*, const Image*, FVector2f, FVector2f, float);
-	//extern template FORCENOINLINE void ImageTransformImpl<3, EAddressMode::ClampToBlack>(Image*, const Image*, FVector2f, FVector2f, float);
-	//extern template FORCENOINLINE void ImageTransformImpl<3, EAddressMode::ClampToEdge>(Image*, const Image*, FVector2f, FVector2f, float);
-	//extern template FORCENOINLINE void ImageTransformImpl<3, EAddressMode::Wrap>(Image*, const Image*, FVector2f, FVector2f, float);
-	//extern template FORCENOINLINE void ImageTransformImpl<4, EAddressMode::ClampToBlack>(Image*, const Image*, FVector2f, FVector2f, float);
-	//extern template FORCENOINLINE void ImageTransformImpl<4, EAddressMode::ClampToEdge>(Image*, const Image*, FVector2f, FVector2f, float);
-	//extern template FORCENOINLINE void ImageTransformImpl<4, EAddressMode::Wrap>(Image*, const Image*, FVector2f, FVector2f, float);
+	//extern template FORCENOINLINE void ImageTransformImpl<1, EAddressMode::ClampToBlack, false>(uint8* DestData, FIntRect DestRect, FIntRect DestCropRect, const uint8* SrcData, FIntRect SrcRect, const FTransform2f& Transform);
+	//extern template FORCENOINLINE void ImageTransformImpl<1, EAddressMode::ClampToEdge, false> (uint8* DestData, FIntRect DestRect, FIntRect DestCropRect, const uint8* SrcData, FIntRect SrcRect, const FTransform2f& Transform);
+	//extern template FORCENOINLINE void ImageTransformImpl<1, EAddressMode::Wrap, false>        (uint8* DestData, FIntRect DestRect, FIntRect DestCropRect, const uint8* SrcData, FIntRect SrcRect, const FTransform2f& Transform);
+	//extern template FORCENOINLINE void ImageTransformImpl<3, EAddressMode::ClampToBlack, false>(uint8* DestData, FIntRect DestRect, FIntRect DestCropRect, const uint8* SrcData, FIntRect SrcRect, const FTransform2f& Transform);
+	//extern template FORCENOINLINE void ImageTransformImpl<3, EAddressMode::ClampToEdge, false> (uint8* DestData, FIntRect DestRect, FIntRect DestCropRect, const uint8* SrcData, FIntRect SrcRect, const FTransform2f& Transform);
+	//extern template FORCENOINLINE void ImageTransformImpl<3, EAddressMode::Wrap, false>        (uint8* DestData, FIntRect DestRect, FIntRect DestCropRect, const uint8* SrcData, FIntRect SrcRect, const FTransform2f& Transform);
+	//extern template FORCENOINLINE void ImageTransformImpl<4, EAddressMode::ClampToBlack, false>(uint8* DestData, FIntRect DestRect, FIntRect DestCropRect, const uint8* SrcData, FIntRect SrcRect, const FTransform2f& Transform);
+	//extern template FORCENOINLINE void ImageTransformImpl<4, EAddressMode::ClampToEdge, false> (uint8* DestData, FIntRect DestRect, FIntRect DestCropRect, const uint8* SrcData, FIntRect SrcRect, const FTransform2f& Transform);
+	//extern template FORCENOINLINE void ImageTransformImpl<4, EAddressMode::Wrap, false>        (uint8* DestData, FIntRect DestRect, FIntRect DestCropRect, const uint8* SrcData, FIntRect SrcRect, const FTransform2f& Transform);
+
+	//extern template FORCENOINLINE void ImageTransformImpl<1, EAddressMode::ClampToBlack, true>(uint8* DestData, FIntRect DestRect, FIntRect DestCropRect, const uint8* SrcData, FIntRect SrcRect, const FTransform2f& Transform);
+	//extern template FORCENOINLINE void ImageTransformImpl<1, EAddressMode::ClampToEdge, true> (uint8* DestData, FIntRect DestRect, FIntRect DestCropRect, const uint8* SrcData, FIntRect SrcRect, const FTransform2f& Transform);
+	//extern template FORCENOINLINE void ImageTransformImpl<1, EAddressMode::Wrap, true>        (uint8* DestData, FIntRect DestRect, FIntRect DestCropRect, const uint8* SrcData, FIntRect SrcRect, const FTransform2f& Transform);
+	//extern template FORCENOINLINE void ImageTransformImpl<3, EAddressMode::ClampToBlack, true>(uint8* DestData, FIntRect DestRect, FIntRect DestCropRect, const uint8* SrcData, FIntRect SrcRect, const FTransform2f& Transform);
+	//extern template FORCENOINLINE void ImageTransformImpl<3, EAddressMode::ClampToEdge, true> (uint8* DestData, FIntRect DestRect, FIntRect DestCropRect, const uint8* SrcData, FIntRect SrcRect, const FTransform2f& Transform);
+	//extern template FORCENOINLINE void ImageTransformImpl<3, EAddressMode::Wrap, true>        (uint8* DestData, FIntRect DestRect, FIntRect DestCropRect, const uint8* SrcData, FIntRect SrcRect, const FTransform2f& Transform);
+	//extern template FORCENOINLINE void ImageTransformImpl<4, EAddressMode::ClampToBlack, true>(uint8* DestData, FIntRect DestRect, FIntRect DestCropRect, const uint8* SrcData, FIntRect SrcRect, const FTransform2f& Transform);
+	//extern template FORCENOINLINE void ImageTransformImpl<4, EAddressMode::ClampToEdge, true> (uint8* DestData, FIntRect DestRect, FIntRect DestCropRect, const uint8* SrcData, FIntRect SrcRect, const FTransform2f& Transform);
+	//extern template FORCENOINLINE void ImageTransformImpl<4, EAddressMode::Wrap, true>        (uint8* DestData, FIntRect DestRect, FIntRect DestCropRect, const uint8* SrcData, FIntRect SrcRect, const FTransform2f& Transform);
 }
diff --git a/Engine/Plugins/Experimental/Mutable/Source/MutableRuntime/Private/MuR/Operations.h b/Engine/Plugins/Experimental/Mutable/Source/MutableRuntime/Private/MuR/Operations.h
--- a/Engine/Plugins/Experimental/Mutable/Source/MutableRuntime/Private/MuR/Operations.h
+++ b/Engine/Plugins/Experimental/Mutable/Source/MutableRuntime/Private/MuR/Operations.h
@@ -839,6 +839,13 @@ namespace mu
 			ADDRESS rotation = 0;
 
 			uint32 AddressMode = 0;
+
+            /** Size of the image to create. If 0, reuse size from base.*/
+            uint16 SizeX = 0;
+            uint16 SizeY = 0;
+
+			uint16 SourceSizeX = 0;
+			uint16 SourceSizeY = 0;
 		};
 
         //-------------------------------------------------------------------------------------
diff --git a/Engine/Plugins/Experimental/Mutable/Source/MutableRuntime/Public/MuR/Model.h b/Engine/Plugins/Experimental/Mutable/Source/MutableRuntime/Public/MuR/Model.h
--- a/Engine/Plugins/Experimental/Mutable/Source/MutableRuntime/Public/MuR/Model.h
+++ b/Engine/Plugins/Experimental/Mutable/Source/MutableRuntime/Public/MuR/Model.h
@@ -13,7 +13,7 @@
 //! data structures. Compiled models are not necessarily compatible when the runtime is updated,
 //! so this version number can be used externally to verify this. It is not used internally, and
 //! serializing models from different versions than this runtime will probably result in a crash.
-#define MUTABLE_COMPILED_MODEL_CODE_VERSION		uint32( 73 )
+#define MUTABLE_COMPILED_MODEL_CODE_VERSION		uint32( 74 )
 #define MUTABLE_PARAMETERS_VERSION              uint32( 2 )
 
 
diff --git a/Engine/Plugins/Experimental/Mutable/Source/MutableTools/Private/MuT/ASTOpImageSwizzle.cpp b/Engine/Plugins/Experimental/Mutable/Source/MutableTools/Private/MuT/ASTOpImageSwizzle.cpp
--- a/Engine/Plugins/Experimental/Mutable/Source/MutableTools/Private/MuT/ASTOpImageSwizzle.cpp
+++ b/Engine/Plugins/Experimental/Mutable/Source/MutableTools/Private/MuT/ASTOpImageSwizzle.cpp
@@ -246,6 +246,17 @@ namespace mu
 			return at;
 		}
 
+		// If all channels are the same, and in the same order, and the source format is the same that we are
+		// setting in the swizzle, then the swizzle won't do anything.
+		if (bAllChannelsAreTheSame && bSameChannelOrder)
+		{
+			FImageDesc SourceDesc = channelSourceAt->GetImageDesc();
+			if (SourceDesc.m_format == Format)
+			{
+				return channelSourceAt;
+			}
+		}
+
 		OP_TYPE sourceType = channelSourceAt->GetOpType();
 
 		if (bAllChannelsAreTheSame && channelSourceAt)
diff --git a/Engine/Plugins/Experimental/Mutable/Source/MutableTools/Private/MuT/ASTOpImageTransform.cpp b/Engine/Plugins/Experimental/Mutable/Source/MutableTools/Private/MuT/ASTOpImageTransform.cpp
--- a/Engine/Plugins/Experimental/Mutable/Source/MutableTools/Private/MuT/ASTOpImageTransform.cpp
+++ b/Engine/Plugins/Experimental/Mutable/Source/MutableTools/Private/MuT/ASTOpImageTransform.cpp
@@ -41,13 +41,18 @@ namespace mu
 	{
 		if (auto other = dynamic_cast<const ASTOpImageTransform*>(&otherUntyped))
 		{
-			return base == other->base &&
+			return 
+				base == other->base &&
 				offsetX == other->offsetX &&
 				offsetY == other->offsetY &&
 				scaleX == other->scaleX &&
 				scaleY == other->scaleY &&
 				rotation == other->rotation &&
-				AddressMode == other->AddressMode;
+				AddressMode == other->AddressMode &&
+				SizeX == other->SizeX &&
+				SizeY == other->SizeY &&
+				SourceSizeX == other->SourceSizeX &&
+				SourceSizeY == other->SourceSizeY;
 		}
 		return false;
 	}
@@ -64,6 +69,10 @@ namespace mu
 		hash_combine(res, scaleY.child().get());
 		hash_combine(res, rotation.child().get());
 		hash_combine(res, std::hash<uint32>()(static_cast<uint32>(AddressMode)));
+		hash_combine(res, SizeX);
+		hash_combine(res, SizeY);
+		hash_combine(res, SourceSizeX);
+		hash_combine(res, SourceSizeY);
 		return res;
 	}
 
@@ -79,6 +88,10 @@ namespace mu
 		n->scaleY = mapChild(scaleY.child());
 		n->rotation = mapChild(rotation.child());
 		n->AddressMode = AddressMode;
+		n->SizeX = SizeX;
+		n->SizeY = SizeY;
+		n->SourceSizeX = SourceSizeX;
+		n->SourceSizeY = SourceSizeY;
 		return n;
 	}
 
@@ -111,7 +124,10 @@ namespace mu
 			Args.scaleY = scaleY ? scaleY->linkedAddress : 0;
 			Args.rotation = rotation ? rotation->linkedAddress : 0;
 			Args.AddressMode = static_cast<uint32>(AddressMode);
-
+			Args.SizeX = SizeX;
+			Args.SizeY = SizeY;
+			Args.SourceSizeX = SourceSizeX;
+			Args.SourceSizeY = SourceSizeY;
 
 			linkedAddress = (OP::ADDRESS)program.m_opAddress.Num();
 			program.m_opAddress.Add((uint32_t)program.m_byteCode.Num());
diff --git a/Engine/Plugins/Experimental/Mutable/Source/MutableTools/Private/MuT/ASTOpImageTransform.h b/Engine/Plugins/Experimental/Mutable/Source/MutableTools/Private/MuT/ASTOpImageTransform.h
--- a/Engine/Plugins/Experimental/Mutable/Source/MutableTools/Private/MuT/ASTOpImageTransform.h
+++ b/Engine/Plugins/Experimental/Mutable/Source/MutableTools/Private/MuT/ASTOpImageTransform.h
@@ -28,6 +28,12 @@ struct FProgram;
 		ASTChild scaleY;
 		ASTChild rotation;
 
+		uint16 SizeX;
+		uint16 SizeY;
+
+		uint16 SourceSizeX;
+		uint16 SourceSizeY;
+
 		EAddressMode AddressMode = EAddressMode::Wrap;
 
 	public:
diff --git a/Engine/Plugins/Experimental/Mutable/Source/MutableTools/Private/MuT/CodeGenerator_Image.cpp b/Engine/Plugins/Experimental/Mutable/Source/MutableTools/Private/MuT/CodeGenerator_Image.cpp
--- a/Engine/Plugins/Experimental/Mutable/Source/MutableTools/Private/MuT/CodeGenerator_Image.cpp
+++ b/Engine/Plugins/Experimental/Mutable/Source/MutableTools/Private/MuT/CodeGenerator_Image.cpp
@@ -601,24 +601,34 @@ namespace mu
 		op->scaleY = ScaleY ? ScaleY : ScaleX;
 		op->rotation = Rotation; 
 		op->AddressMode = node.AddressMode;
+		op->SizeX = node.SizeX;
+		op->SizeY = node.SizeY;
 
-        // Base image
+		// Base image
         Ptr<ASTOp> base;
+		FImageGenerationOptions NewOptions = Options;
+		NewOptions.ImageLayoutStrategy = CompilerOptions::TextureLayoutStrategy::None;
+		NewOptions.LayoutToApply = nullptr;
+		NewOptions.LayoutBlockId = -1;
+		NewOptions.RectSize = {};
+
         if ( node.m_pBase )
         {
 			FImageGenerationResult BaseResult;
-			GenerateImage(Options, BaseResult, node.m_pBase);
+			GenerateImage(NewOptions, BaseResult, node.m_pBase);
 			base = BaseResult.op;
 		}
         else
         {
             // This argument is required
-            base = GenerateMissingImageCode(TEXT("Image Transform Base"), EImageFormat::IF_RGB_UBYTE, node.m_errorContext, Options);
+            base = GenerateMissingImageCode(TEXT("Image Transform Base"), EImageFormat::IF_RGB_UBYTE, node.m_errorContext, NewOptions);
         }
-
-		EImageFormat baseFormat = base->GetImageDesc().m_format;
-        base = GenerateImageSize( base, Options.RectSize);
+		
+		FImageDesc BaseDesc = base->GetImageDesc();
+		
         op->base = base;
+		op->SourceSizeX = BaseDesc.m_size.X;
+		op->SourceSizeY = BaseDesc.m_size.Y;
 
         Result.op = op; 
     }
diff --git a/Engine/Plugins/Experimental/Mutable/Source/MutableTools/Private/MuT/CodeOptimiser_Semantic.cpp b/Engine/Plugins/Experimental/Mutable/Source/MutableTools/Private/MuT/CodeOptimiser_Semantic.cpp
--- a/Engine/Plugins/Experimental/Mutable/Source/MutableTools/Private/MuT/CodeOptimiser_Semantic.cpp
+++ b/Engine/Plugins/Experimental/Mutable/Source/MutableTools/Private/MuT/CodeOptimiser_Semantic.cpp
@@ -1283,12 +1283,18 @@ namespace mu
 				break;
 			}
 
-			case OP_TYPE::IM_TRANSFORM:
+            case OP_TYPE::IM_TRANSFORM:
 			{
-				Ptr<ASTOpImageTransform> nop = mu::Clone<ASTOpImageTransform>(at);
-				Ptr<ASTOp> maskOp = nop->base.child();
-				nop->base = Visit(maskOp, currentSinkingOp);
-				newAt = nop;
+				// It can only sink in the transform if it doesn't have it's own size.
+				const ASTOpImageTransform* typedAt = dynamic_cast<const ASTOpImageTransform*>(at.get());
+				if (typedAt->SizeX == 0 && typedAt->SizeY == 0)
+				{
+					Ptr<ASTOpImageTransform> nop = mu::Clone<ASTOpImageTransform>(at);
+					Ptr<ASTOp> maskOp = nop->base.child();
+					nop->base = Visit(maskOp, currentSinkingOp);
+					newAt = nop;
+				}
+
 				break;
 			}
 
@@ -1372,6 +1378,16 @@ namespace mu
 				at = sourceOp;
                 break;
             }
+            
+            case OP_TYPE::IM_TRANSFORM:
+			{
+				// Set the size in the children and remove resize
+				Ptr<ASTOpImageTransform> sourceOp = mu::Clone<ASTOpImageTransform>(sourceAt.get());
+				sourceOp->SizeX = op.args.ImageResize.size[0];
+				sourceOp->SizeY = op.args.ImageResize.size[1];
+				at = sourceOp;
+				break;
+			}
 
             case OP_TYPE::IM_CONDITIONAL:
             {
@@ -1671,6 +1687,20 @@ namespace mu
 				break;
 			}
 
+            case OP_TYPE::IM_TRANSFORM:
+			{
+				// We can only optimize here if we know the transform result size, otherwise, we will sink the op in the sinker.
+				const ASTOpImageTransform* typedAt = dynamic_cast<const ASTOpImageTransform*>(sourceAt.get());
+				if (typedAt->SizeX != 0 && typedAt->SizeY != 0)
+				{
+					// Set the size in the children and remove resize
+					Ptr<ASTOpImageTransform> sourceOp = mu::Clone<ASTOpImageTransform>(sourceAt.get());
+					sourceOp->SizeX = FMath::CeilToInt32(sourceOp->SizeX * op.args.ImageResizeRel.factor[0]);
+					sourceOp->SizeY = FMath::CeilToInt32(sourceOp->SizeY * op.args.ImageResizeRel.factor[1]);
+					at = sourceOp;
+				}
+				break;
+			}
 
 				// Don't combine. ResizeRel sometimes can resize more children than Resize can do. (see RasterMesh)
 				// It can be combined in an optimization step further in the process, when normal sizes may have been 
diff --git a/Engine/Plugins/Experimental/Mutable/Source/MutableTools/Private/MuT/NodeImageTransform.cpp b/Engine/Plugins/Experimental/Mutable/Source/MutableTools/Private/MuT/NodeImageTransform.cpp
--- a/Engine/Plugins/Experimental/Mutable/Source/MutableTools/Private/MuT/NodeImageTransform.cpp
+++ b/Engine/Plugins/Experimental/Mutable/Source/MutableTools/Private/MuT/NodeImageTransform.cpp
@@ -155,5 +155,29 @@ namespace mu
 	{
 		return m_pD->AddressMode;
 	}
+
+	//---------------------------------------------------------------------------------------------
+	uint16 NodeImageTransform::GetSizeX() const
+	{
+		return m_pD->SizeX;
+	}
+
+	//---------------------------------------------------------------------------------------------
+	void NodeImageTransform::SetSizeX(uint16 SizeX)
+	{
+		m_pD->SizeX = SizeX;
+	}
+
+	//---------------------------------------------------------------------------------------------
+	uint16 NodeImageTransform::GetSizeY() const
+	{
+		return m_pD->SizeY;
+	}
+
+	//---------------------------------------------------------------------------------------------
+	void NodeImageTransform::SetSizeY(uint16 SizeY)
+	{
+		m_pD->SizeY = SizeY;
+	}
 }
 
diff --git a/Engine/Plugins/Experimental/Mutable/Source/MutableTools/Private/MuT/NodeImageTransformPrivate.h b/Engine/Plugins/Experimental/Mutable/Source/MutableTools/Private/MuT/NodeImageTransformPrivate.h
--- a/Engine/Plugins/Experimental/Mutable/Source/MutableTools/Private/MuT/NodeImageTransformPrivate.h
+++ b/Engine/Plugins/Experimental/Mutable/Source/MutableTools/Private/MuT/NodeImageTransformPrivate.h
@@ -31,11 +31,14 @@ namespace mu
 		NodeScalarPtr m_pRotation;
 
 		EAddressMode AddressMode = EAddressMode::Wrap;
+		uint32 SizeX = 0;
+		uint32 SizeY = 0;
+
 
 		//!
 		void Serialise( OutputArchive& arch ) const
 		{
-            uint32 ver = 1;
+            uint32 ver = 2;
 			arch << ver;
 
 			arch << m_pBase;
@@ -45,6 +48,8 @@ namespace mu
 			arch << m_pScaleY;
 			arch << m_pRotation;
 			arch << static_cast<uint32>(AddressMode);
+			arch << SizeX;
+			arch << SizeY;
 		}
 
 		//!
@@ -52,7 +57,7 @@ namespace mu
 		{
             uint32 ver;
 			arch >> ver;
-			check(ver <= 1);
+			check(ver <= 2);
 
 			arch >> m_pBase;
 			arch >> m_pOffsetX;
@@ -62,12 +67,18 @@ namespace mu
 			arch >> m_pRotation;
 
 			uint32 AddressModeValue = static_cast<uint32>(EAddressMode::Wrap);
-			if (ver == 1)
+			if (ver >= 1)
 			{
 				arch >> AddressModeValue;
 			}
 
 			AddressMode = static_cast<EAddressMode>(AddressModeValue);
+
+			if (ver >= 2)
+			{
+				arch >> SizeX;
+				arch >> SizeY;
+			}
 		}
 	};
 
diff --git a/Engine/Plugins/Experimental/Mutable/Source/MutableTools/Public/MuT/NodeImageTransform.h b/Engine/Plugins/Experimental/Mutable/Source/MutableTools/Public/MuT/NodeImageTransform.h
--- a/Engine/Plugins/Experimental/Mutable/Source/MutableTools/Public/MuT/NodeImageTransform.h
+++ b/Engine/Plugins/Experimental/Mutable/Source/MutableTools/Public/MuT/NodeImageTransform.h
@@ -75,6 +75,13 @@ namespace mu
 
 		EAddressMode GetAddressMode() const;
 		void SetAddressMode(EAddressMode AddressMode);
+
+		uint16 GetSizeX() const;
+		void SetSizeX(uint16 SizeX); 
+
+		uint16 GetSizeY() const;
+		void SetSizeY(uint16 SizeY); 
+
 		//-----------------------------------------------------------------------------------------
 		// Interface pattern
 		//-----------------------------------------------------------------------------------------
