[Mutable] Fix crash when accessing the descriptor's requested LODs.

#rnx

diff --git a/Engine/Plugins/Experimental/Mutable/Source/CustomizableObject/Private/MuCO/CustomizableObjectInstance.cpp b/Engine/Plugins/Experimental/Mutable/Source/CustomizableObject/Private/MuCO/CustomizableObjectInstance.cpp
--- a/Engine/Plugins/Experimental/Mutable/Source/CustomizableObject/Private/MuCO/CustomizableObjectInstance.cpp
+++ b/Engine/Plugins/Experimental/Mutable/Source/CustomizableObject/Private/MuCO/CustomizableObjectInstance.cpp
@@ -6341,7 +6341,7 @@ void UCustomizableObjectInstance::SetRequestedLODs(int32 InMinLOD, int32 InMaxLO
 	if (UCustomizableObjectSystem::GetInstance()->IsOnlyGenerateRequestedLODsEnabled())
 	{
 		const int32 ComponentCount = GetNumComponents();
-		MutableUpdateCandidate.RequestedLODLevels.SetNum(ComponentCount);
+		MutableUpdateCandidate.RequestedLODLevels.SetNumZeroed(ComponentCount);
 
 		const TArray<uint16>& GeneratedLODsPerComponent = DescriptorRuntimeHash.GetRequestedLODs();
 
@@ -6377,7 +6377,7 @@ void UCustomizableObjectInstance::SetRequestedLODs(int32 InMinLOD, int32 InMaxLO
 					RequestedLODs |= (1 << LODIndex);
 				}
 				
-				bUpdateRequestedLODs |= (RequestedLODs != Descriptor.RequestedLODLevels[ComponentIndex]);
+				bUpdateRequestedLODs |= (RequestedLODs != MutableUpdateCandidate.RequestedLODLevels[ComponentIndex]);
 
 				// Save new RequestedLODs
 				MutableUpdateCandidate.RequestedLODLevels[ComponentIndex] = RequestedLODs;
