Free texture platform data before assignment in case of double-assign.
#b fabian.giesen

diff --git a/Engine/Plugins/Experimental/Mutable/Source/CustomizableObject/Private/MuCO/CustomizableObjectInstance.cpp b/Engine/Plugins/Experimental/Mutable/Source/CustomizableObject/Private/MuCO/CustomizableObjectInstance.cpp
--- a/Engine/Plugins/Experimental/Mutable/Source/CustomizableObject/Private/MuCO/CustomizableObjectInstance.cpp
+++ b/Engine/Plugins/Experimental/Mutable/Source/CustomizableObject/Private/MuCO/CustomizableObjectInstance.cpp
@@ -5783,11 +5783,6 @@ void UCustomizableInstancePrivateData::BuildMaterials(const TSharedPtr<FMutableO
 											check(PlatformData->Mips[0].SizeY == MutableTexture->GetPlatformData()->Mips[0].SizeY);
 										}
 
-										if (MutableTexture->GetPlatformData())
-										{
-											delete MutableTexture->GetPlatformData();
-										}
-
 										MutableTexture->SetPlatformData(PlatformData);
 										OperationData->ImageToPlatformDataMap.Remove(Image.ImageID);
 									}
diff --git a/Engine/Source/Runtime/Engine/Classes/Engine/Texture2DArray.cpp b/Engine/Source/Runtime/Engine/Classes/Engine/Texture2DArray.cpp
--- a/Engine/Source/Runtime/Engine/Classes/Engine/Texture2DArray.cpp
+++ b/Engine/Source/Runtime/Engine/Classes/Engine/Texture2DArray.cpp
@@ -110,6 +110,12 @@ FTexturePlatformData** UTexture2DArray::GetRunningPlatformData()
 
 void UTexture2DArray::SetPlatformData(FTexturePlatformData* InPlatformData)
 {
+	if (PrivatePlatformData)
+	{
+		ReleaseResource();
+		delete PrivatePlatformData;
+	}
+
 	PrivatePlatformData = InPlatformData;
 }
 
diff --git a/Engine/Source/Runtime/Engine/Private/Texture2D.cpp b/Engine/Source/Runtime/Engine/Private/Texture2D.cpp
--- a/Engine/Source/Runtime/Engine/Private/Texture2D.cpp
+++ b/Engine/Source/Runtime/Engine/Private/Texture2D.cpp
@@ -293,6 +293,11 @@ FTexturePlatformData** UTexture2D::GetRunningPlatformData()
 
 void UTexture2D::SetPlatformData(FTexturePlatformData* InPlatformData)
 {
+	if (PrivatePlatformData)
+	{
+		ReleaseResource();
+		delete PrivatePlatformData;
+	}
 	PrivatePlatformData = InPlatformData;
 }
 
diff --git a/Engine/Source/Runtime/Engine/Private/TextureCube.cpp b/Engine/Source/Runtime/Engine/Private/TextureCube.cpp
--- a/Engine/Source/Runtime/Engine/Private/TextureCube.cpp
+++ b/Engine/Source/Runtime/Engine/Private/TextureCube.cpp
@@ -91,6 +91,11 @@ FTexturePlatformData** UTextureCube::GetRunningPlatformData()
 
 void UTextureCube::SetPlatformData(FTexturePlatformData* InPlatformData)
 {
+	if (PrivatePlatformData)
+	{
+		ReleaseResource();
+		delete PrivatePlatformData;
+	}
 	PrivatePlatformData = InPlatformData;
 }
 
diff --git a/Engine/Source/Runtime/Engine/Private/TextureCubeArray.cpp b/Engine/Source/Runtime/Engine/Private/TextureCubeArray.cpp
--- a/Engine/Source/Runtime/Engine/Private/TextureCubeArray.cpp
+++ b/Engine/Source/Runtime/Engine/Private/TextureCubeArray.cpp
@@ -373,6 +373,11 @@ FTexturePlatformData** UTextureCubeArray::GetRunningPlatformData()
 
 void UTextureCubeArray::SetPlatformData(FTexturePlatformData* InPlatformData)
 {
+	if (PrivatePlatformData)
+	{
+		ReleaseResource();
+		delete PrivatePlatformData;
+	}
 	PrivatePlatformData = InPlatformData;
 }
 
diff --git a/Engine/Source/Runtime/Engine/Private/VolumeTexture.cpp b/Engine/Source/Runtime/Engine/Private/VolumeTexture.cpp
--- a/Engine/Source/Runtime/Engine/Private/VolumeTexture.cpp
+++ b/Engine/Source/Runtime/Engine/Private/VolumeTexture.cpp
@@ -115,6 +115,12 @@ FTexturePlatformData** UVolumeTexture::GetRunningPlatformData()
 
 void UVolumeTexture::SetPlatformData(FTexturePlatformData* InPlatformData)
 {
+	if (PrivatePlatformData)
+	{
+		ReleaseResource();
+		delete PrivatePlatformData;
+	}
+
 	PrivatePlatformData = InPlatformData;
 }
 
