[mutable] Extend ASTC format support (6x6,8x8,10x10). Added low quality compression for 6x6 RGB, 6x6 RGBA and 8x8 RGB. Compression for the rest of the new formats is not available yet.
Extend debugger support to show single R,G,B channels for images.
#rnx

diff --git a/Engine/Plugins/Experimental/Mutable/Source/CustomizableObject/Private/MuCO/CustomizableObjectSystemPrivate.h b/Engine/Plugins/Experimental/Mutable/Source/CustomizableObject/Private/MuCO/CustomizableObjectSystemPrivate.h
--- a/Engine/Plugins/Experimental/Mutable/Source/CustomizableObject/Private/MuCO/CustomizableObjectSystemPrivate.h
+++ b/Engine/Plugins/Experimental/Mutable/Source/CustomizableObject/Private/MuCO/CustomizableObjectSystemPrivate.h
@@ -811,9 +811,6 @@ public:
 	//! Update the stats logged in unreal's stats system. 
 	void UpdateStats();
 
-	// Important!!! Never call when there's a Begin Update thread running!
-	void ReleasePendingMutableInstances();
-
 	/** Update the last set amount of internal memory Mutable can use to build objects. */
 	void UpdateMemoryLimit();
 
diff --git a/Engine/Plugins/Experimental/Mutable/Source/CustomizableObject/Private/MuCO/UnrealMutableImageProvider.cpp b/Engine/Plugins/Experimental/Mutable/Source/CustomizableObject/Private/MuCO/UnrealMutableImageProvider.cpp
--- a/Engine/Plugins/Experimental/Mutable/Source/CustomizableObject/Private/MuCO/UnrealMutableImageProvider.cpp
+++ b/Engine/Plugins/Experimental/Mutable/Source/CustomizableObject/Private/MuCO/UnrealMutableImageProvider.cpp
@@ -108,7 +108,9 @@ mu::EImageFormat GetMutablePixelFormat(EPixelFormat InTextureFormat)
 	case PF_BC5: return mu::EImageFormat::IF_BC5;
 	case PF_G8: return mu::EImageFormat::IF_L_UBYTE;
 	case PF_ASTC_4x4: return mu::EImageFormat::IF_ASTC_4x4_RGBA_LDR;
+	case PF_ASTC_6x6: return mu::EImageFormat::IF_ASTC_6x6_RGBA_LDR;
 	case PF_ASTC_8x8: return mu::EImageFormat::IF_ASTC_8x8_RGBA_LDR;
+	case PF_ASTC_10x10: return mu::EImageFormat::IF_ASTC_10x10_RGBA_LDR;
 	case PF_ASTC_12x12: return mu::EImageFormat::IF_ASTC_12x12_RGBA_LDR;
 	default: return mu::EImageFormat::IF_NONE;
 	}
diff --git a/Engine/Plugins/Experimental/Mutable/Source/CustomizableObjectEditor/CustomizableObjectEditor.Build.cs b/Engine/Plugins/Experimental/Mutable/Source/CustomizableObjectEditor/CustomizableObjectEditor.Build.cs
--- a/Engine/Plugins/Experimental/Mutable/Source/CustomizableObjectEditor/CustomizableObjectEditor.Build.cs
+++ b/Engine/Plugins/Experimental/Mutable/Source/CustomizableObjectEditor/CustomizableObjectEditor.Build.cs
@@ -73,6 +73,11 @@ public class CustomizableObjectEditor : ModuleRules
 				
 				"AssetDefinition",
 				"ContentBrowser",
+				
+				"TextureCompressor",
+				"TextureBuildUtilities",
+				"ImageCore",
+
 			}
 		);
 
diff --git a/Engine/Plugins/Experimental/Mutable/Source/CustomizableObjectEditor/Private/MuCOE/GenerateMutableSource/GenerateMutableSourceSurface.cpp b/Engine/Plugins/Experimental/Mutable/Source/CustomizableObjectEditor/Private/MuCOE/GenerateMutableSource/GenerateMutableSourceSurface.cpp
--- a/Engine/Plugins/Experimental/Mutable/Source/CustomizableObjectEditor/Private/MuCOE/GenerateMutableSource/GenerateMutableSourceSurface.cpp
+++ b/Engine/Plugins/Experimental/Mutable/Source/CustomizableObjectEditor/Private/MuCOE/GenerateMutableSource/GenerateMutableSourceSurface.cpp
@@ -735,6 +735,10 @@ mu::NodeSurfacePtr GenerateMutableSourceSurface(const UEdGraphPin * Pin, FMutabl
 									else if (FormatWithoutPrefix == TEXT("BGRA8")) mutableFormat = mu::EImageFormat::IF_RGBA_UBYTE;
 									else if (PlatformFormat.Contains(TEXT("ASTC")))
 									{
+										// \TODO: ASTC block size to use depends on several project and platform settings. Try to get it.
+										
+										// This is based on GetQualityFormat in TextureFormatASTC.cpp but it would be much better to find
+										// an way of doing it without dependant code.
 										if ((FormatWithoutPrefix == TEXT("AutoASTC")) || (FormatWithoutPrefix == TEXT("RGBAuto")))
 										{
 											mutableFormat = mu::EImageFormat::IF_ASTC_4x4_RGB_LDR;
@@ -742,6 +746,7 @@ mu::NodeSurfacePtr GenerateMutableSourceSurface(const UEdGraphPin * Pin, FMutabl
 										}
 										else if (FormatWithoutPrefix == TEXT("RGB")) mutableFormat = mu::EImageFormat::IF_ASTC_4x4_RGB_LDR;
 										else if (FormatWithoutPrefix == TEXT("RGBA")) mutableFormat = mu::EImageFormat::IF_ASTC_4x4_RGBA_LDR;
+										else if (FormatWithoutPrefix == TEXT("RGBA_HQ")) mutableFormat = mu::EImageFormat::IF_ASTC_4x4_RGBA_LDR;
 										else if (FormatWithoutPrefix == TEXT("NormalRG")) mutableFormat = mu::EImageFormat::IF_ASTC_4x4_RG_LDR;
 										else if (PlatformFormat.Contains(TEXT("ASTC_NormalRG_Precise")))
 										{
diff --git a/Engine/Plugins/Experimental/Mutable/Source/CustomizableObjectEditor/Private/MuCOE/SMutableImageViewer.cpp b/Engine/Plugins/Experimental/Mutable/Source/CustomizableObjectEditor/Private/MuCOE/SMutableImageViewer.cpp
--- a/Engine/Plugins/Experimental/Mutable/Source/CustomizableObjectEditor/Private/MuCOE/SMutableImageViewer.cpp
+++ b/Engine/Plugins/Experimental/Mutable/Source/CustomizableObjectEditor/Private/MuCOE/SMutableImageViewer.cpp
@@ -219,6 +219,12 @@ void SMutableImageViewer::Construct(const FArguments& InArgs)
 				.Text(FText::FromString(TEXT("RGBA")))
 				+ SSegmentedControl<EMutableImageChannels>::Slot(EMutableImageChannels::RGB)
 				.Text(FText::FromString(TEXT("RGB")))
+				+ SSegmentedControl<EMutableImageChannels>::Slot(EMutableImageChannels::R)
+				.Text(FText::FromString(TEXT("R")))
+				+ SSegmentedControl<EMutableImageChannels>::Slot(EMutableImageChannels::G)
+				.Text(FText::FromString(TEXT("G")))
+				+ SSegmentedControl<EMutableImageChannels>::Slot(EMutableImageChannels::B)
+				.Text(FText::FromString(TEXT("B")))
 				+ SSegmentedControl<EMutableImageChannels>::Slot(EMutableImageChannels::A)
 				.Text(FText::FromString(TEXT("A")))
 			]
@@ -251,14 +257,37 @@ void SMutableImageViewer::Tick(const FGeometry& AllottedGeometry, const double I
 
 		if (MutableImage && MutableImage->GetSizeX()>0 && MutableImage->GetSizeY()>0)
 		{
+			// If the image format is not supported in the editor platform (for instance, it is a mobile-specific format)
+			// convert the mutable image to a generic format so that it can be previewed.
+			mu::Ptr<const mu::Image> ImageToConvert = MutableImage;
+
+			// \TODO: This test is a bit weak
+			if (MutableImage->GetFormat() >= mu::EImageFormat::IF_ASTC_4x4_RGB_LDR
+				&&
+				MutableImage->GetFormat() <= mu::EImageFormat::IF_ASTC_10x10_RG_LDR)
+			{
+				mu::FImageOperator ImOp = mu::FImageOperator::GetDefault();
+				int32 Quality = 4;
+				ImageToConvert = ImOp.ImagePixelFormat(Quality, MutableImage.get(), mu::GetUncompressedFormat(MutableImage->GetFormat()) );
+			}
+
 			// Should we extract a single channel?
-			int32 ExtractChannel = (GetImageChannels()==EMutableImageChannels::A) ? 3 : -1;
+			int32 ExtractChannel = -1;
+			switch (GetImageChannels())
+			{
+			case EMutableImageChannels::R: ExtractChannel = 0; break;
+			case EMutableImageChannels::G: ExtractChannel = 1; break;
+			case EMutableImageChannels::B: ExtractChannel = 2; break;
+			case EMutableImageChannels::A: ExtractChannel = 3; break;
+			default:
+				break;
+			}
 
 			FMutableModelImageProperties Props;
 			Props.Filter = TF_Bilinear;
 			Props.SRGB = true;
 			Props.LODBias = 0;
-			UCustomizableInstancePrivateData::ConvertImage(UnrealImage, MutableImage, Props, CurrentVisibleLOD, ExtractChannel);
+			UCustomizableInstancePrivateData::ConvertImage(UnrealImage, ImageToConvert, Props, CurrentVisibleLOD, ExtractChannel);
 			UnrealImage->NeverStream = true;
 			UnrealImage->UpdateResource();
 		}
diff --git a/Engine/Plugins/Experimental/Mutable/Source/CustomizableObjectEditor/Private/MuCOE/SMutableImageViewer.h b/Engine/Plugins/Experimental/Mutable/Source/CustomizableObjectEditor/Private/MuCOE/SMutableImageViewer.h
--- a/Engine/Plugins/Experimental/Mutable/Source/CustomizableObjectEditor/Private/MuCOE/SMutableImageViewer.h
+++ b/Engine/Plugins/Experimental/Mutable/Source/CustomizableObjectEditor/Private/MuCOE/SMutableImageViewer.h
@@ -21,6 +21,9 @@ enum class EMutableImageChannels : uint8
 {
 	RGBA,
 	RGB,
+	R,
+	G,
+	B,
 	A
 };
 
diff --git a/Engine/Plugins/Experimental/Mutable/Source/MutableRuntime/Private/MuR/BlockCompression/Miro/Miro.cpp b/Engine/Plugins/Experimental/Mutable/Source/MutableRuntime/Private/MuR/BlockCompression/Miro/Miro.cpp
--- a/Engine/Plugins/Experimental/Mutable/Source/MutableRuntime/Private/MuR/BlockCompression/Miro/Miro.cpp
+++ b/Engine/Plugins/Experimental/Mutable/Source/MutableRuntime/Private/MuR/BlockCompression/Miro/Miro.cpp
@@ -7,7 +7,6 @@
 #include "Math/UnrealMathSSE.h"
 #include "Math/UnrealMathUtility.h"
 #include "Misc/AssertionMacros.h"
-#include "Templates/UniquePtr.h"
 
 // debug
 //static bool s_debuglog = false;
@@ -2120,10 +2119,10 @@ namespace impl
 	inline constexpr void DCHECK(bool x)
 	{
 		(void)x;
-		 //if (!x)
-		 //{
-		 //    check(x);
-		 //}
+		 if (!x)
+		 {
+		     check(x);
+		 }
 	}
 
 
@@ -2222,7 +2221,7 @@ namespace impl
 			size_t bitoffset_;  // in bits
 		};
 
-		const uint8 bit_reverse_table[256] = {
+		constexpr uint8 bit_reverse_table[256] = {
 			0x00, 0x80, 0x40, 0xC0, 0x20, 0xA0, 0x60, 0xE0, 0x10, 0x90, 0x50, 0xD0,
 			0x30, 0xB0, 0x70, 0xF0, 0x08, 0x88, 0x48, 0xC8, 0x28, 0xA8, 0x68, 0xE8,
 			0x18, 0x98, 0x58, 0xD8, 0x38, 0xB8, 0x78, 0xF8, 0x04, 0x84, 0x44, 0xC4,
@@ -2272,10 +2271,7 @@ namespace impl
 			}
 		}
 
-		inline void copy_bytes(const uint8* source,
-			size_t bytecount,
-			uint8* target,
-			size_t bitoffset)
+		inline void copy_bytes(const uint8* source, size_t bytecount, uint8* target, size_t bitoffset)
 		{
 			for (size_t i = 0; i < bytecount; ++i)
 			{
@@ -2744,7 +2740,7 @@ namespace impl
 			{{102,103,70,71,38},{110,111,78,79,46},{118,119,86,87,54},{126,127,94,95,62},{39,47,55,63,31}}
 		};
 
-		constexpr int8 color_endpoint_range_table[2][12][16] =
+		constexpr int8 color_endpoint_range_table4x4[2][12][16] =
 		{
 			{
 				{-1,-1,-1,-1,-1,-1,-1,-1,-1,-1,-1,-1,-1,-1,-1,-1},
@@ -2759,8 +2755,7 @@ namespace impl
 				{20,20,20,20,20,20,20,20,16,16,16,16,11,11,11,11},
 				{20,20,20,20,20,20,20,20,14,14,14,14,10,10,10,10},
 				{20,20,20,20,19,19,19,19,11,11,11,11,7,7,7,7}
-			}
-			,
+			},
 			{
 				{-1,-1,-1,-1,-1,-1,-1,-1,-1,-1,-1,-1,-1,-1,-1,-1},
 				{20,20,20,20,20,20,20,20,14,14,14,14,9,9,9,9},
@@ -2777,6 +2772,81 @@ namespace impl
 			}
 		};
 
+		constexpr int8 color_endpoint_range_table6x6[2][12][16] = 
+		{ 
+			{
+				{20,20,20,20,20,20,20,20,20,20,20,20,20,20,20,20},
+				{20,20,20,20,20,20,20,20,20,20,20,20,16,16,16,16},
+				{20,20,20,20,20,20,20,20,15,15,15,15,10,10,10,10},
+				{20,20,20,20,16,16,16,16,9,9,9,9,6,6,6,6},
+				{20,20,20,20,8,8,8,8,4,4,4,4,2,2,2,2},
+				{-1,-1,-1,-1,-1,-1,-1,-1,-1,-1,-1,-1,-1,-1,-1,-1},
+				{-1,-1,-1,-1,-1,-1,-1,-1,-1,-1,-1,-1,-1,-1,-1,-1},
+				{-1,-1,-1,-1,-1,-1,-1,-1,-1,-1,-1,-1,-1,-1,-1,-1},
+				{-1,-1,-1,-1,-1,-1,-1,-1,-1,-1,-1,-1,-1,-1,-1,-1},
+				{-1,-1,-1,-1,-1,-1,-1,-1,-1,-1,-1,-1,-1,-1,-1,-1},
+				{-1,-1,-1,-1,-1,-1,-1,-1,-1,-1,-1,-1,-1,-1,-1,-1},
+				{-1,-1,-1,-1,-1,-1,-1,-1,-1,-1,-1,-1,-1,-1,-1,-1}
+			},
+			{
+				{20,20,20,20,19,19,19,19,11,11,11,11,7,7,7,7},
+				{20,20,20,20,11,11,11,11,6,6,6,6,3,3,3,3},
+				{16,16,16,16,6,6,6,6,2,2,2,2,1,1,1,1},
+				{7,7,7,7,1,1,1,1,0,0,0,0,-1,-1,-1,-1},
+				{0,0,0,0,-1,-1,-1,-1,-1,-1,-1,-1,-1,-1,-1,-1},
+				{-1,-1,-1,-1,-1,-1,-1,-1,-1,-1,-1,-1,-1,-1,-1,-1},
+				{-1,-1,-1,-1,-1,-1,-1,-1,-1,-1,-1,-1,-1,-1,-1,-1},
+				{-1,-1,-1,-1,-1,-1,-1,-1,-1,-1,-1,-1,-1,-1,-1,-1},
+				{-1,-1,-1,-1,-1,-1,-1,-1,-1,-1,-1,-1,-1,-1,-1,-1},
+				{-1,-1,-1,-1,-1,-1,-1,-1,-1,-1,-1,-1,-1,-1,-1,-1},
+				{-1,-1,-1,-1,-1,-1,-1,-1,-1,-1,-1,-1,-1,-1,-1,-1},
+				{-1,-1,-1,-1,-1,-1,-1,-1,-1,-1,-1,-1,-1,-1,-1,-1}
+			}
+		};
+
+
+		const int8_t color_endpoint_range_table8x8[2][12][16] = { 
+			{
+				{20,20,20,20,20,20,20,20,19,19,19,19,13,13,13,13},
+				{-1,-1,-1,-1,-1,-1,-1,-1,-1,-1,-1,-1,-1,-1,-1,-1},
+				{-1,-1,-1,-1,-1,-1,-1,-1,-1,-1,-1,-1,-1,-1,-1,-1},
+				{-1,-1,-1,-1,-1,-1,-1,-1,-1,-1,-1,-1,-1,-1,-1,-1},
+				{-1,-1,-1,-1,-1,-1,-1,-1,-1,-1,-1,-1,-1,-1,-1,-1},
+				{-1,-1,-1,-1,-1,-1,-1,-1,-1,-1,-1,-1,-1,-1,-1,-1},
+				{-1,-1,-1,-1,-1,-1,-1,-1,-1,-1,-1,-1,-1,-1,-1,-1},
+				{-1,-1,-1,-1,-1,-1,-1,-1,-1,-1,-1,-1,-1,-1,-1,-1},
+				{-1,-1,-1,-1,-1,-1,-1,-1,-1,-1,-1,-1,-1,-1,-1,-1},
+				{-1,-1,-1,-1,-1,-1,-1,-1,-1,-1,-1,-1,-1,-1,-1,-1},
+				{-1,-1,-1,-1,-1,-1,-1,-1,-1,-1,-1,-1,-1,-1,-1,-1},
+				{-1,-1,-1,-1,-1,-1,-1,-1,-1,-1,-1,-1,-1,-1,-1,-1}
+			},
+			{
+				{20,20,20,20,9,9,9,9,4,4,4,4,2,2,2,2},
+				{-1,-1,-1,-1,-1,-1,-1,-1,-1,-1,-1,-1,-1,-1,-1,-1},
+				{-1,-1,-1,-1,-1,-1,-1,-1,-1,-1,-1,-1,-1,-1,-1,-1},
+				{-1,-1,-1,-1,-1,-1,-1,-1,-1,-1,-1,-1,-1,-1,-1,-1},
+				{-1,-1,-1,-1,-1,-1,-1,-1,-1,-1,-1,-1,-1,-1,-1,-1},
+				{-1,-1,-1,-1,-1,-1,-1,-1,-1,-1,-1,-1,-1,-1,-1,-1},
+				{-1,-1,-1,-1,-1,-1,-1,-1,-1,-1,-1,-1,-1,-1,-1,-1},
+				{-1,-1,-1,-1,-1,-1,-1,-1,-1,-1,-1,-1,-1,-1,-1,-1},
+				{-1,-1,-1,-1,-1,-1,-1,-1,-1,-1,-1,-1,-1,-1,-1,-1},
+				{-1,-1,-1,-1,-1,-1,-1,-1,-1,-1,-1,-1,-1,-1,-1,-1},
+				{-1,-1,-1,-1,-1,-1,-1,-1,-1,-1,-1,-1,-1,-1,-1,-1},
+				{-1,-1,-1,-1,-1,-1,-1,-1,-1,-1,-1,-1,-1,-1,-1,-1}
+			} 
+		};
+
+		//const int8_t color_endpoint_range_table10x10[2][12][16] = { 
+		//	{{-1,-1,-1,-1,-1,-1,-1,-1,-1,-1,-1,-1,-1,-1,-1,-1},{-1,-1,-1,-1,-1,-1,-1,-1,-1,-1,-1,-1,-1,-1,-1,-1},{-1,-1,-1,-1,-1,-1,-1,-1,-1,-1,-1,-1,-1,-1,-1,-1},{-1,-1,-1,-1,-1,-1,-1,-1,-1,-1,-1,-1,-1,-1,-1,-1},{-1,-1,-1,-1,-1,-1,-1,-1,-1,-1,-1,-1,-1,-1,-1,-1},{-1,-1,-1,-1,-1,-1,-1,-1,-1,-1,-1,-1,-1,-1,-1,-1},{-1,-1,-1,-1,-1,-1,-1,-1,-1,-1,-1,-1,-1,-1,-1,-1},{-1,-1,-1,-1,-1,-1,-1,-1,-1,-1,-1,-1,-1,-1,-1,-1},{-1,-1,-1,-1,-1,-1,-1,-1,-1,-1,-1,-1,-1,-1,-1,-1},{-1,-1,-1,-1,-1,-1,-1,-1,-1,-1,-1,-1,-1,-1,-1,-1},{-1,-1,-1,-1,-1,-1,-1,-1,-1,-1,-1,-1,-1,-1,-1,-1},{-1,-1,-1,-1,-1,-1,-1,-1,-1,-1,-1,-1,-1,-1,-1,-1}},
+		//	{{-1,-1,-1,-1,-1,-1,-1,-1,-1,-1,-1,-1,-1,-1,-1,-1},{-1,-1,-1,-1,-1,-1,-1,-1,-1,-1,-1,-1,-1,-1,-1,-1},{-1,-1,-1,-1,-1,-1,-1,-1,-1,-1,-1,-1,-1,-1,-1,-1},{-1,-1,-1,-1,-1,-1,-1,-1,-1,-1,-1,-1,-1,-1,-1,-1},{-1,-1,-1,-1,-1,-1,-1,-1,-1,-1,-1,-1,-1,-1,-1,-1},{-1,-1,-1,-1,-1,-1,-1,-1,-1,-1,-1,-1,-1,-1,-1,-1},{-1,-1,-1,-1,-1,-1,-1,-1,-1,-1,-1,-1,-1,-1,-1,-1},{-1,-1,-1,-1,-1,-1,-1,-1,-1,-1,-1,-1,-1,-1,-1,-1},{-1,-1,-1,-1,-1,-1,-1,-1,-1,-1,-1,-1,-1,-1,-1,-1},{-1,-1,-1,-1,-1,-1,-1,-1,-1,-1,-1,-1,-1,-1,-1,-1},{-1,-1,-1,-1,-1,-1,-1,-1,-1,-1,-1,-1,-1,-1,-1,-1},{-1,-1,-1,-1,-1,-1,-1,-1,-1,-1,-1,-1,-1,-1,-1,-1}} 
+		//};
+
+		//const int8_t color_endpoint_range_table12x12[2][12][16] = { 
+		//	{{-1,-1,-1,-1,-1,-1,-1,-1,-1,-1,-1,-1,-1,-1,-1,-1},{-1,-1,-1,-1,-1,-1,-1,-1,-1,-1,-1,-1,-1,-1,-1,-1},{-1,-1,-1,-1,-1,-1,-1,-1,-1,-1,-1,-1,-1,-1,-1,-1},{-1,-1,-1,-1,-1,-1,-1,-1,-1,-1,-1,-1,-1,-1,-1,-1},{-1,-1,-1,-1,-1,-1,-1,-1,-1,-1,-1,-1,-1,-1,-1,-1},{-1,-1,-1,-1,-1,-1,-1,-1,-1,-1,-1,-1,-1,-1,-1,-1},{-1,-1,-1,-1,-1,-1,-1,-1,-1,-1,-1,-1,-1,-1,-1,-1},{-1,-1,-1,-1,-1,-1,-1,-1,-1,-1,-1,-1,-1,-1,-1,-1},{-1,-1,-1,-1,-1,-1,-1,-1,-1,-1,-1,-1,-1,-1,-1,-1},{-1,-1,-1,-1,-1,-1,-1,-1,-1,-1,-1,-1,-1,-1,-1,-1},{-1,-1,-1,-1,-1,-1,-1,-1,-1,-1,-1,-1,-1,-1,-1,-1},{-1,-1,-1,-1,-1,-1,-1,-1,-1,-1,-1,-1,-1,-1,-1,-1}},
+		//	{{-1,-1,-1,-1,-1,-1,-1,-1,-1,-1,-1,-1,-1,-1,-1,-1},{-1,-1,-1,-1,-1,-1,-1,-1,-1,-1,-1,-1,-1,-1,-1,-1},{-1,-1,-1,-1,-1,-1,-1,-1,-1,-1,-1,-1,-1,-1,-1,-1},{-1,-1,-1,-1,-1,-1,-1,-1,-1,-1,-1,-1,-1,-1,-1,-1},{-1,-1,-1,-1,-1,-1,-1,-1,-1,-1,-1,-1,-1,-1,-1,-1},{-1,-1,-1,-1,-1,-1,-1,-1,-1,-1,-1,-1,-1,-1,-1,-1},{-1,-1,-1,-1,-1,-1,-1,-1,-1,-1,-1,-1,-1,-1,-1,-1},{-1,-1,-1,-1,-1,-1,-1,-1,-1,-1,-1,-1,-1,-1,-1,-1},{-1,-1,-1,-1,-1,-1,-1,-1,-1,-1,-1,-1,-1,-1,-1,-1},{-1,-1,-1,-1,-1,-1,-1,-1,-1,-1,-1,-1,-1,-1,-1,-1},{-1,-1,-1,-1,-1,-1,-1,-1,-1,-1,-1,-1,-1,-1,-1,-1},{-1,-1,-1,-1,-1,-1,-1,-1,-1,-1,-1,-1,-1,-1,-1,-1}} 
+		//};
+
+
 		const uint8 color_unquantize_table[21][256] =
 		{
 			{0,255},
@@ -3821,12 +3891,6 @@ namespace impl
 			size_t r = 0;
 			if (weight_quant < RANGE_32 + 1) r = r_table[weight_quant];
 
-			// Use the first row of Table C.2.8 in the ASTC specification. Beware that this has to be changed if another block-size is used.
-			// Actually weight-grid and block sizes don't need to match. We are defining the weight-grid size here.
-			// \TODO: change to support different block sizes (jordi)
-			size_t a = m - 2;
-			size_t b = n - 4;
-
 			bool d = dual_plane;
 
 			size_t part_value = partition_count - 1;
@@ -3838,15 +3902,65 @@ namespace impl
 			size_t cem = color_endpoint_mode;
 
 			// Block mode
-			orbits8_ptr(pb->data, 0, getbit(r, 1), 1);
-			orbits8_ptr(pb->data, 1, getbit(r, 2), 1);
-			orbits8_ptr(pb->data, 2, 0, 1);
-			orbits8_ptr(pb->data, 3, 0, 1);
-			orbits8_ptr(pb->data, 4, getbit(r, 0), 1);
-			orbits8_ptr(pb->data, 5, a, 2);
-			orbits8_ptr(pb->data, 7, b, 2);
-			orbits8_ptr(pb->data, 9, h, 1);
-			orbits8_ptr(pb->data, 10, d, 1);
+			// Actually weight-grid and block sizes don't need to match. We are defining the weight-grid size here but for
+			// now we make them match.
+			if (BLOCK_WIDTH == 4 && BLOCK_HEIGHT == 4)
+			{
+				// Use the first row of Table C.2.8 in the ASTC specification.
+				size_t a = m - 2;
+				size_t b = n - 4;
+
+				orbits8_ptr(pb->data, 0, getbit(r, 1), 1);
+				orbits8_ptr(pb->data, 1, getbit(r, 2), 1);
+				orbits8_ptr(pb->data, 2, 0, 1);
+				orbits8_ptr(pb->data, 3, 0, 1);
+				orbits8_ptr(pb->data, 4, getbit(r, 0), 1);
+				orbits8_ptr(pb->data, 5, a, 2);
+				orbits8_ptr(pb->data, 7, b, 2);
+				orbits8_ptr(pb->data, 9, h, 1);
+				orbits8_ptr(pb->data, 10, d, 1);
+			}
+			else if (BLOCK_WIDTH == 6 && BLOCK_HEIGHT == 6)
+			{
+				// This encoder doesn't support 6x6 and dual plane: that requires mismatching and interpolating the pixel grid and weight grid
+				check(dual_plane == 0);
+
+				// Use the tenth row of Table C.2.8 in the ASTC specification.
+				size_t a = m - 6;
+				size_t b = n - 6;
+
+				orbits8_ptr(pb->data, 0, 0, 2);
+				orbits8_ptr(pb->data, 2, getbit(r, 1), 1);
+				orbits8_ptr(pb->data, 3, getbit(r, 2), 1);
+				orbits8_ptr(pb->data, 4, getbit(r, 0), 1);
+				orbits8_ptr(pb->data, 5, a, 2);
+				orbits8_ptr(pb->data, 7, 0, 1);
+				orbits8_ptr(pb->data, 8, 1, 1);
+				orbits8_ptr(pb->data, 9, b, 2);
+			}
+			else if (BLOCK_WIDTH == 8 && BLOCK_HEIGHT == 8)
+			{
+				// This encoder doesn't support 8x8 and dual plane: that requires mismatching and interpolating the pixel grid and weight grid
+				check(dual_plane==0);
+
+				// Use the tenth row of Table C.2.8 in the ASTC specification.
+				size_t a = m - 6;
+				size_t b = n - 6;
+
+				orbits8_ptr(pb->data, 0, 0, 2);
+				orbits8_ptr(pb->data, 2, getbit(r, 1), 1);
+				orbits8_ptr(pb->data, 3, getbit(r, 2), 1);
+				orbits8_ptr(pb->data, 4, getbit(r, 0), 1);
+				orbits8_ptr(pb->data, 5, a, 2);
+				orbits8_ptr(pb->data, 7, 0, 1);
+				orbits8_ptr(pb->data, 8, 1, 1);
+				orbits8_ptr(pb->data, 9, b, 2);
+			}
+			else
+			{
+				check(false);
+			}
+
 
 			// Partitions
 			orbits8_ptr(pb->data, 11, part_value, 2);
@@ -3862,6 +3976,7 @@ namespace impl
 			{
 				size_t bits_for_weights = compute_ise_bitcount(2 * BLOCK_WIDTH * BLOCK_HEIGHT, weight_quant);
 
+				// \TODO: dual plane always sets alpha as the separate channel
 				size_t secondPlaneChannel = 3;
 				orbits8_ptr(pb->data, 128 - bits_for_weights - 2, secondPlaneChannel, 2);
 			}
@@ -4038,9 +4153,29 @@ namespace impl
 		}
 
 
+		// This should be compiled out, together with the actual tables.
+		template<int32 BLOCK_WIDTH, int32 BLOCK_HEIGHT>
 		constexpr range_t endpoint_quantization(size_t partitions, range_t weight_quant, color_endpoint_mode_t endpoint_mode)
 		{
-			int8 ce_range = color_endpoint_range_table[partitions - 1][weight_quant][endpoint_mode];
+			int8 ce_range = 0;
+			
+			if (BLOCK_WIDTH == 4 && BLOCK_HEIGHT == 4)
+			{
+				ce_range = color_endpoint_range_table4x4[partitions - 1][weight_quant][endpoint_mode];
+			}
+			else if (BLOCK_WIDTH == 6 && BLOCK_HEIGHT == 6)
+			{
+				ce_range = color_endpoint_range_table6x6[partitions - 1][weight_quant][endpoint_mode];
+			}
+			else if(BLOCK_WIDTH == 8 && BLOCK_HEIGHT == 8)
+			{
+				ce_range = color_endpoint_range_table8x8[partitions - 1][weight_quant][endpoint_mode];
+			}
+			else
+			{
+				check(false);
+			}
+
 			DCHECK(ce_range >= 0 && ce_range <= RANGE_MAX);
 			return static_cast<range_t>(ce_range);
 		}
@@ -4061,14 +4196,30 @@ namespace impl
 
 
 		template<int32 BLOCK_WIDTH, int32 BLOCK_HEIGHT>
-		void encode_luminance(const uint8 texels[BLOCK_WIDTH*BLOCK_HEIGHT],
-			PhysicalBlock* physical_block)
+		void encode_luminance(const uint8 texels[BLOCK_WIDTH*BLOCK_HEIGHT], PhysicalBlock* physical_block)
 		{
 			size_t partition_count = 1;
 
 			color_endpoint_mode_t color_endpoint_mode = CEM_LDR_LUMINANCE_DIRECT;
 			range_t weight_quant = RANGE_32;
-			range_t endpoint_quant = endpoint_quantization(partition_count, weight_quant, color_endpoint_mode);
+			if (BLOCK_WIDTH == 4 && BLOCK_HEIGHT == 4)
+			{
+				weight_quant = RANGE_32;
+			}
+			else if (BLOCK_WIDTH == 6 && BLOCK_HEIGHT == 6)
+			{
+				weight_quant = RANGE_4;
+			}
+			else if (BLOCK_WIDTH == 8 && BLOCK_HEIGHT == 8)
+			{
+				weight_quant = RANGE_2;
+			}
+			else
+			{
+				check(false);
+			}
+
+			range_t endpoint_quant = endpoint_quantization<BLOCK_WIDTH, BLOCK_HEIGHT>(partition_count, weight_quant, color_endpoint_mode);
 
 			uint8 l0 = 255;
 			uint8 l1 = 0;
@@ -4084,14 +4235,13 @@ namespace impl
 
 			uint8 weights_quantized[BLOCK_WIDTH * BLOCK_HEIGHT];
 			calculate_quantized_weights_luminance<BLOCK_WIDTH*BLOCK_HEIGHT>(
-				texels, weight_quant, endpoint_unquantized[0], endpoint_unquantized[1],
-				weights_quantized);
+				texels, weight_quant, endpoint_unquantized[0], endpoint_unquantized[1], weights_quantized);
 
 			uint8 endpoint_ise[MAXIMUM_ENCODED_COLOR_ENDPOINT_BYTES] = { 0 };
-			integer_sequence_encode(endpoint_quantized, 2, RANGE_256, endpoint_ise);
+			integer_sequence_encode(endpoint_quantized, 2, endpoint_quant, endpoint_ise);
 
 			uint8 weights_ise[MAXIMUM_ENCODED_WEIGHT_BYTES + 1] = { 0 };
-			integer_sequence_encode(weights_quantized, BLOCK_WIDTH * BLOCK_HEIGHT, RANGE_32, weights_ise);
+			integer_sequence_encode(weights_quantized, BLOCK_WIDTH * BLOCK_HEIGHT, weight_quant, weights_ise);
 
 			symbolic_to_physical<BLOCK_WIDTH,BLOCK_HEIGHT>(color_endpoint_mode, endpoint_quant, weight_quant,
 				partition_count, endpoint_ise,
@@ -4107,9 +4257,27 @@ namespace impl
 		{
 			constexpr size_t partition_count = 1;
 
-			constexpr color_endpoint_mode_t color_endpoint_mode = CEM_LDR_RGB_DIRECT;
-			constexpr range_t weight_quant = RANGE_12;
-			constexpr range_t endpoint_quant = endpoint_quantization(partition_count, weight_quant, color_endpoint_mode);
+			color_endpoint_mode_t color_endpoint_mode = CEM_LDR_RGB_DIRECT;
+			range_t weight_quant = RANGE_12;
+
+			if (BLOCK_WIDTH == 4 && BLOCK_HEIGHT == 4)
+			{
+				weight_quant = RANGE_12;
+			}
+			else if (BLOCK_WIDTH == 6 && BLOCK_HEIGHT == 6)
+			{
+				weight_quant = RANGE_5; // \TODO: enlarge later?
+			}
+			else if (BLOCK_WIDTH == 8 && BLOCK_HEIGHT == 8)
+			{
+				weight_quant = RANGE_2; // \TODO: enlarge later?
+			}
+			else
+			{
+				check(false);
+			}
+
+			range_t endpoint_quant = endpoint_quantization<BLOCK_WIDTH, BLOCK_HEIGHT>(partition_count, weight_quant, color_endpoint_mode);
 
 			vec3i_t endpoint_unquantized[2];
 			uint8 endpoint_quantized[6];
@@ -4138,34 +4306,29 @@ namespace impl
 			size_t partition_count = 1;
 
 			color_endpoint_mode_t color_endpoint_mode = CEM_LDR_RGBA_DIRECT;
+
 			range_t weight_quant = RANGE_5;
 			range_t endpoint_quant = RANGE_16;
-			//      range_t weight_quant = RANGE_4;
-			//      range_t endpoint_quant = RANGE_48;
-
-
-			//      if (s_debuglog)
-			//      {
-			//          size_t endpointBits = compute_ise_bitcount(8, endpoint_quant);
-			//          size_t weightsBits  = compute_ise_bitcount(BLOCK_TEXEL_COUNT*2, weight_quant);
-			//          UE_LOG(LogMutableCore,Warning," interesting block. %d endpoint bits, %d weight bits", endpointBits, weightsBits );
-			//      }
+			if (BLOCK_WIDTH == 4 && BLOCK_HEIGHT == 4)
+			{
+				weight_quant = RANGE_5;
+				endpoint_quant = RANGE_16;
+				// weight_quant = RANGE_4;
+				// endpoint_quant = RANGE_48;
+			}
+			else if(BLOCK_WIDTH == 6 && BLOCK_HEIGHT == 6)
+			{
+				weight_quant = RANGE_2;
+				endpoint_quant = RANGE_24;
+			}
+			else
+			{
+				check(false);
+			}
 
 			vec4i_t endpoint_unquantized[2];
 			uint8 endpoint_quantized[8];
-			encode_rgba_direct(endpoint_quant, round(e0), round(e1), endpoint_quantized,
-				endpoint_unquantized);
-
-			//      if (s_debuglog)
-			//      {
-			//          UE_LOG(LogMutableCore,Warning," unquantized endpoints: %.3d %.3d %.3d %.3d \t %.3d %.3d %.3d %.3d",
-			//                  endpoint_unquantized[0].x, endpoint_unquantized[0].y, endpoint_unquantized[0].z, endpoint_unquantized[0].w,
-			//                  endpoint_unquantized[1].x, endpoint_unquantized[1].y, endpoint_unquantized[1].z, endpoint_unquantized[1].w );
-
-			//          UE_LOG(LogMutableCore,Warning," quantized endpoints:   %.3d %.3d %.3d %.3d \t %.3d %.3d %.3d %.3d",
-			//                  endpoint_quantized[0], endpoint_quantized[1], endpoint_quantized[2], endpoint_quantized[3],
-			//                  endpoint_quantized[4], endpoint_quantized[5], endpoint_quantized[6], endpoint_quantized[7] );
-			//      }
+			encode_rgba_direct(endpoint_quant, round(e0), round(e1), endpoint_quantized, endpoint_unquantized);
 
 			uint8 weights_quantized[BLOCK_WIDTH * BLOCK_HEIGHT];
 			calculate_quantized_weights_rgb<BLOCK_WIDTH*BLOCK_HEIGHT>(texels, weight_quant,
@@ -4181,24 +4344,6 @@ namespace impl
 				alpha_weights_quantized,
 				3);
 
-			//      if (s_debuglog)
-			//      {
-			//          int i=0;
-			//          UE_LOG(LogMutableCore,Warning," quantized alpha weights:");
-			//          UE_LOG(LogMutableCore,Warning," %.3d %.3d %.3d %.3d ",
-			//                  alpha_weights_quantized[i+0], alpha_weights_quantized[i+1], alpha_weights_quantized[i+2], alpha_weights_quantized[i+3] );
-			//          i+=4;
-			//          UE_LOG(LogMutableCore,Warning," %.3d %.3d %.3d %.3d ",
-			//                  alpha_weights_quantized[i+0], alpha_weights_quantized[i+1], alpha_weights_quantized[i+2], alpha_weights_quantized[i+3] );
-			//          i+=4;
-			//          UE_LOG(LogMutableCore,Warning," %.3d %.3d %.3d %.3d ",
-			//                  alpha_weights_quantized[i+0], alpha_weights_quantized[i+1], alpha_weights_quantized[i+2], alpha_weights_quantized[i+3] );
-			//          i+=4;
-			//          UE_LOG(LogMutableCore,Warning," %.3d %.3d %.3d %.3d ",
-			//                  alpha_weights_quantized[i+0], alpha_weights_quantized[i+1], alpha_weights_quantized[i+2], alpha_weights_quantized[i+3] );
-			//          i+=4;
-			//      }
-
 			uint8 final_weights_quantized[BLOCK_WIDTH * BLOCK_HEIGHT * 2];
 			for (size_t i = 0; i < BLOCK_WIDTH * BLOCK_HEIGHT; i++)
 			{
@@ -4226,24 +4371,40 @@ namespace impl
 			PhysicalBlock* physical_block)
 		{
 			size_t partition_count = 1;
+			bool bDualPartition = false;
 
 			color_endpoint_mode_t color_endpoint_mode = CEM_LDR_RGBA_DIRECT;
-			range_t weight_quant = RANGE_12;
-			range_t endpoint_quant = endpoint_quantization(partition_count, weight_quant, color_endpoint_mode);
+
+			range_t weight_quant = RANGE_5;
+			range_t endpoint_quant = RANGE_16;
+			if (BLOCK_WIDTH == 4 && BLOCK_HEIGHT == 4)
+			{
+				weight_quant = RANGE_8;
+			}
+			else if (BLOCK_WIDTH == 6 && BLOCK_HEIGHT == 6)
+			{
+				weight_quant = RANGE_4;
+			}
+			else if (BLOCK_WIDTH == 8 && BLOCK_HEIGHT == 8)
+			{
+				weight_quant = RANGE_2;
+			}
+			else
+			{
+				check(false);
+			}
+			endpoint_quant = endpoint_quantization<BLOCK_WIDTH,BLOCK_HEIGHT>(partition_count, weight_quant, color_endpoint_mode);
 
 			vec4i_t endpoint_unquantized[2];
 			uint8 endpoint_quantized[8];
-			encode_rgba_direct(endpoint_quant, round(e0), round(e1), endpoint_quantized,
-				endpoint_unquantized);
+			encode_rgba_direct(endpoint_quant, round(e0), round(e1), endpoint_quantized, endpoint_unquantized);
 
-			// Calculate weights ignoring alpha (TODO)
 			uint8 weights_quantized[BLOCK_WIDTH * BLOCK_HEIGHT];
 			calculate_quantized_weights_rgba<BLOCK_WIDTH*BLOCK_HEIGHT>(texels, weight_quant,
 				endpoint_unquantized[0],
 				endpoint_unquantized[1],
 				weights_quantized);
 
-
 			uint8 endpoint_ise[MAXIMUM_ENCODED_COLOR_ENDPOINT_BYTES] = { 0 };
 			integer_sequence_encode(endpoint_quantized, 8, endpoint_quant, endpoint_ise);
 
@@ -4252,52 +4413,98 @@ namespace impl
 
 			symbolic_to_physical<BLOCK_WIDTH, BLOCK_HEIGHT>(color_endpoint_mode, endpoint_quant, weight_quant,
 				partition_count, endpoint_ise,
-				weights_ise, physical_block, false);
+				weights_ise, physical_block, bDualPartition);
 		}
 
 
+		template<int32 BLOCK_WIDTH, int32 BLOCK_HEIGHT>
+		void encode_rg_single_partition(const unorm8_t texels[BLOCK_WIDTH*BLOCK_HEIGHT],
+			vec3f_t e0,
+			vec3f_t e1,
+			PhysicalBlock* physical_block)
+		{
+			size_t partition_count = 1;
+			bool bDualPlane = false;
+
+			color_endpoint_mode_t color_endpoint_mode = CEM_LDR_LUMINANCE_ALPHA_DIRECT;
+			range_t weight_quant = RANGE_6;
+			range_t endpoint_quant = RANGE_64;
+
+			if (BLOCK_WIDTH == 4 && BLOCK_HEIGHT == 4)
+			{
+				weight_quant = RANGE_12;
+				//endpoint_quant = RANGE_64;
+			}
+			else if (BLOCK_WIDTH == 6 && BLOCK_HEIGHT == 6)
+			{
+				weight_quant = RANGE_5;
+				//endpoint_quant = RANGE_256;
+			}
+			else if (BLOCK_WIDTH == 8 && BLOCK_HEIGHT == 8)
+			{
+				weight_quant = RANGE_2;
+				//endpoint_quant = RANGE_256;
+			}
+			else
+			{
+				check(false);
+			}
+			endpoint_quant = endpoint_quantization<BLOCK_WIDTH, BLOCK_HEIGHT>(partition_count, weight_quant, color_endpoint_mode);
+
+			vec3i_t endpoint_unquantized[2];
+			uint8 endpoint_quantized[6];
+			encode_rgb_direct(endpoint_quant, round(e0), round(e1), endpoint_quantized, endpoint_unquantized);
+
+			uint8 weights_quantized[BLOCK_WIDTH * BLOCK_HEIGHT];
+			// TODO: optimize with rg skipping b
+			calculate_quantized_weights_rgb<BLOCK_WIDTH*BLOCK_HEIGHT>(texels, weight_quant, endpoint_unquantized[0], endpoint_unquantized[1], weights_quantized);
+
+			uint8 endpoint_ise[MAXIMUM_ENCODED_COLOR_ENDPOINT_BYTES] = { 0 };
+			integer_sequence_encode(endpoint_quantized, 4, endpoint_quant, endpoint_ise);
+
+			uint8 weights_ise[MAXIMUM_ENCODED_WEIGHT_BYTES + 1] = { 0 };
+			integer_sequence_encode(weights_quantized, BLOCK_WIDTH * BLOCK_HEIGHT, weight_quant, weights_ise);
+
+			symbolic_to_physical<BLOCK_WIDTH,BLOCK_HEIGHT>(color_endpoint_mode, endpoint_quant, weight_quant,
+				partition_count, endpoint_ise,
+				weights_ise, physical_block, bDualPlane);
+		}
+
 
 		template<int32 BLOCK_WIDTH, int32 BLOCK_HEIGHT>
-		void encode_rg_single_partition_dual_plane(const unorm8_t texels[BLOCK_WIDTH*BLOCK_HEIGHT],
+		void encode_rg_single_partition_dual_plane(const unorm8_t texels[BLOCK_WIDTH * BLOCK_HEIGHT],
 			vec4f_t e0,
 			vec4f_t e1,
 			PhysicalBlock* physical_block)
 		{
 			size_t partition_count = 1;
+			bool bDualPlane = true;
 
 			color_endpoint_mode_t color_endpoint_mode = CEM_LDR_LUMINANCE_ALPHA_DIRECT;
 			range_t weight_quant = RANGE_6;
 			range_t endpoint_quant = RANGE_64;
-			//      range_t weight_quant = RANGE_5;
-			//      range_t endpoint_quant = RANGE_128;
-			//      range_t weight_quant = RANGE_8;
-			//      range_t endpoint_quant = RANGE_8;
-
-	//        if (s_debuglog)
-	//        {
-	//          size_t endpointBits = compute_ise_bitcount(4, endpoint_quant);
-	//          size_t weightsBits  = compute_ise_bitcount(BLOCK_TEXEL_COUNT*2, weight_quant);
-	//          UE_LOG(LogMutableCore,Warning," interesting block. %d endpoint bits, %d weight bits", endpointBits, weightsBits );
-	//        }
+
+			if (BLOCK_WIDTH == 4 && BLOCK_HEIGHT == 4)
+			{
+				weight_quant = RANGE_6;
+				endpoint_quant = RANGE_64;
+			}
+			else if (BLOCK_WIDTH == 6 && BLOCK_HEIGHT == 6)
+			{
+				weight_quant = RANGE_2;
+				endpoint_quant = RANGE_256;
+			}
+			else
+			{
+				check(false);
+			}
 
 			vec4i_t endpoint_unquantized[2];
 			uint8 endpoint_quantized[8];
-			encode_rgba_direct(endpoint_quant, round(e0), round(e1), endpoint_quantized,
-				endpoint_unquantized);
-
-			//        if (s_debuglog)
-			//        {
-			//            UE_LOG(LogMutableCore,Warning," unquantized endpoints: %.3d %.3d %.3d %.3d \t %.3d %.3d %.3d %.3d",
-			//                  endpoint_unquantized[0].r, endpoint_unquantized[0].g, endpoint_unquantized[0].b, endpoint_unquantized[0].a,
-			//                  endpoint_unquantized[1].r, endpoint_unquantized[1].g, endpoint_unquantized[1].b, endpoint_unquantized[1].a );
-
-			//            UE_LOG(LogMutableCore,Warning," quantized endpoints:   %.3d %.3d %.3d %.3d \t %.3d %.3d %.3d %.3d",
-			//                  endpoint_quantized[0], endpoint_quantized[1], endpoint_quantized[2], endpoint_quantized[3],
-			//                  endpoint_quantized[4], endpoint_quantized[5], endpoint_quantized[6], endpoint_quantized[7] );
-			//        }
+			encode_rgba_direct(endpoint_quant, round(e0), round(e1), endpoint_quantized, endpoint_unquantized);
 
 			uint8 r_weights_quantized[BLOCK_WIDTH * BLOCK_HEIGHT];
-			calculate_quantized_weights_channel<BLOCK_WIDTH*BLOCK_HEIGHT>(texels, weight_quant,
+			calculate_quantized_weights_channel<BLOCK_WIDTH* BLOCK_HEIGHT>(texels, weight_quant,
 				(uint8)endpoint_unquantized[0].r,
 				(uint8)endpoint_unquantized[1].r,
 				r_weights_quantized,
@@ -4305,42 +4512,12 @@ namespace impl
 
 
 			uint8 g_weights_quantized[BLOCK_WIDTH * BLOCK_HEIGHT];
-			calculate_quantized_weights_channel<BLOCK_WIDTH*BLOCK_HEIGHT>(texels, weight_quant,
+			calculate_quantized_weights_channel<BLOCK_WIDTH* BLOCK_HEIGHT>(texels, weight_quant,
 				(uint8)endpoint_unquantized[0].g,
 				(uint8)endpoint_unquantized[1].g,
 				g_weights_quantized,
 				1);
 
-			//        if (s_debuglog)
-			//        {
-			//            int i=0;
-			//            UE_LOG(LogMutableCore,Warning," r values:");
-			//            for (int j=0;j<4;++j)
-			//            {
-			//                UE_LOG(LogMutableCore,Warning," %.3d %.3d %.3d %.3d ",
-			//                        texels[i].components[2], texels[i+1].components[2], texels[i+2].components[2], texels[i+3].components[2] );
-			//                i+=4;
-			//            }
-
-			//            i=0;
-			//            UE_LOG(LogMutableCore,Warning," quantized r weights:");
-			//            for (int j=0;j<4;++j)
-			//            {
-			//                UE_LOG(LogMutableCore,Warning," %.3d %.3d %.3d %.3d ",
-			//                        r_weights_quantized[i+0], r_weights_quantized[i+1], r_weights_quantized[i+2], r_weights_quantized[i+3] );
-			//                i+=4;
-			//            }
-
-			////            i=0;
-			////            UE_LOG(LogMutableCore,Warning," quantized g weights:");
-			////            for (int j=0;j<4;++j)
-			////            {
-			////                UE_LOG(LogMutableCore,Warning," %.3d %.3d %.3d %.3d ",
-			////                        g_weights_quantized[i+0], g_weights_quantized[i+1], g_weights_quantized[i+2], g_weights_quantized[i+3] );
-			////                i+=4;
-			////            }
-			//        }
-
 			uint8 final_weights_quantized[BLOCK_WIDTH * BLOCK_HEIGHT * 2];
 			for (size_t i = 0; i < BLOCK_WIDTH * BLOCK_HEIGHT; i++)
 			{
@@ -4354,9 +4531,9 @@ namespace impl
 			uint8 weights_ise[MAXIMUM_ENCODED_WEIGHT_BYTES + 1] = { 0 };
 			integer_sequence_encode(final_weights_quantized, BLOCK_WIDTH * BLOCK_HEIGHT * 2, weight_quant, weights_ise);
 
-			symbolic_to_physical<BLOCK_WIDTH,BLOCK_HEIGHT>(color_endpoint_mode, endpoint_quant, weight_quant,
+			symbolic_to_physical<BLOCK_WIDTH, BLOCK_HEIGHT>(color_endpoint_mode, endpoint_quant, weight_quant,
 				partition_count, endpoint_ise,
-				weights_ise, physical_block, true);
+				weights_ise, physical_block, bDualPlane);
 		}
 
 
@@ -4569,7 +4746,7 @@ namespace impl
 
 
 		template<int32 BLOCK_WIDTH, int32 BLOCK_HEIGHT>
-		void compress_block_rgba(const unorm8_t texels[BLOCK_WIDTH * BLOCK_HEIGHT], PhysicalBlock* physical_block)
+		void compress_block_rgba(const unorm8_t texels[BLOCK_WIDTH * BLOCK_HEIGHT], PhysicalBlock* physical_block, bool bDualPlane)
 		{
 			unorm8_t color;
 			bool isSolidRGBA = is_solid_rgba(texels, BLOCK_WIDTH * BLOCK_HEIGHT, &color);
@@ -4611,8 +4788,15 @@ namespace impl
 			//                    re1.x, re1.y, re1.z, re1.w );
 			//        }
 
-					//encode_rgba_single_partition(texels, re0, re1, physical_block);
-			encode_rgba_single_partition_dual_plane<BLOCK_WIDTH, BLOCK_HEIGHT>(texels, re0, re1, physical_block);
+			if (bDualPlane)
+			{
+				// dual plane usually yields worse results!
+				encode_rgba_single_partition_dual_plane<BLOCK_WIDTH, BLOCK_HEIGHT>(texels, re0, re1, physical_block);
+			}
+			else
+			{
+				encode_rgba_single_partition<BLOCK_WIDTH, BLOCK_HEIGHT>(texels, re0, re1, physical_block);
+			}
 		}
 
 
@@ -4639,10 +4823,20 @@ namespace impl
 				g1 = FMath::Max(g1, texels[i].components[1]);
 			}
 
-			vec4f_t re0(r0, g0, r0, g0);
-			vec4f_t re1(r1, g1, r1, g1);
-
-			encode_rg_single_partition_dual_plane<BLOCK_WIDTH, BLOCK_HEIGHT>(texels, re0, re1, physical_block);
+			// Avoid dual-plane by default
+			bool bUseDualPlane = false;
+			if (bUseDualPlane)
+			{
+				vec4f_t re0(r0, g0, r0, g0);
+				vec4f_t re1(r1, g1, r1, g1);
+				encode_rg_single_partition_dual_plane<BLOCK_WIDTH, BLOCK_HEIGHT>(texels, re0, re1, physical_block);
+			}
+			else
+			{
+				vec3f_t re0(r0, g0, 0);
+				vec3f_t re1(r1, g1, 0);
+				encode_rg_single_partition<BLOCK_WIDTH, BLOCK_HEIGHT>(texels, re0, re1, physical_block);
+			}
 		}
 
 	}
@@ -6125,7 +6319,8 @@ namespace impl
 
 			const Block128 blockData(data);
 
-			// (anticto) shortcut for blank blocks
+			// (anticto) shortcut for blank blocks. This is not correct really, an all-zero block is invalid but we support it 
+			// so that zero-initialized ASTC textures are actually black.
 			if (blockData.isZero())
 			{
 				int32 pix = 0;
@@ -7364,7 +7559,7 @@ namespace miro
 							toPixel[0] = Block[py * BLOCK_SIZE * 4 + px * 4 + 0];
 							toPixel[1] = Block[py * BLOCK_SIZE * 4 + px * 4 + 1];
 							toPixel[2] = Block[py * BLOCK_SIZE * 4 + px * 4 + 2];
-							toPixel[3] = 0;
+							toPixel[3] = 255;
 						}
 					}
 				}
@@ -7384,8 +7579,8 @@ namespace miro
 		astcrt::PhysicalBlock physical_block_zero = { {0} };
 		astcrt::PhysicalBlock* dst_re = reinterpret_cast<astcrt::PhysicalBlock*>(to);
 
-		uint32 bx = FMath::DivideAndRoundUp(sx, 4u);
-		uint32 by = FMath::DivideAndRoundUp(sy, 4u);
+		uint32 bx = FMath::DivideAndRoundUp(sx, BLOCK_SIZE);
+		uint32 by = FMath::DivideAndRoundUp(sy, BLOCK_SIZE);
 
 		//for ( uint32 y=0; y<by; ++y )
 		ParallelFor(by,
@@ -7420,7 +7615,8 @@ namespace miro
 					}
 
 					*rowTo = physical_block_zero;
-					astcrt::compress_block_rgba<BLOCK_SIZE, BLOCK_SIZE>((astcrt::unorm8_t*)block, rowTo);
+					// \TODO: We are forcing to not use dual plane, this should be per-block.
+					astcrt::compress_block_rgba<BLOCK_SIZE, BLOCK_SIZE>((astcrt::unorm8_t*)block, rowTo, false);
 
 					++rowTo;
 				}
@@ -7472,7 +7668,9 @@ namespace miro
 					}
 
 					*rowTo = physical_block_zero;
-					astcrt::compress_block_rgba<BLOCK_SIZE, BLOCK_SIZE>((astcrt::unorm8_t*)block, rowTo);
+					// \TODO: We are forcing to not use dual plane, this should be per-block.
+					// In this case we shoiuld weight-out A in all calculations, or even just compress as RGB.
+					astcrt::compress_block_rgba<BLOCK_SIZE, BLOCK_SIZE>((astcrt::unorm8_t*)block, rowTo, false);
 
 					++rowTo;
 				}
@@ -7760,7 +7958,7 @@ namespace miro
 								uint8* ToPixel = to + (yi * sx + xi) * 3;
 								ToPixel[0] = Block[BlockY * BLOCK_SIZE * 4 + BlockX * 4 + 0];
 								ToPixel[1] = Block[BlockY * BLOCK_SIZE * 4 + BlockX * 4 + 3];
-								ToPixel[2] = 0;
+								ToPixel[2] = 255;
 							}
 						}
 					}
@@ -7803,8 +8001,8 @@ namespace miro
 								uint8* ToPixel = to + (yi * sx + xi) * 4;
 								ToPixel[0] = Block[BlockY * BLOCK_SIZE * 4 + BlockX * 4 + 0];
 								ToPixel[1] = Block[BlockY * BLOCK_SIZE * 4 + BlockX * 4 + 3];
-								ToPixel[2] = 0;
-								ToPixel[3] = 0;
+								ToPixel[2] = 255;
+								ToPixel[3] = 255;
 							}
 						}
 					}
@@ -7897,6 +8095,16 @@ namespace miro
 	//---------------------------------------------------------------------------------------------
 	// 6x6
 	//---------------------------------------------------------------------------------------------
+	void RGBA_to_ASTC6x6RGBAL(uint32 sx, uint32 sy, const uint8* from, uint8* to, int32 Quality)
+	{
+		Generic_RGBA_to_ASTCRGBAL<6>(sx, sy, from, to, Quality);
+	}
+
+	void RGB_to_ASTC6x6RGBAL(uint32 sx, uint32 sy, const uint8* from, uint8* to, int32 Quality)
+	{
+		Generic_RGB_to_ASTCRGBAL<6>(sx, sy, from, to, Quality);
+	}
+
 	void ASTC6x6RGBAL_to_RGBA(uint32 sx, uint32 sy, const uint8* from, uint8* to)
 	{
 		Generic_ASTCRGBAL_to_RGBA<6>(sx, sy, from, to);
@@ -7907,6 +8115,16 @@ namespace miro
 		Generic_ASTCRGBAL_to_RGB<6>(sx, sy, from, to);
 	}
 
+	void RGB_to_ASTC6x6RGBL(uint32 sx, uint32 sy, const uint8* from, uint8* to, int32 Quality)
+	{
+		Generic_RGB_to_ASTCRGBL<6>(sx, sy, from, to, Quality);
+	}
+
+	void RGBA_to_ASTC6x6RGBL(uint32 sx, uint32 sy, const uint8* from, uint8* to, int32 Quality)
+	{
+		Generic_RGBA_to_ASTCRGBL<6>(sx, sy, from, to, Quality);
+	}
+
 	void ASTC6x6RGBL_to_RGB(uint32 sx, uint32 sy, const uint8* from, uint8* to)
 	{
 		Generic_ASTCRGBL_to_RGB<6>(sx, sy, from, to);
@@ -7917,6 +8135,16 @@ namespace miro
 		Generic_ASTCRGBL_to_RGBA<6>(sx, sy, from, to);
 	}
 
+	void RGB_to_ASTC6x6RGL(uint32 sx, uint32 sy, const uint8* from, uint8* to, int32 Quality)
+	{
+		Generic_RGB_to_ASTCRGL<6>(sx, sy, from, to, Quality);
+	}
+
+	void RGBA_to_ASTC6x6RGL(uint32 sx, uint32 sy, const uint8* from, uint8* to, int32 Quality)
+	{
+		Generic_RGBA_to_ASTCRGL<6>(sx, sy, from, to, Quality);
+	}
+
 	void ASTC6x6RGL_to_RGB(uint32 sx, uint32 sy, const uint8* from, uint8* to)
 	{
 		Generic_ASTCRGL_to_RGB<6>(sx, sy, from, to);
@@ -7931,6 +8159,16 @@ namespace miro
 	//---------------------------------------------------------------------------------------------
 	// 8x8
 	//---------------------------------------------------------------------------------------------
+	void RGBA_to_ASTC8x8RGBAL(uint32 sx, uint32 sy, const uint8* from, uint8* to, int32 Quality)
+	{
+		Generic_RGBA_to_ASTCRGBAL<8>(sx, sy, from, to, Quality);
+	}
+
+	void RGB_to_ASTC8x8RGBAL(uint32 sx, uint32 sy, const uint8* from, uint8* to, int32 Quality)
+	{
+		Generic_RGB_to_ASTCRGBAL<8>(sx, sy, from, to, Quality);
+	}
+
 	void ASTC8x8RGBAL_to_RGBA(uint32 sx, uint32 sy, const uint8* from, uint8* to)
 	{
 		Generic_ASTCRGBAL_to_RGBA<8>(sx, sy, from, to);
@@ -7941,6 +8179,16 @@ namespace miro
 		Generic_ASTCRGBAL_to_RGB<8>(sx, sy, from, to);
 	}
 
+	void RGB_to_ASTC8x8RGBL(uint32 sx, uint32 sy, const uint8* from, uint8* to, int32 Quality)
+	{
+		Generic_RGB_to_ASTCRGBL<8>(sx, sy, from, to, Quality);
+	}
+
+	void RGBA_to_ASTC8x8RGBL(uint32 sx, uint32 sy, const uint8* from, uint8* to, int32 Quality)
+	{
+		Generic_RGBA_to_ASTCRGBL<6>(sx, sy, from, to, Quality);
+	}
+
 	void ASTC8x8RGBL_to_RGB(uint32 sx, uint32 sy, const uint8* from, uint8* to)
 	{
 		Generic_ASTCRGBL_to_RGB<8>(sx, sy, from, to);
@@ -7951,6 +8199,16 @@ namespace miro
 		Generic_ASTCRGBL_to_RGBA<8>(sx, sy, from, to);
 	}
 
+	void RGB_to_ASTC8x8RGL(uint32 sx, uint32 sy, const uint8* from, uint8* to, int32 Quality)
+	{
+		Generic_RGB_to_ASTCRGL<8>(sx, sy, from, to, Quality);
+	}
+
+	void RGBA_to_ASTC8x8RGL(uint32 sx, uint32 sy, const uint8* from, uint8* to, int32 Quality)
+	{
+		Generic_RGBA_to_ASTCRGL<8>(sx, sy, from, to, Quality);
+	}
+
 	void ASTC8x8RGL_to_RGB(uint32 sx, uint32 sy, const uint8* from, uint8* to)
 	{
 		Generic_ASTCRGL_to_RGB<8>(sx, sy, from, to);
@@ -7962,6 +8220,40 @@ namespace miro
 	}
 
 
+	//---------------------------------------------------------------------------------------------
+	// 10x10
+	//---------------------------------------------------------------------------------------------
+	void ASTC10x10RGBAL_to_RGBA(uint32 sx, uint32 sy, const uint8* from, uint8* to)
+	{
+		Generic_ASTCRGBAL_to_RGBA<10>(sx, sy, from, to);
+	}
+
+	void ASTC10x10RGBAL_to_RGB(uint32 sx, uint32 sy, const uint8* from, uint8* to)
+	{
+		Generic_ASTCRGBAL_to_RGB<10>(sx, sy, from, to);
+	}
+
+	void ASTC10x10RGBL_to_RGB(uint32 sx, uint32 sy, const uint8* from, uint8* to)
+	{
+		Generic_ASTCRGBL_to_RGB<10>(sx, sy, from, to);
+	}
+
+	void ASTC10x10RGBL_to_RGBA(uint32 sx, uint32 sy, const uint8* from, uint8* to)
+	{
+		Generic_ASTCRGBL_to_RGBA<10>(sx, sy, from, to);
+	}
+
+	void ASTC10x10RGL_to_RGB(uint32 sx, uint32 sy, const uint8* from, uint8* to)
+	{
+		Generic_ASTCRGL_to_RGB<10>(sx, sy, from, to);
+	}
+
+	void ASTC10x10RGL_to_RGBA(uint32 sx, uint32 sy, const uint8* from, uint8* to)
+	{
+		Generic_ASTCRGL_to_RGBA<10>(sx, sy, from, to);
+	}
+
+
 	//---------------------------------------------------------------------------------------------
 	// 12x12
 	//---------------------------------------------------------------------------------------------
@@ -7999,4 +8291,3 @@ namespace miro
 
 
 }
-
diff --git a/Engine/Plugins/Experimental/Mutable/Source/MutableRuntime/Private/MuR/BlockCompression/Miro/Miro.h b/Engine/Plugins/Experimental/Mutable/Source/MutableRuntime/Private/MuR/BlockCompression/Miro/Miro.h
--- a/Engine/Plugins/Experimental/Mutable/Source/MutableRuntime/Private/MuR/BlockCompression/Miro/Miro.h
+++ b/Engine/Plugins/Experimental/Mutable/Source/MutableRuntime/Private/MuR/BlockCompression/Miro/Miro.h
@@ -76,62 +76,71 @@ namespace miro
 
 #if MIRO_INCLUDE_ASTC
 
-    //! ASTC 4x4 RGBA low support
+    //! ASTC 4x4 support
     extern void RGBA_to_ASTC4x4RGBAL( uint32 sx, uint32 sy, const uint8* From, uint8* To, int32 Quality );
     extern void RGB_to_ASTC4x4RGBAL( uint32 sx, uint32 sy, const uint8* From, uint8* To, int32 Quality );
     extern void ASTC4x4RGBAL_to_RGBA( uint32 sx, uint32 sy, const uint8* From, uint8* To );
     extern void ASTC4x4RGBAL_to_RGB( uint32 sx, uint32 sy, const uint8* From, uint8* To );
 
-    //! ASTC 4x4 RGB low support
     extern void RGB_to_ASTC4x4RGBL( uint32 sx, uint32 sy, const uint8* From, uint8* To, int32 Quality );
     extern void RGBA_to_ASTC4x4RGBL( uint32 sx, uint32 sy, const uint8* From, uint8* To, int32 Quality );
     extern void ASTC4x4RGBL_to_RGB( uint32 sx, uint32 sy, const uint8* From, uint8* To );
     extern void ASTC4x4RGBL_to_RGBA( uint32 sx, uint32 sy, const uint8* From, uint8* To );
 	extern void L_to_ASTC4x4RGBL(uint32 sx, uint32 sy, const uint8* From, uint8* To, int32 Quality);
-    
-	//! Convenience transforms
-    extern void ASTC4x4RGBL_to_ASTC4x4RGBAL( uint32 sx, uint32 sy, const uint8* From, uint8* To );
-    extern void ASTC4x4RGBAL_to_ASTC4x4RGBL( uint32 sx, uint32 sy, const uint8* From, uint8* To );
 
-    //! ASTC 4x4 RG low support
     extern void RGB_to_ASTC4x4RGL( uint32 sx, uint32 sy, const uint8* From, uint8* To, int32 Quality );
     extern void RGBA_to_ASTC4x4RGL( uint32 sx, uint32 sy, const uint8* From, uint8* To, int32 Quality );
     extern void ASTC4x4RGL_to_RGB( uint32 sx, uint32 sy, const uint8* From, uint8* To );
     extern void ASTC4x4RGL_to_RGBA( uint32 sx, uint32 sy, const uint8* From, uint8* To );
 
-	//! ASTC 6x6 RGBA low support
+	extern void ASTC4x4RGBL_to_ASTC4x4RGBAL(uint32 sx, uint32 sy, const uint8* From, uint8* To);
+	extern void ASTC4x4RGBAL_to_ASTC4x4RGBL(uint32 sx, uint32 sy, const uint8* From, uint8* To);
+
+	//! ASTC 6x6 support
+	extern void RGBA_to_ASTC6x6RGBAL(uint32 sx, uint32 sy, const uint8* From, uint8* To, int32 Quality);
+	extern void RGB_to_ASTC6x6RGBAL(uint32 sx, uint32 sy, const uint8* From, uint8* To, int32 Quality);
 	extern void ASTC6x6RGBAL_to_RGBA(uint32 sx, uint32 sy, const uint8* From, uint8* To);
 	extern void ASTC6x6RGBAL_to_RGB(uint32 sx, uint32 sy, const uint8* From, uint8* To);
-
-	//! ASTC 6x6 RGB low support
+	
+	extern void RGB_to_ASTC6x6RGBL(uint32 sx, uint32 sy, const uint8* From, uint8* To, int32 Quality);
+	extern void RGBA_to_ASTC6x6RGBL(uint32 sx, uint32 sy, const uint8* From, uint8* To, int32 Quality);
 	extern void ASTC6x6RGBL_to_RGB(uint32 sx, uint32 sy, const uint8* From, uint8* To);
 	extern void ASTC6x6RGBL_to_RGBA(uint32 sx, uint32 sy, const uint8* From, uint8* To);
 
-	//! ASTC 6x6 RG low support
+	extern void RGB_to_ASTC6x6RGL(uint32 sx, uint32 sy, const uint8* From, uint8* To, int32 Quality);
+	extern void RGBA_to_ASTC6x6RGL(uint32 sx, uint32 sy, const uint8* From, uint8* To, int32 Quality);
 	extern void ASTC6x6RGL_to_RGB(uint32 sx, uint32 sy, const uint8* From, uint8* To);
 	extern void ASTC6x6RGL_to_RGBA(uint32 sx, uint32 sy, const uint8* From, uint8* To);
 
-	//! ASTC 8x8 RGBA low support
+	//! ASTC 8x8 support
+	extern void RGBA_to_ASTC8x8RGBAL(uint32 sx, uint32 sy, const uint8* From, uint8* To, int32 Quality);
+	extern void RGB_to_ASTC8x8RGBAL(uint32 sx, uint32 sy, const uint8* From, uint8* To, int32 Quality);
 	extern void ASTC8x8RGBAL_to_RGBA(uint32 sx, uint32 sy, const uint8* From, uint8* To);
 	extern void ASTC8x8RGBAL_to_RGB(uint32 sx, uint32 sy, const uint8* From, uint8* To);
 
-	//! ASTC 8x8 RGB low support
+	extern void RGB_to_ASTC8x8RGBL(uint32 sx, uint32 sy, const uint8* From, uint8* To, int32 Quality);
+	extern void RGBA_to_ASTC8x8RGBL(uint32 sx, uint32 sy, const uint8* From, uint8* To, int32 Quality);
 	extern void ASTC8x8RGBL_to_RGB(uint32 sx, uint32 sy, const uint8* From, uint8* To);
 	extern void ASTC8x8RGBL_to_RGBA(uint32 sx, uint32 sy, const uint8* From, uint8* To);
 
-	//! ASTC 8x8 RG low support
+	extern void RGB_to_ASTC8x8RGL(uint32 sx, uint32 sy, const uint8* From, uint8* To, int32 Quality);
+	extern void RGBA_to_ASTC8x8RGL(uint32 sx, uint32 sy, const uint8* From, uint8* To, int32 Quality);
 	extern void ASTC8x8RGL_to_RGB(uint32 sx, uint32 sy, const uint8* From, uint8* To);
 	extern void ASTC8x8RGL_to_RGBA(uint32 sx, uint32 sy, const uint8* From, uint8* To);
 
-	//! ASTC 12x12 RGBA low support
+	//! ASTC 10x10 support
+	extern void ASTC10x10RGBAL_to_RGBA(uint32 sx, uint32 sy, const uint8* From, uint8* To);
+	extern void ASTC10x10RGBAL_to_RGB(uint32 sx, uint32 sy, const uint8* From, uint8* To);
+	extern void ASTC10x10RGBL_to_RGB(uint32 sx, uint32 sy, const uint8* From, uint8* To);
+	extern void ASTC10x10RGBL_to_RGBA(uint32 sx, uint32 sy, const uint8* From, uint8* To);
+	extern void ASTC10x10RGL_to_RGB(uint32 sx, uint32 sy, const uint8* From, uint8* To);
+	extern void ASTC10x10RGL_to_RGBA(uint32 sx, uint32 sy, const uint8* From, uint8* To);
+
+	//! ASTC 12x12 support
 	extern void ASTC12x12RGBAL_to_RGBA(uint32 sx, uint32 sy, const uint8* From, uint8* To);
 	extern void ASTC12x12RGBAL_to_RGB(uint32 sx, uint32 sy, const uint8* From, uint8* To);
-
-	//! ASTC 12x12 RGB low support
 	extern void ASTC12x12RGBL_to_RGB(uint32 sx, uint32 sy, const uint8* From, uint8* To);
 	extern void ASTC12x12RGBL_to_RGBA(uint32 sx, uint32 sy, const uint8* From, uint8* To);
-
-	//! ASTC 12x12 RG low support
 	extern void ASTC12x12RGL_to_RGB(uint32 sx, uint32 sy, const uint8* From, uint8* To);
 	extern void ASTC12x12RGL_to_RGBA(uint32 sx, uint32 sy, const uint8* From, uint8* To);
 
diff --git a/Engine/Plugins/Experimental/Mutable/Source/MutableRuntime/Private/MuR/Image.cpp b/Engine/Plugins/Experimental/Mutable/Source/MutableRuntime/Private/MuR/Image.cpp
--- a/Engine/Plugins/Experimental/Mutable/Source/MutableRuntime/Private/MuR/Image.cpp
+++ b/Engine/Plugins/Experimental/Mutable/Source/MutableRuntime/Private/MuR/Image.cpp
@@ -63,12 +63,18 @@ namespace mu
 		FImageFormatData(4, 4, 16, 4, {252, 253, 255, 255, 255, 255, 255, 255, 0, 0, 0, 0, 0, 0, 0, 0}),     // IF_ASTC_4x4_RGBA_LDR
         FImageFormatData( 4, 4, 16, 2 ),	// IF_ASTC_4x4_RG_LDR // TODO: check black block for RG.
 		
-		FImageFormatData(8, 8, 16, 3),		// IF_ASTC_8x8_RGB_LDR,
-		FImageFormatData(8, 8, 16, 4),		// IF_ASTC_8x8_RGBA_LDR,
+		FImageFormatData(8, 8, 16, 3, {252, 253, 255, 255, 255, 255, 255, 255, 0, 0, 0, 0, 0, 0, 255, 255}),	// IF_ASTC_8x8_RGB_LDR,
+		FImageFormatData(8, 8, 16, 4, {252, 253, 255, 255, 255, 255, 255, 255, 0, 0, 0, 0, 0, 0, 0, 0}),		// IF_ASTC_8x8_RGBA_LDR,
 		FImageFormatData(8, 8, 16, 2),		// IF_ASTC_8x8_RG_LDR,
-		FImageFormatData(12, 12, 16, 3),	// IF_ASTC_12x12_RGB_LDR
-		FImageFormatData(12, 12, 16, 4),	// IF_ASTC_12x12_RGBA_LDR
+		FImageFormatData(12, 12, 16, 3, {252, 253, 255, 255, 255, 255, 255, 255, 0, 0, 0, 0, 0, 0, 255, 255}),	// IF_ASTC_12x12_RGB_LDR
+		FImageFormatData(12, 12, 16, 4, {252, 253, 255, 255, 255, 255, 255, 255, 0, 0, 0, 0, 0, 0, 0, 0}),		// IF_ASTC_12x12_RGBA_LDR
 		FImageFormatData(12, 12, 16, 2),	// IF_ASTC_12x12_RG_LDR
+		FImageFormatData(6, 6, 16, 3, {252, 253, 255, 255, 255, 255, 255, 255, 0, 0, 0, 0, 0, 0, 255, 255}),	// IF_ASTC_6x6_RGB_LDR,
+		FImageFormatData(6, 6, 16, 4, {252, 253, 255, 255, 255, 255, 255, 255, 0, 0, 0, 0, 0, 0, 0, 0}),		// IF_ASTC_6x6_RGBA_LDR,
+		FImageFormatData(6, 6, 16, 2),		// IF_ASTC_6x6_RG_LDR,
+		FImageFormatData(10, 10, 16, 3, {252, 253, 255, 255, 255, 255, 255, 255, 0, 0, 0, 0, 0, 0, 255, 255}),	// IF_ASTC_6x6_RGB_LDR,
+		FImageFormatData(10, 10, 16, 4, {252, 253, 255, 255, 255, 255, 255, 255, 0, 0, 0, 0, 0, 0, 0, 0}),		// IF_ASTC_6x6_RGBA_LDR,
+		FImageFormatData(10, 10, 16, 2),	// IF_ASTC_6x6_RG_LDR,
 
     };
 
diff --git a/Engine/Plugins/Experimental/Mutable/Source/MutableRuntime/Private/MuR/ImagePrivate.h b/Engine/Plugins/Experimental/Mutable/Source/MutableRuntime/Private/MuR/ImagePrivate.h
--- a/Engine/Plugins/Experimental/Mutable/Source/MutableRuntime/Private/MuR/ImagePrivate.h
+++ b/Engine/Plugins/Experimental/Mutable/Source/MutableRuntime/Private/MuR/ImagePrivate.h
@@ -38,10 +38,22 @@ namespace mu
         case EImageFormat::IF_BC3: r = EImageFormat::IF_RGBA_UBYTE; break;
         case EImageFormat::IF_BC4: r = EImageFormat::IF_L_UBYTE; break;
         case EImageFormat::IF_BC5: r = EImageFormat::IF_RGB_UBYTE; break;
-        case EImageFormat::IF_ASTC_4x4_RGB_LDR: r = EImageFormat::IF_RGB_UBYTE; break;
-        case EImageFormat::IF_ASTC_4x4_RGBA_LDR: r = EImageFormat::IF_RGBA_UBYTE; break;
-        case EImageFormat::IF_ASTC_4x4_RG_LDR: r = EImageFormat::IF_RGB_UBYTE; break;
-        default: break;
+		case EImageFormat::IF_ASTC_4x4_RGB_LDR: r = EImageFormat::IF_RGB_UBYTE; break;
+		case EImageFormat::IF_ASTC_4x4_RGBA_LDR: r = EImageFormat::IF_RGBA_UBYTE; break;
+		case EImageFormat::IF_ASTC_4x4_RG_LDR: r = EImageFormat::IF_RGB_UBYTE; break;
+		case EImageFormat::IF_ASTC_6x6_RGB_LDR: r = EImageFormat::IF_RGB_UBYTE; break;
+		case EImageFormat::IF_ASTC_6x6_RGBA_LDR: r = EImageFormat::IF_RGBA_UBYTE; break;
+		case EImageFormat::IF_ASTC_6x6_RG_LDR: r = EImageFormat::IF_RGB_UBYTE; break;
+		case EImageFormat::IF_ASTC_8x8_RGB_LDR: r = EImageFormat::IF_RGB_UBYTE; break;
+		case EImageFormat::IF_ASTC_8x8_RGBA_LDR: r = EImageFormat::IF_RGBA_UBYTE; break;
+		case EImageFormat::IF_ASTC_8x8_RG_LDR: r = EImageFormat::IF_RGB_UBYTE; break;
+		case EImageFormat::IF_ASTC_10x10_RGB_LDR: r = EImageFormat::IF_RGB_UBYTE; break;
+		case EImageFormat::IF_ASTC_10x10_RGBA_LDR: r = EImageFormat::IF_RGBA_UBYTE; break;
+		case EImageFormat::IF_ASTC_10x10_RG_LDR: r = EImageFormat::IF_RGB_UBYTE; break;
+		case EImageFormat::IF_ASTC_12x12_RGB_LDR: r = EImageFormat::IF_RGB_UBYTE; break;
+		case EImageFormat::IF_ASTC_12x12_RGBA_LDR: r = EImageFormat::IF_RGBA_UBYTE; break;
+		case EImageFormat::IF_ASTC_12x12_RG_LDR: r = EImageFormat::IF_RGB_UBYTE; break;
+		default: break;
 		}
 
 		return r;
diff --git a/Engine/Plugins/Experimental/Mutable/Source/MutableRuntime/Private/MuR/OpImageMipmap.cpp b/Engine/Plugins/Experimental/Mutable/Source/MutableRuntime/Private/MuR/OpImageMipmap.cpp
--- a/Engine/Plugins/Experimental/Mutable/Source/MutableRuntime/Private/MuR/OpImageMipmap.cpp
+++ b/Engine/Plugins/Experimental/Mutable/Source/MutableRuntime/Private/MuR/OpImageMipmap.cpp
@@ -510,89 +510,67 @@ namespace mu
         check(Dest->GetSizeY() == Base->GetSizeY());
         check(Dest->GetFormat() == Base->GetFormat());
 
-        switch (Base->GetFormat())
-        {
-        case EImageFormat::IF_L_UBYTE:
-        case EImageFormat::IF_RGB_UBYTE:
-        case EImageFormat::IF_RGBA_UBYTE:
-        case EImageFormat::IF_BGRA_UBYTE:
-            break;
-
-        // Bad cases: we need to decompress, mip and then recompress. It may be necessary to
-        // generate the latests mips after composing, or in unoptimised code.
-        // TODO: Make sure that the code optimisations avoid this cases generating separate
-        // operations to generate the mip tail.
-        case EImageFormat::IF_BC1:
-        case EImageFormat::IF_BC2:
-        case EImageFormat::IF_BC3:
-        case EImageFormat::IF_BC4:
-        case EImageFormat::IF_BC5:
-        case EImageFormat::IF_ASTC_4x4_RGB_LDR:
-        case EImageFormat::IF_ASTC_4x4_RGBA_LDR:
-        case EImageFormat::IF_ASTC_4x4_RG_LDR:
-        {
-			// Uncompress the last mip that we already have
-			FIntVector2 UncompressedSize = Base->CalculateMipSize(StartLevel);
-			Scratch.Uncompressed = CreateImage( 
-				(uint16)UncompressedSize[0], (uint16)UncompressedSize[1], 
-				1, 
-				EImageFormat::IF_RGBA_UBYTE, EInitializationType::NotInitialized );
+		EImageFormat BaseFormat = Base->GetFormat();
+		if (mu::IsCompressedFormat(BaseFormat))
+		{
+			// Is it a block format?
+			if (mu::GetImageFormatData(BaseFormat).PixelsPerBlockX > 1)
+			{
+				// Uncompress the last mip that we already have
+				FIntVector2 UncompressedSize = Base->CalculateMipSize(StartLevel);
+				Scratch.Uncompressed = CreateImage(
+					(uint16)UncompressedSize[0], (uint16)UncompressedSize[1],
+					1,
+					EImageFormat::IF_RGBA_UBYTE, EInitializationType::NotInitialized);
+
+				FIntVector2 UncompressedMipsSize = Base->CalculateMipSize(StartLevel + 1);
+				// Generate the mipmaps from there on
+				Scratch.UncompressedMips = CreateImage(
+					(uint16)UncompressedMipsSize[0], (uint16)UncompressedMipsSize[1],
+					FMath::Max(1, LevelCount - StartLevel - 1),
+					EImageFormat::IF_RGBA_UBYTE, EInitializationType::NotInitialized);
+
+				// Compress the mipmapped image
+				Scratch.CompressedMips = CreateImage(
+					(uint16)UncompressedMipsSize[0], (uint16)UncompressedMipsSize[1],
+					Scratch.UncompressedMips->GetLODCount(),
+					Base->GetFormat(), EInitializationType::NotInitialized);
+			}
+			else
+			{
+				// It's probably an RLE compressed format
 
-			FIntVector2 UncompressedMipsSize = Base->CalculateMipSize(StartLevel + 1);
-			// Generate the mipmaps from there on
-			Scratch.UncompressedMips = CreateImage( 
-				(uint16)UncompressedMipsSize[0], (uint16)UncompressedMipsSize[1],
-				FMath::Max(1, LevelCount - StartLevel - 1), 
-				EImageFormat::IF_RGBA_UBYTE, EInitializationType::NotInitialized);
+				// Uncompress the last mip that we already have
+				FIntVector2 UncompressedSize = Base->CalculateMipSize(StartLevel);
+				Scratch.Uncompressed = CreateImage(
+					(uint16)UncompressedSize[0], (uint16)UncompressedSize[1],
+					1,
+					EImageFormat::IF_L_UBYTE, EInitializationType::NotInitialized);
 
-			// Compress the mipmapped image
-			Scratch.CompressedMips = CreateImage(
-				(uint16)UncompressedMipsSize[0], (uint16)UncompressedMipsSize[1],
-				Scratch.UncompressedMips->GetLODCount(),
-				Base->GetFormat(), EInitializationType::NotInitialized);
-
-            break;
-        }
-
-        case EImageFormat::IF_L_UBIT_RLE:
-        case EImageFormat::IF_L_UBYTE_RLE:
-        {
-            // Uncompress the last mip that we already have
-			FIntVector2 UncompressedSize = Base->CalculateMipSize(StartLevel);
-            Scratch.Uncompressed = CreateImage(
-                (uint16)UncompressedSize[0], (uint16)UncompressedSize[1],
-                1,
-                EImageFormat::IF_L_UBYTE, EInitializationType::NotInitialized);
-
-
-			FIntVector2 UncompressedMipsSize = Base->CalculateMipSize(StartLevel + 1);
-            // Generate the mipmaps from there on
-            Scratch.UncompressedMips = CreateImage(
-                (uint16)UncompressedMipsSize[0], (uint16)UncompressedMipsSize[1],
-				FMath::Max(1, LevelCount - StartLevel - 1),
-                EImageFormat::IF_L_UBYTE, EInitializationType::NotInitialized);
-
-
-            // Compress the mipmapped image
-            Scratch.CompressedMips = CreateImage(
-                (uint16)UncompressedMipsSize[0], (uint16)UncompressedMipsSize[1],
-                Scratch.UncompressedMips->GetLODCount(),
-                Base->GetFormat(), EInitializationType::NotInitialized);
-
-            // Preallocate ample memory for the compressed data
-			uint32 TotalMemory = Scratch.UncompressedMips->GetDataSize();
-            Scratch.CompressedMips->m_data.SetNumUninitialized(TotalMemory);
-
-            // Preallocate ample memory for the destination data
-			TotalMemory = Base->GetDataSize() + Scratch.UncompressedMips->GetDataSize();
-            Dest->m_data.SetNumUninitialized(TotalMemory);
-
-            break;
-        }
-
-        default:
-            checkf( false, TEXT("Format not implemented in mipmap generation."));
-        }
+
+				FIntVector2 UncompressedMipsSize = Base->CalculateMipSize(StartLevel + 1);
+				// Generate the mipmaps from there on
+				Scratch.UncompressedMips = CreateImage(
+					(uint16)UncompressedMipsSize[0], (uint16)UncompressedMipsSize[1],
+					FMath::Max(1, LevelCount - StartLevel - 1),
+					EImageFormat::IF_L_UBYTE, EInitializationType::NotInitialized);
+
+
+				// Compress the mipmapped image
+				Scratch.CompressedMips = CreateImage(
+					(uint16)UncompressedMipsSize[0], (uint16)UncompressedMipsSize[1],
+					Scratch.UncompressedMips->GetLODCount(),
+					Base->GetFormat(), EInitializationType::NotInitialized);
+
+				// Preallocate ample memory for the compressed data
+				uint32 TotalMemory = Scratch.UncompressedMips->GetDataSize();
+				Scratch.CompressedMips->m_data.SetNumUninitialized(TotalMemory);
+
+				// Preallocate ample memory for the destination data
+				TotalMemory = Base->GetDataSize() + Scratch.UncompressedMips->GetDataSize();
+				Dest->m_data.SetNumUninitialized(TotalMemory);
+			}
+		}
     }
 
 
@@ -654,36 +632,11 @@ namespace mu
 
 		FIntVector2 sourceSize = Base->CalculateMipSize(StartLevel);
 
-		switch (Base->GetFormat())
+		EImageFormat BaseFormat = Base->GetFormat();
+		if (mu::IsCompressedFormat(BaseFormat))
 		{
-		case EImageFormat::IF_L_UBYTE:
-			GenerateMipmapsUint8<1>(LevelCount - StartLevel - 1, pSourceBuf, pDestBuf, sourceSize, Settings);
-			break;
-
-		case EImageFormat::IF_RGB_UBYTE:
-			GenerateMipmapsUint8<3>(LevelCount - StartLevel - 1, pSourceBuf, pDestBuf, sourceSize, Settings);
-			break;
+			// Bad case.
 
-		case EImageFormat::IF_BGRA_UBYTE:
-		case EImageFormat::IF_RGBA_UBYTE:
-			GenerateMipmapsUint8<4>(LevelCount - StartLevel - 1, pSourceBuf, pDestBuf, sourceSize, Settings);
-			break;
-
-			// Bad cases: we need to decompress, mip and then recompress. It may be necessary to
-			// generate the latests mips after composing, or in unoptimised code.
-			// TODO: Make sure that the code optimisations avoid this cases generating separate
-			// operations to generate the mip tail.
-		case EImageFormat::IF_BC1:
-		case EImageFormat::IF_BC2:
-		case EImageFormat::IF_BC3:
-		case EImageFormat::IF_BC4:
-		case EImageFormat::IF_BC5:
-		case EImageFormat::IF_ASTC_4x4_RGB_LDR:
-		case EImageFormat::IF_ASTC_4x4_RGBA_LDR:
-		case EImageFormat::IF_ASTC_4x4_RG_LDR:
-		case EImageFormat::IF_L_UBIT_RLE:
-		case EImageFormat::IF_L_UBYTE_RLE:
-		{
 			// Uncompress the last mip that we already have
 			bool bSuccess = false;
 			ImagePixelFormat(bSuccess, CompressionQuality, Scratch.Uncompressed.get(), Base, StartLevel);
@@ -721,11 +674,27 @@ namespace mu
 			}
 
 			FMemory::Memcpy(pDestBuf, Scratch.CompressedMips->GetData(), Scratch.CompressedMips->GetDataSize());
-			break;
 		}
+		else
+		{
+			switch (Base->GetFormat())
+			{
+			case EImageFormat::IF_L_UBYTE:
+				GenerateMipmapsUint8<1>(LevelCount - StartLevel - 1, pSourceBuf, pDestBuf, sourceSize, Settings);
+				break;
 
-		default:
-			checkf(false, TEXT("Format not implemented in mipmap generation."));
+			case EImageFormat::IF_RGB_UBYTE:
+				GenerateMipmapsUint8<3>(LevelCount - StartLevel - 1, pSourceBuf, pDestBuf, sourceSize, Settings);
+				break;
+
+			case EImageFormat::IF_BGRA_UBYTE:
+			case EImageFormat::IF_RGBA_UBYTE:
+				GenerateMipmapsUint8<4>(LevelCount - StartLevel - 1, pSourceBuf, pDestBuf, sourceSize, Settings);
+				break;
+
+			default:
+				checkf(false, TEXT("Format not implemented in mipmap generation."));
+			}
 		}
 	}
 
diff --git a/Engine/Plugins/Experimental/Mutable/Source/MutableRuntime/Private/MuR/OpImagePixelFormat.cpp b/Engine/Plugins/Experimental/Mutable/Source/MutableRuntime/Private/MuR/OpImagePixelFormat.cpp
--- a/Engine/Plugins/Experimental/Mutable/Source/MutableRuntime/Private/MuR/OpImagePixelFormat.cpp
+++ b/Engine/Plugins/Experimental/Mutable/Source/MutableRuntime/Private/MuR/OpImagePixelFormat.cpp
@@ -563,6 +563,42 @@ void FImageOperator::ImagePixelFormat( bool& bOutSuccess, int32 CompressionQuali
                 break;
             }
 
+			case EImageFormat::IF_ASTC_6x6_RGB_LDR:
+			{
+				for (int m = 0; m < resultLODCount; ++m)
+				{
+					FIntVector2 mipSize = pResult->CalculateMipSize(m);
+					miro::ASTC6x6RGBL_to_RGB(mipSize[0], mipSize[1],
+						Base->GetMipData(baseLOD + m),
+						pResult->GetMipData(m));
+				}
+				break;
+			}
+
+			case EImageFormat::IF_ASTC_6x6_RGBA_LDR:
+			{
+				for (int m = 0; m < resultLODCount; ++m)
+				{
+					FIntVector2 mipSize = pResult->CalculateMipSize(m);
+					miro::ASTC6x6RGBAL_to_RGB(mipSize[0], mipSize[1],
+						Base->GetMipData(baseLOD + m),
+						pResult->GetMipData(m));
+				}
+				break;
+			}
+
+			case EImageFormat::IF_ASTC_6x6_RG_LDR:
+			{
+				for (int m = 0; m < resultLODCount; ++m)
+				{
+					FIntVector2 mipSize = pResult->CalculateMipSize(m);
+					miro::ASTC6x6RGL_to_RGB(mipSize[0], mipSize[1],
+						Base->GetMipData(baseLOD + m),
+						pResult->GetMipData(m));
+				}
+				break;
+			}
+
 			case EImageFormat::IF_ASTC_8x8_RGB_LDR:
 			{
 				for (int m = 0; m < resultLODCount; ++m)
@@ -599,6 +635,42 @@ void FImageOperator::ImagePixelFormat( bool& bOutSuccess, int32 CompressionQuali
 				break;
 			}
 
+			case EImageFormat::IF_ASTC_10x10_RGB_LDR:
+			{
+				for (int m = 0; m < resultLODCount; ++m)
+				{
+					FIntVector2 mipSize = pResult->CalculateMipSize(m);
+					miro::ASTC10x10RGBL_to_RGB(mipSize[0], mipSize[1],
+						Base->GetMipData(baseLOD + m),
+						pResult->GetMipData(m));
+				}
+				break;
+			}
+
+			case EImageFormat::IF_ASTC_10x10_RGBA_LDR:
+			{
+				for (int m = 0; m < resultLODCount; ++m)
+				{
+					FIntVector2 mipSize = pResult->CalculateMipSize(m);
+					miro::ASTC10x10RGBAL_to_RGB(mipSize[0], mipSize[1],
+						Base->GetMipData(baseLOD + m),
+						pResult->GetMipData(m));
+				}
+				break;
+			}
+
+			case EImageFormat::IF_ASTC_10x10_RG_LDR:
+			{
+				for (int m = 0; m < resultLODCount; ++m)
+				{
+					FIntVector2 mipSize = pResult->CalculateMipSize(m);
+					miro::ASTC10x10RGL_to_RGB(mipSize[0], mipSize[1],
+						Base->GetMipData(baseLOD + m),
+						pResult->GetMipData(m));
+				}
+				break;
+			}
+
 			case EImageFormat::IF_ASTC_12x12_RGB_LDR:
 			{
 				for (int m = 0; m < resultLODCount; ++m)
@@ -814,6 +886,30 @@ void FImageOperator::ImagePixelFormat( bool& bOutSuccess, int32 CompressionQuali
                 break;
             }
 
+			case EImageFormat::IF_ASTC_6x6_RGBA_LDR:
+			{
+				for (int m = 0; m < resultLODCount; ++m)
+				{
+					FIntVector2 mipSize = pResult->CalculateMipSize(m);
+					miro::ASTC6x6RGBAL_to_RGBA(mipSize[0], mipSize[1],
+						Base->GetMipData(baseLOD + m),
+						pResult->GetMipData(m));
+				}
+				break;
+			}
+
+			case EImageFormat::IF_ASTC_6x6_RG_LDR:
+			{
+				for (int m = 0; m < resultLODCount; ++m)
+				{
+					FIntVector2 mipSize = pResult->CalculateMipSize(m);
+					miro::ASTC6x6RGL_to_RGBA(mipSize[0], mipSize[1],
+						Base->GetMipData(baseLOD + m),
+						pResult->GetMipData(m));
+				}
+				break;
+			}
+
 			case EImageFormat::IF_ASTC_8x8_RGB_LDR:
 			{
 				for (int m = 0; m < resultLODCount; ++m)
@@ -850,6 +946,54 @@ void FImageOperator::ImagePixelFormat( bool& bOutSuccess, int32 CompressionQuali
 				break;
 			}
 
+			case EImageFormat::IF_ASTC_6x6_RGB_LDR:
+			{
+				for (int m = 0; m < resultLODCount; ++m)
+				{
+					FIntVector2 mipSize = pResult->CalculateMipSize(m);
+					miro::ASTC6x6RGBL_to_RGBA(mipSize[0], mipSize[1],
+						Base->GetMipData(baseLOD + m),
+						pResult->GetMipData(m));
+				}
+				break;
+			}
+
+			case EImageFormat::IF_ASTC_10x10_RGB_LDR:
+			{
+				for (int m = 0; m < resultLODCount; ++m)
+				{
+					FIntVector2 mipSize = pResult->CalculateMipSize(m);
+					miro::ASTC10x10RGBL_to_RGBA(mipSize[0], mipSize[1],
+						Base->GetMipData(baseLOD + m),
+						pResult->GetMipData(m));
+				}
+				break;
+			}
+
+			case EImageFormat::IF_ASTC_10x10_RGBA_LDR:
+			{
+				for (int m = 0; m < resultLODCount; ++m)
+				{
+					FIntVector2 mipSize = pResult->CalculateMipSize(m);
+					miro::ASTC10x10RGBAL_to_RGBA(mipSize[0], mipSize[1],
+						Base->GetMipData(baseLOD + m),
+						pResult->GetMipData(m));
+				}
+				break;
+			}
+
+			case EImageFormat::IF_ASTC_10x10_RG_LDR:
+			{
+				for (int m = 0; m < resultLODCount; ++m)
+				{
+					FIntVector2 mipSize = pResult->CalculateMipSize(m);
+					miro::ASTC10x10RGL_to_RGBA(mipSize[0], mipSize[1],
+						Base->GetMipData(baseLOD + m),
+						pResult->GetMipData(m));
+				}
+				break;
+			}
+
 			case EImageFormat::IF_ASTC_12x12_RGB_LDR:
 			{
 				for (int m = 0; m < resultLODCount; ++m)
@@ -1361,6 +1505,275 @@ void FImageOperator::ImagePixelFormat( bool& bOutSuccess, int32 CompressionQuali
             break;
         }
 
+		case EImageFormat::IF_ASTC_6x6_RGB_LDR:
+		{
+			switch (Base->GetFormat())
+			{
+			case EImageFormat::IF_RGBA_UBYTE:
+				for (int m = 0; m < resultLODCount; ++m)
+				{
+					FIntVector2 mipSize = pResult->CalculateMipSize(m);
+					miro::RGBA_to_ASTC6x6RGBL(
+						mipSize[0], mipSize[1], Base->GetMipData(baseLOD + m),
+						pResult->GetMipData(m), CompressionQuality);
+				}
+				break;
+
+			case EImageFormat::IF_RGB_UBYTE:
+				for (int m = 0; m < resultLODCount; ++m)
+				{
+					FIntVector2 mipSize = pResult->CalculateMipSize(m);
+					miro::RGB_to_ASTC6x6RGBL(
+						mipSize[0], mipSize[1], Base->GetMipData(baseLOD + m),
+						pResult->GetMipData(m), CompressionQuality);
+				}
+				break;
+
+			//case EImageFormat::IF_ASTC_6x6_RGBA_LDR:
+			//	for (int m = 0; m < resultLODCount; ++m)
+			//	{
+			//		FIntVector2 mipSize = pResult->CalculateMipSize(m);
+			//		miro::ASTC6x6RGBAL_to_ASTC6x6RGBL(mipSize[0], mipSize[1],
+			//			Base->GetMipData(baseLOD + m),
+			//			pResult->GetMipData(m));
+			//	}
+			//	break;
+
+			//case EImageFormat::IF_L_UBYTE:
+			//	for (int m = 0; m < resultLODCount; ++m)
+			//	{
+			//		FIntVector2 mipSize = pResult->CalculateMipSize(m);
+			//		miro::L_to_ASTC6x6RGBL(mipSize[0], mipSize[1],
+			//			Base->GetMipData(baseLOD + m),
+			//			pResult->GetMipData(m), CompressionQuality);
+			//	}
+			//	break;
+
+			default:
+				// Case not implemented
+				check(false);
+
+			}
+
+			break;
+		}
+
+		case EImageFormat::IF_ASTC_6x6_RGBA_LDR:
+		{
+			switch (Base->GetFormat())
+			{
+			case EImageFormat::IF_RGBA_UBYTE:
+				for (int m = 0; m < resultLODCount; ++m)
+				{
+					FIntVector2 mipSize = pResult->CalculateMipSize(m);
+					miro::RGBA_to_ASTC6x6RGBAL(
+						mipSize[0], mipSize[1], Base->GetMipData(baseLOD + m),
+						pResult->GetMipData(m), CompressionQuality);
+				}
+				break;
+
+			case EImageFormat::IF_RGB_UBYTE:
+				for (int m = 0; m < resultLODCount; ++m)
+				{
+					FIntVector2 mipSize = pResult->CalculateMipSize(m);
+					miro::RGB_to_ASTC6x6RGBAL(
+						mipSize[0], mipSize[1], Base->GetMipData(baseLOD + m),
+						pResult->GetMipData(m), CompressionQuality);
+				}
+				break;
+
+			//case EImageFormat::IF_ASTC_6x6_RGB_LDR:
+			//	for (int m = 0; m < resultLODCount; ++m)
+			//	{
+			//		FIntVector2 mipSize = pResult->CalculateMipSize(m);
+			//		miro::ASTC6x6RGBL_to_ASTC6x6RGBAL(mipSize[0], mipSize[1],
+			//			Base->GetMipData(baseLOD + m),
+			//			pResult->GetMipData(m));
+			//	}
+			//	break;
+
+			default:
+				// Case not implemented
+				check(false);
+
+			}
+
+			break;
+		}
+
+		case EImageFormat::IF_ASTC_6x6_RG_LDR:
+		{
+			switch (Base->GetFormat())
+			{
+			case EImageFormat::IF_RGBA_UBYTE:
+				for (int m = 0; m < resultLODCount; ++m)
+				{
+					FIntVector2 mipSize = pResult->CalculateMipSize(m);
+					miro::RGBA_to_ASTC6x6RGL(
+						mipSize[0], mipSize[1], Base->GetMipData(baseLOD + m),
+						pResult->GetMipData(m), CompressionQuality);
+				}
+				break;
+
+			case EImageFormat::IF_RGB_UBYTE:
+				for (int m = 0; m < resultLODCount; ++m)
+				{
+					FIntVector2 mipSize = pResult->CalculateMipSize(m);
+					miro::RGB_to_ASTC6x6RGL(
+						mipSize[0], mipSize[1], Base->GetMipData(baseLOD + m),
+						pResult->GetMipData(m), CompressionQuality);
+				}
+				break;
+
+			//case EImageFormat::IF_ASTC_6x6_RG_LDR:
+			//	for (int m = 0; m < resultLODCount; ++m)
+			//	{
+			//		FIntVector2 mipSize = pResult->CalculateMipSize(m);
+			//		// Hack that actually works because of block size.
+			//		miro::ASTC6x6RGBAL_to_ASTC6x6RGBL(mipSize[0], mipSize[1],
+			//			Base->GetMipData(baseLOD + m),
+			//			pResult->GetMipData(m));
+			//	}
+			//	break;
+
+			default:
+				// Case not implemented
+				check(false);
+
+			}
+
+			break;
+		}
+
+		case EImageFormat::IF_ASTC_8x8_RGB_LDR:
+		{
+			switch (Base->GetFormat())
+			{
+			case EImageFormat::IF_RGBA_UBYTE:
+				for (int m = 0; m < resultLODCount; ++m)
+				{
+					FIntVector2 mipSize = pResult->CalculateMipSize(m);
+					miro::RGBA_to_ASTC8x8RGBL(
+						mipSize[0], mipSize[1], Base->GetMipData(baseLOD + m),
+						pResult->GetMipData(m), CompressionQuality);
+				}
+				break;
+
+			case EImageFormat::IF_RGB_UBYTE:
+				for (int m = 0; m < resultLODCount; ++m)
+				{
+					FIntVector2 mipSize = pResult->CalculateMipSize(m);
+					miro::RGB_to_ASTC8x8RGBL(
+						mipSize[0], mipSize[1], Base->GetMipData(baseLOD + m),
+						pResult->GetMipData(m), CompressionQuality);
+				}
+				break;
+
+				//case EImageFormat::IF_ASTC_8x8_RGBA_LDR:
+				//	for (int m = 0; m < resultLODCount; ++m)
+				//	{
+				//		FIntVector2 mipSize = pResult->CalculateMipSize(m);
+				//		miro::ASTC8x8RGBAL_to_ASTC8x8RGBL(mipSize[0], mipSize[1],
+				//			Base->GetMipData(baseLOD + m),
+				//			pResult->GetMipData(m));
+				//	}
+				//	break;
+
+				//case EImageFormat::IF_L_UBYTE:
+				//	for (int m = 0; m < resultLODCount; ++m)
+				//	{
+				//		FIntVector2 mipSize = pResult->CalculateMipSize(m);
+				//		miro::L_to_ASTC8x8RGBL(mipSize[0], mipSize[1],
+				//			Base->GetMipData(baseLOD + m),
+				//			pResult->GetMipData(m), CompressionQuality);
+				//	}
+				//	break;
+
+			default:
+				// Case not implemented
+				check(false);
+
+			}
+
+			break;
+		}
+
+		case EImageFormat::IF_ASTC_8x8_RGBA_LDR:
+		{
+			switch (Base->GetFormat())
+			{
+			case EImageFormat::IF_RGBA_UBYTE:
+				for (int m = 0; m < resultLODCount; ++m)
+				{
+					FIntVector2 mipSize = pResult->CalculateMipSize(m);
+					miro::RGBA_to_ASTC8x8RGBAL(
+						mipSize[0], mipSize[1], Base->GetMipData(baseLOD + m),
+						pResult->GetMipData(m), CompressionQuality);
+				}
+				break;
+
+			case EImageFormat::IF_RGB_UBYTE:
+				for (int m = 0; m < resultLODCount; ++m)
+				{
+					FIntVector2 mipSize = pResult->CalculateMipSize(m);
+					miro::RGB_to_ASTC8x8RGBAL(
+						mipSize[0], mipSize[1], Base->GetMipData(baseLOD + m),
+						pResult->GetMipData(m), CompressionQuality);
+				}
+				break;
+
+			//case EImageFormat::IF_ASTC_8x8_RGB_LDR:
+			//	for (int m = 0; m < resultLODCount; ++m)
+			//	{
+			//		FIntVector2 mipSize = pResult->CalculateMipSize(m);
+			//		miro::ASTC8x8RGBL_to_ASTC8x8RGBAL(mipSize[0], mipSize[1],
+			//			Base->GetMipData(baseLOD + m),
+			//			pResult->GetMipData(m));
+			//	}
+			//	break;
+
+			default:
+				// Case not implemented
+				check(false);
+
+			}
+
+			break;
+		}
+
+		case EImageFormat::IF_ASTC_8x8_RG_LDR:
+		{
+			switch (Base->GetFormat())
+			{
+			case EImageFormat::IF_RGBA_UBYTE:
+				for (int m = 0; m < resultLODCount; ++m)
+				{
+					FIntVector2 mipSize = pResult->CalculateMipSize(m);
+					miro::RGBA_to_ASTC8x8RGL(
+						mipSize[0], mipSize[1], Base->GetMipData(baseLOD + m),
+						pResult->GetMipData(m), CompressionQuality);
+				}
+				break;
+
+			case EImageFormat::IF_RGB_UBYTE:
+				for (int m = 0; m < resultLODCount; ++m)
+				{
+					FIntVector2 mipSize = pResult->CalculateMipSize(m);
+					miro::RGB_to_ASTC8x8RGL(
+						mipSize[0], mipSize[1], Base->GetMipData(baseLOD + m),
+						pResult->GetMipData(m), CompressionQuality);
+				}
+				break;
+
+			default:
+				// Case not implemented
+				check(false);
+
+			}
+
+			break;
+		}
+
         default:
             // Case not implemented
             check( false );
diff --git a/Engine/Plugins/Experimental/Mutable/Source/MutableRuntime/Public/MuR/Image.h b/Engine/Plugins/Experimental/Mutable/Source/MutableRuntime/Public/MuR/Image.h
--- a/Engine/Plugins/Experimental/Mutable/Source/MutableRuntime/Public/MuR/Image.h
+++ b/Engine/Plugins/Experimental/Mutable/Source/MutableRuntime/Public/MuR/Image.h
@@ -78,6 +78,12 @@ namespace mu
 		IF_ASTC_12x12_RGB_LDR,
 		IF_ASTC_12x12_RGBA_LDR,
 		IF_ASTC_12x12_RG_LDR,
+		IF_ASTC_6x6_RGB_LDR,
+		IF_ASTC_6x6_RGBA_LDR,
+		IF_ASTC_6x6_RG_LDR,
+		IF_ASTC_10x10_RGB_LDR,
+		IF_ASTC_10x10_RGBA_LDR,
+		IF_ASTC_10x10_RG_LDR,
 
         IF_COUNT
 	};
diff --git a/Engine/Plugins/Experimental/Mutable/Source/MutableTools/Private/MuT/TypeInfo.cpp b/Engine/Plugins/Experimental/Mutable/Source/MutableTools/Private/MuT/TypeInfo.cpp
--- a/Engine/Plugins/Experimental/Mutable/Source/MutableTools/Private/MuT/TypeInfo.cpp
+++ b/Engine/Plugins/Experimental/Mutable/Source/MutableTools/Private/MuT/TypeInfo.cpp
@@ -45,8 +45,18 @@ namespace mu
 		"ASTC_12x12_RGB_LDR",
 		"ASTC_12x12_RGBA_LDR",
 		"ASTC_12x12_RG_LDR",
+
+		"ASTC_6x6_RGB_LDR",
+		"ASTC_6x6_RGBA_LDR",
+		"ASTC_6x6_RG_LDR",
+		"ASTC_10x10_RGB_LDR",
+		"ASTC_10x10_RGBA_LDR",
+		"ASTC_10x10_RG_LDR",
+
 	};
 
+	static_assert(sizeof(TypeInfo::s_imageFormatName) / sizeof(void*) == int32(EImageFormat::IF_COUNT));
+
 
 	const char* TypeInfo::s_meshBufferSemanticName[MBS_COUNT] =
 	{
@@ -77,6 +87,8 @@ namespace mu
 
 	};
 
+	static_assert(sizeof(TypeInfo::s_meshBufferSemanticName) / sizeof(void*) == int32(MBS_COUNT));
+
 	const char* TypeInfo::s_meshBufferFormatName[MBF_COUNT] =
 	{
 		"None",
@@ -105,12 +117,15 @@ namespace mu
 		"Float64",
 	};
 
-	const char* TypeInfo::s_projectorTypeName[static_cast<uint32>(mu::PROJECTOR_TYPE::COUNT)] =
+	static_assert(sizeof(TypeInfo::s_meshBufferFormatName) / sizeof(void*) == int32(MBF_COUNT));
+
+	const char* TypeInfo::s_projectorTypeName[static_cast<uint32>(PROJECTOR_TYPE::COUNT)] =
 	{
 		"Planar",
 		"Cylindrical",
 		"Wrapping",
 	};
+	static_assert(sizeof(TypeInfo::s_projectorTypeName) / sizeof(void*) == int32(PROJECTOR_TYPE::COUNT));
 
 	const char* TypeInfo::s_curveInterpolationModeName[static_cast<uint8_t>(mu::CurveKeyFrame::InterpMode::Count)] =
 	{
@@ -150,5 +165,6 @@ namespace mu
 		"Lighten", // TODO: This name is not descriptive.
 		"NormalCombine",
 	};
+	static_assert(sizeof(TypeInfo::s_blendModeName) / sizeof(void*) == int32(EBlendType::_BT_COUNT));
 
 }
\ No newline at end of file
