[Mutable] Initialize the generated meshes with the LOD quality settings specified in the CO.

#rnx

diff --git a/Engine/Plugins/Experimental/Mutable/Source/CustomizableObject/Classes/MuCO/CustomizableObject.h b/Engine/Plugins/Experimental/Mutable/Source/CustomizableObject/Classes/MuCO/CustomizableObject.h
--- a/Engine/Plugins/Experimental/Mutable/Source/CustomizableObject/Classes/MuCO/CustomizableObject.h
+++ b/Engine/Plugins/Experimental/Mutable/Source/CustomizableObject/Classes/MuCO/CustomizableObject.h
@@ -989,7 +989,6 @@ struct FMutableLODSettings
 {
 	GENERATED_BODY()
 
-#if WITH_EDITORONLY_DATA
 
 	/** Minimum LOD to render per Platform. */
 	UPROPERTY(EditAnywhere, Category = LODSettings, meta = (DisplayName = "Minimum LOD"))
@@ -999,6 +998,8 @@ struct FMutableLODSettings
 	UPROPERTY(EditAnywhere, Category = LODSettings, meta = (DisplayName = "Quality Level Minimum LOD"))
 	FPerQualityLevelInt MinQualityLevelLOD;
 
+#if WITH_EDITORONLY_DATA
+
 	/** Override the LOD Streaming settings from the reference skeletal meshes.*/
 	UPROPERTY(EditAnywhere, Category = LODSettings, meta = (DisplayName = "Override LOD Streaming Settings"))
 	bool bOverrideLODStreamingSettings = true;
@@ -1666,6 +1667,9 @@ public:
 
 	const TArray<FName>& GetBoneNamesArray() const;
 
+	/** Return the MinLOD index to generate based on the active LODSettings (PerPlatformMinLOD or PerQualityLevelMinLOD) */
+	int32 GetMinLODIndex() const;
+
 private:
 	
 	/** Stores the bone names of all the bones that can possibly use the generated meshes */
diff --git a/Engine/Plugins/Experimental/Mutable/Source/CustomizableObject/Classes/MuCO/CustomizableObjectSystem.h b/Engine/Plugins/Experimental/Mutable/Source/CustomizableObject/Classes/MuCO/CustomizableObjectSystem.h
--- a/Engine/Plugins/Experimental/Mutable/Source/CustomizableObject/Classes/MuCO/CustomizableObjectSystem.h
+++ b/Engine/Plugins/Experimental/Mutable/Source/CustomizableObject/Classes/MuCO/CustomizableObjectSystem.h
@@ -250,6 +250,10 @@ public:
 	FEditorCompileSettings EditorSettings;
 
 #endif
+	
+	// Return the current MinLodQualityLevel for skeletal meshes.
+	int32 GetSkeletalMeshMinLODQualityLevel() const;
+
 	bool IsSupport16BitBoneIndexEnabled() const;
 
 	bool IsProgressiveMipStreamingEnabled() const;
diff --git a/Engine/Plugins/Experimental/Mutable/Source/CustomizableObject/Private/MuCO/CustomizableInstancePrivateData.h b/Engine/Plugins/Experimental/Mutable/Source/CustomizableObject/Private/MuCO/CustomizableInstancePrivateData.h
--- a/Engine/Plugins/Experimental/Mutable/Source/CustomizableObject/Private/MuCO/CustomizableInstancePrivateData.h
+++ b/Engine/Plugins/Experimental/Mutable/Source/CustomizableObject/Private/MuCO/CustomizableInstancePrivateData.h
@@ -268,7 +268,7 @@ private:
 
 	void InitLastUpdateData(const TSharedPtr<FMutableOperationData>& OperationData);
 
-	void InitSkeletalMeshData(const TSharedPtr<FMutableOperationData>& OperationData, USkeletalMesh* SkeletalMesh, const FMutableRefSkeletalMeshData* RefSkeletalMeshData, int32 ComponentIndex);
+	void InitSkeletalMeshData(const TSharedPtr<FMutableOperationData>& OperationData, USkeletalMesh* SkeletalMesh, const FMutableRefSkeletalMeshData* RefSkeletalMeshData, const UCustomizableObject& CustomizableObject, int32 ComponentIndex);
 
 	bool BuildSkeletonData(const TSharedPtr<FMutableOperationData>& OperationData, USkeletalMesh& SkeletalMesh, const FMutableRefSkeletalMeshData& RefSkeletalMeshData, UCustomizableObject& CustomizableObject, int32 ComponentIndex);
 	void BuildMeshSockets(const TSharedPtr<FMutableOperationData>& OperationData, USkeletalMesh* SkeletalMesh, const FMutableRefSkeletalMeshData* RefSkeletalMeshData, UCustomizableObjectInstance* CustomizableObjectInstance, mu::MeshPtrConst MutableMesh);
diff --git a/Engine/Plugins/Experimental/Mutable/Source/CustomizableObject/Private/MuCO/CustomizableObject.cpp b/Engine/Plugins/Experimental/Mutable/Source/CustomizableObject/Private/MuCO/CustomizableObject.cpp
--- a/Engine/Plugins/Experimental/Mutable/Source/CustomizableObject/Private/MuCO/CustomizableObject.cpp
+++ b/Engine/Plugins/Experimental/Mutable/Source/CustomizableObject/Private/MuCO/CustomizableObject.cpp
@@ -1778,6 +1778,26 @@ const TArray<FName>& UCustomizableObject::GetBoneNamesArray() const
 }
 
 
+int32 UCustomizableObject::GetMinLODIndex() const
+{
+	int32 MinLODIdx = 0;
+
+	if (GEngine && GEngine->UseSkeletalMeshMinLODPerQualityLevels)
+	{
+		if (UCustomizableObjectSystem::GetInstance() != nullptr)
+		{
+			MinLODIdx = LODSettings.MinQualityLevelLOD.GetValue(UCustomizableObjectSystem::GetInstance()->GetSkeletalMeshMinLODQualityLevel());
+		}
+	}
+	else
+	{
+		MinLODIdx = LODSettings.MinLOD.GetValue();
+	}
+
+	return FMath::Max(MinLODIdx, LODSettings.FirstLODAvailable);
+}
+
+
 FGuid UCustomizableObject::GetCompilationGuid() const
 {
 	return CompilationGuid;
diff --git a/Engine/Plugins/Experimental/Mutable/Source/CustomizableObject/Private/MuCO/CustomizableObjectInstance.cpp b/Engine/Plugins/Experimental/Mutable/Source/CustomizableObject/Private/MuCO/CustomizableObjectInstance.cpp
--- a/Engine/Plugins/Experimental/Mutable/Source/CustomizableObject/Private/MuCO/CustomizableObjectInstance.cpp
+++ b/Engine/Plugins/Experimental/Mutable/Source/CustomizableObject/Private/MuCO/CustomizableObjectInstance.cpp
@@ -1956,7 +1956,7 @@ bool UCustomizableInstancePrivateData::UpdateSkeletalMesh_PostBeginUpdate0(UCust
 		INC_DWORD_STAT(STAT_MutableNumSkeletalMeshes);
 
 		// Set up the default information any mesh from this component will have (LODArrayInfos, RenderData, Mesh settings, etc). 
-		InitSkeletalMeshData(OperationData, SkeletalMesh, RefSkeletalMeshData, Component.Id);
+		InitSkeletalMeshData(OperationData, SkeletalMesh, RefSkeletalMeshData, *CustomizableObject, Component.Id);
 
 		if (Component.Mesh)
 		{
@@ -3329,7 +3329,7 @@ void UCustomizableInstancePrivateData::ConvertImage(UTexture2D* Texture, mu::Ima
 }
 
 
-void UCustomizableInstancePrivateData::InitSkeletalMeshData(const TSharedPtr<FMutableOperationData>& OperationData, USkeletalMesh* SkeletalMesh, const FMutableRefSkeletalMeshData* RefSkeletalMeshData, int32 ComponentIndex)
+void UCustomizableInstancePrivateData::InitSkeletalMeshData(const TSharedPtr<FMutableOperationData>& OperationData, USkeletalMesh* SkeletalMesh, const FMutableRefSkeletalMeshData* RefSkeletalMeshData, const UCustomizableObject& CustomizableObject, int32 ComponentIndex)
 {
 	MUTABLE_CPUPROFILER_SCOPE(UCustomizableInstancePrivateData::InitSkeletalMesh);
 
@@ -3344,8 +3344,8 @@ void UCustomizableInstancePrivateData::InitSkeletalMeshData(const TSharedPtr<FMu
 	SkeletalMesh->SetShadowPhysicsAsset(RefSkeletalMeshData->ShadowPhysicsAsset.Get());
 
 	// Set Min LOD
-	SkeletalMesh->SetMinLod(FirstLODAvailable);
-	SkeletalMesh->SetQualityLevelMinLod(FirstLODAvailable);
+	SkeletalMesh->SetMinLod(FMath::Max(CustomizableObject.LODSettings.MinLOD.GetDefault(), (int32)FirstLODAvailable));
+	SkeletalMesh->SetQualityLevelMinLod(CustomizableObject.LODSettings.MinQualityLevelLOD);
 
 	SkeletalMesh->SetHasBeenSimplified(false);
 	SkeletalMesh->SetHasVertexColors(false);
@@ -6317,7 +6317,8 @@ void UCustomizableObjectInstance::SetRequestedLODs(int32 InMinLOD, int32 InMaxLO
 	FMutableUpdateCandidate MutableUpdateCandidate(this);
 
 	// Clamp Min LOD
-	InMinLOD = FMath::Min(FMath::Max(InMinLOD, static_cast<int32>(PrivateData->FirstLODAvailable)), PrivateData->FirstLODAvailable + PrivateData->NumMaxLODsToStream);
+	const int32 MinLODIdx = GetCustomizableObject()->GetMinLODIndex();
+	InMinLOD = FMath::Min(FMath::Max(InMinLOD, MinLODIdx), MinLODIdx + PrivateData->NumMaxLODsToStream);
 
 	// Clamp Max LOD
 	InMaxLOD = FMath::Max(FMath::Min3(InMaxLOD, PrivateData->NumLODsAvailable - 1, (int32)MAX_MESH_LOD_COUNT), InMinLOD);
diff --git a/Engine/Plugins/Experimental/Mutable/Source/CustomizableObject/Private/MuCO/CustomizableObjectSystem.cpp b/Engine/Plugins/Experimental/Mutable/Source/CustomizableObject/Private/MuCO/CustomizableObjectSystem.cpp
--- a/Engine/Plugins/Experimental/Mutable/Source/CustomizableObject/Private/MuCO/CustomizableObjectSystem.cpp
+++ b/Engine/Plugins/Experimental/Mutable/Source/CustomizableObject/Private/MuCO/CustomizableObjectSystem.cpp
@@ -107,6 +107,19 @@ static TAutoConsoleVariable<int32> CVarGeneratedResourcesCacheSize(
 	TEXT("Limit the number of resources (images and meshes) that will be tracked for reusal. Each tracked resource uses a small amout of memory for its key."),
 	ECVF_Scalability);
 
+int32 FCustomizableObjectSystemPrivate::SkeletalMeshMinLodQualityLevel = -1;
+static void CVarMutableSinkFunction()
+{
+	if (UCustomizableObjectSystem::IsCreated())
+	{
+		FCustomizableObjectSystemPrivate* PrivateSystem = UCustomizableObjectSystem::GetInstance()->GetPrivate();
+
+		static const IConsoleVariable* CVarSkeletalMeshMinLodQualityLevelCVarName = IConsoleManager::Get().FindConsoleVariable(TEXT("r.SkeletalMesh.MinLodQualityLevel"));
+		PrivateSystem->SkeletalMeshMinLodQualityLevel = CVarSkeletalMeshMinLodQualityLevelCVarName ? CVarSkeletalMeshMinLodQualityLevelCVarName->GetInt() : INDEX_NONE;
+	}
+}
+static FAutoConsoleVariableSink CVarMutableSink(FConsoleCommandDelegate::CreateStatic(&CVarMutableSinkFunction));
+
 
 bool FMutablePendingInstanceWork::ArePendingUpdatesEmpty() const
 {
@@ -330,6 +343,8 @@ void UCustomizableObjectSystem::InitSystem()
 	const IConsoleVariable* CVarSupport16BitBoneIndex = IConsoleManager::Get().FindConsoleVariable(TEXT("r.GPUSkin.Support16BitBoneIndex"));
 	Private->bSupport16BitBoneIndex = CVarSupport16BitBoneIndex ? CVarSupport16BitBoneIndex->GetBool() : false;
 
+	CVarMutableSinkFunction();
+
 	Private->CurrentMutableOperation = nullptr;
 	Private->CurrentInstanceBeingUpdated = nullptr;
 
@@ -3134,6 +3149,13 @@ int32 UCustomizableObjectSystem::GetAverageBuildTime() const
 }
 
 
+int32 UCustomizableObjectSystem::GetSkeletalMeshMinLODQualityLevel() const
+{
+	check(Private != nullptr);
+	return Private->SkeletalMeshMinLodQualityLevel;
+}
+
+
 bool UCustomizableObjectSystem::IsSupport16BitBoneIndexEnabled() const
 {
 	check(Private != nullptr);
diff --git a/Engine/Plugins/Experimental/Mutable/Source/CustomizableObject/Private/MuCO/CustomizableObjectSystemPrivate.h b/Engine/Plugins/Experimental/Mutable/Source/CustomizableObject/Private/MuCO/CustomizableObjectSystemPrivate.h
--- a/Engine/Plugins/Experimental/Mutable/Source/CustomizableObject/Private/MuCO/CustomizableObjectSystemPrivate.h
+++ b/Engine/Plugins/Experimental/Mutable/Source/CustomizableObject/Private/MuCO/CustomizableObjectSystemPrivate.h
@@ -656,6 +656,7 @@ public:
 	static int32 EnableSkipGenerateResidentMips;
 	static int32 EnableOnlyGenerateRequestedLODs;
 	static int32 MaxTextureSizeToGenerate;
+	static int32 SkeletalMeshMinLodQualityLevel;
 	static bool bEnableMutableReusePreviousUpdateData;
 
 	/** */
