[Mutable] Add Image explicit output size to TextureTransform operations.

- Now an output image size can be specified next to a reference image. If no reference images is present, the size and lod settings of the base texture is used.
- Added check to image transform operations to keep source image aspect ratio.
- Updated ImageTransform tooltip.

#rnx

diff --git a/Engine/Plugins/Experimental/Mutable/Source/CustomizableObject/Classes/MuCO/CustomizableObject.h b/Engine/Plugins/Experimental/Mutable/Source/CustomizableObject/Classes/MuCO/CustomizableObject.h
--- a/Engine/Plugins/Experimental/Mutable/Source/CustomizableObject/Classes/MuCO/CustomizableObject.h
+++ b/Engine/Plugins/Experimental/Mutable/Source/CustomizableObject/Classes/MuCO/CustomizableObject.h
@@ -1440,7 +1440,7 @@ private:
 	// This is a manual version number for the binary blobs in this asset.
 	// Increasing it invalidates all the previously compiled models.
 	// Warning: If while merging code both versions have changed, take the highest+1.
-	static const int32 CurrentSupportedVersion = 400;
+	static const int32 CurrentSupportedVersion = 401;
 
 public:
 
diff --git a/Engine/Plugins/Experimental/Mutable/Source/CustomizableObjectEditor/Classes/MuCOE/Nodes/CustomizableObjectNodeTextureTransform.h b/Engine/Plugins/Experimental/Mutable/Source/CustomizableObjectEditor/Classes/MuCOE/Nodes/CustomizableObjectNodeTextureTransform.h
--- a/Engine/Plugins/Experimental/Mutable/Source/CustomizableObjectEditor/Classes/MuCOE/Nodes/CustomizableObjectNodeTextureTransform.h
+++ b/Engine/Plugins/Experimental/Mutable/Source/CustomizableObjectEditor/Classes/MuCOE/Nodes/CustomizableObjectNodeTextureTransform.h
@@ -28,6 +28,25 @@ public:
 	UPROPERTY(EditAnywhere, Category = CustomizableObject)
 	ETextureTransformAddressMode AddressMode = ETextureTransformAddressMode::Wrap;
 
+
+	/** If true, the operation will keep the source aspect ratio regardless of the output image size */
+	UPROPERTY(EditAnywhere, Category = CustomizableObject)
+	bool bKeepAspectRatio = false;
+
+	/** Set the width of the Texture. If greater than zero, it overrides the Reference Texture width. */
+	UPROPERTY(EditAnywhere, Category = CustomizableObject)
+	uint32 TextureSizeX = 0;
+
+	/** Set the height of the Texture. If greater than zero, it overrides the Reference Texture height. */
+	UPROPERTY(EditAnywhere, Category = CustomizableObject)
+	uint32 TextureSizeY = 0;
+
+	/** Reference Texture used to decide the texture properties of the mutable-generated textures
+	* connected to this material (e.g. LODBias, Size X,...). If null, mutable default texture properties will be applied. */
+	UPROPERTY(EditAnywhere, Category = CustomizableObject)
+	TObjectPtr<UTexture2D> ReferenceTexture = nullptr;
+
+
 	// Begin EdGraphNode interface
 	FText GetNodeTitle(ENodeTitleType::Type TittleType) const override;
 	FLinearColor GetNodeTitleColor() const override;
diff --git a/Engine/Plugins/Experimental/Mutable/Source/CustomizableObjectEditor/Private/MuCOE/GenerateMutableSource/GenerateMutableSourceImage.cpp b/Engine/Plugins/Experimental/Mutable/Source/CustomizableObjectEditor/Private/MuCOE/GenerateMutableSource/GenerateMutableSourceImage.cpp
--- a/Engine/Plugins/Experimental/Mutable/Source/CustomizableObjectEditor/Private/MuCOE/GenerateMutableSource/GenerateMutableSourceImage.cpp
+++ b/Engine/Plugins/Experimental/Mutable/Source/CustomizableObjectEditor/Private/MuCOE/GenerateMutableSource/GenerateMutableSourceImage.cpp
@@ -803,6 +803,31 @@ mu::NodeImagePtr GenerateMutableSourceImage(const UEdGraphPin* Pin, FMutableGrap
 			default: { check(false); return mu::EAddressMode::None; }
 			}
 		}));
+
+		FUintVector2 TextureSize(TypedNodeTransform->TextureSizeX, TypedNodeTransform->TextureSizeY);
+
+		// Calculating Texture size using Reference texture parameters
+		if (TypedNodeTransform->ReferenceTexture)
+		{
+			const int32 LODBias = ComputeLODBias(
+				GenerationContext, TypedNodeTransform->ReferenceTexture, TypedNodeTransform->ReferenceTexture->MaxTextureSize, 
+				nullptr, INDEX_NONE, false);
+
+			if (TextureSize.X > 0 && TextureSize.Y > 0)
+			{
+				TextureSize.X = TextureSize.X >> LODBias;
+				TextureSize.Y = TextureSize.Y >> LODBias;
+			}
+			else
+			{
+				TextureSize.X = TypedNodeTransform->ReferenceTexture->GetImportedSize().X >> LODBias;
+				TextureSize.Y = TypedNodeTransform->ReferenceTexture->GetImportedSize().Y >> LODBias;
+			}
+		}
+		
+		TransformNode->SetKeepAspectRatio(TypedNodeTransform->bKeepAspectRatio);
+		TransformNode->SetSizeX(TextureSize.X);
+		TransformNode->SetSizeY(TextureSize.Y);
 	}
 
 	else if (const UCustomizableObjectNodeTextureSaturate* TypedNodeSaturate = Cast<UCustomizableObjectNodeTextureSaturate>(Node))
diff --git a/Engine/Plugins/Experimental/Mutable/Source/CustomizableObjectEditor/Private/MuCOE/Nodes/CustomizableObjectNodeTextureTransform.cpp b/Engine/Plugins/Experimental/Mutable/Source/CustomizableObjectEditor/Private/MuCOE/Nodes/CustomizableObjectNodeTextureTransform.cpp
--- a/Engine/Plugins/Experimental/Mutable/Source/CustomizableObjectEditor/Private/MuCOE/Nodes/CustomizableObjectNodeTextureTransform.cpp
+++ b/Engine/Plugins/Experimental/Mutable/Source/CustomizableObjectEditor/Private/MuCOE/Nodes/CustomizableObjectNodeTextureTransform.cpp
@@ -60,7 +60,11 @@ FLinearColor UCustomizableObjectNodeTextureTransform::GetNodeTitleColor() const
 
 FText UCustomizableObjectNodeTextureTransform::GetTooltipText() const
 {
-	return LOCTEXT("Texture_Transform_Tooltip", "Applies linear transform, Offset followed by Scale followed by Rotation, to the content of Base Texture. Samples outside the image get tiled.");
+	return LOCTEXT("Texture_Transform_Tooltip", 
+			"Applies a linear transform, rotation and scale around the center of the image plus translation, "
+			"to the content of Base Texture. Rotation is in the range [0 .. 1], 1 being full rotation, offset " 
+			"and scale are in output image normalized coordinates with origin at the center of the image. " 
+			"If Keep Aspect Ratio is set, an scaling factor preserving aspect ratio will be used as identity.");
 }
 
 void UCustomizableObjectNodeTextureTransform::Serialize(FArchive& Ar)
diff --git a/Engine/Plugins/Experimental/Mutable/Source/MutableRuntime/Private/MuR/CodeRunner.cpp b/Engine/Plugins/Experimental/Mutable/Source/MutableRuntime/Private/MuR/CodeRunner.cpp
--- a/Engine/Plugins/Experimental/Mutable/Source/MutableRuntime/Private/MuR/CodeRunner.cpp
+++ b/Engine/Plugins/Experimental/Mutable/Source/MutableRuntime/Private/MuR/CodeRunner.cpp
@@ -4732,8 +4732,8 @@ namespace mu
 			{
 				const TArray<FScheduledOp, TFixedAllocator<2>> Deps = 
 				{
-					FScheduledOp(Args.scaleX, item),
-					FScheduledOp(Args.scaleY, item),
+					FScheduledOp(Args.ScaleX, item),
+					FScheduledOp(Args.ScaleY, item),
 				};
 
                 AddOp(FScheduledOp(item.At, item, 1), Deps);
@@ -4744,9 +4744,9 @@ namespace mu
 			{
             	MUTABLE_CPUPROFILER_SCOPE(IM_TRANSFORM_1)
 
-                const FVector2f Scale = FVector2f(
-                        Args.scaleX ? LoadScalar(FCacheAddress(Args.scaleX, item)) : 1.0f,
-                        Args.scaleY ? LoadScalar(FCacheAddress(Args.scaleY, item)) : 1.0f);
+                FVector2f Scale = FVector2f(
+                        Args.ScaleX ? LoadScalar(FCacheAddress(Args.ScaleX, item)) : 1.0f,
+                        Args.ScaleY ? LoadScalar(FCacheAddress(Args.ScaleY, item)) : 1.0f);
 	
 				using FUint16Vector2 = UE::Math::TIntVector2<uint16>;
 				const FUint16Vector2 DestSizeI = Invoke([&]() 
@@ -4769,15 +4769,19 @@ namespace mu
 
 				const FVector2f DestSize = FVector2f(DestSizeI.X, DestSizeI.Y);
 				const FVector2f SourceSize = FVector2f(Args.SourceSizeX, Args.SourceSizeY);
+	
+				if (Args.bKeepAspectRatio)
+				{
+					const float DestAspectOverSrcAspect = (DestSize.X * SourceSize.Y) / (DestSize.Y * SourceSize.X); 
 
-				const float DestAspectRatio = DestSize.X / DestSize.Y;
-				const float SrcAspecRatio = SourceSize.X / SourceSize.Y;
+					Scale *= DestAspectOverSrcAspect > 1.0f 
+							? FVector2f(1.0f/DestAspectOverSrcAspect, 1.0f) 
+							: FVector2f(1.0f, DestAspectOverSrcAspect); 
+				}
 			
-				const bool bKeepSourceAspectRatio = true;
-				
-				const FTransform2f Transform = FTransform2f(FVector2f(-0.5f)).
-					Concatenate(FTransform2f(FScale2f(Scale))).
-					Concatenate(FTransform2f(FVector2f(0.5f)));
+				const FTransform2f Transform = FTransform2f(FVector2f(-0.5f))
+					.Concatenate(FTransform2f(FScale2f(Scale)))
+					.Concatenate(FTransform2f(FVector2f(0.5f)));
 
 				FBox2f NormalizedCropRect(ForceInit);
 				NormalizedCropRect += Transform.TransformPoint(FVector2f(0.0f, 0.0f));
@@ -4803,10 +4807,10 @@ namespace mu
 				const uint8 Mip = static_cast<uint8>(FMath::Max(0, FMath::FloorToInt(HeapData.ImageTransform.MipValue)));
 				const TArray<FScheduledOp, TFixedAllocator<4>> Deps = 
 				{
-					FScheduledOp::FromOpAndOptions(Args.base, item, Mip),
-					FScheduledOp(Args.offsetX,  item),
-					FScheduledOp(Args.offsetY,  item),
-					FScheduledOp(Args.rotation, item) 
+					FScheduledOp::FromOpAndOptions(Args.Base, item, Mip),
+					FScheduledOp(Args.OffsetX,  item),
+					FScheduledOp(Args.OffsetY,  item),
+					FScheduledOp(Args.Rotation, item) 
 				};
 				
                 AddOp(FScheduledOp(item.At, item, 2, HeapDataAddress), Deps);
@@ -4820,18 +4824,18 @@ namespace mu
 				const FScheduledOpData HeapData = m_heapData[item.CustomState];
 
 				const uint8 Mip = static_cast<uint8>(FMath::Max(0, FMath::FloorToInt(HeapData.ImageTransform.MipValue)));
-				Ptr<const Image> Source = LoadImage(FCacheAddress(Args.base, item.ExecutionIndex, Mip));
+				Ptr<const Image> Source = LoadImage(FCacheAddress(Args.Base, item.ExecutionIndex, Mip));
 
 				const FVector2f Offset = FVector2f(
-                        Args.offsetX ? LoadScalar(FCacheAddress(Args.offsetX, item)) : 0.0f,
-                        Args.offsetY ? LoadScalar(FCacheAddress(Args.offsetY, item)) : 0.0f);
+                        Args.OffsetX ? LoadScalar(FCacheAddress(Args.OffsetX, item)) : 0.0f,
+                        Args.OffsetY ? LoadScalar(FCacheAddress(Args.OffsetY, item)) : 0.0f);
 
                 const FVector2f Scale = FVector2f(
 						FPlatformMath::LoadHalf(&HeapData.ImageTransform.ScaleXEncodedHalf),
 						FPlatformMath::LoadHalf(&HeapData.ImageTransform.ScaleYEncodedHalf));
 
 				// Map Range 0-1 to a full rotation
-                const float Rotation = LoadScalar(FCacheAddress(Args.rotation, item)) * UE_TWO_PI;
+                const float Rotation = LoadScalar(FCacheAddress(Args.Rotation, item)) * UE_TWO_PI;
 	
 				EImageFormat SourceFormat = Source->GetFormat();
 				EImageFormat Format = GetUncompressedFormat(SourceFormat);
@@ -6642,11 +6646,20 @@ namespace mu
             {
             case 0:
 			{
-				AddOp(FScheduledOp(item.At, item, 1), FScheduledOp(Args.base, item));
+				AddOp(FScheduledOp(item.At, item, 1), FScheduledOp(Args.Base, item));	
                 break;
 			}
             case 1:
             {
+				m_heapImageDesc[item.CustomState].m_lods = 1;
+				m_heapImageDesc[item.CustomState].m_format = GetUncompressedFormat(m_heapImageDesc[item.CustomState].m_format);
+				
+				if (!(Args.SizeX == 0 && Args.SizeY == 0))
+				{
+					m_heapImageDesc[item.CustomState].m_size[0] = Args.SizeX;
+					m_heapImageDesc[item.CustomState].m_size[1] = Args.SizeY;
+				}
+
 				StoreValidDesc(item);
                 break;
             }
diff --git a/Engine/Plugins/Experimental/Mutable/Source/MutableRuntime/Private/MuR/Operations.cpp b/Engine/Plugins/Experimental/Mutable/Source/MutableRuntime/Private/MuR/Operations.cpp
--- a/Engine/Plugins/Experimental/Mutable/Source/MutableRuntime/Private/MuR/Operations.cpp
+++ b/Engine/Plugins/Experimental/Mutable/Source/MutableRuntime/Private/MuR/Operations.cpp
@@ -797,12 +797,12 @@ namespace mu
 		case OP_TYPE::IM_TRANSFORM:
 		{
 			OP::ImageTransformArgs Args = program.GetOpArgs<OP::ImageTransformArgs>(at);
-			f(Args.base);
-			f(Args.offsetX);
-			f(Args.offsetY);
-			f(Args.scaleX);
-			f(Args.scaleY);
-			f(Args.rotation);
+			f(Args.Base);
+			f(Args.OffsetX);
+			f(Args.OffsetY);
+			f(Args.ScaleX);
+			f(Args.ScaleY);
+			f(Args.Rotation);
 			break;
 		}
 
diff --git a/Engine/Plugins/Experimental/Mutable/Source/MutableRuntime/Private/MuR/Operations.h b/Engine/Plugins/Experimental/Mutable/Source/MutableRuntime/Private/MuR/Operations.h
--- a/Engine/Plugins/Experimental/Mutable/Source/MutableRuntime/Private/MuR/Operations.h
+++ b/Engine/Plugins/Experimental/Mutable/Source/MutableRuntime/Private/MuR/Operations.h
@@ -831,14 +831,15 @@ namespace mu
 
 		struct ImageTransformArgs
 		{
-			ADDRESS base = 0;
-			ADDRESS offsetX = 0;
-			ADDRESS offsetY = 0;
-			ADDRESS scaleX = 0;
-			ADDRESS scaleY = 0;
-			ADDRESS rotation = 0;
-
-			uint32 AddressMode = 0;
+			ADDRESS Base = 0;
+			ADDRESS OffsetX = 0;
+			ADDRESS OffsetY = 0;
+			ADDRESS ScaleX = 0;
+			ADDRESS ScaleY = 0;
+			ADDRESS Rotation = 0;
+
+			uint32 AddressMode      : 31;
+			uint32 bKeepAspectRatio : 1;
 
             /** Size of the image to create. If 0, reuse size from base.*/
             uint16 SizeX = 0;
diff --git a/Engine/Plugins/Experimental/Mutable/Source/MutableRuntime/Public/MuR/Model.h b/Engine/Plugins/Experimental/Mutable/Source/MutableRuntime/Public/MuR/Model.h
--- a/Engine/Plugins/Experimental/Mutable/Source/MutableRuntime/Public/MuR/Model.h
+++ b/Engine/Plugins/Experimental/Mutable/Source/MutableRuntime/Public/MuR/Model.h
@@ -13,7 +13,7 @@
 //! data structures. Compiled models are not necessarily compatible when the runtime is updated,
 //! so this version number can be used externally to verify this. It is not used internally, and
 //! serializing models from different versions than this runtime will probably result in a crash.
-#define MUTABLE_COMPILED_MODEL_CODE_VERSION		uint32( 74 )
+#define MUTABLE_COMPILED_MODEL_CODE_VERSION		uint32( 75 )
 #define MUTABLE_PARAMETERS_VERSION              uint32( 2 )
 
 
diff --git a/Engine/Plugins/Experimental/Mutable/Source/MutableTools/Private/MuT/ASTOpImageTransform.cpp b/Engine/Plugins/Experimental/Mutable/Source/MutableTools/Private/MuT/ASTOpImageTransform.cpp
--- a/Engine/Plugins/Experimental/Mutable/Source/MutableTools/Private/MuT/ASTOpImageTransform.cpp
+++ b/Engine/Plugins/Experimental/Mutable/Source/MutableTools/Private/MuT/ASTOpImageTransform.cpp
@@ -18,12 +18,12 @@ namespace mu
 
 	//-------------------------------------------------------------------------------------------------
 	ASTOpImageTransform::ASTOpImageTransform()
-		: base(this)
-		, offsetX(this)
-		, offsetY(this)
-		, scaleX(this)
-		, scaleY(this)
-		, rotation(this)
+		: Base(this)
+		, OffsetX(this)
+		, OffsetY(this)
+		, ScaleX(this)
+		, ScaleY(this)
+		, Rotation(this)
 	{
 	}
 
@@ -37,22 +37,23 @@ namespace mu
 
 
 	//-------------------------------------------------------------------------------------------------
-	bool ASTOpImageTransform::IsEqual(const ASTOp& otherUntyped) const
+	bool ASTOpImageTransform::IsEqual(const ASTOp& OtherUntyped) const
 	{
-		if (auto other = dynamic_cast<const ASTOpImageTransform*>(&otherUntyped))
+		if (auto Other = dynamic_cast<const ASTOpImageTransform*>(&OtherUntyped))
 		{
 			return 
-				base == other->base &&
-				offsetX == other->offsetX &&
-				offsetY == other->offsetY &&
-				scaleX == other->scaleX &&
-				scaleY == other->scaleY &&
-				rotation == other->rotation &&
-				AddressMode == other->AddressMode &&
-				SizeX == other->SizeX &&
-				SizeY == other->SizeY &&
-				SourceSizeX == other->SourceSizeX &&
-				SourceSizeY == other->SourceSizeY;
+				Base 	 == Other->Base &&
+				OffsetX  == Other->OffsetX &&
+				OffsetY  == Other->OffsetY &&
+				ScaleX 	 == Other->ScaleX &&
+				ScaleY 	 == Other->ScaleY &&
+				Rotation == Other->Rotation &&
+				AddressMode == Other->AddressMode &&
+				SizeX == Other->SizeX &&
+				SizeY == Other->SizeY &&
+				SourceSizeX == Other->SourceSizeX &&
+				SourceSizeY == Other->SourceSizeY &&
+				bKeepAspectRatio == Other->bKeepAspectRatio;
 		}
 		return false;
 	}
@@ -61,55 +62,57 @@ namespace mu
 	//-------------------------------------------------------------------------------------------------
 	uint64 ASTOpImageTransform::Hash() const
 	{
-		uint64 res = std::hash<OP_TYPE>()(OP_TYPE::IM_TRANSFORM);
-		hash_combine(res, base.child().get());
-		hash_combine(res, offsetX.child().get());
-		hash_combine(res, offsetY.child().get());
-		hash_combine(res, scaleX.child().get());
-		hash_combine(res, scaleY.child().get());
-		hash_combine(res, rotation.child().get());
-		hash_combine(res, std::hash<uint32>()(static_cast<uint32>(AddressMode)));
-		hash_combine(res, SizeX);
-		hash_combine(res, SizeY);
-		hash_combine(res, SourceSizeX);
-		hash_combine(res, SourceSizeY);
-		return res;
+		uint64 Res = std::hash<OP_TYPE>()(OP_TYPE::IM_TRANSFORM);
+		hash_combine(Res, Base.child().get());
+		hash_combine(Res, OffsetX.child().get());
+		hash_combine(Res, OffsetY.child().get());
+		hash_combine(Res, ScaleX.child().get());
+		hash_combine(Res, ScaleY.child().get());
+		hash_combine(Res, Rotation.child().get());
+		hash_combine(Res, std::hash<uint32>()(static_cast<uint32>(AddressMode)));
+		hash_combine(Res, SizeX);
+		hash_combine(Res, SizeY);
+		hash_combine(Res, SourceSizeX);
+		hash_combine(Res, SourceSizeY);
+		hash_combine(Res, bKeepAspectRatio);
+		return Res;
 	}
 
 
 	//-------------------------------------------------------------------------------------------------
-	mu::Ptr<ASTOp> ASTOpImageTransform::Clone(MapChildFuncRef mapChild) const
+	mu::Ptr<ASTOp> ASTOpImageTransform::Clone(MapChildFuncRef MapChild) const
 	{
-		Ptr<ASTOpImageTransform> n = new ASTOpImageTransform();
-		n->base = mapChild(base.child());
-		n->offsetX = mapChild(offsetX.child());
-		n->offsetY = mapChild(offsetY.child());
-		n->scaleX = mapChild(scaleX.child());
-		n->scaleY = mapChild(scaleY.child());
-		n->rotation = mapChild(rotation.child());
-		n->AddressMode = AddressMode;
-		n->SizeX = SizeX;
-		n->SizeY = SizeY;
-		n->SourceSizeX = SourceSizeX;
-		n->SourceSizeY = SourceSizeY;
-		return n;
+		Ptr<ASTOpImageTransform> NewOp = new ASTOpImageTransform();
+		NewOp->Base     = MapChild(Base.child());
+		NewOp->OffsetX  = MapChild(OffsetX.child());
+		NewOp->OffsetY  = MapChild(OffsetY.child());
+		NewOp->ScaleX   = MapChild(ScaleX.child());
+		NewOp->ScaleY   = MapChild(ScaleY.child());
+		NewOp->Rotation = MapChild(Rotation.child());
+		NewOp->AddressMode = AddressMode;
+		NewOp->SizeX = SizeX;
+		NewOp->SizeY = SizeY;
+		NewOp->SourceSizeX = SourceSizeX;
+		NewOp->SourceSizeY = SourceSizeY;
+		NewOp->bKeepAspectRatio = bKeepAspectRatio;
+		return NewOp;
 	}
 
 
 	//-------------------------------------------------------------------------------------------------
-	void ASTOpImageTransform::ForEachChild(const TFunctionRef<void(ASTChild&)> f)
+	void ASTOpImageTransform::ForEachChild(const TFunctionRef<void(ASTChild&)> Func)
 	{
-		f(base);
-		f(offsetX);
-		f(offsetY);
-		f(scaleX);
-		f(scaleY);
-		f(rotation);
+		Func(Base);
+		Func(OffsetX);
+		Func(OffsetY);
+		Func(ScaleX);
+		Func(ScaleY);
+		Func(Rotation);
 	}
 
 
 	//-------------------------------------------------------------------------------------------------
-	void ASTOpImageTransform::Link(FProgram& program, FLinkerOptions*)
+	void ASTOpImageTransform::Link(FProgram& Program, FLinkerOptions*)
 	{
 		// Already linked?
 		if (!linkedAddress)
@@ -117,69 +120,79 @@ namespace mu
 			OP::ImageTransformArgs Args;
 			FMemory::Memzero(Args);
 
-			Args.base = base ? base->linkedAddress : 0;
-			Args.offsetX = offsetX ? offsetX->linkedAddress : 0;
-			Args.offsetY = offsetY ? offsetY->linkedAddress : 0;
-			Args.scaleX = scaleX ? scaleX->linkedAddress : 0;
-			Args.scaleY = scaleY ? scaleY->linkedAddress : 0;
-			Args.rotation = rotation ? rotation->linkedAddress : 0;
+			Args.Base     = Base     ? Base->linkedAddress     : 0;
+			Args.OffsetX  = OffsetX  ? OffsetX->linkedAddress  : 0;
+			Args.OffsetY  = OffsetY  ? OffsetY->linkedAddress  : 0;
+			Args.ScaleX   = ScaleX   ? ScaleX->linkedAddress   : 0;
+			Args.ScaleY   = ScaleY   ? ScaleY->linkedAddress   : 0;
+			Args.Rotation = Rotation ? Rotation->linkedAddress : 0;
 			Args.AddressMode = static_cast<uint32>(AddressMode);
+			Args.bKeepAspectRatio = bKeepAspectRatio;
 			Args.SizeX = SizeX;
 			Args.SizeY = SizeY;
 			Args.SourceSizeX = SourceSizeX;
 			Args.SourceSizeY = SourceSizeY;
 
-			linkedAddress = (OP::ADDRESS)program.m_opAddress.Num();
-			program.m_opAddress.Add((uint32_t)program.m_byteCode.Num());
-			AppendCode(program.m_byteCode, OP_TYPE::IM_TRANSFORM);
-			AppendCode(program.m_byteCode, Args);
+			linkedAddress = (OP::ADDRESS)Program.m_opAddress.Num();
+			Program.m_opAddress.Add((uint32_t)Program.m_byteCode.Num());
+			AppendCode(Program.m_byteCode, OP_TYPE::IM_TRANSFORM);
+			AppendCode(Program.m_byteCode, Args);
 		}
 	}
 
 	//-------------------------------------------------------------------------------------------------
-	FImageDesc ASTOpImageTransform::GetImageDesc(bool returnBestOption, FGetImageDescContext* context) const
+	FImageDesc ASTOpImageTransform::GetImageDesc(bool bReturnBestOption, FGetImageDescContext* Context) const
 	{
-		FImageDesc res;
+		FImageDesc Result;
 
 		// Local context in case it is necessary
-		FGetImageDescContext localContext;
-		if (!context)
+		FGetImageDescContext LocalContext;
+		if (!Context)
 		{
-			context = &localContext;
+			Context = &LocalContext;
 		}
 		else
 		{
 			// Cached result?
-			FImageDesc* PtrValue = context->m_results.Find(this);
+			FImageDesc* PtrValue = Context->m_results.Find(this);
 			if (PtrValue)
 			{
 				return *PtrValue;
 			}
 		}
 
+
 		// Actual work
-		if (base)
+		if (Base)
 		{
-			res = base->GetImageDesc(returnBestOption, context);
+			Result = Base->GetImageDesc(bReturnBestOption, Context);
+			
+			Result.m_format = GetUncompressedFormat(Result.m_format); 
+			Result.m_lods = 1;
+			
+			if (!(SizeX == 0 && SizeY == 0))
+			{
+				Result.m_size = FImageSize(SizeX, SizeY);
+			}
 		}
 
 
 		// Cache the result
-		if (context)
+		if (Context)
 		{
-			context->m_results.Add(this, res);
+			Context->m_results.Add(this, Result);
 		}
 
-		return res;
+		return Result;
 	}
 
 
 	//-------------------------------------------------------------------------------------------------
 	void ASTOpImageTransform::GetLayoutBlockSize(int* pBlockX, int* pBlockY)
 	{
-		if (base)
+		if (Base)
 		{
-			base->GetLayoutBlockSize(pBlockX, pBlockY);
+			Base->GetLayoutBlockSize(pBlockX, pBlockY);
 		}
 	}
 
@@ -187,9 +200,21 @@ namespace mu
 	//-------------------------------------------------------------------------------------------------
 	mu::Ptr<ImageSizeExpression> ASTOpImageTransform::GetImageSizeExpression() const
 	{
-		if (base)
+		if (Base)
 		{
-			return base->GetImageSizeExpression();
+			if (SizeX == 0 && SizeY == 0)
+			{
+				return Base->GetImageSizeExpression();
+			}
+			else
+			{
+				Ptr<ImageSizeExpression> SizeExpr = new ImageSizeExpression;
+				SizeExpr->type = ImageSizeExpression::ISET_CONSTANT;
+				SizeExpr->size[0] = SizeX;
+				SizeExpr->size[1] = SizeY;
+
+				return SizeExpr;
+			}
 		}
 
 		return nullptr;
diff --git a/Engine/Plugins/Experimental/Mutable/Source/MutableTools/Private/MuT/ASTOpImageTransform.h b/Engine/Plugins/Experimental/Mutable/Source/MutableTools/Private/MuT/ASTOpImageTransform.h
--- a/Engine/Plugins/Experimental/Mutable/Source/MutableTools/Private/MuT/ASTOpImageTransform.h
+++ b/Engine/Plugins/Experimental/Mutable/Source/MutableTools/Private/MuT/ASTOpImageTransform.h
@@ -21,12 +21,12 @@ struct FProgram;
 	{
 	public:
 
-		ASTChild base;
-		ASTChild offsetX;
-		ASTChild offsetY;
-		ASTChild scaleX;
-		ASTChild scaleY;
-		ASTChild rotation;
+		ASTChild Base;
+		ASTChild OffsetX;
+		ASTChild OffsetY;
+		ASTChild ScaleX;
+		ASTChild ScaleY;
+		ASTChild Rotation;
 
 		uint16 SizeX;
 		uint16 SizeY;
@@ -35,6 +35,7 @@ struct FProgram;
 		uint16 SourceSizeY;
 
 		EAddressMode AddressMode = EAddressMode::Wrap;
+		bool bKeepAspectRatio = false;
 
 	public:
 
@@ -44,11 +45,11 @@ struct FProgram;
 
 		OP_TYPE GetOpType() const override { return OP_TYPE::IM_TRANSFORM; }
 		uint64 Hash() const override;
-		bool IsEqual(const ASTOp& otherUntyped) const override;
-		Ptr<ASTOp> Clone(MapChildFuncRef mapChild) const override;
+		bool IsEqual(const ASTOp& OtherUntyped) const override;
+		Ptr<ASTOp> Clone(MapChildFuncRef MapChild) const override;
 		void ForEachChild(const TFunctionRef<void(ASTChild&)>) override;
-		void Link(FProgram& program, FLinkerOptions* Options) override;
-		FImageDesc GetImageDesc(bool returnBestOption, FGetImageDescContext* context) const override;
+		void Link(FProgram& Program, FLinkerOptions* Options) override;
+		FImageDesc GetImageDesc(bool bReturnBestOption, FGetImageDescContext* Context) const override;
 		void GetLayoutBlockSize(int* pBlockX, int* pBlockY) override;
 		Ptr<ImageSizeExpression> GetImageSizeExpression() const override;
 	};
diff --git a/Engine/Plugins/Experimental/Mutable/Source/MutableTools/Private/MuT/CodeGenerator_Image.cpp b/Engine/Plugins/Experimental/Mutable/Source/MutableTools/Private/MuT/CodeGenerator_Image.cpp
--- a/Engine/Plugins/Experimental/Mutable/Source/MutableTools/Private/MuT/CodeGenerator_Image.cpp
+++ b/Engine/Plugins/Experimental/Mutable/Source/MutableTools/Private/MuT/CodeGenerator_Image.cpp
@@ -558,79 +558,80 @@ namespace mu
 
 	void CodeGenerator::GenerateImage_Transform(const FImageGenerationOptions& Options, FImageGenerationResult& Result, const NodeImageTransform* InNode)
     {
-		const NodeImageTransform::Private& node = *InNode->GetPrivate();
+		const NodeImageTransform::Private& Node = *InNode->GetPrivate();
 
         MUTABLE_CPUPROFILER_SCOPE(NodeImageTransform);
 
-        Ptr<ASTOpImageTransform> op = new ASTOpImageTransform();
+        Ptr<ASTOpImageTransform> Op = new ASTOpImageTransform();
 
 		Ptr<ASTOp> OffsetX;
-		if (node.m_pOffsetX)
+		if (Node.m_pOffsetX)
 		{
-			OffsetX = Generate(node.m_pOffsetX);
+			OffsetX = Generate(Node.m_pOffsetX);
 		}
 
 		Ptr<ASTOp> OffsetY;
-		if (node.m_pOffsetY)
+		if (Node.m_pOffsetY)
 		{
-			OffsetY = Generate(node.m_pOffsetY);
+			OffsetY = Generate(Node.m_pOffsetY);
 		}
 	
 		Ptr<ASTOp> ScaleX;
-		if (node.m_pScaleX)
+		if (Node.m_pScaleX)
 		{
-			ScaleX = Generate(node.m_pScaleX);
+			ScaleX = Generate(Node.m_pScaleX);
 		}
 	
 		Ptr<ASTOp> ScaleY;
-		if (node.m_pScaleY)
+		if (Node.m_pScaleY)
 		{
-			ScaleY = Generate(node.m_pScaleY);
+			ScaleY = Generate(Node.m_pScaleY);
 		}
 
 		Ptr<ASTOp> Rotation;
-		if (node.m_pRotation)
+		if (Node.m_pRotation)
 		{
-			Rotation = Generate(node.m_pRotation);
+			Rotation = Generate(Node.m_pRotation);
 		}
 
 		// If one of the inputs (offset or scale) is missig assume unifrom translation/scaling 
-		op->offsetX = OffsetX ? OffsetX : OffsetY;
-		op->offsetY = OffsetY ? OffsetY : OffsetX;
- 		op->scaleX = ScaleX ? ScaleX : ScaleY;
-		op->scaleY = ScaleY ? ScaleY : ScaleX;
-		op->rotation = Rotation; 
-		op->AddressMode = node.AddressMode;
-		op->SizeX = node.SizeX;
-		op->SizeY = node.SizeY;
+		Op->OffsetX = OffsetX ? OffsetX : OffsetY;
+		Op->OffsetY = OffsetY ? OffsetY : OffsetX;
+ 		Op->ScaleX = ScaleX ? ScaleX : ScaleY;
+		Op->ScaleY = ScaleY ? ScaleY : ScaleX;
+		Op->Rotation = Rotation; 
+		Op->AddressMode = Node.AddressMode;
+		Op->SizeX = Node.SizeX;
+		Op->SizeY = Node.SizeY;
+		Op->bKeepAspectRatio = Node.bKeepAspectRatio;
 
 		// Base image
-        Ptr<ASTOp> base;
+        Ptr<ASTOp> Base;
 		FImageGenerationOptions NewOptions = Options;
 		NewOptions.ImageLayoutStrategy = CompilerOptions::TextureLayoutStrategy::None;
 		NewOptions.LayoutToApply = nullptr;
 		NewOptions.LayoutBlockId = -1;
 		NewOptions.RectSize = {};
 
-        if ( node.m_pBase )
+        if (Node.m_pBase)
         {
 			FImageGenerationResult BaseResult;
-			GenerateImage(NewOptions, BaseResult, node.m_pBase);
-			base = BaseResult.op;
+			GenerateImage(NewOptions, BaseResult, Node.m_pBase);
+			Base = BaseResult.op;
 		}
         else
         {
             // This argument is required
-            base = GenerateMissingImageCode(TEXT("Image Transform Base"), EImageFormat::IF_RGB_UBYTE, node.m_errorContext, NewOptions);
+            Base = GenerateMissingImageCode(TEXT("Image Transform Base"), EImageFormat::IF_RGB_UBYTE, Node.m_errorContext, NewOptions);
         }
 		
-		FImageDesc BaseDesc = base->GetImageDesc();
+		FImageDesc BaseDesc = Base->GetImageDesc();
 		
-        op->base = base;
-		op->SourceSizeX = BaseDesc.m_size.X;
-		op->SourceSizeY = BaseDesc.m_size.Y;
+        Op->Base = Base;
+		Op->SourceSizeX = BaseDesc.m_size.X;
+		Op->SourceSizeY = BaseDesc.m_size.Y;
 
-        Result.op = op; 
+        Result.op = Op; 
     }
 
     //---------------------------------------------------------------------------------------------
diff --git a/Engine/Plugins/Experimental/Mutable/Source/MutableTools/Private/MuT/CodeOptimiser_Semantic.cpp b/Engine/Plugins/Experimental/Mutable/Source/MutableTools/Private/MuT/CodeOptimiser_Semantic.cpp
--- a/Engine/Plugins/Experimental/Mutable/Source/MutableTools/Private/MuT/CodeOptimiser_Semantic.cpp
+++ b/Engine/Plugins/Experimental/Mutable/Source/MutableTools/Private/MuT/CodeOptimiser_Semantic.cpp
@@ -1286,13 +1286,13 @@ namespace mu
             case OP_TYPE::IM_TRANSFORM:
 			{
 				// It can only sink in the transform if it doesn't have it's own size.
-				const ASTOpImageTransform* typedAt = dynamic_cast<const ASTOpImageTransform*>(at.get());
-				if (typedAt->SizeX == 0 && typedAt->SizeY == 0)
+				const ASTOpImageTransform* TypedAt = dynamic_cast<const ASTOpImageTransform*>(at.get());
+				if (TypedAt->SizeX == 0 && TypedAt->SizeY == 0)
 				{
-					Ptr<ASTOpImageTransform> nop = mu::Clone<ASTOpImageTransform>(at);
-					Ptr<ASTOp> maskOp = nop->base.child();
-					nop->base = Visit(maskOp, currentSinkingOp);
-					newAt = nop;
+					Ptr<ASTOpImageTransform> NewOp = mu::Clone<ASTOpImageTransform>(at);
+					Ptr<ASTOp> MaskOp = NewOp->Base.child();
+					NewOp->Base = Visit(MaskOp, currentSinkingOp);
+					newAt = NewOp;
 				}
 
 				break;
diff --git a/Engine/Plugins/Experimental/Mutable/Source/MutableTools/Private/MuT/ImageDescGenerator.cpp b/Engine/Plugins/Experimental/Mutable/Source/MutableTools/Private/MuT/ImageDescGenerator.cpp
--- a/Engine/Plugins/Experimental/Mutable/Source/MutableTools/Private/MuT/ImageDescGenerator.cpp
+++ b/Engine/Plugins/Experimental/Mutable/Source/MutableTools/Private/MuT/ImageDescGenerator.cpp
@@ -523,13 +523,19 @@ namespace mu
 	}
 
     //---------------------------------------------------------------------------------------------
-    Ptr<ASTOp> ImageDescGenerator::Visit( const NodeImageTransform::Private& node )
+    Ptr<ASTOp> ImageDescGenerator::Visit(const NodeImageTransform::Private& node)
     {
         m_desc = MUTABLE_MISSING_IMAGE_DESC;
 
-        if ( NodeImage* pBase = node.m_pBase.get() )
+        if (NodeImage* pBase = node.m_pBase.get())
         {
-            Generate( *pBase->GetBasePrivate() );
+            Generate(*pBase->GetBasePrivate());
+            
+            if (!(node.SizeX == 0 && node.SizeY == 0))
+            {
+                m_desc.m_size[0] = static_cast<uint16>(node.SizeX);  
+                m_desc.m_size[1] = static_cast<uint16>(node.SizeY);  
+            }
         }
 
         return 0;
diff --git a/Engine/Plugins/Experimental/Mutable/Source/MutableTools/Private/MuT/NodeImageTransform.cpp b/Engine/Plugins/Experimental/Mutable/Source/MutableTools/Private/MuT/NodeImageTransform.cpp
--- a/Engine/Plugins/Experimental/Mutable/Source/MutableTools/Private/MuT/NodeImageTransform.cpp
+++ b/Engine/Plugins/Experimental/Mutable/Source/MutableTools/Private/MuT/NodeImageTransform.cpp
@@ -179,5 +179,17 @@ namespace mu
 	{
 		m_pD->SizeY = SizeY;
 	}
+
+	//---------------------------------------------------------------------------------------------
+	bool NodeImageTransform::GetKeepAspectRatio() const
+	{
+		return m_pD->bKeepAspectRatio;;
+	}
+
+	//---------------------------------------------------------------------------------------------
+	void NodeImageTransform::SetKeepAspectRatio(bool bKeepAspectRatio)
+	{
+		m_pD->bKeepAspectRatio = bKeepAspectRatio;
+	}
 }
 
diff --git a/Engine/Plugins/Experimental/Mutable/Source/MutableTools/Private/MuT/NodeImageTransformPrivate.h b/Engine/Plugins/Experimental/Mutable/Source/MutableTools/Private/MuT/NodeImageTransformPrivate.h
--- a/Engine/Plugins/Experimental/Mutable/Source/MutableTools/Private/MuT/NodeImageTransformPrivate.h
+++ b/Engine/Plugins/Experimental/Mutable/Source/MutableTools/Private/MuT/NodeImageTransformPrivate.h
@@ -33,12 +33,16 @@ namespace mu
 		EAddressMode AddressMode = EAddressMode::Wrap;
 		uint32 SizeX = 0;
 		uint32 SizeY = 0;
+		
+		bool bKeepAspectRatio = false;
+		uint8 UnusedPadding[sizeof(NodeImagePtr) - sizeof(bool)] = {0}; 
+		static_assert(sizeof(NodeImagePtr) - sizeof(bool) >= 1);
 
 
 		//!
-		void Serialise( OutputArchive& arch ) const
+		void Serialise(OutputArchive& arch) const
 		{
-            uint32 ver = 2;
+            uint32 ver = 3;
 			arch << ver;
 
 			arch << m_pBase;
@@ -50,14 +54,15 @@ namespace mu
 			arch << static_cast<uint32>(AddressMode);
 			arch << SizeX;
 			arch << SizeY;
+			arch << bKeepAspectRatio;
 		}
 
 		//!
-		void Unserialise( InputArchive& arch )
+		void Unserialise(InputArchive& arch)
 		{
             uint32 ver;
 			arch >> ver;
-			check(ver <= 2);
+			check(ver <= 3);
 
 			arch >> m_pBase;
 			arch >> m_pOffsetX;
@@ -79,6 +84,11 @@ namespace mu
 				arch >> SizeX;
 				arch >> SizeY;
 			}
+
+			if (ver >= 3)
+			{
+				arch >> bKeepAspectRatio;
+			}
 		}
 	};
 
diff --git a/Engine/Plugins/Experimental/Mutable/Source/MutableTools/Public/MuT/NodeImageTransform.h b/Engine/Plugins/Experimental/Mutable/Source/MutableTools/Public/MuT/NodeImageTransform.h
--- a/Engine/Plugins/Experimental/Mutable/Source/MutableTools/Public/MuT/NodeImageTransform.h
+++ b/Engine/Plugins/Experimental/Mutable/Source/MutableTools/Public/MuT/NodeImageTransform.h
@@ -82,6 +82,9 @@ namespace mu
 		uint16 GetSizeY() const;
 		void SetSizeY(uint16 SizeY); 
 
+		bool GetKeepAspectRatio() const;
+		void SetKeepAspectRatio(bool bKeepAspectRatio);
+
 		//-----------------------------------------------------------------------------------------
 		// Interface pattern
 		//-----------------------------------------------------------------------------------------
