[Mutable] ImageSwizzle vs ImageTransform sink optimizations.

#rnx

diff --git a/Engine/Plugins/Experimental/Mutable/Source/MutableTools/Private/MuT/ASTOpImageSwizzle.cpp b/Engine/Plugins/Experimental/Mutable/Source/MutableTools/Private/MuT/ASTOpImageSwizzle.cpp
--- a/Engine/Plugins/Experimental/Mutable/Source/MutableTools/Private/MuT/ASTOpImageSwizzle.cpp
+++ b/Engine/Plugins/Experimental/Mutable/Source/MutableTools/Private/MuT/ASTOpImageSwizzle.cpp
@@ -19,6 +19,7 @@
 #include "MuT/ASTOpImageLayer.h"
 #include "MuT/ASTOpImageLayerColor.h"
 #include "MuT/ASTOpImageRasterMesh.h"
+#include "MuT/ASTOpImageTransform.h"
 #include "MuT/ASTOpSwitch.h"
 #include "MuT/StreamsPrivate.h"
 
@@ -245,7 +246,7 @@ namespace mu
 		{
 			return at;
 		}
-
+		
 		// If all channels are the same, and in the same order, and the source format is the same that we are
 		// setting in the swizzle, then the swizzle won't do anything.
 		if (bAllChannelsAreTheSame && bSameChannelOrder)
@@ -381,6 +382,16 @@ namespace mu
 				break;
 			}
 
+			case OP_TYPE::IM_TRANSFORM:
+			{
+				Ptr<ASTOpImageTransform> NewTransform = mu::Clone<ASTOpImageTransform>(channelSourceAt);
+				Ptr<ASTOpImageSwizzle> NewSwizzle = mu::Clone<ASTOpImageSwizzle>(this);
+				ReplaceAllSources(NewSwizzle, NewTransform->Base.child());
+				NewTransform->Base = NewSwizzle;
+				at = NewTransform;
+				break;
+			}
+
 			//case OP_TYPE::IM_RESIZE:
 			//{
 			//	Ptr<ASTOpFixed> NewResize = mu::Clone<ASTOpFixed>(channelSourceAt);
@@ -704,6 +715,57 @@ namespace mu
 
 			}
 
+			// Swizzle down compatible image transforms.
+			if (!at && sourceType == OP_TYPE::IM_TRANSFORM)
+			{
+				const ASTOpImageTransform* FirstTransform = dynamic_cast<const ASTOpImageTransform*>(Sources[0].child().get());
+				check(FirstTransform);
+
+				bool bAreAllTransformsCompatible = true;
+				for (int32 c = 1; c < MUTABLE_OP_MAX_SWIZZLE_CHANNELS; ++c)
+				{
+					if (Sources[c])
+					{
+						const ASTOpImageTransform* Typed = dynamic_cast<const ASTOpImageTransform*>(Sources[c].child().get());
+						check(Typed);
+
+						// Compare all args but the source image
+						if (Typed->ScaleX.child() != FirstTransform->ScaleX.child()
+							|| Typed->ScaleY.child() != FirstTransform->ScaleY.child()
+							|| Typed->OffsetX.child() != FirstTransform->OffsetX.child()
+							|| Typed->OffsetY.child() != FirstTransform->OffsetY.child()
+							|| Typed->Rotation.child() != FirstTransform->Rotation.child()
+							|| Typed->SourceSizeX != FirstTransform->SourceSizeX
+							|| Typed->SourceSizeY != FirstTransform->SourceSizeY
+							)
+						{
+							bAreAllTransformsCompatible = false;
+							break;
+						}
+					}
+				}
+
+				if (bAreAllTransformsCompatible)
+				{
+					// Move the swizzle down all the paths
+					Ptr<ASTOpImageTransform> NewTransform = mu::Clone<ASTOpImageTransform>(FirstTransform);
+
+					Ptr<ASTOpImageSwizzle> NewSwizzle = mu::Clone<ASTOpImageSwizzle>(this);
+					for (int c = 0; c < MUTABLE_OP_MAX_SWIZZLE_CHANNELS; ++c)
+					{
+						const ASTOpImageTransform* ChannelTransform = dynamic_cast<const ASTOpImageTransform*>(Sources[c].child().get());
+						if (ChannelTransform)
+						{
+							NewSwizzle->Sources[c] = ChannelTransform->Base.child();
+						}
+					}
+
+					NewTransform->Base = NewSwizzle;
+
+					at = NewTransform;
+				}
+			}
+
 			// Swizzle down compatible resizes.
 			//if (!at && sourceType == OP_TYPE::IM_RESIZE)
 			//{
