#rnx

diff --git a/Engine/Plugins/Experimental/Mutable/Source/CustomizableObject/Classes/MuCO/CustomizableObjectSystem.h b/Engine/Plugins/Experimental/Mutable/Source/CustomizableObject/Classes/MuCO/CustomizableObjectSystem.h
--- a/Engine/Plugins/Experimental/Mutable/Source/CustomizableObject/Classes/MuCO/CustomizableObjectSystem.h
+++ b/Engine/Plugins/Experimental/Mutable/Source/CustomizableObject/Classes/MuCO/CustomizableObjectSystem.h
@@ -38,29 +38,6 @@ struct FGuid;
 // This is used to hide Mutable SDK members in the public headers.
 class FCustomizableObjectSystemPrivate;
 
-
-// Mutable stats
-DECLARE_STATS_GROUP(TEXT("Mutable"), STATGROUP_Mutable, STATCAT_Advanced);
-DECLARE_DWORD_ACCUMULATOR_STAT_EXTERN(TEXT("Num Created Mutable Skeletal Meshes"), STAT_MutableNumSkeletalMeshes, STATGROUP_Mutable, );
-DECLARE_DWORD_ACCUMULATOR_STAT_EXTERN(TEXT("Num Cached Mutable Skeletal Meshes"), STAT_MutableNumCachedSkeletalMeshes, STATGROUP_Mutable, );
-DECLARE_DWORD_ACCUMULATOR_STAT_EXTERN(TEXT("Num Allocated Mutable Skeletal Meshes"), STAT_MutableNumAllocatedSkeletalMeshes, STATGROUP_Mutable, );
-DECLARE_DWORD_ACCUMULATOR_STAT_EXTERN(TEXT("Num Instances at LOD 0"), STAT_MutableNumInstancesLOD0, STATGROUP_Mutable, );
-DECLARE_DWORD_ACCUMULATOR_STAT_EXTERN(TEXT("Num Instances at LOD 1"), STAT_MutableNumInstancesLOD1, STATGROUP_Mutable, );
-DECLARE_DWORD_ACCUMULATOR_STAT_EXTERN(TEXT("Num Instances at LOD 2 or more"), STAT_MutableNumInstancesLOD2, STATGROUP_Mutable, );
-DECLARE_DWORD_ACCUMULATOR_STAT_EXTERN(TEXT("Skeletal Mesh Resource Memory"), STAT_MutableSkeletalMeshResourceMemory, STATGROUP_Mutable, );
-DECLARE_DWORD_ACCUMULATOR_STAT_EXTERN(TEXT("Num Created Mutable Textures"), STAT_MutableNumTextures, STATGROUP_Mutable, );
-DECLARE_DWORD_ACCUMULATOR_STAT_EXTERN(TEXT("Num Cached Mutable Textures"), STAT_MutableNumCachedTextures, STATGROUP_Mutable, );
-DECLARE_DWORD_ACCUMULATOR_STAT_EXTERN(TEXT("Num Allocated Mutable Textures"), STAT_MutableNumAllocatedTextures, STATGROUP_Mutable, );
-DECLARE_DWORD_ACCUMULATOR_STAT_EXTERN(TEXT("Texture Resource Memory"), STAT_MutableTextureResourceMemory, STATGROUP_Mutable, );
-DECLARE_DWORD_ACCUMULATOR_STAT_EXTERN(TEXT("Texture Generated Memory"), STAT_MutableTextureGeneratedMemory, STATGROUP_Mutable, );
-DECLARE_DWORD_ACCUMULATOR_STAT_EXTERN(TEXT("Texture Locked Memory"), STAT_MutableTextureCacheMemory, STATGROUP_Mutable, );
-DECLARE_DWORD_ACCUMULATOR_STAT_EXTERN(TEXT("Pending Instance Updates"), STAT_MutablePendingInstanceUpdates, STATGROUP_Mutable, );
-DECLARE_DWORD_ACCUMULATOR_STAT_EXTERN(TEXT("Abandoned Instance Updates"), STAT_MutableAbandonedInstanceUpdates, STATGROUP_Mutable, );
-DECLARE_DWORD_ACCUMULATOR_STAT_EXTERN(TEXT("Last Instance Build Time"), STAT_MutableInstanceBuildTime, STATGROUP_Mutable, );
-DECLARE_DWORD_ACCUMULATOR_STAT_EXTERN(TEXT("Avrg Instance Build Time"), STAT_MutableInstanceBuildTimeAvrg, STATGROUP_Mutable, );
-DECLARE_DWORD_ACCUMULATOR_STAT_EXTERN(TEXT("Streaming Ops"), STAT_MutableStreamingOps, STATGROUP_Mutable, );
-
-
 extern TAutoConsoleVariable<bool> CVarClearWorkingMemoryOnUpdateEnd;
 
 extern TAutoConsoleVariable<bool> CVarReuseImagesBetweenInstances;
diff --git a/Engine/Plugins/Experimental/Mutable/Source/CustomizableObject/Private/MuCO/CustomizableObjectInstance.cpp b/Engine/Plugins/Experimental/Mutable/Source/CustomizableObject/Private/MuCO/CustomizableObjectInstance.cpp
--- a/Engine/Plugins/Experimental/Mutable/Source/CustomizableObject/Private/MuCO/CustomizableObjectInstance.cpp
+++ b/Engine/Plugins/Experimental/Mutable/Source/CustomizableObject/Private/MuCO/CustomizableObjectInstance.cpp
@@ -17,6 +17,7 @@
 #include "Materials/Material.h"
 #include "Materials/MaterialInstanceDynamic.h"
 #include "Modules/ModuleManager.h"
+
 #include "MuCO/CustomizableInstanceLODManagement.h"
 #include "MuCO/CustomizableInstancePrivateData.h"
 #include "MuCO/CustomizableObjectExtension.h"
@@ -27,7 +28,9 @@
 #include "MuCO/ICustomizableObjectModule.h"
 #include "MuCO/UnrealConversionUtils.h"
 #include "MuCO/UnrealPortabilityHelpers.h"
+#include "MuCO/LogBenchmarkUtil.h"
 #include "MuR/Model.h"
+
 #include "PhysicsEngine/PhysicsAsset.h"
 #include "Rendering/Texture2DResource.h"
 #include "RenderingThread.h"
diff --git a/Engine/Plugins/Experimental/Mutable/Source/CustomizableObject/Private/MuCO/CustomizableObjectSystem.cpp b/Engine/Plugins/Experimental/Mutable/Source/CustomizableObject/Private/MuCO/CustomizableObjectSystem.cpp
--- a/Engine/Plugins/Experimental/Mutable/Source/CustomizableObject/Private/MuCO/CustomizableObjectSystem.cpp
+++ b/Engine/Plugins/Experimental/Mutable/Source/CustomizableObject/Private/MuCO/CustomizableObjectSystem.cpp
@@ -7,11 +7,9 @@
 #include "AssetRegistry/AssetRegistryModule.h"
 #include "Engine/SkeletalMesh.h"
 #include "Engine/SkinnedAsset.h"
-#include "GameFramework/Pawn.h"
 #include "Engine/SkinnedAssetCommon.h"
 #include "GameFramework/PlayerController.h"
 #include "Interfaces/ITargetPlatform.h"
-#include "Kismet/GameplayStatics.h"
 #include "MuCO/CustomizableInstanceLODManagement.h"
 #include "MuCO/CustomizableInstancePrivateData.h"
 #include "MuCO/CustomizableObjectPrivate.h"
@@ -27,7 +25,6 @@
 #include "MuCO/UnrealPortabilityHelpers.h"
 #include "MuR/Model.h"
 #include "MuR/Settings.h"
-#include "TextureResource.h"
 #include "UObject/UObjectIterator.h"
 #include "Widgets/Notifications/SNotificationList.h"
 #include "ContentStreaming.h"
@@ -53,24 +50,6 @@ namespace impl
 	void Task_Game_Callbacks_Work(UCustomizableObjectInstance* CustomizableObjectInstance);
 }
 
-DEFINE_STAT(STAT_MutableNumSkeletalMeshes);
-DEFINE_STAT(STAT_MutableNumCachedSkeletalMeshes);
-DEFINE_STAT(STAT_MutableNumAllocatedSkeletalMeshes);
-DEFINE_STAT(STAT_MutableNumInstancesLOD0);
-DEFINE_STAT(STAT_MutableNumInstancesLOD1);
-DEFINE_STAT(STAT_MutableNumInstancesLOD2);
-DEFINE_STAT(STAT_MutableSkeletalMeshResourceMemory);
-DEFINE_STAT(STAT_MutableNumTextures);
-DEFINE_STAT(STAT_MutableNumCachedTextures);
-DEFINE_STAT(STAT_MutableNumAllocatedTextures);
-DEFINE_STAT(STAT_MutableTextureResourceMemory);
-DEFINE_STAT(STAT_MutableTextureGeneratedMemory);
-DEFINE_STAT(STAT_MutableTextureCacheMemory);
-DEFINE_STAT(STAT_MutablePendingInstanceUpdates);
-DEFINE_STAT(STAT_MutableAbandonedInstanceUpdates);
-DEFINE_STAT(STAT_MutableInstanceBuildTime);
-DEFINE_STAT(STAT_MutableInstanceBuildTimeAvrg);
-DEFINE_STAT(STAT_MutableStreamingOps);
 
 DECLARE_CYCLE_STAT(TEXT("MutablePendingRelease Time"), STAT_MutablePendingRelease, STATGROUP_Game);
 DECLARE_CYCLE_STAT(TEXT("MutableTask"), STAT_MutableTask, STATGROUP_Game);
@@ -353,10 +332,10 @@ void UCustomizableObjectSystem::InitSystem()
 	Private->TickDelegateHandle = FTSTicker::GetCoreTicker().AddTicker(Private->TickDelegate, 0.f);
 #endif // !UE_SERVER
 
-	Private->TotalBuildMs = 0;
-	Private->TotalBuiltInstances = 0;
-	Private->NumInstances = 0;
-	Private->TextureMemoryUsed = 0;
+	Private->MutableStats.TotalBuildMs = 0;
+	Private->MutableStats.TotalBuiltInstances = 0;
+	Private->MutableStats.NumInstances = 0;
+	Private->MutableStats.TextureMemoryUsed = 0;
 
 #if !(UE_BUILD_SHIPPING || UE_BUILD_TEST)
 	SET_DWORD_STAT(STAT_MutableNumSkeletalMeshes, 0);
@@ -480,217 +459,36 @@ void UCustomizableObjectSystem::SetNewCompilerFunc(FCustomizableObjectCompilerBa
 
 
 void FCustomizableObjectSystemPrivate::CreatedTexture(UTexture2D* Texture)
-{
-	// TODO: Do not keep this array unless we are gathering stats!
-	TextureTrackerArray.Add(Texture);
-
-	INC_DWORD_STAT(STAT_MutableNumTextures);
-}
-
-
-int32 FCustomizableObjectSystemPrivate::EnableMutableAnimInfoDebugging = 0;
-
-static FAutoConsoleVariableRef CVarEnableMutableAnimInfoDebugging(
-	TEXT("mutable.EnableMutableAnimInfoDebugging"), FCustomizableObjectSystemPrivate::EnableMutableAnimInfoDebugging,
-	TEXT("If set to 1 or greater print on screen the animation info of the pawn's Customizable Object Instance. Anim BPs, slots and tags will be displayed."
-	"If the root Customizable Object is recompiled after this command is run, the used skeletal meshes will also be displayed."),
-	ECVF_Default);
-
-
-void FCustomizableObjectSystemPrivate::UpdateStats()
 {
 #if !(UE_BUILD_SHIPPING || UE_BUILD_TEST)
 	bool bLogEnabled = true;
 #else
-	bool bLogEnabled = LogBenchmarkUtil::isLoggingActive();
+	bool bLogEnabled = LogBenchmarkUtil::IsLoggingActive();
 #endif
 	if (bLogEnabled)
 	{
-		CountAllocatedSkeletalMesh = 0;
-
-		int CountLOD0 = 0;
-		int CountLOD1 = 0;
-		int CountLOD2 = 0;
-		int CountTotal = 0;
-
-		for (TObjectIterator<UCustomizableObjectInstance> CustomizableObjectInstance; CustomizableObjectInstance; ++CustomizableObjectInstance)
-		{
-			if ( IsValidChecked(*CustomizableObjectInstance) && CustomizableObjectInstance->GetPrivate())
-			{
-				++CountTotal;
-
-				for (int32 ComponentIndex = 0; ComponentIndex < CustomizableObjectInstance->SkeletalMeshes.Num(); ++ComponentIndex)
-				{
-					if (CustomizableObjectInstance->SkeletalMeshes[ComponentIndex] && CustomizableObjectInstance->SkeletalMeshes[ComponentIndex]->GetResourceForRendering())
-					{
-						CountAllocatedSkeletalMesh++;
-
-						if (CustomizableObjectInstance->GetCurrentMinLOD() < 1)
-						{
-							++CountLOD0;
-						}
-						else if (CustomizableObjectInstance->GetCurrentMaxLOD() < 2)
-						{
-							++CountLOD1;
-						}
-						else
-						{
-							++CountLOD2;
-						}
-					}
-				}
-			}
-		}
-
-		NumInstances = CountLOD0 + CountLOD1 + CountLOD2;
-		TotalInstances = CountTotal;
-		NumPendingInstances = MutablePendingInstanceWork.Num();
-		//SET_DWORD_STAT(STAT_MutableSkeletalMeshResourceMemory, Size / 1024.f);
-		SET_DWORD_STAT(STAT_MutablePendingInstanceUpdates, MutablePendingInstanceWork.Num());
-
-		uint64 Size = 0;
-		uint32 CountAllocated = 0;
-		for (TWeakObjectPtr<UTexture2D>& Tracker : TextureTrackerArray)
-		{
-			if (Tracker.IsValid() && Tracker->GetResource())
-			{
-				CountAllocated++;
+		MutableStats.TextureTrackerArray.Add(Texture);
 
-				if (Tracker->GetResource()->TextureRHI)
-				{
-					Size += Tracker->CalcTextureMemorySizeEnum(TMC_AllMips);
-				}
-			}
-		}
-
-		TextureMemoryUsed = int64_t(Size / 1024);
-
-		uint64 SizeGenerated = 0;
-		for (TObjectIterator<UCustomizableObjectInstance> CustomizableObjectInstance; CustomizableObjectInstance; ++CustomizableObjectInstance)
+		for (auto Iterator = MutableStats.TextureTrackerArray.CreateIterator(); Iterator; ++Iterator)
 		{
-			if (IsValidChecked(*CustomizableObjectInstance) && CustomizableObjectInstance->GetPrivate() && CustomizableObjectInstance->HasAnySkeletalMesh())
+			if (Iterator->IsStale())
 			{
-				bool bHasResourceForRendering = false;
-				for (int32 MeshIndex = 0; !bHasResourceForRendering && MeshIndex < CustomizableObjectInstance->SkeletalMeshes.Num(); ++MeshIndex)
-				{
-					bHasResourceForRendering = CustomizableObjectInstance->GetSkeletalMesh(MeshIndex) && CustomizableObjectInstance->GetSkeletalMesh(MeshIndex)->GetResourceForRendering();
-				}
-
-				if (bHasResourceForRendering)
-				{
-					for (const FGeneratedTexture& GeneratedTextures : CustomizableObjectInstance->GetPrivate()->GeneratedTextures)
-					{
-						if (GeneratedTextures.Texture)
-						{
-							check(GeneratedTextures.Texture != nullptr);
-							SizeGenerated += GeneratedTextures.Texture->CalcTextureMemorySizeEnum(TMC_AllMips);
-						}
-					}
-				}
+				Iterator.RemoveCurrent();
 			}
 		}
+	}
 
-		if (LogBenchmarkUtil::isLoggingActive())
-		{
-			LogBenchmarkUtil::updateStat("customizable_objects", (int32)CountAllocatedSkeletalMesh);
-			LogBenchmarkUtil::updateStat("pending_instance_updates", MutablePendingInstanceWork.Num());
-			LogBenchmarkUtil::updateStat("allocated_textures", (int32)CountAllocated);
-			LogBenchmarkUtil::updateStat("texture_resource_memory", (long double)Size / 1048576.0L);
-			LogBenchmarkUtil::updateStat("texture_generated_memory", (long double)SizeGenerated / 1048576.0L);
-		}
-
-#if !(UE_BUILD_SHIPPING || UE_BUILD_TEST)
-		SET_DWORD_STAT(STAT_MutableNumInstancesLOD0, CountLOD0);
-		SET_DWORD_STAT(STAT_MutableNumInstancesLOD1, CountLOD1);
-		SET_DWORD_STAT(STAT_MutableNumInstancesLOD2, CountLOD2);
-		SET_DWORD_STAT(STAT_MutableNumAllocatedSkeletalMeshes, CountAllocatedSkeletalMesh);
-		SET_DWORD_STAT(STAT_MutablePendingInstanceUpdates, MutablePendingInstanceWork.Num());
-		SET_DWORD_STAT(STAT_MutableNumAllocatedTextures, CountAllocated);
-		SET_DWORD_STAT(STAT_MutableTextureResourceMemory, Size / 1024.f);
-		SET_DWORD_STAT(STAT_MutableTextureGeneratedMemory, SizeGenerated / 1024.f);
-#endif
-
-#if WITH_EDITORONLY_DATA
-		if (FCustomizableObjectSystemPrivate::IsMutableAnimInfoDebuggingEnabled())
-		{
-			if (GEngine)
-			{
-				bool bFoundPlayer = false;
-				int32 MsgIndex = 15820; // Arbitrary big value to prevent collisions with other on-screen messages
-
-				for (TObjectIterator<UCustomizableSkeletalComponent> CustomizableSkeletalComponent; CustomizableSkeletalComponent; ++CustomizableSkeletalComponent)
-				{
-					AActor* ParentActor = CustomizableSkeletalComponent->GetAttachmentRootActor();
-					UCustomizableObjectInstance* Instance = CustomizableSkeletalComponent->CustomizableObjectInstance;
-
-					APawn* PlayerPawn = nullptr;
-					if (UWorld* World = CustomizableSkeletalComponent->GetWorld())
-					{
-						PlayerPawn = UGameplayStatics::GetPlayerPawn(World, 0);						
-					}
-					
-					if (ParentActor && (ParentActor == PlayerPawn) && Instance)
-					{
-						bFoundPlayer = true;
-
-						FString TagString;
-						const FGameplayTagContainer& Tags = Instance->GetAnimationGameplayTags();
-
-						for (const FGameplayTag& Tag : Tags)
-						{
-							TagString += !TagString.IsEmpty() ? FString(TEXT(", ")) : FString();
-							TagString += Tag.ToString();
-						}												
-												
-						GEngine->AddOnScreenDebugMessage(MsgIndex++, .0f, FColor::Green, TEXT("Animation tags: ") + TagString);
-
-						check(Instance->GetPrivate() != nullptr);
-						FCustomizableInstanceComponentData* ComponentData = Instance->GetPrivate()->GetComponentData(CustomizableSkeletalComponent->ComponentIndex);
-
-						if (ComponentData)
-						{
-							for (TPair<FName, TSoftClassPtr<UAnimInstance>>& Entry : ComponentData->AnimSlotToBP)
-							{
-								FString AnimBPSlot;
-
-								AnimBPSlot += Entry.Key.ToString() + FString("-") + Entry.Value.GetAssetName();
-								GEngine->AddOnScreenDebugMessage(MsgIndex++, .0f, FColor::Green, AnimBPSlot);
-							}
-						}
-
-						GEngine->AddOnScreenDebugMessage(MsgIndex++, .0f, FColor::Green, TEXT("Slots-AnimBP: "));
-
-						if (ComponentData)
-						{
-							if (ComponentData->MeshPartPaths.IsEmpty())
-							{
-								GEngine->AddOnScreenDebugMessage(MsgIndex++, .0f, FColor::Magenta,
-									TEXT("No meshes found. In order to see the meshes compile the pawn's root CustomizableObject after the 'mutable.EnableMutableAnimInfoDebugging 1' command has been run."));
-							}
-
-							for (const FString& MeshPath : ComponentData->MeshPartPaths)
-							{
-								GEngine->AddOnScreenDebugMessage(MsgIndex++, .0f, FColor::Magenta, MeshPath);
-							}
-						}
+	INC_DWORD_STAT(STAT_MutableNumTextures);
+}
 
-						GEngine->AddOnScreenDebugMessage(MsgIndex++, .0f, FColor::Magenta, TEXT("Meshes: "));
 
-						GEngine->AddOnScreenDebugMessage(MsgIndex++, .0f, FColor::Cyan,
-							TEXT("Player Pawn Mutable Mesh/Animation info for component ") + FString::Printf(TEXT("%d"), 
-							CustomizableSkeletalComponent->ComponentIndex));
-					}
-				}
+int32 FCustomizableObjectSystemPrivate::EnableMutableAnimInfoDebugging = 0;
 
-				if (!bFoundPlayer)
-				{
-					GEngine->AddOnScreenDebugMessage(MsgIndex, .0f, FColor::Yellow, TEXT("Mutable Animation info: N/A"));
-				}
-			}
-		}
-#endif
-	}
-}
+static FAutoConsoleVariableRef CVarEnableMutableAnimInfoDebugging(
+	TEXT("mutable.EnableMutableAnimInfoDebugging"), FCustomizableObjectSystemPrivate::EnableMutableAnimInfoDebugging,
+	TEXT("If set to 1 or greater print on screen the animation info of the pawn's Customizable Object Instance. Anim BPs, slots and tags will be displayed."
+	"If the root Customizable Object is recompiled after this command is run, the used skeletal meshes will also be displayed."),
+	ECVF_Default);
 
 
 void FCustomizableObjectSystemPrivate::AddReferencedObjects(FReferenceCollector& Collector)
@@ -1994,18 +1792,8 @@ namespace impl
 			CustomizableObjectInstancePrivateData->TexturesToRelease.Empty();
 		}
 
-		// TODO: T2927
-		if (LogBenchmarkUtil::isLoggingActive())
-		{
-			double deltaSeconds = FPlatformTime::Seconds() - CustomizableObjectSystemPrivateData->CurrentMutableOperation->StartUpdateTime;
-			int32 deltaMs = int32(deltaSeconds * 1000);
-
-			LogBenchmarkUtil::updateStat("customizable_instance_build_time", deltaMs);
-			CustomizableObjectSystemPrivateData->TotalBuildMs += deltaMs;
-			CustomizableObjectSystemPrivateData->TotalBuiltInstances++;
-			SET_DWORD_STAT(STAT_MutableInstanceBuildTime, deltaMs);
-			SET_DWORD_STAT(STAT_MutableInstanceBuildTimeAvrg, CustomizableObjectSystemPrivateData->TotalBuildMs / CustomizableObjectSystemPrivateData->TotalBuiltInstances);
-		}
+		LogBenchmarkUtil::UpdateBuildTimeStats(CustomizableObjectSystemPrivateData->MutableStats, 
+												CustomizableObjectSystemPrivateData->CurrentMutableOperation->StartUpdateTime);
 
 		// End Update
 		System->ClearCurrentMutableOperation();
@@ -2425,7 +2213,7 @@ namespace impl
 			return;
 		}
 
-		if (LogBenchmarkUtil::isLoggingActive())
+		if (LogBenchmarkUtil::IsLoggingActive())
 		{
 			Operation->StartUpdateTime = FPlatformTime::Seconds();
 		}
@@ -2847,22 +2635,9 @@ bool UCustomizableObjectSystem::Tick(float DeltaTime)
 	{
 		AdvanceCurrentOperation();
 	}
-	
-	// TODO: T2927
-	if (LogBenchmarkUtil::isLoggingActive())
-	{
-		uint64 SizeCache = 0;
-		for (const UTexture2D* CachedTextures : ProtectedCachedTextures)
-		{
-			if (CachedTextures)
-			{
-				SizeCache += CachedTextures->CalcTextureMemorySizeEnum(TMC_AllMips);
-			}
-		}
-		SET_DWORD_STAT(STAT_MutableTextureCacheMemory, SizeCache / 1024.f);
-	}
 
-	Private->UpdateStats();
+	Private->MutableStats.MutablePendingInstanceWorkCount = Private->MutablePendingInstanceWork.Num();
+	LogBenchmarkUtil::UpdateStats(Private->MutableStats, ProtectedCachedTextures);
 
 #if WITH_EDITOR
 	TickRecompileCustomizableObjects();
@@ -3121,31 +2896,31 @@ FMutableOperation FMutableOperation::CreateInstanceUpdate(UCustomizableObjectIns
 int32 UCustomizableObjectSystem::GetNumInstances() const
 {
 	check(Private != nullptr);
-	return Private->NumInstances;
+	return Private->MutableStats.NumInstances;
 }
 
 int32 UCustomizableObjectSystem::GetNumPendingInstances() const
 {
 	check(Private != nullptr);
-	return Private->NumPendingInstances;
+	return Private->MutableStats.NumPendingInstances;
 }
 
 int32 UCustomizableObjectSystem::GetTotalInstances() const
 {
 	check(Private != nullptr);
-	return Private->TotalInstances;
+	return Private->MutableStats.TotalInstances;
 }
 
 int32 UCustomizableObjectSystem::GetTextureMemoryUsed() const
 {
 	check(Private != nullptr);
-	return int32(Private->TextureMemoryUsed);
+	return int32(Private->MutableStats.TextureMemoryUsed);
 }
 
 int32 UCustomizableObjectSystem::GetAverageBuildTime() const
 {
 	check(Private != nullptr);
-	return Private->TotalBuiltInstances == 0 ? 0 : Private->TotalBuildMs / Private->TotalBuiltInstances;
+	return Private->MutableStats.TotalBuiltInstances == 0 ? 0 : Private->MutableStats.TotalBuildMs / Private->MutableStats.TotalBuiltInstances;
 }
 
 
@@ -3235,14 +3010,16 @@ void UCustomizableObjectSystem::AddUncompiledCOWarning(const UCustomizableObject
 	UE_LOG(LogMutable, Error, TEXT("%s"), *ErrorString);
 }
 
+
 void UCustomizableObjectSystem::EnableBenchmark()
 {
-	LogBenchmarkUtil::startLogging();
+	LogBenchmarkUtil::StartLogging();
 }
 
+
 void UCustomizableObjectSystem::EndBenchmark()
 {
-	LogBenchmarkUtil::shutdownAndSaveResults();
+	LogBenchmarkUtil::ShutdownAndSaveResults();
 }
 
 
diff --git a/Engine/Plugins/Experimental/Mutable/Source/CustomizableObject/Private/MuCO/CustomizableObjectSystemPrivate.h b/Engine/Plugins/Experimental/Mutable/Source/CustomizableObject/Private/MuCO/CustomizableObjectSystemPrivate.h
--- a/Engine/Plugins/Experimental/Mutable/Source/CustomizableObject/Private/MuCO/CustomizableObjectSystemPrivate.h
+++ b/Engine/Plugins/Experimental/Mutable/Source/CustomizableObject/Private/MuCO/CustomizableObjectSystemPrivate.h
@@ -610,6 +610,36 @@ struct FMutableReleasePlatformOperationData
 };
 
 
+struct FMutableStats
+{
+	// Stores the number of pending instance updates, LOD Updates, discards and releases last tick
+	int32 MutablePendingInstanceWorkCount = 0;
+
+	// Total time spent building instances
+	int32 TotalBuildMs = 0;
+
+	// Total number of instances built or updated.
+	int32 TotalBuiltInstances = 0;
+
+	// Number of instances alive (built)
+	int32 NumInstances = 0;
+
+	// Number of CustomizableObjectInstances waiting to be updated.
+	int32 NumPendingInstances = 0;
+
+	// Total number of CustomizableObjectInstances, including not built.
+	int32 TotalInstances = 0;
+
+	// Total memory in bytes used for generated textures
+	int64_t TextureMemoryUsed = 0;
+
+	uint32 CountAllocatedSkeletalMesh = 0;
+
+	// \TODO: Remove this array if we are not gathering stats!
+	TArray<TWeakObjectPtr<class UTexture2D>> TextureTrackerArray;
+};
+
+
 class FCustomizableObjectSystemPrivate : public FGCObject
 {
 public:
@@ -765,39 +795,19 @@ public:
 	bool IsReplaceDiscardedWithReferenceMeshEnabled() const { return bReplaceDiscardedWithReferenceMesh; }
 	void SetReplaceDiscardedWithReferenceMeshEnabled(bool bIsEnabled) { bReplaceDiscardedWithReferenceMesh = bIsEnabled; }
 
-	int32 GetCountAllocatedSkeletalMesh() { return CountAllocatedSkeletalMesh; }
+	int32 GetCountAllocatedSkeletalMesh() { return MutableStats.CountAllocatedSkeletalMesh; }
+
+	mutable FMutableStats MutableStats;
 
 	bool bReplaceDiscardedWithReferenceMesh = false;
 	bool bReleaseTexturesImmediately = false;
 
 	bool bSupport16BitBoneIndex = false;
 
-	// Statistics: Total time spent building instances
-	int32 TotalBuildMs = 0;
-
-	// Statistics: Total number of instances built or updated.
-	int32 TotalBuiltInstances = 0;
-
-	// Statistics: Number of instances alive (built)
-	int32 NumInstances = 0;
-
-	// Statistics: number of CustomizableObjectInstances waiting to be updated.
-	int32 NumPendingInstances = 0;
-
-	// Statistics: Total number of CustomizableObjectInstances, including not built.
-	int32 TotalInstances = 0;
-
-	// Statistics: total memory in bytes used for generated textures
-	int64_t TextureMemoryUsed = 0;
-
-	mutable uint32 CountAllocatedSkeletalMesh = 0;
-
 	static FCustomizableObjectCompilerBase* (*NewCompilerFunc)();
 
 	void CreatedTexture(UTexture2D* Texture);
 
-	// \TODO: Remove this array if we are not gathering stats!
-	TArray<TWeakObjectPtr<class UTexture2D>> TextureTrackerArray;
 	TMap<FMutableImageCacheKey, uint32> TextureReferenceCount; // Keeps a count of texture usage to decide if they have to be blocked from GC during an update
 
 	// This is protected from GC by AddReferencedObjects
@@ -809,9 +819,6 @@ public:
 	FTSTicker::FDelegateHandle TickDelegateHandle;
 	FTickerDelegate TickDelegate;
 
-	//! Update the stats logged in unreal's stats system. 
-	void UpdateStats();
-
 	/** Update the last set amount of internal memory Mutable can use to build objects. */
 	void UpdateMemoryLimit();
 
diff --git a/Engine/Plugins/Experimental/Mutable/Source/CustomizableObject/Private/MuCO/LogBenchmarkUtil.cpp b/Engine/Plugins/Experimental/Mutable/Source/CustomizableObject/Private/MuCO/LogBenchmarkUtil.cpp
--- a/Engine/Plugins/Experimental/Mutable/Source/CustomizableObject/Private/MuCO/LogBenchmarkUtil.cpp
+++ b/Engine/Plugins/Experimental/Mutable/Source/CustomizableObject/Private/MuCO/LogBenchmarkUtil.cpp
@@ -4,18 +4,54 @@
 
 #include "HAL/PlatformFileManager.h"
 #include "Misc/FileHelper.h"
+#include "Kismet/GameplayStatics.h"
+#include "Engine/SkeletalMesh.h"
+#include "TextureResource.h"
+#include "GameFramework/Pawn.h"
+
+#if WITH_EDITOR
+#include "Editor/EditorEngine.h"
+#endif
+
 #include "MuCO/UnrealPortabilityHelpers.h"
+#include "MuCO/CustomizableObjectSystemPrivate.h"
+#include "MuCO/CustomizableObjectSystem.h"
+#include "MuCO/CustomizableInstancePrivateData.h"
+#include "MuCO/CustomizableSkeletalComponent.h"
+
+
+extern ENGINE_API UEngine* GEngine;
+
+#if WITH_EDITOR
+extern UNREALED_API UEditorEngine* GEditor;
+#endif
+
+DEFINE_STAT(STAT_MutableNumSkeletalMeshes);
+DEFINE_STAT(STAT_MutableNumCachedSkeletalMeshes);
+DEFINE_STAT(STAT_MutableNumAllocatedSkeletalMeshes);
+DEFINE_STAT(STAT_MutableNumInstancesLOD0);
+DEFINE_STAT(STAT_MutableNumInstancesLOD1);
+DEFINE_STAT(STAT_MutableNumInstancesLOD2);
+DEFINE_STAT(STAT_MutableNumTextures);
+DEFINE_STAT(STAT_MutableNumCachedTextures);
+DEFINE_STAT(STAT_MutableNumAllocatedTextures);
+DEFINE_STAT(STAT_MutableTextureResourceMemory);
+DEFINE_STAT(STAT_MutableTextureGeneratedMemory);
+DEFINE_STAT(STAT_MutableTextureCacheMemory);
+DEFINE_STAT(STAT_MutablePendingInstanceUpdates);
+DEFINE_STAT(STAT_MutableAbandonedInstanceUpdates);
+DEFINE_STAT(STAT_MutableInstanceBuildTime);
+DEFINE_STAT(STAT_MutableInstanceBuildTimeAvrg);
 
 bool LogBenchmarkUtil::bLoggingActive = false;
 
-LogBenchmarkUtil::~LogBenchmarkUtil() {}
 
-void LogBenchmarkUtil::startLogging()
+void LogBenchmarkUtil::StartLogging()
 {
 	bLoggingActive = true;
 }
 
-void LogBenchmarkUtil::shutdownAndSaveResults()
+void LogBenchmarkUtil::ShutdownAndSaveResults()
 {
 	if (bLoggingActive)
 	{
@@ -23,100 +59,327 @@ void LogBenchmarkUtil::shutdownAndSaveResults()
 		FString FileName = FString("BenchmarkResult.txt");
 		FString AbsoluteFilePath = SaveDirectory + "/" + FileName;
 		IPlatformFile& PlatformFile = FPlatformFileManager::Get().GetPlatformFile();
+
 		if (PlatformFile.CreateDirectoryTree(*SaveDirectory))
 		{
-			FString currentTime = FDateTime::Now().ToString();
-			FString DateStamp = LINE_TERMINATOR + FString("[" + currentTime + "]LogBenchmark: (string) ") + FString("benchmark_results_location : ") + SaveDirectory;
+			FString CurrentTime = FDateTime::Now().ToString();
+			FString DateStamp = LINE_TERMINATOR + FString("[" + CurrentTime + "]LogBenchmark: (string) ") + FString("benchmark_results_location : ") + SaveDirectory;
 			FFileHelper::SaveStringToFile(DateStamp, *AbsoluteFilePath, FFileHelper::EEncodingOptions::ForceAnsi);
 
-			FString keyName;
-			int32 maxValue;
-			float averageValue;
+			FString KeyName;
+			int32 MaxValue;
+			float AverageValue;
 
 			// Iterate over the integer stats and log them to a file
-			for (TPair<FName, IntStat> ist : getInstance().IntStats)
+			for (TPair<FName, IntStat> Ist : GetInstance().IntStats)
 			{
-				currentTime = FDateTime::Now().ToString();
+				CurrentTime = FDateTime::Now().ToString();
 
-				keyName = FString(ist.Key.ToString());
-				maxValue = ist.Value.max;
-				averageValue = (long double)ist.Value.total / (long double)ist.Value.totalCalls;
+				KeyName = FString(Ist.Key.ToString());
+				MaxValue = Ist.Value.Max;
+				AverageValue = (long double)Ist.Value.Total / (long double)Ist.Value.TotalCalls;
 
 				FString TextToSave = LINE_TERMINATOR;
-				TextToSave += FString("[" + currentTime + "]");
+				TextToSave += FString("[" + CurrentTime + "]");
 				TextToSave += FString("LogBenchmark: (int) peak_");
-				TextToSave += FString::Printf(TEXT("%s : %lld"), *keyName, maxValue);
+				TextToSave += FString::Printf(TEXT("%s : %lld"), *KeyName, MaxValue);
 				TextToSave += LINE_TERMINATOR;
-				TextToSave += FString("[" + currentTime + "]");
+				TextToSave += FString("[" + CurrentTime + "]");
 				TextToSave += FString("LogBenchmark: (float) average_");
-				TextToSave += FString::Printf(TEXT("%s : %f"), *keyName, averageValue);
+				TextToSave += FString::Printf(TEXT("%s : %f"), *KeyName, AverageValue);
 
 				FFileHelper::SaveStringToFile(TextToSave, *AbsoluteFilePath, FFileHelper::EEncodingOptions::ForceAnsi, &IFileManager::Get(), FILEWRITE_Append);
 			}
 
 			// Iterate over the Float stats and log them to a file
-			for (TPair<FName, FloatStat> dst : getInstance().FloatStats)
+			for (TPair<FName, FloatStat> dst : GetInstance().FloatStats)
 			{
-				currentTime = FDateTime::Now().ToString();
+				CurrentTime = FDateTime::Now().ToString();
 
-				keyName = FString(dst.Key.ToString());
-				maxValue = dst.Value.max;
-				averageValue = dst.Value.total / dst.Value.totalCalls;
+				KeyName = FString(dst.Key.ToString());
+				MaxValue = dst.Value.Max;
+				AverageValue = dst.Value.Total / dst.Value.TotalCalls;
 
 				FString TextToSave = LINE_TERMINATOR;
-				TextToSave += FString("[" + currentTime + "]");
+				TextToSave += FString("[" + CurrentTime + "]");
 				TextToSave += FString("LogBenchmark: (int) peak_");
-				TextToSave += FString::Printf(TEXT("%s : %lld"), *keyName, maxValue);
+				TextToSave += FString::Printf(TEXT("%s : %lld"), *KeyName, MaxValue);
 				TextToSave += LINE_TERMINATOR;
-				TextToSave += FString("[" + currentTime + "]");
+				TextToSave += FString("[" + CurrentTime + "]");
 				TextToSave += FString("LogBenchmark: (float) average_");
-				TextToSave += FString::Printf(TEXT("%s : %f"), *keyName, averageValue);
+				TextToSave += FString::Printf(TEXT("%s : %f"), *KeyName, AverageValue);
 
 				FFileHelper::SaveStringToFile(TextToSave, *AbsoluteFilePath, FFileHelper::EEncodingOptions::ForceAnsi, &IFileManager::Get(), FILEWRITE_Append);
 			}
 		}
+
 		bLoggingActive = false;
 	}
 }
 
-void LogBenchmarkUtil::updateStat(const FName& stat, int32 value)
+void LogBenchmarkUtil::UpdateStat(const FName& Stat, int32 Value)
 {
 	if (bLoggingActive)
 	{
-		if (getInstance().IntStats.Contains(stat))
+		if (GetInstance().IntStats.Contains(Stat))
 		{
-			getInstance().IntStats[stat] += value;
+			GetInstance().IntStats[Stat] += Value;
 		}
 		else
 		{
-			getInstance().IntStats.Add(stat) += value;
+			GetInstance().IntStats.Add(Stat) += Value;
 		}
 	}
 }
 
-void LogBenchmarkUtil::updateStat(const FName & stat, double value)
+void LogBenchmarkUtil::UpdateStat(const FName & Stat, double Value)
 {
-	updateStat(stat, (long double)value);
+	UpdateStat(Stat, (long double)Value);
 }
 
-void LogBenchmarkUtil::updateStat(const FName& stat, long double value)
+void LogBenchmarkUtil::UpdateStat(const FName& Stat, long double Value)
 {
 	if (bLoggingActive)
 	{
-		if (getInstance().FloatStats.Contains(stat))
+		if (GetInstance().FloatStats.Contains(Stat))
 		{
-			getInstance().FloatStats[stat] += value;
+			GetInstance().FloatStats[Stat] += Value;
 		}
 		else
 		{
-			getInstance().FloatStats.Add(stat) += value;
+			GetInstance().FloatStats.Add(Stat) += Value;
 		}
 	}
 }
 
-bool LogBenchmarkUtil::isLoggingActive()
+bool LogBenchmarkUtil::IsLoggingActive()
 {
 	return bLoggingActive;
 }
 
-LogBenchmarkUtil::LogBenchmarkUtil() {}
+
+void LogBenchmarkUtil::UpdateStats(FMutableStats& StatsToUpdate, const TArray< TObjectPtr<UTexture2D> >& ProtectedCachedTextures)
+{
+	if (LogBenchmarkUtil::IsLoggingActive())
+	{
+		uint64 SizeCache = 0;
+
+		for (const UTexture2D* CachedTextures : ProtectedCachedTextures)
+		{
+			if (CachedTextures)
+			{
+				SizeCache += CachedTextures->CalcTextureMemorySizeEnum(TMC_AllMips);
+			}
+		}
+
+		SET_DWORD_STAT(STAT_MutableTextureCacheMemory, SizeCache / 1024.f);
+	}
+
+#if !(UE_BUILD_SHIPPING || UE_BUILD_TEST)
+	bool bLogEnabled = true;
+#else
+	bool bLogEnabled = LogBenchmarkUtil::IsLoggingActive();
+#endif
+	if (bLogEnabled)
+	{
+		const int32& MutablePendingInstanceWorkCount = StatsToUpdate.MutablePendingInstanceWorkCount;
+		StatsToUpdate.CountAllocatedSkeletalMesh = 0;
+
+		int32 CountLOD0 = 0;
+		int32 CountLOD1 = 0;
+		int32 CountLOD2 = 0;
+		int32 CountTotal = 0;
+
+		for (TObjectIterator<UCustomizableObjectInstance> CustomizableObjectInstance; CustomizableObjectInstance; ++CustomizableObjectInstance)
+		{
+			if (IsValidChecked(*CustomizableObjectInstance) && CustomizableObjectInstance->GetPrivate())
+			{
+				++CountTotal;
+
+				for (int32 ComponentIndex = 0; ComponentIndex < CustomizableObjectInstance->SkeletalMeshes.Num(); ++ComponentIndex)
+				{
+					if (CustomizableObjectInstance->SkeletalMeshes[ComponentIndex] && CustomizableObjectInstance->SkeletalMeshes[ComponentIndex]->GetResourceForRendering())
+					{
+						StatsToUpdate.CountAllocatedSkeletalMesh++;
+
+						if (CustomizableObjectInstance->GetCurrentMinLOD() < 1)
+						{
+							++CountLOD0;
+						}
+						else if (CustomizableObjectInstance->GetCurrentMaxLOD() < 2)
+						{
+							++CountLOD1;
+						}
+						else
+						{
+							++CountLOD2;
+						}
+					}
+				}
+			}
+		}
+
+		StatsToUpdate.NumInstances = CountLOD0 + CountLOD1 + CountLOD2;
+		StatsToUpdate.TotalInstances = CountTotal;
+		StatsToUpdate.NumPendingInstances = MutablePendingInstanceWorkCount;
+		SET_DWORD_STAT(STAT_MutablePendingInstanceUpdates, MutablePendingInstanceWorkCount);
+
+		uint64 Size = 0;
+		uint32 CountAllocated = 0;
+		for (TWeakObjectPtr<UTexture2D>& Tracker : StatsToUpdate.TextureTrackerArray)
+		{
+			if (Tracker.IsValid() && Tracker->GetResource())
+			{
+				CountAllocated++;
+
+				if (Tracker->GetResource()->TextureRHI)
+				{
+					Size += Tracker->CalcTextureMemorySizeEnum(TMC_AllMips);
+				}
+			}
+		}
+
+		StatsToUpdate.TextureMemoryUsed = int64_t(Size / 1024);
+
+		uint64 SizeGenerated = 0;
+		for (TObjectIterator<UCustomizableObjectInstance> CustomizableObjectInstance; CustomizableObjectInstance; ++CustomizableObjectInstance)
+		{
+			if (IsValidChecked(*CustomizableObjectInstance) && CustomizableObjectInstance->GetPrivate() && CustomizableObjectInstance->HasAnySkeletalMesh())
+			{
+				bool bHasResourceForRendering = false;
+				for (int32 MeshIndex = 0; !bHasResourceForRendering && MeshIndex < CustomizableObjectInstance->SkeletalMeshes.Num(); ++MeshIndex)
+				{
+					bHasResourceForRendering = CustomizableObjectInstance->GetSkeletalMesh(MeshIndex) && CustomizableObjectInstance->GetSkeletalMesh(MeshIndex)->GetResourceForRendering();
+				}
+
+				if (bHasResourceForRendering)
+				{
+					for (const FGeneratedTexture& GeneratedTextures : CustomizableObjectInstance->GetPrivate()->GeneratedTextures)
+					{
+						if (GeneratedTextures.Texture)
+						{
+							check(GeneratedTextures.Texture != nullptr);
+							SizeGenerated += GeneratedTextures.Texture->CalcTextureMemorySizeEnum(TMC_AllMips);
+						}
+					}
+				}
+			}
+		}
+
+		if (LogBenchmarkUtil::IsLoggingActive())
+		{
+			LogBenchmarkUtil::UpdateStat("customizable_objects", (int32)StatsToUpdate.CountAllocatedSkeletalMesh);
+			LogBenchmarkUtil::UpdateStat("pending_instance_updates", MutablePendingInstanceWorkCount);
+			LogBenchmarkUtil::UpdateStat("allocated_textures", (int32)CountAllocated);
+			LogBenchmarkUtil::UpdateStat("texture_resource_memory", (long double)Size / 1048576.0L);
+			LogBenchmarkUtil::UpdateStat("texture_generated_memory", (long double)SizeGenerated / 1048576.0L);
+		}
+
+#if !(UE_BUILD_SHIPPING || UE_BUILD_TEST)
+		SET_DWORD_STAT(STAT_MutableNumInstancesLOD0, CountLOD0);
+		SET_DWORD_STAT(STAT_MutableNumInstancesLOD1, CountLOD1);
+		SET_DWORD_STAT(STAT_MutableNumInstancesLOD2, CountLOD2);
+		SET_DWORD_STAT(STAT_MutableNumAllocatedSkeletalMeshes, StatsToUpdate.CountAllocatedSkeletalMesh);
+		SET_DWORD_STAT(STAT_MutablePendingInstanceUpdates, MutablePendingInstanceWorkCount);
+		SET_DWORD_STAT(STAT_MutableNumAllocatedTextures, CountAllocated);
+		SET_DWORD_STAT(STAT_MutableTextureResourceMemory, Size / 1024.f);
+		SET_DWORD_STAT(STAT_MutableTextureGeneratedMemory, SizeGenerated / 1024.f);
+#endif
+
+#if WITH_EDITORONLY_DATA
+		if (UCustomizableObjectSystem::GetInstance()->IsMutableAnimInfoDebuggingEnabled())
+		{
+			if (GEngine)
+			{
+				bool bFoundPlayer = false;
+				int32 MsgIndex = 15820; // Arbitrary big value to prevent collisions with other on-screen messages
+
+				for (TObjectIterator<UCustomizableSkeletalComponent> CustomizableSkeletalComponent; CustomizableSkeletalComponent; ++CustomizableSkeletalComponent)
+				{
+					AActor* ParentActor = CustomizableSkeletalComponent->GetAttachmentRootActor();
+					UCustomizableObjectInstance* Instance = CustomizableSkeletalComponent->CustomizableObjectInstance;
+
+					APawn* PlayerPawn = nullptr;
+					if (UWorld* World = CustomizableSkeletalComponent->GetWorld())
+					{
+						PlayerPawn = UGameplayStatics::GetPlayerPawn(World, 0);
+					}
+
+					if (ParentActor && (ParentActor == PlayerPawn) && Instance)
+					{
+						bFoundPlayer = true;
+
+						FString TagString;
+						const FGameplayTagContainer& Tags = Instance->GetAnimationGameplayTags();
+
+						for (const FGameplayTag& Tag : Tags)
+						{
+							TagString += !TagString.IsEmpty() ? FString(TEXT(", ")) : FString();
+							TagString += Tag.ToString();
+						}
+
+						GEngine->AddOnScreenDebugMessage(MsgIndex++, .0f, FColor::Green, TEXT("Animation tags: ") + TagString);
+
+						check(Instance->GetPrivate() != nullptr);
+						FCustomizableInstanceComponentData* ComponentData = Instance->GetPrivate()->GetComponentData(CustomizableSkeletalComponent->ComponentIndex);
+
+						if (ComponentData)
+						{
+							for (TPair<FName, TSoftClassPtr<UAnimInstance>>& Entry : ComponentData->AnimSlotToBP)
+							{
+								FString AnimBPSlot;
+
+								AnimBPSlot += Entry.Key.ToString() + FString("-") + Entry.Value.GetAssetName();
+								GEngine->AddOnScreenDebugMessage(MsgIndex++, .0f, FColor::Green, AnimBPSlot);
+							}
+						}
+
+						GEngine->AddOnScreenDebugMessage(MsgIndex++, .0f, FColor::Green, TEXT("Slots-AnimBP: "));
+
+						if (ComponentData)
+						{
+							if (ComponentData->MeshPartPaths.IsEmpty())
+							{
+								GEngine->AddOnScreenDebugMessage(MsgIndex++, .0f, FColor::Magenta,
+									TEXT("No meshes found. In order to see the meshes compile the pawn's root CustomizableObject after the 'mutable.EnableMutableAnimInfoDebugging 1' command has been run."));
+							}
+
+							for (const FString& MeshPath : ComponentData->MeshPartPaths)
+							{
+								GEngine->AddOnScreenDebugMessage(MsgIndex++, .0f, FColor::Magenta, MeshPath);
+							}
+						}
+
+						GEngine->AddOnScreenDebugMessage(MsgIndex++, .0f, FColor::Magenta, TEXT("Meshes: "));
+
+						GEngine->AddOnScreenDebugMessage(MsgIndex++, .0f, FColor::Cyan,
+							TEXT("Player Pawn Mutable Mesh/Animation info for component ") + FString::Printf(TEXT("%d"),
+								CustomizableSkeletalComponent->ComponentIndex));
+					}
+				}
+
+				if (!bFoundPlayer)
+				{
+					GEngine->AddOnScreenDebugMessage(MsgIndex, .0f, FColor::Yellow, TEXT("Mutable Animation info: N/A"));
+				}
+			}
+		}
+#endif
+	}
+}
+
+
+void LogBenchmarkUtil::UpdateBuildTimeStats(FMutableStats& StatsToUpdate, double StartUpdateTime)
+{
+	if (LogBenchmarkUtil::IsLoggingActive())
+	{
+		double DeltaSeconds = FPlatformTime::Seconds() - StartUpdateTime;
+		int32 DeltaMs = int32(DeltaSeconds * 1000);
+
+		LogBenchmarkUtil::UpdateStat("customizable_instance_build_time", DeltaMs);
+		StatsToUpdate.TotalBuildMs += DeltaMs;
+		StatsToUpdate.TotalBuiltInstances++;
+		SET_DWORD_STAT(STAT_MutableInstanceBuildTime, DeltaMs);
+		SET_DWORD_STAT(STAT_MutableInstanceBuildTimeAvrg, StatsToUpdate.TotalBuildMs / StatsToUpdate.TotalBuiltInstances);
+	}
+}
diff --git a/Engine/Plugins/Experimental/Mutable/Source/CustomizableObject/Private/MuCO/LogBenchmarkUtil.h b/Engine/Plugins/Experimental/Mutable/Source/CustomizableObject/Private/MuCO/LogBenchmarkUtil.h
--- a/Engine/Plugins/Experimental/Mutable/Source/CustomizableObject/Private/MuCO/LogBenchmarkUtil.h
+++ b/Engine/Plugins/Experimental/Mutable/Source/CustomizableObject/Private/MuCO/LogBenchmarkUtil.h
@@ -4,43 +4,73 @@
 
 #include "Containers/Map.h"
 #include "UObject/NameTypes.h"
+#include "Stats/Stats.h"
+
+// Mutable stats
+DECLARE_STATS_GROUP(TEXT("Mutable"), STATGROUP_Mutable, STATCAT_Advanced);
+DECLARE_DWORD_ACCUMULATOR_STAT_EXTERN(TEXT("Num Created Mutable Skeletal Meshes"), STAT_MutableNumSkeletalMeshes, STATGROUP_Mutable, );
+DECLARE_DWORD_ACCUMULATOR_STAT_EXTERN(TEXT("Num Cached Mutable Skeletal Meshes"), STAT_MutableNumCachedSkeletalMeshes, STATGROUP_Mutable, );
+DECLARE_DWORD_ACCUMULATOR_STAT_EXTERN(TEXT("Num Allocated Mutable Skeletal Meshes"), STAT_MutableNumAllocatedSkeletalMeshes, STATGROUP_Mutable, );
+DECLARE_DWORD_ACCUMULATOR_STAT_EXTERN(TEXT("Num Instances at LOD 0"), STAT_MutableNumInstancesLOD0, STATGROUP_Mutable, );
+DECLARE_DWORD_ACCUMULATOR_STAT_EXTERN(TEXT("Num Instances at LOD 1"), STAT_MutableNumInstancesLOD1, STATGROUP_Mutable, );
+DECLARE_DWORD_ACCUMULATOR_STAT_EXTERN(TEXT("Num Instances at LOD 2 or more"), STAT_MutableNumInstancesLOD2, STATGROUP_Mutable, );
+DECLARE_DWORD_ACCUMULATOR_STAT_EXTERN(TEXT("Num Created Mutable Textures"), STAT_MutableNumTextures, STATGROUP_Mutable, );
+DECLARE_DWORD_ACCUMULATOR_STAT_EXTERN(TEXT("Num Cached Mutable Textures"), STAT_MutableNumCachedTextures, STATGROUP_Mutable, );
+DECLARE_DWORD_ACCUMULATOR_STAT_EXTERN(TEXT("Num Allocated Mutable Textures"), STAT_MutableNumAllocatedTextures, STATGROUP_Mutable, );
+DECLARE_DWORD_ACCUMULATOR_STAT_EXTERN(TEXT("Texture Resource Memory"), STAT_MutableTextureResourceMemory, STATGROUP_Mutable, );
+DECLARE_DWORD_ACCUMULATOR_STAT_EXTERN(TEXT("Texture Generated Memory"), STAT_MutableTextureGeneratedMemory, STATGROUP_Mutable, );
+DECLARE_DWORD_ACCUMULATOR_STAT_EXTERN(TEXT("Texture Locked Memory"), STAT_MutableTextureCacheMemory, STATGROUP_Mutable, );
+DECLARE_DWORD_ACCUMULATOR_STAT_EXTERN(TEXT("Pending Instance Updates"), STAT_MutablePendingInstanceUpdates, STATGROUP_Mutable, );
+DECLARE_DWORD_ACCUMULATOR_STAT_EXTERN(TEXT("Abandoned Instance Updates"), STAT_MutableAbandonedInstanceUpdates, STATGROUP_Mutable, );
+DECLARE_DWORD_ACCUMULATOR_STAT_EXTERN(TEXT("Last Instance Build Time"), STAT_MutableInstanceBuildTime, STATGROUP_Mutable, );
+DECLARE_DWORD_ACCUMULATOR_STAT_EXTERN(TEXT("Avrg Instance Build Time"), STAT_MutableInstanceBuildTimeAvrg, STATGROUP_Mutable, );
+
+struct FMutableStats;
 
 class LogBenchmarkUtil
 {
 public:
-	~LogBenchmarkUtil();
-	static void startLogging();
-	static void shutdownAndSaveResults();
-	static void updateStat(const FName& stat, int32 value);
-	static void updateStat(const FName& stat, double value);
-	static void updateStat(const FName& stat, long double value);
-	static bool isLoggingActive();
+	static void StartLogging();
+	static void ShutdownAndSaveResults();
+
+	//! Update the stats logged in unreal's stats system.
+	static void UpdateStats(FMutableStats& StatsToUpdate, const TArray< TObjectPtr<class UTexture2D> >& ProtectedCachedTextures);
+
+	static void UpdateBuildTimeStats(FMutableStats& StatsToUpdate, double StartUpdateTime);
+
+	static void UpdateStat(const FName& Stat, int32 Value);
+	static void UpdateStat(const FName& Stat, double Value);
+	static void UpdateStat(const FName& Stat, long double Value);
+
+	static bool IsLoggingActive();
 
 private:
-	struct IntStat {
-		int32 total;
-		int32 max;
-		int totalCalls;
+	struct IntStat
+	{
+		int32 Total;
+		int32 Max;
+		int32 TotalCalls;
 
-		IntStat& operator+=(const int32& rhs)
+		IntStat& operator+=(const int32& Rhs)
 		{
-			total += rhs;
-			max = FGenericPlatformMath::Max(max, rhs);
-			totalCalls++;
+			Total += Rhs;
+			Max = FGenericPlatformMath::Max(Max, Rhs);
+			TotalCalls++;
 			return *this;
 		}
 	};
 
-	struct FloatStat {
-		long double total;
-		long double max;
-		int totalCalls;
+	struct FloatStat
+	{
+		long double Total;
+		long double Max;
+		int32 TotalCalls;
 
-		FloatStat& operator+=(const long double& rhs)
+		FloatStat& operator+=(const long double& Rhs)
 		{
-			total += rhs;
-			max = FGenericPlatformMath::Max(max, rhs);
-			totalCalls++;
+			Total += Rhs;
+			Max = FGenericPlatformMath::Max(Max, Rhs);
+			TotalCalls++;
 			return *this;
 		}
 	};
@@ -49,11 +79,11 @@ private:
 	TMap<FName, IntStat> IntStats;
 	TMap<FName, FloatStat> FloatStats;
 
-
-	static LogBenchmarkUtil &getInstance() {
-		static LogBenchmarkUtil instance;
-		return instance;
+	static LogBenchmarkUtil &GetInstance()
+	{
+		static LogBenchmarkUtil Instance;
+		return Instance;
 	}
 
-	LogBenchmarkUtil();
+	LogBenchmarkUtil() = default;
 };
diff --git a/Engine/Plugins/Experimental/Mutable/Source/CustomizableObject/Private/MuCO/UnrealMutableModelDiskStreamer.cpp b/Engine/Plugins/Experimental/Mutable/Source/CustomizableObject/Private/MuCO/UnrealMutableModelDiskStreamer.cpp
--- a/Engine/Plugins/Experimental/Mutable/Source/CustomizableObject/Private/MuCO/UnrealMutableModelDiskStreamer.cpp
+++ b/Engine/Plugins/Experimental/Mutable/Source/CustomizableObject/Private/MuCO/UnrealMutableModelDiskStreamer.cpp
@@ -12,6 +12,7 @@
 #include "HAL/PlatformFileManager.h"
 #endif
 
+DEFINE_STAT(STAT_MutableStreamingOps);
 
 //-------------------------------------------------------------------------------------------------
 #if WITH_EDITOR
diff --git a/Engine/Plugins/Experimental/Mutable/Source/CustomizableObject/Public/MuCO/UnrealMutableModelDiskStreamer.h b/Engine/Plugins/Experimental/Mutable/Source/CustomizableObject/Public/MuCO/UnrealMutableModelDiskStreamer.h
--- a/Engine/Plugins/Experimental/Mutable/Source/CustomizableObject/Public/MuCO/UnrealMutableModelDiskStreamer.h
+++ b/Engine/Plugins/Experimental/Mutable/Source/CustomizableObject/Public/MuCO/UnrealMutableModelDiskStreamer.h
@@ -3,6 +3,7 @@
 #pragma once
 
 #include "MuCO/CustomizableObject.h"
+#include "MuCO/LogBenchmarkUtil.h"
 #include "MuR/Serialisation.h"
 
 class FArchive;
@@ -12,6 +13,7 @@ class UCustomizableObject;
 namespace mu { class Model; }
 struct FMutableStreamableBlock;
 
+DECLARE_DWORD_ACCUMULATOR_STAT_EXTERN(TEXT("Streaming Ops"), STAT_MutableStreamingOps, STATGROUP_Mutable, );
 
 
 class UnrealMutableInputStream : public mu::InputStream
