[Mutable] Fix possible issue with swizzle vs transform sink optimizations.

- Swizzle channel compatibility check was not exhaustive.

#rnx

diff --git a/Engine/Plugins/Experimental/Mutable/Source/MutableTools/Private/MuT/ASTOpImageSwizzle.cpp b/Engine/Plugins/Experimental/Mutable/Source/MutableTools/Private/MuT/ASTOpImageSwizzle.cpp
--- a/Engine/Plugins/Experimental/Mutable/Source/MutableTools/Private/MuT/ASTOpImageSwizzle.cpp
+++ b/Engine/Plugins/Experimental/Mutable/Source/MutableTools/Private/MuT/ASTOpImageSwizzle.cpp
@@ -735,8 +735,12 @@ namespace mu
 							|| Typed->OffsetX.child() != FirstTransform->OffsetX.child()
 							|| Typed->OffsetY.child() != FirstTransform->OffsetY.child()
 							|| Typed->Rotation.child() != FirstTransform->Rotation.child()
+							|| Typed->SizeX != FirstTransform->SizeX
+							|| Typed->SizeY != FirstTransform->SizeY
 							|| Typed->SourceSizeX != FirstTransform->SourceSizeX
 							|| Typed->SourceSizeY != FirstTransform->SourceSizeY
+							|| Typed->AddressMode != FirstTransform->AddressMode
+							|| Typed->bKeepAspectRatio != FirstTransform->bKeepAspectRatio
 							)
 						{
 							bAreAllTransformsCompatible = false;
