[mutable] Removed some deprecated operations.
#rnx

diff --git a/Engine/Plugins/Experimental/Mutable/Source/CustomizableObject/Classes/MuCO/CustomizableObject.h b/Engine/Plugins/Experimental/Mutable/Source/CustomizableObject/Classes/MuCO/CustomizableObject.h
--- a/Engine/Plugins/Experimental/Mutable/Source/CustomizableObject/Classes/MuCO/CustomizableObject.h
+++ b/Engine/Plugins/Experimental/Mutable/Source/CustomizableObject/Classes/MuCO/CustomizableObject.h
@@ -1440,7 +1440,7 @@ private:
 	// This is a manual version number for the binary blobs in this asset.
 	// Increasing it invalidates all the previously compiled models.
 	// Warning: If while merging code both versions have changed, take the highest+1.
-	static const int32 CurrentSupportedVersion = 401;
+	static const int32 CurrentSupportedVersion = 402;
 
 public:
 
diff --git a/Engine/Plugins/Experimental/Mutable/Source/CustomizableObjectEditor/Private/MuCOE/SMutableCodeViewer.cpp b/Engine/Plugins/Experimental/Mutable/Source/CustomizableObjectEditor/Private/MuCOE/SMutableCodeViewer.cpp
--- a/Engine/Plugins/Experimental/Mutable/Source/CustomizableObjectEditor/Private/MuCOE/SMutableCodeViewer.cpp
+++ b/Engine/Plugins/Experimental/Mutable/Source/CustomizableObjectEditor/Private/MuCOE/SMutableCodeViewer.cpp
@@ -1022,7 +1022,7 @@ void SMutableCodeViewer::GenerateNavigationDropdownElements()
 	
 	// Add an extra operation type that will represent the constant resource based navigation type
 	{
-		const FText EntryName = FText::FromString("CONST_RESOURCE_BASED_NAV");	
+		const FText EntryName = FText::FromString("Selected Constant");	
 		const FSlateColor EntryColor = FSlateColor(FLinearColor(0.35f ,0.35f,1.0f,1));
 		ConstantBasedNavigationEntry = MakeShared<FMutableOperationElement>(mu::OP_TYPE::NONE,EntryName,0,EntryColor);
 		
diff --git a/Engine/Plugins/Experimental/Mutable/Source/MutableRuntime/Private/MuR/CodeRunner.cpp b/Engine/Plugins/Experimental/Mutable/Source/MutableRuntime/Private/MuR/CodeRunner.cpp
--- a/Engine/Plugins/Experimental/Mutable/Source/MutableRuntime/Private/MuR/CodeRunner.cpp
+++ b/Engine/Plugins/Experimental/Mutable/Source/MutableRuntime/Private/MuR/CodeRunner.cpp
@@ -50,7 +50,6 @@
 #include "MuR/OpMeshGeometryOperation.h"
 #include "MuR/OpMeshMerge.h"
 #include "MuR/OpMeshMorph.h"
-#include "MuR/OpMeshRemapIndices.h"
 #include "MuR/OpMeshRemove.h"
 #include "MuR/OpMeshReshape.h"
 #include "MuR/OpMeshTransform.h"
@@ -1851,63 +1850,6 @@ namespace mu
             break;
         }
 
-        case OP_TYPE::ME_EXTRACTFACEGROUP:
-        {
-			OP::MeshExtractFaceGroupArgs args = pModel->GetPrivate()->m_program.GetOpArgs<OP::MeshExtractFaceGroupArgs>(item.At);
-            switch (item.Stage)
-            {
-            case 0:
-			{
-				if (args.source)
-				{
-					AddOp(FScheduledOp(item.At, item, 1),
-						FScheduledOp(args.source, item));
-				}
-				else
-				{
-					StoreMesh(item, nullptr);
-				}
-
-				break;
-			}
-            case 1:
-            {
-				MUTABLE_CPUPROFILER_SCOPE(ME_EXTRACTFACEGROUP_1)
-
-                Ptr<const Mesh> Source = LoadMesh(FCacheAddress(args.source, item));
-				if (Source)
-				{
-					Ptr<Mesh> Result = CreateMesh();
-
-					bool bOutSuccess = false;
-					MeshExtractFaceGroup(Result.get(), Source.get(), args.group, bOutSuccess);
-
-					Release(Source);
-					if (!bOutSuccess)
-					{
-						check(false);
-						Release(Result);
-						StoreMesh(item, nullptr);
-					}
-					else
-					{
-						StoreMesh(item, Result);
-					}
-				}
-				else
-				{
-					StoreMesh(item, nullptr);
-				}
-                break;
-            }
-
-            default:
-                check(false);
-            }
-
-            break;
-        }
-
         case OP_TYPE::ME_TRANSFORM:
         {
 			OP::MeshTransformArgs args = pModel->GetPrivate()->m_program.GetOpArgs<OP::MeshTransformArgs>(item.At);
@@ -2194,64 +2136,6 @@ namespace mu
 			break;
 		}
 
-        case OP_TYPE::ME_REMAPINDICES:
-        {
-			OP::MeshRemapIndicesArgs args = pModel->GetPrivate()->m_program.GetOpArgs<OP::MeshRemapIndicesArgs>(item.At);
-            switch (item.Stage)
-            {
-            case 0:
-			{
-				if (args.source)
-				{
-					AddOp(FScheduledOp(item.At, item, 1),
-						FScheduledOp(args.source, item),
-						FScheduledOp(args.reference, item));
-				}
-				else
-				{
-					StoreMesh(item, nullptr);
-				}
-				break;
-			}
-            case 1:
-            {
-    			MUTABLE_CPUPROFILER_SCOPE(ME_REMAPINDICES_1)
-
-                Ptr<const Mesh> Source = LoadMesh(FCacheAddress(args.source,item));
-                Ptr<const Mesh> pReference = LoadMesh(FCacheAddress(args.reference,item));
-
-                // Only if both are valid.
-                MeshPtr Result = CreateMesh(Source ? Source->GetDataSize() : 0);
-                
-				bool bOutSuccess = false;
-				if (Source && pReference)
-                {	
-                    MeshRemapIndices(Result.get(), Source.get(), pReference.get(), bOutSuccess);
-                }
-			
-				Release(pReference);
-				if (!bOutSuccess)
-				{
-					Release(Result);
-					StoreMesh(item, Source);
-				}
-				else
-				{	
-					Release(Source);
-					StoreMesh(item, Result);
-				}
-
-                break;
-            }
-
-            default:
-                check(false);
-            }
-
-            break;
-        }
-
-
         case OP_TYPE::ME_APPLYPOSE:
         {
 			OP::MeshApplyPoseArgs args = pModel->GetPrivate()->m_program.GetOpArgs<OP::MeshApplyPoseArgs>(item.At);
@@ -5497,84 +5381,6 @@ namespace mu
             break;
         }
 
-        case OP_TYPE::CO_IMAGESIZE:
-        {
-			OP::ColourImageSizeArgs args = pModel->GetPrivate()->m_program.GetOpArgs<OP::ColourImageSizeArgs>(item.At);
-            switch (item.Stage)
-            {
-            case 0:
-                    AddOp( FScheduledOp( item.At, item, 1),
-                           FScheduledOp( args.image, item) );
-                break;
-
-            case 1:
-            {
-                Ptr<const Image> pImage = LoadImage( FCacheAddress(args.image,item) );
-
-				FVector4f result = FVector4f( (float)pImage->GetSizeX(), (float)pImage->GetSizeY(), 0.0f, 0.0f );
-
-				Release(pImage);
-                StoreColor( item, result );
-                break;
-            }
-
-            default:
-                check(false);
-            }
-
-            break;
-        }
-
-        case OP_TYPE::CO_LAYOUTBLOCKTRANSFORM:
-        {
-			OP::ColourLayoutBlockTransformArgs args = pModel->GetPrivate()->m_program.GetOpArgs<OP::ColourLayoutBlockTransformArgs>(item.At);
-            switch (item.Stage)
-            {
-            case 0:
-                    AddOp( FScheduledOp( item.At, item, 1),
-                           FScheduledOp( args.layout, item) );
-                break;
-
-            case 1:
-            {
-                Ptr<const Layout> pLayout = LoadLayout( FCacheAddress(args.layout,item) );
-
-				FVector4f result = FVector4f(0,0,0,0);
-                if ( pLayout )
-                {
-                    int relBlockIndex = pLayout->FindBlock( args.block );
-
-                    if( relBlockIndex >=0 )
-                    {
-                        box< UE::Math::TIntVector2<uint16> > rectInblocks;
-                        pLayout->GetBlock
-                                (
-                                    relBlockIndex,
-                                    &rectInblocks.min[0], &rectInblocks.min[1],
-                                    &rectInblocks.size[0], &rectInblocks.size[1]
-                                    );
-
-                        // Convert the rect from blocks to pixels
-                        FIntPoint grid = pLayout->GetGridSize();
-
-                        result = FVector4f( float(rectInblocks.min[0]) / float(grid[0]),
-                                              float(rectInblocks.min[1]) / float(grid[1]),
-                                              float(rectInblocks.size[0]) / float(grid[0]),
-                                              float(rectInblocks.size[1]) / float(grid[1]) );
-                    }
-                }
-
-                StoreColor( item, result );
-                break;
-            }
-
-            default:
-                check(false);
-            }
-
-            break;
-        }
-
         case OP_TYPE::CO_FROMSCALARS:
         {
 			OP::ColourFromScalarsArgs args = pModel->GetPrivate()->m_program.GetOpArgs<OP::ColourFromScalarsArgs>(item.At);
diff --git a/Engine/Plugins/Experimental/Mutable/Source/MutableRuntime/Private/MuR/CodeRunnerTasked.cpp b/Engine/Plugins/Experimental/Mutable/Source/MutableRuntime/Private/MuR/CodeRunnerTasked.cpp
--- a/Engine/Plugins/Experimental/Mutable/Source/MutableRuntime/Private/MuR/CodeRunnerTasked.cpp
+++ b/Engine/Plugins/Experimental/Mutable/Source/MutableRuntime/Private/MuR/CodeRunnerTasked.cpp
@@ -137,8 +137,6 @@ namespace mu
 		    
 		{   0,   0,   0,   0 },	// CO_SAMPLEIMAGE
 		{   0,   0,   0,   0 },	// CO_SWIZZLE
-		{   0,   0,   0,   0 },	// CO_IMAGESIZE
-		{   0,   0,   0,   0 },	// CO_LAYOUTBLOCKTRANSFORM
 		{   0,   0,   0,   0 },	// CO_FROMSCALARS
 		{   0,   0,   0,   0 },	// CO_ARITHMETIC
 		    
@@ -186,7 +184,6 @@ namespace mu
 		{   0,   0,   0,   0 },	// ME_SETSKELETON
 		{   0,   0,   0,   0 },	// ME_PROJECT
 		{   0,   0,   0,   0 },	// ME_APPLYPOSE
-		{   0,   0,   0,   0 },	// ME_REMAPINDICES
 		{   0,   0,   0,   0 },	// ME_GEOMETRYOPERATION
 		{   0, -20,   0,   0 },	// ME_BINDSHAPE
 		{   0, -10,   0,   0 },	// ME_APPLYSHAPE
diff --git a/Engine/Plugins/Experimental/Mutable/Source/MutableRuntime/Private/MuR/OpMeshRemapIndices.h b/Engine/Plugins/Experimental/Mutable/Source/MutableRuntime/Private/MuR/OpMeshRemapIndices.h
deleted file mode 100644
diff --git a/Engine/Plugins/Experimental/Mutable/Source/MutableRuntime/Private/MuR/Operations.cpp b/Engine/Plugins/Experimental/Mutable/Source/MutableRuntime/Private/MuR/Operations.cpp
--- a/Engine/Plugins/Experimental/Mutable/Source/MutableRuntime/Private/MuR/Operations.cpp
+++ b/Engine/Plugins/Experimental/Mutable/Source/MutableRuntime/Private/MuR/Operations.cpp
@@ -73,8 +73,6 @@ namespace mu
 
 		{ DT_COLOUR,		false,	{ 0, 0, 0, 0, 0, 0, 0, 0, 0, 0, 0, 0, 0, 0, 0, 0, 0, 0, 0, 0, 0, 0, 0, 0, 0, 0, 0, 0, 0 }		},	// CO_SAMPLEIMAGE
 		{ DT_COLOUR,		false,	{ 0, 0, 0, 0, 0, 0, 0, 0, 0, 0, 0, 0, 0, 0, 0, 0, 0, 0, 0, 0, 0, 0, 0, 0, 0, 0, 0, 0, 0 }		},	// CO_SWIZZLE
-		{ DT_COLOUR,		false,	{ 0, 0, 0, 0, 0, 0, 0, 0, 0, 0, 0, 0, 0, 0, 0, 0, 0, 0, 0, 0, 0, 0, 0, 0, 0, 0, 0, 0, 0 }		},	// CO_IMAGESIZE
-		{ DT_COLOUR,		false,	{ 0, 0, 0, 0, 0, 0, 0, 0, 0, 0, 0, 0, 0, 0, 0, 0, 0, 0, 0, 0, 0, 0, 0, 0, 0, 0, 0, 0, 0 }		},	// CO_LAYOUTBLOCKTRANSFORM
 		{ DT_COLOUR,		false,	{ 0, 0, 0, 0, 0, 0, 0, 0, 0, 0, 0, 0, 0, 0, 0, 0, 0, 0, 0, 0, 0, 0, 0, 0, 0, 0, 0, 0, 0 }		},	// CO_FROMSCALARS
 		{ DT_COLOUR,		false,	{ 0, 0, 0, 0, 0, 0, 0, 0, 0, 0, 0, 0, 0, 0, 0, 0, 0, 0, 0, 0, 0, 0, 0, 0, 0, 0, 0, 0, 0 }		},	// CO_ARITHMETIC
 
@@ -115,14 +113,12 @@ namespace mu
 		{ DT_MESH,			true,	{ 0, 0, 0, 0, 0, 0, 0, 0, 0, 0, 0, 0, 0, 0, 0, 0, 0, 0, 0, 0, 0, 0, 0, 0, 0, 0, 0, 0, 0 }		},	// ME_REMOVEMASK
 		{ DT_MESH,			true,	{ 0, 0, 0, 0, 0, 0, 0, 0, 0, 0, 0, 0, 0, 0, 0, 0, 0, 0, 0, 0, 0, 0, 0, 0, 0, 0, 0, 0, 0 }		},	// ME_FORMAT
 		{ DT_MESH,			true,	{ 0, 0, 0, 0, 0, 0, 0, 0, 0, 0, 0, 0, 0, 0, 0, 0, 0, 0, 0, 0, 0, 0, 0, 0, 0, 0, 0, 0, 0 }		},	// ME_EXTRACTLAYOUTBLOCK
-		{ DT_MESH,			true,	{ 0, 0, 0, 0, 0, 0, 0, 0, 0, 0, 0, 0, 0, 0, 0, 0, 0, 0, 0, 0, 0, 0, 0, 0, 0, 0, 0, 0, 0 }		},	// ME_EXTRACTFACEGROUP
 		{ DT_MESH,			true,	{ 0, 0, 0, 0, 0, 0, 0, 0, 0, 0, 0, 0, 0, 0, 0, 0, 0, 0, 0, 0, 0, 0, 0, 0, 0, 0, 0, 0, 0 }		},	// ME_TRANSFORM
 		{ DT_MESH,			true,	{ 0, 0, 0, 0, 0, 0, 0, 0, 0, 0, 0, 0, 0, 0, 0, 0, 0, 0, 0, 0, 0, 0, 0, 0, 0, 0, 0, 0, 0 }		},	// ME_CLIPMORPHPLANE
 		{ DT_MESH,			true,	{ 0, 0, 0, 0, 0, 0, 0, 0, 0, 0, 0, 0, 0, 0, 0, 0, 0, 0, 0, 0, 0, 0, 0, 0, 0, 0, 0, 0, 0 }		},	// ME_CLIPWITHMESH
 		{ DT_MESH,			true,	{ 0, 0, 0, 0, 0, 0, 0, 0, 0, 0, 0, 0, 0, 0, 0, 0, 0, 0, 0, 0, 0, 0, 0, 0, 0, 0, 0, 0, 0 }		},	// ME_SETSKELETON
 		{ DT_MESH,			true,	{ 0, 0, 0, 0, 0, 0, 0, 0, 0, 0, 0, 0, 0, 0, 0, 0, 0, 0, 0, 0, 0, 0, 0, 0, 0, 0, 0, 0, 0 }		},	// ME_PROJECT
 		{ DT_MESH,			true,	{ 0, 0, 0, 0, 0, 0, 0, 0, 0, 0, 0, 0, 0, 0, 0, 0, 0, 0, 0, 0, 0, 0, 0, 0, 0, 0, 0, 0, 0 }		},	// ME_APPLYPOSE
-		{ DT_MESH,			true,	{ 0, 0, 0, 0, 0, 0, 0, 0, 0, 0, 0, 0, 0, 0, 0, 0, 0, 0, 0, 0, 0, 0, 0, 0, 0, 0, 0, 0, 0 }		},	// ME_REMAPINDICES
 		{ DT_MESH,			true,	{ 0, 0, 0, 0, 0, 0, 0, 0, 0, 0, 0, 0, 0, 0, 0, 0, 0, 0, 0, 0, 0, 0, 0, 0, 0, 0, 0, 0, 0 }		},	// ME_GEOMETRYOPERATION
 		{ DT_MESH,			true,	{ 0, 0, 0, 0, 0, 0, 0, 0, 0, 0, 0, 0, 0, 0, 0, 0, 0, 0, 0, 0, 0, 0, 0, 0, 0, 0, 0, 0, 0 }		},	// ME_BINDSHAPE
 		{ DT_MESH,			true,	{ 0, 0, 0, 0, 0, 0, 0, 0, 0, 0, 0, 0, 0, 0, 0, 0, 0, 0, 0, 0, 0, 0, 0, 0, 0, 0, 0, 0, 0 }		},	// ME_APPLYSHAPE
@@ -230,14 +226,6 @@ namespace mu
             }
             break;
 
-        case OP_TYPE::CO_IMAGESIZE:
-            f(&op.args.ColourSampleImage.image );
-            break;
-
-        case OP_TYPE::CO_LAYOUTBLOCKTRANSFORM:
-            f(&op.args.ColourLayoutBlockTransform.layout );
-            break;
-
         case OP_TYPE::CO_FROMSCALARS:
 			for (int32 t = 0; t < MUTABLE_OP_MAX_SWIZZLE_CHANNELS; ++t)
 			{
@@ -345,10 +333,6 @@ namespace mu
             f(&op.args.MeshMaskDiff.fragment );
             break;
 
-        case OP_TYPE::ME_EXTRACTFACEGROUP:
-            f(&op.args.MeshExtractFaceGroup.source );
-            break;
-
         case OP_TYPE::ME_CLIPMORPHPLANE:
             f(&op.args.MeshClipMorphPlane.source);
             break;
@@ -545,20 +529,6 @@ namespace mu
             break;
         }
 
-        case OP_TYPE::CO_IMAGESIZE:
-        {
-			OP::ColourSampleImageArgs args = program.GetOpArgs<OP::ColourSampleImageArgs>(at);
-            f(args.image );
-            break;
-        }
-
-        case OP_TYPE::CO_LAYOUTBLOCKTRANSFORM:
-        {
-			OP::ColourLayoutBlockTransformArgs args = program.GetOpArgs<OP::ColourLayoutBlockTransformArgs>(at);
-            f(args.layout );
-            break;
-        }
-
         case OP_TYPE::CO_FROMSCALARS:
         {
 			OP::ColourFromScalarsArgs args = program.GetOpArgs<OP::ColourFromScalarsArgs>(at);
@@ -930,13 +900,6 @@ namespace mu
             break;
         }
 
-        case OP_TYPE::ME_EXTRACTFACEGROUP:
-        {
-			OP::MeshExtractFaceGroupArgs args = program.GetOpArgs<OP::MeshExtractFaceGroupArgs>(at);
-            f(args.source );
-            break;
-        }
-
         case OP_TYPE::ME_CLIPMORPHPLANE:
         {
 			OP::MeshClipMorphPlaneArgs args = program.GetOpArgs<OP::MeshClipMorphPlaneArgs>(at);
@@ -968,15 +931,7 @@ namespace mu
 
 			break;
 		}
- 	
-        case OP_TYPE::ME_REMAPINDICES :
-        {
-			OP::MeshRemapIndicesArgs args = program.GetOpArgs<OP::MeshRemapIndicesArgs>(at);
-            f(args.source);
-            f(args.reference);
-            break;
-        }
-
+ 
         case OP_TYPE::ME_SETSKELETON :
         {
 			OP::MeshSetSkeletonArgs args = program.GetOpArgs<OP::MeshSetSkeletonArgs>(at);
diff --git a/Engine/Plugins/Experimental/Mutable/Source/MutableRuntime/Private/MuR/Operations.h b/Engine/Plugins/Experimental/Mutable/Source/MutableRuntime/Private/MuR/Operations.h
--- a/Engine/Plugins/Experimental/Mutable/Source/MutableRuntime/Private/MuR/Operations.h
+++ b/Engine/Plugins/Experimental/Mutable/Source/MutableRuntime/Private/MuR/Operations.h
@@ -122,14 +122,6 @@ namespace mu
         //! Make a color by shuffling channels from other colours.
         CO_SWIZZLE,
 
-        //! Make a "color" from the size of an image. Colours are used as a generic vector here.
-		//! \TODO: Deprecated?
-		CO_IMAGESIZE,
-
-        //! Encode a layout block transformation in a vector: position = (x,y) size = ( z, w )
-		//! \TODO: Deprecated?
-		CO_LAYOUTBLOCKTRANSFORM,
-
         //! Compose a vector from 4 scalars
         CO_FROMSCALARS,
 
@@ -257,10 +249,6 @@ namespace mu
         //! Extract a fragment of a mesh containing a specific layout block.
         ME_EXTRACTLAYOUTBLOCK,
 
-        //! Extract a fragment of a mesh containing a specific face group.
-		//! \TODO: Deprecated?
-		ME_EXTRACTFACEGROUP,
-
         //! Apply a transform in a 4x4 matrix to the geometry channels of the mesh
         ME_TRANSFORM,
 
@@ -280,11 +268,6 @@ namespace mu
         //! Deform a skinned mesh applying a skeletal pose
         ME_APPLYPOSE,
 
-		//! Remap the vertex index buffer with the ones in the reference mesh, using the current
-		//! ones as relative to that mesh.
-		//! \TODO: Deprecated?
-		ME_REMAPINDICES,
-
 		//! Apply a geometry core operation to a mesh.
 		//! \TODO: Deprecated?
 		ME_GEOMETRYOPERATION,
@@ -501,18 +484,7 @@ namespace mu
             ADDRESS sources[ MUTABLE_OP_MAX_SWIZZLE_CHANNELS ];
         };
 
-        struct ColourImageSizeArgs
-        {
-            ADDRESS image;
-        };
-
-        struct ColourLayoutBlockTransformArgs
-        {
-            ADDRESS layout;
-            uint16 block;
-        };
-
-        struct ColourMultiplyArgs
+		struct ColourMultiplyArgs
         {
             ADDRESS a,b;
         };
@@ -914,13 +886,6 @@ namespace mu
 
         };
 
-
-        struct MeshExtractFaceGroupArgs
-        {
-            ADDRESS source;
-            uint32 group;
-        };
-
 		struct MeshTransformArgs
 		{
 			ADDRESS source;
@@ -1143,9 +1108,7 @@ namespace mu
             //-------------------------------------------------------------------------------------
             ColourSampleImageArgs ColourSampleImage;
             ColourSwizzleArgs ColourSwizzle;
-            ColourImageSizeArgs ColourImageSize;
             ColourMultiplyArgs ColourMultiply;
-            ColourLayoutBlockTransformArgs ColourLayoutBlockTransform;
             ColourFromScalarsArgs ColourFromScalars;
             ArithmeticArgs ColourArithmetic;
 
@@ -1173,7 +1136,6 @@ namespace mu
             MeshMergeArgs MeshMerge;
             MeshInterpolateArgs MeshInterpolate;
             MeshMaskDiffArgs MeshMaskDiff;
-            MeshExtractFaceGroupArgs MeshExtractFaceGroup;
 			MeshClipMorphPlaneArgs MeshClipMorphPlane;
             MeshClipWithMeshArgs MeshClipWithMesh;
             MeshSetSkeletonArgs MeshSetSkeleton;
diff --git a/Engine/Plugins/Experimental/Mutable/Source/MutableRuntime/Public/MuR/Model.h b/Engine/Plugins/Experimental/Mutable/Source/MutableRuntime/Public/MuR/Model.h
--- a/Engine/Plugins/Experimental/Mutable/Source/MutableRuntime/Public/MuR/Model.h
+++ b/Engine/Plugins/Experimental/Mutable/Source/MutableRuntime/Public/MuR/Model.h
@@ -13,7 +13,7 @@
 //! data structures. Compiled models are not necessarily compatible when the runtime is updated,
 //! so this version number can be used externally to verify this. It is not used internally, and
 //! serializing models from different versions than this runtime will probably result in a crash.
-#define MUTABLE_COMPILED_MODEL_CODE_VERSION		uint32( 75 )
+#define MUTABLE_COMPILED_MODEL_CODE_VERSION		uint32( 76 )
 #define MUTABLE_PARAMETERS_VERSION              uint32( 2 )
 
 
diff --git a/Engine/Plugins/Experimental/Mutable/Source/MutableTools/Private/MuT/AST.cpp b/Engine/Plugins/Experimental/Mutable/Source/MutableTools/Private/MuT/AST.cpp
--- a/Engine/Plugins/Experimental/Mutable/Source/MutableTools/Private/MuT/AST.cpp
+++ b/Engine/Plugins/Experimental/Mutable/Source/MutableTools/Private/MuT/AST.cpp
@@ -1210,11 +1210,6 @@ FImageDesc ASTOpFixed::GetImageDesc( bool returnBestOption, FGetImageDescContext
 		res = GetImageDesc(op.args.ImageInvert.base, returnBestOption, context);
 		break;
 
-    case OP_TYPE::CO_IMAGESIZE:
-        res = GetImageDesc( op.args.ColourImageSize.image, returnBestOption, context );
-        res.m_format = EImageFormat::IF_NONE;
-        break;
-
 	case OP_TYPE::IM_COLOURMAP:
 		res = GetImageDesc( op.args.ImageColourMap.base, returnBestOption, context );
 		break;
@@ -1542,8 +1537,6 @@ bool ASTOpFixed::IsColourConstant(FVector4f& colour) const
 
     case OP_TYPE::CO_SAMPLEIMAGE:
     case OP_TYPE::CO_SWIZZLE:
-    case OP_TYPE::CO_IMAGESIZE:
-    case OP_TYPE::CO_LAYOUTBLOCKTRANSFORM:
     case OP_TYPE::CO_FROMSCALARS:
     case OP_TYPE::CO_ARITHMETIC:
 
diff --git a/Engine/Plugins/Experimental/Mutable/Source/MutableTools/Private/MuT/ASTOpMeshRemapIndices.cpp b/Engine/Plugins/Experimental/Mutable/Source/MutableTools/Private/MuT/ASTOpMeshRemapIndices.cpp
deleted file mode 100644
diff --git a/Engine/Plugins/Experimental/Mutable/Source/MutableTools/Private/MuT/ASTOpMeshRemapIndices.h b/Engine/Plugins/Experimental/Mutable/Source/MutableTools/Private/MuT/ASTOpMeshRemapIndices.h
deleted file mode 100644
diff --git a/Engine/Plugins/Experimental/Mutable/Source/MutableTools/Private/MuT/CodeGenerator_Mesh.cpp b/Engine/Plugins/Experimental/Mutable/Source/MutableTools/Private/MuT/CodeGenerator_Mesh.cpp
--- a/Engine/Plugins/Experimental/Mutable/Source/MutableTools/Private/MuT/CodeGenerator_Mesh.cpp
+++ b/Engine/Plugins/Experimental/Mutable/Source/MutableTools/Private/MuT/CodeGenerator_Mesh.cpp
@@ -948,16 +948,9 @@ class Node;
                 }
             }
 
-            else if ( node.m_fragmentType==NodeMeshFragment::FT_FACE_GROUP )
+            else
             {
-				// \TODO: Deprecated?
-                Ptr<ASTOpFixed> op = new ASTOpFixed();
-                OutResult.meshOp = op;
-
-                op->op.type = OP_TYPE::ME_EXTRACTFACEGROUP;
-
-                op->SetChild( op->op.args.MeshExtractFaceGroup.source, BaseResult.meshOp );
-                op->op.args.MeshExtractFaceGroup.group = node.m_layoutOrGroup;
+				check(false);
             }
 
         }
diff --git a/Engine/Plugins/Experimental/Mutable/Source/MutableTools/Private/MuT/CodeOptimiser_Semantic.cpp b/Engine/Plugins/Experimental/Mutable/Source/MutableTools/Private/MuT/CodeOptimiser_Semantic.cpp
--- a/Engine/Plugins/Experimental/Mutable/Source/MutableTools/Private/MuT/CodeOptimiser_Semantic.cpp
+++ b/Engine/Plugins/Experimental/Mutable/Source/MutableTools/Private/MuT/CodeOptimiser_Semantic.cpp
@@ -1116,7 +1116,7 @@ namespace mu
 				// still an integer after relative scale
 				bool acceptable = false;
 				{
-					auto typedAt = dynamic_cast<const ASTOpImageCompose*>(at.get());
+					const ASTOpImageCompose* typedAt = dynamic_cast<const ASTOpImageCompose*>(at.get());
 					Ptr<ASTOp> originalBaseOp = typedAt->Base.child();
 
 					// \todo: recursion-proof cache?
@@ -1283,7 +1283,7 @@ namespace mu
 				break;
 			}
 
-            case OP_TYPE::IM_TRANSFORM:
+			case OP_TYPE::IM_TRANSFORM:
 			{
 				// It can only sink in the transform if it doesn't have it's own size.
 				const ASTOpImageTransform* TypedAt = dynamic_cast<const ASTOpImageTransform*>(at.get());
@@ -1378,8 +1378,8 @@ namespace mu
 				at = sourceOp;
                 break;
             }
-            
-            case OP_TYPE::IM_TRANSFORM:
+
+			case OP_TYPE::IM_TRANSFORM:
 			{
 				// Set the size in the children and remove resize
 				Ptr<ASTOpImageTransform> sourceOp = mu::Clone<ASTOpImageTransform>(sourceAt.get());
@@ -1656,7 +1656,7 @@ namespace mu
 
 			case OP_TYPE::IM_BLANKLAYOUT:
 			{
-				auto newOp = mu::Clone<ASTOpFixed>(sourceAt);
+				Ptr<ASTOpFixed> newOp = mu::Clone<ASTOpFixed>(sourceAt);
 
 				newOp->op.args.ImageBlankLayout.blockSize[0] =
 					uint16(newOp->op.args.ImageBlankLayout.blockSize[0]
@@ -1672,22 +1672,16 @@ namespace mu
 
 			case OP_TYPE::IM_PLAINCOLOUR:
 			{
-				auto newOp = mu::Clone<ASTOpFixed>(sourceAt);
+				Ptr<ASTOpFixed> newOp = mu::Clone<ASTOpFixed>(sourceAt);
 
-				newOp->op.args.ImagePlainColour.size[0] =
-					uint16(newOp->op.args.ImagePlainColour.size[0]
-						* op.args.ImageResizeRel.factor[0]
-						+ 0.5f);
-				newOp->op.args.ImagePlainColour.size[1] =
-					uint16(newOp->op.args.ImagePlainColour.size[1]
-						* op.args.ImageResizeRel.factor[1]
-						+ 0.5f);
+				newOp->op.args.ImagePlainColour.size[0] = FMath::CeilToInt( newOp->op.args.ImagePlainColour.size[0] * op.args.ImageResizeRel.factor[0] );
+				newOp->op.args.ImagePlainColour.size[1] = FMath::CeilToInt( newOp->op.args.ImagePlainColour.size[1] * op.args.ImageResizeRel.factor[1] );
 				newOp->op.args.ImagePlainColour.LODs = 1; // TODO
 				at = newOp;
 				break;
 			}
 
-            case OP_TYPE::IM_TRANSFORM:
+			case OP_TYPE::IM_TRANSFORM:
 			{
 				// We can only optimize here if we know the transform result size, otherwise, we will sink the op in the sinker.
 				const ASTOpImageTransform* typedAt = dynamic_cast<const ASTOpImageTransform*>(sourceAt.get());
diff --git a/Engine/Plugins/Experimental/Mutable/Source/MutableTools/Private/MuT/DataPacker.cpp b/Engine/Plugins/Experimental/Mutable/Source/MutableTools/Private/MuT/DataPacker.cpp
--- a/Engine/Plugins/Experimental/Mutable/Source/MutableTools/Private/MuT/DataPacker.cpp
+++ b/Engine/Plugins/Experimental/Mutable/Source/MutableTools/Private/MuT/DataPacker.cpp
@@ -619,22 +619,6 @@ namespace mu
                 break;
             }
 
-			case OP_TYPE::ME_EXTRACTFACEGROUP:
-			{
-				recurse = false;
-
-				const ASTOpFixed* op = dynamic_cast<const ASTOpFixed*>(node.get());
-
-				// todo: check if we really need all of them
-				uint64_t newState = currentSemantics;
-				newState |= (UINT64_C(1) << MBS_LAYOUTBLOCK);
-				newState |= (UINT64_C(1) << MBS_CHART);
-				newState |= (UINT64_C(1) << MBS_VERTEXINDEX);
-
-				RecurseWithState(op->children[op->op.args.MeshExtractFaceGroup.source].child(), newState);
-				break;
-			}
-
 			case OP_TYPE::LA_FROMMESH:
 			{
 				recurse = false;
diff --git a/Engine/Plugins/Experimental/Mutable/Source/MutableTools/Private/MuT/ErrorLog.cpp b/Engine/Plugins/Experimental/Mutable/Source/MutableTools/Private/MuT/ErrorLog.cpp
--- a/Engine/Plugins/Experimental/Mutable/Source/MutableTools/Private/MuT/ErrorLog.cpp
+++ b/Engine/Plugins/Experimental/Mutable/Source/MutableTools/Private/MuT/ErrorLog.cpp
@@ -202,7 +202,7 @@ namespace mu
     
     // clang-format off
 
-    const TCHAR* s_opNames[ size_t(OP_TYPE::COUNT) ] =
+    const TCHAR* s_opNames[ int32(OP_TYPE::COUNT) ] =
 	{
 		TEXT("NONE             "),
 
@@ -257,8 +257,6 @@ namespace mu
 		
 		TEXT("CO_SAMPLEIMAGE   "),
 		TEXT("CO_SWIZZLE       "),
-		TEXT("CO_IMAGESIZE     "),
-		TEXT("CO_LAYOUTBLOCKTR "),
 		TEXT("CO_FROMSCALARS   "),
 		TEXT("CO_ARITHMETIC    "),
 
@@ -299,14 +297,12 @@ namespace mu
         TEXT("ME_REMOVEMASK    "),
         TEXT("ME_FORMAT        "),
         TEXT("ME_EXTRACTLABLOCK"),
-        TEXT("ME_EXTRACTFACEGRP"),
 		TEXT("ME_TRANSFORM     "),
 		TEXT("ME_CLIPMORPHPLANE"),
         TEXT("ME_CLIPWITHMESH  "),
         TEXT("ME_SETSKELETON   "),
         TEXT("ME_PROJECT       "),
         TEXT("ME_APPLYPOSE     "),
-        TEXT("ME_REMAPINDICES  "),
 		TEXT("ME_GEOMETRYOP	   "),
 		TEXT("ME_BINDSHAPE	   "),
 		TEXT("ME_APPLYSHAPE	   "),
@@ -330,6 +326,8 @@ namespace mu
 		TEXT("LA_FROMMESH	   "),
 	};
 
+	static_assert(sizeof(s_opNames) / sizeof(void*) == int32(OP_TYPE::COUNT));
+
     // clang-format on
 
     //---------------------------------------------------------------------------------------------
diff --git a/Engine/Plugins/Experimental/Mutable/Source/MutableTools/Public/MuT/NodeMeshFragment.h b/Engine/Plugins/Experimental/Mutable/Source/MutableTools/Public/MuT/NodeMeshFragment.h
--- a/Engine/Plugins/Experimental/Mutable/Source/MutableTools/Public/MuT/NodeMeshFragment.h
+++ b/Engine/Plugins/Experimental/Mutable/Source/MutableTools/Public/MuT/NodeMeshFragment.h
@@ -53,7 +53,7 @@ namespace mu
 
         typedef enum {
             FT_LAYOUT_BLOCKS = 0,
-            FT_FACE_GROUP
+            FT_FACE_GROUP_DEPRECATED
         } FRAGMENT_TYPE;
 
         //! Get the type of mesh fragment to extract.
