[Mutable] Fix generation of wrongly sized textures when a size limitation is applied on top of the image and streaming is disabled.

#rnx

diff --git a/Engine/Plugins/Experimental/Mutable/Source/CustomizableObject/Classes/MuCO/CustomizableObjectSystem.h b/Engine/Plugins/Experimental/Mutable/Source/CustomizableObject/Classes/MuCO/CustomizableObjectSystem.h
--- a/Engine/Plugins/Experimental/Mutable/Source/CustomizableObject/Classes/MuCO/CustomizableObjectSystem.h
+++ b/Engine/Plugins/Experimental/Mutable/Source/CustomizableObject/Classes/MuCO/CustomizableObjectSystem.h
@@ -102,6 +102,9 @@ struct FMutableImageReference
 	uint8 LOD = 0;
 	uint8 Component = 0;
 	uint8 Image = 0;
+
+	/** */
+	uint8 BaseMip = 0;
 };
 
 
diff --git a/Engine/Plugins/Experimental/Mutable/Source/CustomizableObject/Private/MuCO/CustomizableObjectInstance.cpp b/Engine/Plugins/Experimental/Mutable/Source/CustomizableObject/Private/MuCO/CustomizableObjectInstance.cpp
--- a/Engine/Plugins/Experimental/Mutable/Source/CustomizableObject/Private/MuCO/CustomizableObjectInstance.cpp
+++ b/Engine/Plugins/Experimental/Mutable/Source/CustomizableObject/Private/MuCO/CustomizableObjectInstance.cpp
@@ -5846,6 +5846,7 @@ void UCustomizableInstancePrivateData::BuildMaterials(const TSharedPtr<FMutableO
 														MutableMipDataProviderFactory->ImageRef.LOD = uint8(Image.BaseLOD);
 														MutableMipDataProviderFactory->ImageRef.Component = uint8(ComponentIndex);
 														MutableMipDataProviderFactory->ImageRef.Image = uint8(ImageIndex);
+														MutableMipDataProviderFactory->ImageRef.BaseMip = uint8(Image.BaseMip);
 														MutableMipDataProviderFactory->UpdateContext = UpdateContext;
 														MutableTexture->AddAssetUserData(MutableMipDataProviderFactory);
 													}
diff --git a/Engine/Plugins/Experimental/Mutable/Source/CustomizableObject/Private/MuCO/CustomizableObjectMipDataProvider.cpp b/Engine/Plugins/Experimental/Mutable/Source/CustomizableObject/Private/MuCO/CustomizableObjectMipDataProvider.cpp
--- a/Engine/Plugins/Experimental/Mutable/Source/CustomizableObject/Private/MuCO/CustomizableObjectMipDataProvider.cpp
+++ b/Engine/Plugins/Experimental/Mutable/Source/CustomizableObject/Private/MuCO/CustomizableObjectMipDataProvider.cpp
@@ -153,31 +153,12 @@ namespace impl
 				// runtime cannot remember all the resources it has built, and only remembers a fixed amount.
 				mu::FResourceID MipImageID = Instance->GetImageId(ImageRef.LOD, ImageRef.Component, SurfaceIndex, ImageRef.Image);
 
-				UCustomizableObjectSystem* COSystem = UCustomizableObjectSystem::GetInstance();
-				FCustomizableObjectSystemPrivate* SystemPrivate = COSystem ? COSystem->GetPrivate() : nullptr;
-				int32 MaxTextureSizeToGenerate = SystemPrivate ? SystemPrivate->MaxTextureSizeToGenerate : 0;
-
-				int32 ExtraMipsToSkip = 0;
-				
-				if (MaxTextureSizeToGenerate > 0)
-				{
-					// TODO: Why do we need to do this again? The full size should be stored in the initial image creation.
-					mu::FImageDesc ImageDesc;
-					System->GetImageDesc(InstanceID, MipImageID, ImageDesc);
-
-					uint16 MaxSize = FMath::Max(ImageDesc.m_size[0], ImageDesc.m_size[1]);
-
-					if (MaxSize > MaxTextureSizeToGenerate)
-					{
-						ExtraMipsToSkip = FMath::CeilLogTwo(MaxSize / MaxTextureSizeToGenerate);
-					}
-				}
 
 				mu::ImagePtrConst Image;
 				{
 					MUTABLE_CPUPROFILER_SCOPE(GetImage);
 
-					Image = System->GetImage(InstanceID, MipImageID, OperationData->MipsToSkip + ExtraMipsToSkip, ImageRef.LOD);
+					Image = System->GetImage(InstanceID, MipImageID, ImageRef.BaseMip + OperationData->MipsToSkip, ImageRef.LOD);
 				}
 
 				check(Image);
diff --git a/Engine/Plugins/Experimental/Mutable/Source/CustomizableObject/Private/MuCO/CustomizableObjectSystem.cpp b/Engine/Plugins/Experimental/Mutable/Source/CustomizableObject/Private/MuCO/CustomizableObjectSystem.cpp
--- a/Engine/Plugins/Experimental/Mutable/Source/CustomizableObject/Private/MuCO/CustomizableObjectSystem.cpp
+++ b/Engine/Plugins/Experimental/Mutable/Source/CustomizableObject/Private/MuCO/CustomizableObjectSystem.cpp
@@ -1331,6 +1331,7 @@ namespace impl
 							Image.FullImageSizeX = 0;
 							Image.FullImageSizeY = 0;
 							Image.BaseLOD = BaseLODIndex;
+							Image.BaseMip = 0;
 						}
 
 						// Vectors
@@ -1425,7 +1426,10 @@ namespace impl
 
 				if (MaxTextureSizeToGenerate > 0 && MaxSize > MaxTextureSizeToGenerate)
 				{
-					Reduction = MaxSize / MaxTextureSizeToGenerate;
+					// Find the reduction factor, and the BaseMip of the texture.
+					const uint32 NextPowerOfTwo = FMath::RoundUpToPowerOfTwo(FMath::DivideAndRoundUp(MaxSize, MaxTextureSizeToGenerate));
+					Reduction = FMath::Max(NextPowerOfTwo, 2U); // At least divide the texture by a factor of two
+					Image.BaseMip = FMath::FloorLog2(Reduction);
 				}
 
 				Image.FullImageSizeX = ImageDesc.m_size[0] / Reduction;
@@ -1462,7 +1466,7 @@ namespace impl
 				}
 				else
 				{
-					Image.Image = System->GetImage(OperationData->InstanceID, Image.ImageID, MipsToSkip, Image.BaseLOD);
+					Image.Image = System->GetImage(OperationData->InstanceID, Image.ImageID, Image.BaseMip + MipsToSkip, Image.BaseLOD);
 				}
 
 				check(Image.Image);
diff --git a/Engine/Plugins/Experimental/Mutable/Source/CustomizableObject/Private/MuCO/CustomizableObjectSystemPrivate.h b/Engine/Plugins/Experimental/Mutable/Source/CustomizableObject/Private/MuCO/CustomizableObjectSystemPrivate.h
--- a/Engine/Plugins/Experimental/Mutable/Source/CustomizableObject/Private/MuCO/CustomizableObjectSystemPrivate.h
+++ b/Engine/Plugins/Experimental/Mutable/Source/CustomizableObject/Private/MuCO/CustomizableObjectSystemPrivate.h
@@ -417,6 +417,7 @@ struct FInstanceUpdateData
 		
 		// LOD of the ImageId. If the texture is shared between LOD, first LOD where this image can be found. 
 		int32 BaseLOD;
+		int32 BaseMip;
 		
 		uint16 FullImageSizeX, FullImageSizeY;
 		mu::Ptr<const mu::Image> Image;
