[Mutable] Fix an FString serialization issue causing double zero terminating characters when constructing it from a TArray<TCHAR>.

#rnx

diff --git a/Engine/Plugins/Experimental/Mutable/Source/CustomizableObject/Private/MuCO/UnrealMutableImageProvider.cpp b/Engine/Plugins/Experimental/Mutable/Source/CustomizableObject/Private/MuCO/UnrealMutableImageProvider.cpp
--- a/Engine/Plugins/Experimental/Mutable/Source/CustomizableObject/Private/MuCO/UnrealMutableImageProvider.cpp
+++ b/Engine/Plugins/Experimental/Mutable/Source/CustomizableObject/Private/MuCO/UnrealMutableImageProvider.cpp
@@ -469,7 +469,7 @@ void FUnrealMutableImageProvider::CacheImage(FName Id, bool bUser)
 {
 	check( IsInGameThread() );
 
-	if (Id == FName())
+	if (Id == NAME_None)
 	{
 		return;	
 	}	
@@ -564,7 +564,7 @@ void FUnrealMutableImageProvider::UnCacheImage(FName Id, bool bUser)
 	// TODO: Review GM
 	//check(IsInGameThread());
 
-	if (Id == FName())
+	if (Id == NAME_None)
 	{
 		return;	
 	}
diff --git a/Engine/Plugins/Experimental/Mutable/Source/MutableRuntime/Private/MuR/Serialisation.cpp b/Engine/Plugins/Experimental/Mutable/Source/MutableRuntime/Private/MuR/Serialisation.cpp
--- a/Engine/Plugins/Experimental/Mutable/Source/MutableRuntime/Private/MuR/Serialisation.cpp
+++ b/Engine/Plugins/Experimental/Mutable/Source/MutableRuntime/Private/MuR/Serialisation.cpp
@@ -53,8 +53,8 @@ namespace mu
     {
 	    TArray<TCHAR> Data;
     	arch >> Data;
-    	
-    	t = FString(Data);
+
+    	t = FString(Data.GetData()); // Construct from raw pointer to avoid double zero terminating character
     }
 
     //---------------------------------------------------------------------------------------------
