[Mutable] fix FImageSwizzleTask_Prepare CPUPROFILER_SCOPE name.

#rnx

diff --git a/Engine/Plugins/Experimental/Mutable/Source/MutableRuntime/Private/MuR/CodeRunnerTasked.cpp b/Engine/Plugins/Experimental/Mutable/Source/MutableRuntime/Private/MuR/CodeRunnerTasked.cpp
--- a/Engine/Plugins/Experimental/Mutable/Source/MutableRuntime/Private/MuR/CodeRunnerTasked.cpp
+++ b/Engine/Plugins/Experimental/Mutable/Source/MutableRuntime/Private/MuR/CodeRunnerTasked.cpp
@@ -1378,7 +1378,7 @@ namespace mu
 	bool FImageSwizzleTask::Prepare(CodeRunner* Runner, bool& bOutFailed)
 	{
 		// This runs in the mutable Runner thread
-		MUTABLE_CPUPROFILER_SCOPE(FImageLayerColourTask_Prepare);
+		MUTABLE_CPUPROFILER_SCOPE(FImageSwizzleTask_Prepare);
 		bOutFailed = false;
 
 		for (int i = 0; i < MUTABLE_OP_MAX_SWIZZLE_CHANNELS; ++i)
