[mutable] Fix error in 8x8 rgba compression mistakenly using 6x6 blocks.
#rnx

diff --git a/Engine/Plugins/Experimental/Mutable/Source/MutableRuntime/Private/MuR/BlockCompression/Miro/Miro.cpp b/Engine/Plugins/Experimental/Mutable/Source/MutableRuntime/Private/MuR/BlockCompression/Miro/Miro.cpp
--- a/Engine/Plugins/Experimental/Mutable/Source/MutableRuntime/Private/MuR/BlockCompression/Miro/Miro.cpp
+++ b/Engine/Plugins/Experimental/Mutable/Source/MutableRuntime/Private/MuR/BlockCompression/Miro/Miro.cpp
@@ -8186,7 +8186,7 @@ namespace miro
 
 	void RGBA_to_ASTC8x8RGBL(uint32 sx, uint32 sy, const uint8* from, uint8* to, int32 Quality)
 	{
-		Generic_RGBA_to_ASTCRGBL<6>(sx, sy, from, to, Quality);
+		Generic_RGBA_to_ASTCRGBL<8>(sx, sy, from, to, Quality);
 	}
 
 	void ASTC8x8RGBL_to_RGB(uint32 sx, uint32 sy, const uint8* from, uint8* to)
