#rnx

diff --git a/Engine/Plugins/Experimental/Mutable/Source/CustomizableObject/Private/MuCO/CustomizableInstanceLODManagement.cpp b/Engine/Plugins/Experimental/Mutable/Source/CustomizableObject/Private/MuCO/CustomizableInstanceLODManagement.cpp
--- a/Engine/Plugins/Experimental/Mutable/Source/CustomizableObject/Private/MuCO/CustomizableInstanceLODManagement.cpp
+++ b/Engine/Plugins/Experimental/Mutable/Source/CustomizableObject/Private/MuCO/CustomizableInstanceLODManagement.cpp
@@ -324,6 +324,20 @@ void UCustomizableInstanceLODManagement::UpdateInstanceDistsAndLODs(FMutableInst
 				{
 					COI->SetIsBeingUsedByComponentInPlay(true);
 
+					int32 ComponentIndex = CustomizableSkeletalComponent->ComponentIndex;
+
+#if WITH_EDITOR
+					// If the instance is generated but the component doesn't have a mesh, set it.
+					// Can happen when duplicating instances in the editor.
+					if (COI->SkeletalMeshes.IsValidIndex(ComponentIndex) && 
+						Parent->GetSkeletalMeshAsset() == nullptr &&
+						COI->SkeletalMeshes[ComponentIndex])
+					{
+						// As the instance is already generated, this will be very fast and just set the mesh and call the delegates
+						COI->UpdateSkeletalMeshAsync();
+					}
+#endif
+
 					// If it's the local player set max priority
 					const USceneComponent* const AttachParentParentComponent = Parent->GetAttachParent();
 					AActor* ParentParentActor = AttachParentParentComponent ? AttachParentParentComponent->GetOwner() : nullptr;
@@ -343,9 +357,9 @@ void UCustomizableInstanceLODManagement::UpdateInstanceDistsAndLODs(FMutableInst
 					LODTracker.MinLOD = FMath::Min(LODTracker.MinLOD, Parent->bOverrideMinLod ? Parent->MinLodModel : 0);
 
 					// If the parent component have a SkeletalMesh use the RequestedLODLevel of the component as reference to know which LODs mutable should generate.
-					if (UE_MUTABLE_GETSKELETALMESHASSET(Parent) && LODTracker.RequestedLODsPerComponent.IsValidIndex(CustomizableSkeletalComponent->ComponentIndex))
+					if (UE_MUTABLE_GETSKELETALMESHASSET(Parent) && LODTracker.RequestedLODsPerComponent.IsValidIndex(ComponentIndex))
 					{
-						LODTracker.RequestedLODsPerComponent[CustomizableSkeletalComponent->ComponentIndex] |= 1 << Parent->GetPredictedLODLevel();
+						LODTracker.RequestedLODsPerComponent[ComponentIndex] |= 1 << Parent->GetPredictedLODLevel();
 					}
 				}
 			}
