[Mutable] Optimize image fill color operations.

#rnx

diff --git a/Engine/Plugins/Experimental/Mutable/Source/MutableRuntime/Private/MuR/Image.cpp b/Engine/Plugins/Experimental/Mutable/Source/MutableRuntime/Private/MuR/Image.cpp
--- a/Engine/Plugins/Experimental/Mutable/Source/MutableRuntime/Private/MuR/Image.cpp
+++ b/Engine/Plugins/Experimental/Mutable/Source/MutableRuntime/Private/MuR/Image.cpp
@@ -1201,59 +1201,114 @@ namespace mu
 		{
 		case EImageFormat::IF_RGB_UBYTE:
 		{
-			// TODO: Optimize: don't write bytes one by one
-			int pixelCount = Target->CalculatePixelCount();
+			int32 PixelCount = Target->CalculatePixelCount();
 			uint8* pData = Target->GetData();
-			uint8 r = uint8(FMath::Clamp(255.0f * Color[0], 0.0f, 255.0f));
-			uint8 g = uint8(FMath::Clamp(255.0f * Color[1], 0.0f, 255.0f));
-			uint8 b = uint8(FMath::Clamp(255.0f * Color[2], 0.0f, 255.0f));
-			for (int p = 0; p < pixelCount; ++p)
+			uint32 R = uint32(FMath::Clamp(255.0f * Color[0], 0.0f, 255.0f));
+			uint32 G = uint32(FMath::Clamp(255.0f * Color[1], 0.0f, 255.0f));
+			uint32 B = uint32(FMath::Clamp(255.0f * Color[2], 0.0f, 255.0f));
+			
+			const uint32 PixelData = R | (G << 8) | (B << 16);
+
+			if (PixelData == 0)
+			{
+				FMemory::Memzero(pData, PixelCount*3);
+				break;
+			}
+
+			if ((R == G) & (R == B))
+			{
+				FMemory::Memset(pData, static_cast<uint8>(R), PixelCount*3);
+				break;
+			}
+
+			const uint64 TwoPixelsData = (uint64(PixelData) << 24) | PixelData; 
+
+			for (int32 P = 0; P < PixelCount >> 1; ++P)
 			{
-				pData[0] = r;
-				pData[1] = g;
-				pData[2] = b;
-				pData += 3;
+				FMemory::Memcpy(&pData[P * 6], &TwoPixelsData, 6);
 			}
+
+			if (PixelCount & 1)
+			{
+				FMemory::Memcpy(&pData[(PixelCount >> 1) * 6], &PixelData, 3);
+			}
+
 			break;
 		}
 
 		case EImageFormat::IF_RGBA_UBYTE:
 		{
-			// TODO: Optimize: don't write bytes one by one
-			int pixelCount = Target->CalculatePixelCount();
+			int32 PixelCount = Target->CalculatePixelCount();
 			uint8* pData = Target->GetData();
-			uint8 r = uint8(FMath::Clamp(255.0f * Color[0], 0.0f, 255.0f));
-			uint8 g = uint8(FMath::Clamp(255.0f * Color[1], 0.0f, 255.0f));
-			uint8 b = uint8(FMath::Clamp(255.0f * Color[2], 0.0f, 255.0f));
-			uint8 a = uint8(FMath::Clamp(255.0f * Color[3], 0.0f, 255.0f));
-			for (int p = 0; p < pixelCount; ++p)
+			uint32 R = uint32(FMath::Clamp(255.0f * Color[0], 0.0f, 255.0f));
+			uint32 G = uint32(FMath::Clamp(255.0f * Color[1], 0.0f, 255.0f));
+			uint32 B = uint32(FMath::Clamp(255.0f * Color[2], 0.0f, 255.0f));
+			uint32 A = uint32(FMath::Clamp(255.0f * Color[3], 0.0f, 255.0f));
+		
+			const uint32 PixelData = R | (G << 8) | (B << 16) | (A << 24);  
+
+			if (PixelData == 0)
+			{
+				FMemory::Memzero(pData, PixelCount*4);
+				break;
+			}
+
+			if ((R == G) & (R == B) & (R == A))
+			{
+				FMemory::Memset(pData, static_cast<uint8>(R), PixelCount*4);
+				break;
+			}
+
+			const uint64 TwoPixelsData = (uint64(PixelData) << 32) | PixelData; 
+
+			for (int32 P = 0; P < (PixelCount >> 1); ++P)
+			{
+				FMemory::Memcpy(&pData[P * 8], &TwoPixelsData, 8);
+			}
+
+			if (PixelCount & 1)
 			{
-				pData[0] = r;
-				pData[1] = g;
-				pData[2] = b;
-				pData[3] = a;
-				pData += 4;
+				FMemory::Memcpy(&pData[(PixelCount >> 1) * 8], &PixelData, 4);
 			}
+
 			break;
 		}
 
 		case EImageFormat::IF_BGRA_UBYTE:
 		{
-			// TODO: Optimize: don't write bytes one by one
-			int pixelCount = Target->CalculatePixelCount();
+			int32 PixelCount = Target->CalculatePixelCount();
 			uint8* pData = Target->GetData();
-			uint8 r = uint8(FMath::Clamp(255.0f * Color[0], 0.0f, 255.0f));
-			uint8 g = uint8(FMath::Clamp(255.0f * Color[1], 0.0f, 255.0f));
-			uint8 b = uint8(FMath::Clamp(255.0f * Color[2], 0.0f, 255.0f));
-			uint8 a = uint8(FMath::Clamp(255.0f * Color[3], 0.0f, 255.0f));
-			for (int p = 0; p < pixelCount; ++p)
+			uint32 R = uint32(FMath::Clamp(255.0f * Color[0], 0.0f, 255.0f));
+			uint32 G = uint32(FMath::Clamp(255.0f * Color[1], 0.0f, 255.0f));
+			uint32 B = uint32(FMath::Clamp(255.0f * Color[2], 0.0f, 255.0f));
+			uint32 A = uint32(FMath::Clamp(255.0f * Color[3], 0.0f, 255.0f));
+	
+			const uint32 PixelData = R | (G << 8) | (B << 16) | (A << 24);  
+
+			if (PixelData == 0)
+			{
+				FMemory::Memzero(pData, PixelCount*4);
+				break;
+			}
+
+			if ((R == G) & (R == B) & (R == A))
+			{
+				FMemory::Memset(pData, static_cast<uint8>(R), PixelCount*4);
+				break;
+			}
+
+			const uint64 TwoPixelsData = (uint64(PixelData) << 32) | PixelData; 
+
+			for (int32 P = 0; P < (PixelCount >> 1); ++P)
+			{
+				FMemory::Memcpy(&pData[P * 8], &TwoPixelsData, 8);
+			}
+
+			if (PixelCount & 1)
 			{
-				pData[0] = b;
-				pData[1] = g;
-				pData[2] = r;
-				pData[3] = a;
-				pData += 4;
+				FMemory::Memcpy(&pData[(PixelCount >> 1) * 8], &PixelData, 4);
 			}
+
 			break;
 		}
 
