[Mutable] Deduplicate Variation Nodes.

- Deduplicated code from all Variation Nodes.
- Simplified Switch Nodes code.
- No longer required a hardcoded list of all nodes.
- Minor code cleanup.

#rnx

diff --git a/Engine/Plugins/Experimental/Mutable/Source/CustomizableObject/Public/MuCO/CustomizableObjectCustomVersion.h b/Engine/Plugins/Experimental/Mutable/Source/CustomizableObject/Public/MuCO/CustomizableObjectCustomVersion.h
--- a/Engine/Plugins/Experimental/Mutable/Source/CustomizableObject/Public/MuCO/CustomizableObjectCustomVersion.h
+++ b/Engine/Plugins/Experimental/Mutable/Source/CustomizableObject/Public/MuCO/CustomizableObjectCustomVersion.h
@@ -92,6 +92,8 @@ struct CUSTOMIZABLEOBJECT_API FCustomizableObjectCustomVersion
 
 		MeshReshapeVertexColorUsageSelection,
 
+		DeduplicateNodeVariant,
+
 		// -----<new versions can be added above this line>-------------------------------------------------
 		VersionPlusOne,
 		LatestVersion = VersionPlusOne - 1
diff --git a/Engine/Plugins/Experimental/Mutable/Source/CustomizableObjectEditor/Classes/MuCOE/EdGraphSchema_CustomizableObject.h b/Engine/Plugins/Experimental/Mutable/Source/CustomizableObjectEditor/Classes/MuCOE/EdGraphSchema_CustomizableObject.h
--- a/Engine/Plugins/Experimental/Mutable/Source/CustomizableObjectEditor/Classes/MuCOE/EdGraphSchema_CustomizableObject.h
+++ b/Engine/Plugins/Experimental/Mutable/Source/CustomizableObjectEditor/Classes/MuCOE/EdGraphSchema_CustomizableObject.h
@@ -87,8 +87,6 @@ class CUSTOMIZABLEOBJECTEDITOR_API UEdGraphSchema_CustomizableObject : public UE
 public:
 	GENERATED_BODY()
 
-	UEdGraphSchema_CustomizableObject();
-
 	// Allowable PinType.PinCategory values
 	static const FName PC_Object;
 	static const FName PC_Material;
@@ -106,12 +104,6 @@ public:
 	static const FName PC_MaterialAsset;
 	static const FName PC_Wildcard;
 
-private:
-	/** List of all node types. */
-	TArray<UClass*> NodeTypes;
-	
-public:
-
 	// EdGraphSchema interface
 	virtual void GetAssetsGraphHoverMessage(const TArray<FAssetData>& Assets, const UEdGraph* HoverGraph, FString& OutTooltipText, bool& OutOkIcon) const override;
 	virtual void GetGraphContextActions(FGraphContextMenuBuilder& ContextMenuBuilder) const override;
diff --git a/Engine/Plugins/Experimental/Mutable/Source/CustomizableObjectEditor/Classes/MuCOE/Nodes/CustomizableObjectNodeColorSwitch.h b/Engine/Plugins/Experimental/Mutable/Source/CustomizableObjectEditor/Classes/MuCOE/Nodes/CustomizableObjectNodeColorSwitch.h
--- a/Engine/Plugins/Experimental/Mutable/Source/CustomizableObjectEditor/Classes/MuCOE/Nodes/CustomizableObjectNodeColorSwitch.h
+++ b/Engine/Plugins/Experimental/Mutable/Source/CustomizableObjectEditor/Classes/MuCOE/Nodes/CustomizableObjectNodeColorSwitch.h
@@ -6,9 +6,6 @@
 
 #include "CustomizableObjectNodeColorSwitch.generated.h"
 
-class FArchive;
-class UObject;
-
 
 UCLASS()
 class CUSTOMIZABLEOBJECTEDITOR_API UCustomizableObjectNodeColorSwitch : public UCustomizableObjectNodeSwitchBase
@@ -16,15 +13,10 @@ class CUSTOMIZABLEOBJECTEDITOR_API UCustomizableObjectNodeColorSwitch : public U
 public:
 	GENERATED_BODY()
 
-	// UObject interface
-	virtual void Serialize(FArchive& Ar) override;
-	
-	// UCustomizableObjectNodeSwitchBase interface
-	FString GetOutputPinName() const override;
-
-	FName GetCategory() const override;
+	// UCustomizableObjectNode interface
+	virtual void BackwardsCompatibleFixup() override;
 
-	// Own interface
-	FString GetPinPrefix() const;
+	// UCustomizableObjectNodeSwitchBase interface
+	virtual FName GetCategory() const override;
 };
 
diff --git a/Engine/Plugins/Experimental/Mutable/Source/CustomizableObjectEditor/Classes/MuCOE/Nodes/CustomizableObjectNodeColorVariation.h b/Engine/Plugins/Experimental/Mutable/Source/CustomizableObjectEditor/Classes/MuCOE/Nodes/CustomizableObjectNodeColorVariation.h
--- a/Engine/Plugins/Experimental/Mutable/Source/CustomizableObjectEditor/Classes/MuCOE/Nodes/CustomizableObjectNodeColorVariation.h
+++ b/Engine/Plugins/Experimental/Mutable/Source/CustomizableObjectEditor/Classes/MuCOE/Nodes/CustomizableObjectNodeColorVariation.h
@@ -2,71 +2,27 @@
 
 #pragma once
 
-#include "MuCOE/Nodes/CustomizableObjectNode.h"
+#include "MuCOE/Nodes/CustomizableObjectNodeVariation.h"
 
 #include "CustomizableObjectNodeColorVariation.generated.h"
 
-namespace ENodeTitleType { enum Type : int; }
-
-class UCustomizableObjectNodeRemapPins;
-class UEdGraphPin;
-class UObject;
-struct FPropertyChangedEvent;
-
-
-USTRUCT()
-struct CUSTOMIZABLEOBJECTEDITOR_API FCustomizableObjectColorVariation
-{
-	GENERATED_USTRUCT_BODY()
-
-	UPROPERTY(EditAnywhere, Category = CustomizableObject)
-	FString Tag;
-};
+struct FCustomizableObjectColorVariation;
 
 
 UCLASS()
-class CUSTOMIZABLEOBJECTEDITOR_API UCustomizableObjectNodeColorVariation : public UCustomizableObjectNode
+class CUSTOMIZABLEOBJECTEDITOR_API UCustomizableObjectNodeColorVariation : public UCustomizableObjectNodeVariation
 {
 public:
 	GENERATED_BODY()
 
-	UCustomizableObjectNodeColorVariation();
-
-	UPROPERTY(EditAnywhere, Category = CustomizableObject)
-	TArray<FCustomizableObjectColorVariation> Variations;
-
-	// UObject interface.
-	void PostEditChangeProperty(FPropertyChangedEvent& PropertyChangedEvent) override;
-
-	// EdGraphNode interface
-	FText GetNodeTitle(ENodeTitleType::Type TitleType) const override;
-	FLinearColor GetNodeTitleColor() const override;
-	FText GetTooltipText() const override;
+	// Deprecated properties
+	UPROPERTY()
+	TArray<FCustomizableObjectColorVariation> Variations_DEPRECATED;
 
 	// UCustomizableObjectNode interface
-	void AllocateDefaultPins(UCustomizableObjectNodeRemapPins* RemapPins) override;
-
-	UEdGraphPin* OutputPin() const
-	{
-		return FindPin(TEXT("Color"));
-	}
-
-	UEdGraphPin* DefaultPin() const
-	{
-		return FindPin(TEXT("Default"));
-	}
-
-	int32 GetNumVariations() const
-	{
-		return Variations.Num();
-	}
-
-	UEdGraphPin* VariationPin(int Index) const
-	{
-		return FindPin(FString::Printf(TEXT("Variation %d"),Index));
-	}
-
-private:
-	TMap<UEdGraphPin*, UEdGraphPin*> RemapPinsMap;
+	virtual void BackwardsCompatibleFixup() override;
+	
+	// UCustomizableObjectNodeVariation interface
+	virtual FName GetCategory() const override;
 };
 
diff --git a/Engine/Plugins/Experimental/Mutable/Source/CustomizableObjectEditor/Classes/MuCOE/Nodes/CustomizableObjectNodeFloatSwitch.h b/Engine/Plugins/Experimental/Mutable/Source/CustomizableObjectEditor/Classes/MuCOE/Nodes/CustomizableObjectNodeFloatSwitch.h
--- a/Engine/Plugins/Experimental/Mutable/Source/CustomizableObjectEditor/Classes/MuCOE/Nodes/CustomizableObjectNodeFloatSwitch.h
+++ b/Engine/Plugins/Experimental/Mutable/Source/CustomizableObjectEditor/Classes/MuCOE/Nodes/CustomizableObjectNodeFloatSwitch.h
@@ -6,9 +6,6 @@
 
 #include "CustomizableObjectNodeFloatSwitch.generated.h"
 
-class FArchive;
-class UObject;
-
 
 UCLASS()
 class CUSTOMIZABLEOBJECTEDITOR_API UCustomizableObjectNodeFloatSwitch : public UCustomizableObjectNodeSwitchBase
@@ -16,14 +13,9 @@ class CUSTOMIZABLEOBJECTEDITOR_API UCustomizableObjectNodeFloatSwitch : public U
 public:
 	GENERATED_BODY()
 
-	// UObject interface
-	virtual void Serialize(FArchive& Ar) override;
-	
-	// UCustomizableObjectNodeSwitchBase interface
-	FString GetOutputPinName() const override;
-
-	FName GetCategory() const override;
+	// UCustomizableObjectNode interface
+	virtual void BackwardsCompatibleFixup() override;
 
-	// Own interface
-	FString GetPinPrefix() const;
+	// UCustomizableObjectNodeSwitchBase interface
+	virtual FName GetCategory() const override;
 };
diff --git a/Engine/Plugins/Experimental/Mutable/Source/CustomizableObjectEditor/Classes/MuCOE/Nodes/CustomizableObjectNodeFloatVariation.h b/Engine/Plugins/Experimental/Mutable/Source/CustomizableObjectEditor/Classes/MuCOE/Nodes/CustomizableObjectNodeFloatVariation.h
--- a/Engine/Plugins/Experimental/Mutable/Source/CustomizableObjectEditor/Classes/MuCOE/Nodes/CustomizableObjectNodeFloatVariation.h
+++ b/Engine/Plugins/Experimental/Mutable/Source/CustomizableObjectEditor/Classes/MuCOE/Nodes/CustomizableObjectNodeFloatVariation.h
@@ -2,68 +2,27 @@
 
 #pragma once
 
-#include "MuCOE/Nodes/CustomizableObjectNode.h"
+#include "MuCOE/Nodes/CustomizableObjectNodeVariation.h"
 
 #include "CustomizableObjectNodeFloatVariation.generated.h"
 
-namespace ENodeTitleType { enum Type : int; }
-
-class UCustomizableObjectNodeRemapPins;
-class UEdGraphPin;
-class UObject;
-struct FPropertyChangedEvent;
-
-
-USTRUCT()
-struct CUSTOMIZABLEOBJECTEDITOR_API FCustomizableObjectFloatVariation
-{
-	GENERATED_USTRUCT_BODY()
-
-	UPROPERTY(EditAnywhere, Category = CustomizableObject)
-	FString Tag;
-};
+struct FCustomizableObjectFloatVariation;
 
 
 UCLASS()
-class CUSTOMIZABLEOBJECTEDITOR_API UCustomizableObjectNodeFloatVariation : public UCustomizableObjectNode
+class CUSTOMIZABLEOBJECTEDITOR_API UCustomizableObjectNodeFloatVariation : public UCustomizableObjectNodeVariation
 {
 public:
 	GENERATED_BODY()
-
-	UCustomizableObjectNodeFloatVariation();
-
-	UPROPERTY(EditAnywhere, Category = CustomizableObject)
-	TArray<FCustomizableObjectFloatVariation> Variations;
-
-	// UObject interface.
-	void PostEditChangeProperty(FPropertyChangedEvent& PropertyChangedEvent) override;
-
-	// EdGraphNode interface
-	FText GetNodeTitle(ENodeTitleType::Type TitleType) const override;
-	FLinearColor GetNodeTitleColor() const override;
-	FText GetTooltipText() const override;
+	
+	// Deprecated properties
+	UPROPERTY()
+	TArray<FCustomizableObjectFloatVariation> Variations_DEPRECATED;
 
 	// UCustomizableObjectNode interface
-	void AllocateDefaultPins(UCustomizableObjectNodeRemapPins* RemapPins) override;
-
-	UEdGraphPin* OutputPin() const
-	{
-		return FindPin(TEXT("Float"));
-	}
-
-	UEdGraphPin* DefaultPin() const
-	{
-		return FindPin(TEXT("Default"));
-	}
-
-	int32 GetNumVariations() const
-	{
-		return Variations.Num();
-	}
-
-	UEdGraphPin* VariationPin(int Index) const
-	{
-		return FindPin(FString::Printf(TEXT("Variation %d"),Index));
-	}
+	virtual void BackwardsCompatibleFixup() override;
+	
+	// UCustomizableObjectNodeVariation interface
+	virtual FName GetCategory() const override;
 };
 
diff --git a/Engine/Plugins/Experimental/Mutable/Source/CustomizableObjectEditor/Classes/MuCOE/Nodes/CustomizableObjectNodeMaterialVariation.h b/Engine/Plugins/Experimental/Mutable/Source/CustomizableObjectEditor/Classes/MuCOE/Nodes/CustomizableObjectNodeMaterialVariation.h
--- a/Engine/Plugins/Experimental/Mutable/Source/CustomizableObjectEditor/Classes/MuCOE/Nodes/CustomizableObjectNodeMaterialVariation.h
+++ b/Engine/Plugins/Experimental/Mutable/Source/CustomizableObjectEditor/Classes/MuCOE/Nodes/CustomizableObjectNodeMaterialVariation.h
@@ -2,26 +2,11 @@
 
 #pragma once
 
-#include "MuCOE/Nodes/CustomizableObjectNodeMaterialBase.h"
+#include "MuCOE/Nodes/CustomizableObjectNodeVariation.h"
 
 #include "CustomizableObjectNodeMaterialVariation.generated.h"
 
-namespace ENodeTitleType { enum Type : int; }
-
-class UCustomizableObjectNodeRemapPins;
-class UEdGraphPin;
-class UObject;
-struct FPropertyChangedEvent;
-
-
-USTRUCT()
-struct CUSTOMIZABLEOBJECTEDITOR_API FCustomizableObjectMaterialVariation
-{
-	GENERATED_USTRUCT_BODY()
-
-	UPROPERTY(EditAnywhere, Category = CustomizableObject)
-	FString Tag;
-};
+struct FCustomizableObjectMaterialVariation;
 
 
 UENUM(BlueprintType)
@@ -33,48 +18,24 @@ enum class ECustomizableObjectNodeMaterialVariationType : uint8
 
 
 UCLASS()
-class CUSTOMIZABLEOBJECTEDITOR_API UCustomizableObjectNodeMaterialVariation : public UCustomizableObjectNodeMaterialBase
+class CUSTOMIZABLEOBJECTEDITOR_API UCustomizableObjectNodeMaterialVariation : public UCustomizableObjectNodeVariation
 {
 public:
 	GENERATED_BODY()
 
-	UCustomizableObjectNodeMaterialVariation();
-	
 	UPROPERTY(EditAnywhere, Category = CustomizableObject)
 	ECustomizableObjectNodeMaterialVariationType Type = ECustomizableObjectNodeMaterialVariationType::Tag;
-
-	UPROPERTY(EditAnywhere, Category = CustomizableObject)
-	TArray<FCustomizableObjectMaterialVariation> Variations;
-
-	// UObject interface.
-	void PostEditChangeProperty(FPropertyChangedEvent& PropertyChangedEvent) override;
-
-	// Begin EdGraphNode interface
-	FText GetNodeTitle(ENodeTitleType::Type TitleType) const override;
-	FLinearColor GetNodeTitleColor() const override;
-	FText GetTooltipText() const override;
+	
+	// Deprecated properties
+	UPROPERTY()
+	TArray<FCustomizableObjectMaterialVariation> Variations_DEPRECATED;
 
 	// UCustomizableObjectNode interface
-	void AllocateDefaultPins(UCustomizableObjectNodeRemapPins* RemapPins) override;
-
-	UEdGraphPin* OutputPin() const
-	{
-		return FindPin(TEXT("Material"));
-	}
-
-	UEdGraphPin* DefaultPin() const
-	{
-		return FindPin(TEXT("Default"));
-	}
-
-	int32 GetNumVariations() const
-	{
-		return Variations.Num();
-	}
+	virtual void BackwardsCompatibleFixup() override;
+	
+	// UCustomizableObjectNodeVariation interface
+	virtual FName GetCategory() const override;
 
-	UEdGraphPin* VariationPin(int Index) const
-	{
-		return FindPin(FString::Printf(TEXT("Variation %d"),Index));
-	}
+	bool IsInputPinArray() const;
 };
 
diff --git a/Engine/Plugins/Experimental/Mutable/Source/CustomizableObjectEditor/Classes/MuCOE/Nodes/CustomizableObjectNodeMeshSwitch.h b/Engine/Plugins/Experimental/Mutable/Source/CustomizableObjectEditor/Classes/MuCOE/Nodes/CustomizableObjectNodeMeshSwitch.h
--- a/Engine/Plugins/Experimental/Mutable/Source/CustomizableObjectEditor/Classes/MuCOE/Nodes/CustomizableObjectNodeMeshSwitch.h
+++ b/Engine/Plugins/Experimental/Mutable/Source/CustomizableObjectEditor/Classes/MuCOE/Nodes/CustomizableObjectNodeMeshSwitch.h
@@ -6,9 +6,6 @@
 
 #include "CustomizableObjectNodeMeshSwitch.generated.h"
 
-class FArchive;
-class UObject;
-
 
 UCLASS()
 class CUSTOMIZABLEOBJECTEDITOR_API UCustomizableObjectNodeMeshSwitch : public UCustomizableObjectNodeSwitchBase
@@ -16,14 +13,10 @@ class CUSTOMIZABLEOBJECTEDITOR_API UCustomizableObjectNodeMeshSwitch : public UC
 public:
 	GENERATED_BODY()
 
-	// UObject interface
-	virtual void Serialize(FArchive& Ar) override;
+	// UCustomizableObjectNode interface
+	virtual void BackwardsCompatibleFixup() override;
 	
 	// UCustomizableObjectNodeSwitchBase interface
-	FString GetOutputPinName() const override;
-
-	FName GetCategory() const override;
-
-	// Own interface
-	FString GetPinPrefix() const;
+	virtual FName GetCategory() const override;
 };
+
diff --git a/Engine/Plugins/Experimental/Mutable/Source/CustomizableObjectEditor/Classes/MuCOE/Nodes/CustomizableObjectNodeMeshVariation.h b/Engine/Plugins/Experimental/Mutable/Source/CustomizableObjectEditor/Classes/MuCOE/Nodes/CustomizableObjectNodeMeshVariation.h
--- a/Engine/Plugins/Experimental/Mutable/Source/CustomizableObjectEditor/Classes/MuCOE/Nodes/CustomizableObjectNodeMeshVariation.h
+++ b/Engine/Plugins/Experimental/Mutable/Source/CustomizableObjectEditor/Classes/MuCOE/Nodes/CustomizableObjectNodeMeshVariation.h
@@ -2,68 +2,27 @@
 
 #pragma once
 
-#include "MuCOE/Nodes/CustomizableObjectNode.h"
+#include "MuCOE/Nodes/CustomizableObjectNodeVariation.h"
 
 #include "CustomizableObjectNodeMeshVariation.generated.h"
 
-namespace ENodeTitleType { enum Type : int; }
-
-class UCustomizableObjectNodeRemapPins;
-class UEdGraphPin;
-class UObject;
-struct FPropertyChangedEvent;
-
-
-USTRUCT()
-struct CUSTOMIZABLEOBJECTEDITOR_API FCustomizableObjectMeshVariation
-{
-	GENERATED_USTRUCT_BODY()
-
-	UPROPERTY(EditAnywhere, Category = CustomizableObject)
-	FString Tag;
-};
+struct FCustomizableObjectMeshVariation;
 
 
 UCLASS()
-class CUSTOMIZABLEOBJECTEDITOR_API UCustomizableObjectNodeMeshVariation : public UCustomizableObjectNode
+class CUSTOMIZABLEOBJECTEDITOR_API UCustomizableObjectNodeMeshVariation : public UCustomizableObjectNodeVariation
 {
 public:
 	GENERATED_BODY()
 
-	UCustomizableObjectNodeMeshVariation();
-
-	UPROPERTY(EditAnywhere, Category = CustomizableObject)
-	TArray<FCustomizableObjectMeshVariation> Variations;
-
-	// UObject interface.
-	void PostEditChangeProperty(FPropertyChangedEvent& PropertyChangedEvent) override;
-
-	// EdGraphNode interface
-	FText GetNodeTitle(ENodeTitleType::Type TitleType) const override;
-	FLinearColor GetNodeTitleColor() const override;
-	FText GetTooltipText() const override;
+	// Deprecated properties
+	UPROPERTY()
+	TArray<FCustomizableObjectMeshVariation> Variations_DEPRECATED;
 
 	// UCustomizableObjectNode interface
-	void AllocateDefaultPins(UCustomizableObjectNodeRemapPins* RemapPins) override;
-
-	UEdGraphPin* OutputPin() const
-	{
-		return FindPin(TEXT("Mesh"));
-	}
-
-	UEdGraphPin* DefaultPin() const
-	{
-		return FindPin(TEXT("Default"));
-	}
-
-	int32 GetNumVariations() const
-	{
-		return Variations.Num();
-	}
-
-	UEdGraphPin* VariationPin(int Index) const
-	{
-		return FindPin(FString::Printf(TEXT("Variation %d"),Index));
-	}
+	virtual void BackwardsCompatibleFixup() override;
+	
+	// UCustomizableObjectNodeVariation interface
+	virtual FName GetCategory() const override;
 };
 
diff --git a/Engine/Plugins/Experimental/Mutable/Source/CustomizableObjectEditor/Classes/MuCOE/Nodes/CustomizableObjectNodePassThroughTextureSwitch.h b/Engine/Plugins/Experimental/Mutable/Source/CustomizableObjectEditor/Classes/MuCOE/Nodes/CustomizableObjectNodePassThroughTextureSwitch.h
--- a/Engine/Plugins/Experimental/Mutable/Source/CustomizableObjectEditor/Classes/MuCOE/Nodes/CustomizableObjectNodePassThroughTextureSwitch.h
+++ b/Engine/Plugins/Experimental/Mutable/Source/CustomizableObjectEditor/Classes/MuCOE/Nodes/CustomizableObjectNodePassThroughTextureSwitch.h
@@ -12,11 +12,7 @@ class CUSTOMIZABLEOBJECTEDITOR_API UCustomizableObjectNodePassThroughTextureSwit
 {
 public:
 	GENERATED_BODY()
-
+	
 	// UCustomizableObjectNodeSwitchBase interface
-	FString GetOutputPinName() const override;
-
-	FName GetCategory() const override;
-
-	FString GetPinPrefix() const;
+	virtual FName GetCategory() const override;
 };
diff --git a/Engine/Plugins/Experimental/Mutable/Source/CustomizableObjectEditor/Classes/MuCOE/Nodes/CustomizableObjectNodeSwitchBase.cpp b/Engine/Plugins/Experimental/Mutable/Source/CustomizableObjectEditor/Classes/MuCOE/Nodes/CustomizableObjectNodeSwitchBase.cpp
--- a/Engine/Plugins/Experimental/Mutable/Source/CustomizableObjectEditor/Classes/MuCOE/Nodes/CustomizableObjectNodeSwitchBase.cpp
+++ b/Engine/Plugins/Experimental/Mutable/Source/CustomizableObjectEditor/Classes/MuCOE/Nodes/CustomizableObjectNodeSwitchBase.cpp
@@ -65,7 +65,7 @@ void UCustomizableObjectNodeSwitchBase::AllocateDefaultPins(UCustomizableObjectN
 {
 	const UEdGraphSchema_CustomizableObject* Schema = GetDefault<UEdGraphSchema_CustomizableObject>();
 
-	UEdGraphPin* OutputPin = CustomCreatePin(EGPD_Output, GetCategory(), FName(GetOutputPinName()));
+	UEdGraphPin* OutputPin = CustomCreatePin(EGPD_Output, GetCategory(), FName(UEdGraphSchema_CustomizableObject::GetPinCategoryName(GetCategory()).ToString()));
 	OutputPin->bDefaultValueIsIgnored = true;
 
 	OutputPinReference = FEdGraphPinReference(OutputPin);
@@ -200,5 +200,18 @@ void UCustomizableObjectNodeSwitchBase::LinkPostEditChangePropertyDelegate(const
 	}
 }
 
+
+FString UCustomizableObjectNodeSwitchBase::GetOutputPinName() const
+{
+	return FString();
+}
+
+
+FString UCustomizableObjectNodeSwitchBase::GetPinPrefix() const
+{
+	return UEdGraphSchema_CustomizableObject::GetPinCategoryName(GetCategory()).ToString() + " ";
+}
+
+
 #undef LOCTEXT_NAMESPACE
 
diff --git a/Engine/Plugins/Experimental/Mutable/Source/CustomizableObjectEditor/Classes/MuCOE/Nodes/CustomizableObjectNodeSwitchBase.h b/Engine/Plugins/Experimental/Mutable/Source/CustomizableObjectEditor/Classes/MuCOE/Nodes/CustomizableObjectNodeSwitchBase.h
--- a/Engine/Plugins/Experimental/Mutable/Source/CustomizableObjectEditor/Classes/MuCOE/Nodes/CustomizableObjectNodeSwitchBase.h
+++ b/Engine/Plugins/Experimental/Mutable/Source/CustomizableObjectEditor/Classes/MuCOE/Nodes/CustomizableObjectNodeSwitchBase.h
@@ -14,31 +14,29 @@ class UObject;
 struct FFrame;
 
 
-UCLASS()
+UCLASS(Abstract)
 class CUSTOMIZABLEOBJECTEDITOR_API UCustomizableObjectNodeSwitchBase : public UCustomizableObjectNode
 {
 public:
 	GENERATED_BODY()
 
-
 	// UObject interface
 	virtual void Serialize(FArchive& Ar) override;
 	
 	// Begin EdGraphNode interface
-	FText GetNodeTitle(ENodeTitleType::Type TitleType) const;
-	FLinearColor GetNodeTitleColor() const override;
-	FText GetTooltipText() const override;
-	void PinConnectionListChanged(UEdGraphPin* Pin) override;
+	virtual FText GetNodeTitle(ENodeTitleType::Type TitleType) const override;
+	virtual FLinearColor GetNodeTitleColor() const override;
+	virtual FText GetTooltipText() const override;
+	virtual void PinConnectionListChanged(UEdGraphPin* Pin) override;
 	virtual void PostPasteNode() override;
 
-
 	// UCustomizableObjectNode interface
 	virtual void PostBackwardsCompatibleFixup() override;
-	void AllocateDefaultPins(UCustomizableObjectNodeRemapPins* RemapPins) override;
-	void ReconstructNode(UCustomizableObjectNodeRemapPins* RemapPins) override;
+	virtual void AllocateDefaultPins(UCustomizableObjectNodeRemapPins* RemapPins) override;
+	virtual void ReconstructNode(UCustomizableObjectNodeRemapPins* RemapPins) override;
 
-	/** Get the output pin catergory. Override. */
-	virtual FName GetCategory() const { return FName(); };
+	/** Get the output pin category. Override. */
+	virtual FName GetCategory() const PURE_VIRTUAL(UCustomizableObjectNodeSwitchBase::GetCategory, return {}; );
 
 	UEdGraphPin* OutputPin() const;
 
@@ -54,13 +52,14 @@ public:
 	/** Links the PostEditChangeProperty delegate */
 	void LinkPostEditChangePropertyDelegate(const UEdGraphPin& Pin);
 
-protected:
-	/** Get the pin prefix. Used for retrocompatibility. Override. */
-	virtual FString GetPinPrefix() const { return FString(); };
-
 	/** Get the ouput pin name. Override. */
-	virtual FString GetOutputPinName() const { return FString(); };
+	virtual FString GetOutputPinName() const;
 
+private:
+	/** Get the pin prefix. Used for retrocompatibility. Override. */
+	virtual FString GetPinPrefix() const;
+
+protected:
 	UPROPERTY()
 	FEdGraphPinReference OutputPinReference;
 
@@ -70,7 +69,7 @@ private:
 
 	void ReloadEnumParam();
 
-	/** Last NodeEnumParameter connected. Used to remove the callback once desconnected. */
+	/** Last NodeEnumParameter connected. Used to remove the callback once disconnected. */
 	TWeakObjectPtr<UCustomizableObjectNode> LastNodeEnumParameterConnected;
 
 	/** NodeEnumParameter property changed callback function. Reconstructs the node. */
@@ -83,7 +82,6 @@ private:
 	UPROPERTY()
 	TArray<FString> ReloadingElementsNames;
 
-private:
 	UPROPERTY()
 	FEdGraphPinReference SwitchParameterPinReference;
 };
diff --git a/Engine/Plugins/Experimental/Mutable/Source/CustomizableObjectEditor/Classes/MuCOE/Nodes/CustomizableObjectNodeTextureSwitch.h b/Engine/Plugins/Experimental/Mutable/Source/CustomizableObjectEditor/Classes/MuCOE/Nodes/CustomizableObjectNodeTextureSwitch.h
--- a/Engine/Plugins/Experimental/Mutable/Source/CustomizableObjectEditor/Classes/MuCOE/Nodes/CustomizableObjectNodeTextureSwitch.h
+++ b/Engine/Plugins/Experimental/Mutable/Source/CustomizableObjectEditor/Classes/MuCOE/Nodes/CustomizableObjectNodeTextureSwitch.h
@@ -6,26 +6,15 @@
 
 #include "CustomizableObjectNodeTextureSwitch.generated.h"
 
-class FArchive;
-class UObject;
-
 
 UCLASS()
 class CUSTOMIZABLEOBJECTEDITOR_API UCustomizableObjectNodeTextureSwitch : public UCustomizableObjectNodeSwitchBase
 {
 public:
 	GENERATED_BODY()
-
-	// UObject interface
-	virtual void Serialize(FArchive& Ar) override;
-
 	// UCustomizableObjectNode interface
 	virtual void BackwardsCompatibleFixup() override;
-
+	
 	// UCustomizableObjectNodeSwitchBase interface
-	FString GetOutputPinName() const override;
-
-	FName GetCategory() const override;
-
-	FString GetPinPrefix() const;
+	virtual FName GetCategory() const override;
 };
diff --git a/Engine/Plugins/Experimental/Mutable/Source/CustomizableObjectEditor/Classes/MuCOE/Nodes/CustomizableObjectNodeVariation.h b/Engine/Plugins/Experimental/Mutable/Source/CustomizableObjectEditor/Classes/MuCOE/Nodes/CustomizableObjectNodeVariation.h
new file mode 100644
--- /dev/null
+++ b/Engine/Plugins/Experimental/Mutable/Source/CustomizableObjectEditor/Classes/MuCOE/Nodes/CustomizableObjectNodeVariation.h
@@ -0,0 +1,74 @@
+// Copyright Epic Games, Inc. All Rights Reserved.
+
+#pragma once
+
+#include "MuCOE/Nodes/CustomizableObjectNode.h"
+
+#include "CustomizableObjectNodeVariation.generated.h"
+
+namespace ENodeTitleType
+{
+	enum Type : int;
+}
+
+class UCustomizableObjectNodeRemapPins;
+class UEdGraphPin;
+class UObject;
+struct FPropertyChangedEvent;
+
+
+USTRUCT()
+struct CUSTOMIZABLEOBJECTEDITOR_API FCustomizableObjectVariation
+{
+	GENERATED_USTRUCT_BODY()
+
+	UPROPERTY(EditAnywhere, Category = CustomizableObject)
+	FString Tag;
+};
+
+
+UCLASS(Abstract)
+class CUSTOMIZABLEOBJECTEDITOR_API UCustomizableObjectNodeVariation : public UCustomizableObjectNode
+{
+	GENERATED_BODY()
+
+protected:
+	UPROPERTY(EditAnywhere, Category = CustomizableObject)
+	TArray<FCustomizableObjectVariation> Variations;
+
+private:
+	TArray<FEdGraphPinReference> VariationsPins;
+	
+public:
+	// UObject interface.
+	virtual void PostEditChangeProperty(FPropertyChangedEvent& PropertyChangedEvent) override;
+
+	// EdGraphNode interface
+	virtual FText GetNodeTitle(ENodeTitleType::Type TitleType) const override;
+	virtual FLinearColor GetNodeTitleColor() const override;
+	virtual FText GetTooltipText() const override;
+	
+	// UCustomizableObjectNode interface
+	virtual void BackwardsCompatibleFixup() override;
+	virtual void AllocateDefaultPins(UCustomizableObjectNodeRemapPins* RemapPins) override;
+
+	// Own interface
+	/** Return the pin category of this node. */
+	virtual FName GetCategory() const PURE_VIRTUAL(UCustomizableObjectNodeVariation::GetCategory, return {}; );
+
+	/** Return true if all inputs pins should be array. */
+	virtual bool IsInputPinArray() const;
+
+	/** Return the number of variations (input pins excluding the Default Pin). */
+	int32 GetNumVariations() const;
+
+	/** Get the variation at the given index. */
+	const FCustomizableObjectVariation& GetVariation(int32 Index) const;
+
+	/** Get the Default Input Pin. */
+	UEdGraphPin* DefaultPin() const;
+
+	/** Get the Variation Input Pin. */
+	UEdGraphPin* VariationPin(int32 Index) const;
+};
+
diff --git a/Engine/Plugins/Experimental/Mutable/Source/CustomizableObjectEditor/Private/MuCOE/EdGraphSchema_CustomizableObject.cpp b/Engine/Plugins/Experimental/Mutable/Source/CustomizableObjectEditor/Private/MuCOE/EdGraphSchema_CustomizableObject.cpp
--- a/Engine/Plugins/Experimental/Mutable/Source/CustomizableObjectEditor/Private/MuCOE/EdGraphSchema_CustomizableObject.cpp
+++ b/Engine/Plugins/Experimental/Mutable/Source/CustomizableObjectEditor/Private/MuCOE/EdGraphSchema_CustomizableObject.cpp
@@ -228,73 +228,6 @@ const FName UEdGraphSchema_CustomizableObject::PC_MaterialAsset("materialAsset")
 const FName UEdGraphSchema_CustomizableObject::PC_Wildcard("wildcard");
 
 
-UEdGraphSchema_CustomizableObject::UEdGraphSchema_CustomizableObject()
-{
-	NodeTypes.Add(UCustomizableObjectNodeAnimationPose::StaticClass());
-	NodeTypes.Add(UCustomizableObjectNodeColorArithmeticOp::StaticClass());
-	NodeTypes.Add(UCustomizableObjectNodeColorConstant::StaticClass());
-	NodeTypes.Add(UCustomizableObjectNodeColorFromFloats::StaticClass());
-	NodeTypes.Add(UCustomizableObjectNodeColorParameter::StaticClass());
-	NodeTypes.Add(UCustomizableObjectNodeColorSwitch::StaticClass());
-	NodeTypes.Add(UCustomizableObjectNodeColorVariation::StaticClass());
-	NodeTypes.Add(UCustomizableObjectNodeCopyMaterial::StaticClass());
-	NodeTypes.Add(UCustomizableObjectNodeCurve::StaticClass());
-	NodeTypes.Add(UCustomizableObjectNodeEditMaterial::StaticClass());
-	NodeTypes.Add(UCustomizableObjectNodeEnumParameter::StaticClass());
-	NodeTypes.Add(UCustomizableObjectNodeExtendMaterial::StaticClass());
-	NodeTypes.Add(UCustomizableObjectNodeFloatConstant::StaticClass());
-	NodeTypes.Add(UCustomizableObjectNodeFloatParameter::StaticClass());
-	NodeTypes.Add(UCustomizableObjectNodeFloatSwitch::StaticClass());
-	NodeTypes.Add(UCustomizableObjectNodeFloatVariation::StaticClass());
-	NodeTypes.Add(UCustomizableObjectNodeGroupProjectorParameter::StaticClass());
-	NodeTypes.Add(UCustomizableObjectNodeMaterial::StaticClass());
-	NodeTypes.Add(UCustomizableObjectNodeMaterialVariation::StaticClass());
-	NodeTypes.Add(UCustomizableObjectNodeMeshClipDeform::StaticClass());
-	NodeTypes.Add(UCustomizableObjectNodeMeshClipMorph::StaticClass());
-	NodeTypes.Add(UCustomizableObjectNodeMeshClipWithMesh::StaticClass());
-	NodeTypes.Add(UCustomizableObjectNodeMeshGeometryOperation::StaticClass());
-	NodeTypes.Add(UCustomizableObjectNodeMeshMorph::StaticClass());
-	NodeTypes.Add(UCustomizableObjectNodeMeshMorphStackApplication::StaticClass());
-	NodeTypes.Add(UCustomizableObjectNodeMeshMorphStackDefinition::StaticClass());
-	NodeTypes.Add(UCustomizableObjectNodeMeshReshape::StaticClass());
-	NodeTypes.Add(UCustomizableObjectNodeMeshSwitch::StaticClass());
-	NodeTypes.Add(UCustomizableObjectNodeMeshVariation::StaticClass());
-	NodeTypes.Add(UCustomizableObjectNodeMorphMaterial::StaticClass());
-	NodeTypes.Add(UCustomizableObjectNodeObject::StaticClass());
-	NodeTypes.Add(UCustomizableObjectNodeObjectGroup::StaticClass());
-	NodeTypes.Add(UCustomizableObjectNodeProjectorConstant::StaticClass());
-	NodeTypes.Add(UCustomizableObjectNodeProjectorParameter::StaticClass());
-	NodeTypes.Add(UCustomizableObjectNodeRemoveMesh::StaticClass());
-	NodeTypes.Add(UCustomizableObjectNodeRemoveMeshBlocks::StaticClass());
-	NodeTypes.Add(UCustomizableObjectNodeSkeletalMesh::StaticClass());
-	NodeTypes.Add(UCustomizableObjectNodeStaticMesh::StaticClass());
-	NodeTypes.Add(UCustomizableObjectNodeTable::StaticClass());
-	NodeTypes.Add(UCustomizableObjectNodeTexture::StaticClass());
-	NodeTypes.Add(UCustomizableObjectNodeTextureBinarise::StaticClass());
-	NodeTypes.Add(UCustomizableObjectNodeTextureColourMap::StaticClass());
-	NodeTypes.Add(UCustomizableObjectNodeTextureFromChannels::StaticClass());
-	NodeTypes.Add(UCustomizableObjectNodeTextureInterpolate::StaticClass());
-	NodeTypes.Add(UCustomizableObjectNodeTextureInvert::StaticClass());
-	NodeTypes.Add(UCustomizableObjectNodeTextureLayer::StaticClass());
-	NodeTypes.Add(UCustomizableObjectNodeTextureParameter::StaticClass());
-	NodeTypes.Add(UCustomizableObjectNodeTextureProject::StaticClass());
-	NodeTypes.Add(UCustomizableObjectNodeTextureSample::StaticClass());
-	NodeTypes.Add(UCustomizableObjectNodeTextureSaturate::StaticClass());
-	NodeTypes.Add(UCustomizableObjectNodeTextureSwitch::StaticClass());
-	NodeTypes.Add(UCustomizableObjectNodeTextureToChannels::StaticClass());
-	NodeTypes.Add(UCustomizableObjectNodeTextureTransform::StaticClass());
-	NodeTypes.Add(UCustomizableObjectNodeTextureVariation::StaticClass());
-
-	NodeTypes.Sort([](UClass& A, UClass& B) -> bool
-	{
-		const UCustomizableObjectNode* NodeA = CastChecked<UCustomizableObjectNode>(A.GetDefaultObject());
-		const UCustomizableObjectNode* NodeB = CastChecked<UCustomizableObjectNode>(B.GetDefaultObject());
-		
-		return NodeA->GetNodeTitle(ENodeTitleType::ListView).CompareTo(NodeB->GetNodeTitle(ENodeTitleType::ListView)) < 0;
-	});
-}
-
-
 TSharedPtr<FCustomizableObjectSchemaAction_NewNode> UEdGraphSchema_CustomizableObject::AddNewNodeAction(FGraphActionListBuilderBase& ContextMenuBuilder, const FString& Category, const FText& MenuDesc, const FText& Tooltip, const int32 Grouping, const FString& Keywords)
 {
 	TSharedPtr<FCustomizableObjectSchemaAction_NewNode> NewActionNode = TSharedPtr<FCustomizableObjectSchemaAction_NewNode>(new FCustomizableObjectSchemaAction_NewNode(Category, MenuDesc, Tooltip, Grouping, FText::FromString(Keywords)));
@@ -793,6 +726,33 @@ bool UEdGraphSchema_CustomizableObject::ShouldHidePinDefaultValue(UEdGraphPin* P
 void UEdGraphSchema_CustomizableObject::GetContextMenuActionsReconstructAllChildNodes(UToolMenu* Menu, TWeakObjectPtr<UGraphNodeContextMenuContext> WeakContext) const
 {
 	FToolMenuSection& SubSection = Menu->AddSection("Section");
+
+	TArray<UClass*> NodeTypes;
+
+	for (TObjectIterator<UCustomizableObjectNode> It(RF_NoFlags); It; ++It)
+	{
+		const UCustomizableObjectNode* Node = *It;
+		if (!Node->HasAllFlags(RF_ClassDefaultObject))
+		{
+			continue; // Only interested in CDOs
+		}
+
+		UClass* Class = Node->GetClass();
+		if (Class->HasAnyClassFlags(CLASS_Abstract))
+		{
+			continue;
+		}
+		
+		NodeTypes.Add(Class);
+	}
+	
+	NodeTypes.Sort([](UClass& A, UClass& B) -> bool
+	{
+		const UCustomizableObjectNode* NodeA = CastChecked<UCustomizableObjectNode>(A.GetDefaultObject());
+		const UCustomizableObjectNode* NodeB = CastChecked<UCustomizableObjectNode>(B.GetDefaultObject());
+		
+		return NodeA->GetNodeTitle(ENodeTitleType::ListView).CompareTo(NodeB->GetNodeTitle(ENodeTitleType::ListView)) < 0;
+	});
 	
 	for (UClass* NodeType : NodeTypes)
 	{
diff --git a/Engine/Plugins/Experimental/Mutable/Source/CustomizableObjectEditor/Private/MuCOE/GenerateMutableSource/GenerateMutableSource.cpp b/Engine/Plugins/Experimental/Mutable/Source/CustomizableObjectEditor/Private/MuCOE/GenerateMutableSource/GenerateMutableSource.cpp
--- a/Engine/Plugins/Experimental/Mutable/Source/CustomizableObjectEditor/Private/MuCOE/GenerateMutableSource/GenerateMutableSource.cpp
+++ b/Engine/Plugins/Experimental/Mutable/Source/CustomizableObjectEditor/Private/MuCOE/GenerateMutableSource/GenerateMutableSource.cpp
@@ -1219,7 +1219,7 @@ bool AffectsCurrentComponent(const UEdGraphPin* Pin, FMutableGraphGenerationCont
 			bAffectsCurrentComponent = true;
 		}
 
-		for (int VariationIndex = 0; VariationIndex < TypedNodeVar->Variations.Num(); ++VariationIndex)
+		for (int VariationIndex = 0; VariationIndex < TypedNodeVar->GetNumVariations(); ++VariationIndex)
 		{
 			if (UEdGraphPin* VariationPin = TypedNodeVar->VariationPin(VariationIndex))
 			{
diff --git a/Engine/Plugins/Experimental/Mutable/Source/CustomizableObjectEditor/Private/MuCOE/GenerateMutableSource/GenerateMutableSourceColor.cpp b/Engine/Plugins/Experimental/Mutable/Source/CustomizableObjectEditor/Private/MuCOE/GenerateMutableSource/GenerateMutableSourceColor.cpp
--- a/Engine/Plugins/Experimental/Mutable/Source/CustomizableObjectEditor/Private/MuCOE/GenerateMutableSource/GenerateMutableSourceColor.cpp
+++ b/Engine/Plugins/Experimental/Mutable/Source/CustomizableObjectEditor/Private/MuCOE/GenerateMutableSource/GenerateMutableSourceColor.cpp
@@ -252,13 +252,14 @@ mu::NodeColourPtr GenerateMutableSourceColor(const UEdGraphPin* Pin, FMutableGra
 			GenerationContext.Compiler->CompilerLog(LOCTEXT("ColorVarMissingDef", "Color variation node requires a default value."), Node);
 		}
 
-		ColorNode->SetVariationCount(TypedNodeColorVar->Variations.Num());
-		for (int VariationIndex = 0; VariationIndex < TypedNodeColorVar->Variations.Num(); ++VariationIndex)
+		const int32 NumVariations = TypedNodeColorVar->GetNumVariations();
+		ColorNode->SetVariationCount(NumVariations);
+		for (int VariationIndex = 0; VariationIndex < NumVariations; ++VariationIndex)
 		{
 			const UEdGraphPin* VariationPin = TypedNodeColorVar->VariationPin(VariationIndex);
 			if (!VariationPin) continue;
 
-			ColorNode->SetVariationTag(VariationIndex, StringCast<ANSICHAR>(*TypedNodeColorVar->Variations[VariationIndex].Tag).Get());
+			ColorNode->SetVariationTag(VariationIndex, StringCast<ANSICHAR>(*TypedNodeColorVar->GetVariation(VariationIndex).Tag).Get());
 			if (const UEdGraphPin* ConnectedPin = FollowInputPin(*VariationPin))
 			{
 				mu::NodeColourPtr ChildNode = GenerateMutableSourceColor(ConnectedPin, GenerationContext);
diff --git a/Engine/Plugins/Experimental/Mutable/Source/CustomizableObjectEditor/Private/MuCOE/GenerateMutableSource/GenerateMutableSourceFloat.cpp b/Engine/Plugins/Experimental/Mutable/Source/CustomizableObjectEditor/Private/MuCOE/GenerateMutableSource/GenerateMutableSourceFloat.cpp
--- a/Engine/Plugins/Experimental/Mutable/Source/CustomizableObjectEditor/Private/MuCOE/GenerateMutableSource/GenerateMutableSourceFloat.cpp
+++ b/Engine/Plugins/Experimental/Mutable/Source/CustomizableObjectEditor/Private/MuCOE/GenerateMutableSource/GenerateMutableSourceFloat.cpp
@@ -280,13 +280,14 @@ mu::NodeScalarPtr GenerateMutableSourceFloat(const UEdGraphPin* Pin, FMutableGra
 			GenerationContext.Compiler->CompilerLog(LOCTEXT("FloatVarMissingDef", "Float variation node requires a default value."), Node);
 		}
 
-		FloatNode->SetVariationCount(TypedNodeFloatVar->Variations.Num());
-		for (int VariationIndex = 0; VariationIndex < TypedNodeFloatVar->Variations.Num(); ++VariationIndex)
+		const int32 NumVariations = TypedNodeFloatVar->GetNumVariations();
+		FloatNode->SetVariationCount(NumVariations);
+		for (int VariationIndex = 0; VariationIndex < NumVariations; ++VariationIndex)
 		{
 			UEdGraphPin* VariationPin = TypedNodeFloatVar->VariationPin(VariationIndex);
 			if (!VariationPin) continue;
 
-			FloatNode->SetVariationTag(VariationIndex, StringCast<ANSICHAR>(*TypedNodeFloatVar->Variations[VariationIndex].Tag).Get());
+			FloatNode->SetVariationTag(VariationIndex, StringCast<ANSICHAR>(*TypedNodeFloatVar->GetVariation(VariationIndex).Tag).Get());
 			if (const UEdGraphPin* ConnectedPin = FollowInputPin(*VariationPin))
 			{
 				mu::NodeScalarPtr ChildNode = GenerateMutableSourceFloat(ConnectedPin, GenerationContext);
diff --git a/Engine/Plugins/Experimental/Mutable/Source/CustomizableObjectEditor/Private/MuCOE/GenerateMutableSource/GenerateMutableSourceMesh.cpp b/Engine/Plugins/Experimental/Mutable/Source/CustomizableObjectEditor/Private/MuCOE/GenerateMutableSource/GenerateMutableSourceMesh.cpp
--- a/Engine/Plugins/Experimental/Mutable/Source/CustomizableObjectEditor/Private/MuCOE/GenerateMutableSource/GenerateMutableSourceMesh.cpp
+++ b/Engine/Plugins/Experimental/Mutable/Source/CustomizableObjectEditor/Private/MuCOE/GenerateMutableSource/GenerateMutableSourceMesh.cpp
@@ -3264,13 +3264,14 @@ mu::NodeMeshPtr GenerateMutableSourceMesh(const UEdGraphPin* Pin,
 			GenerationContext.Compiler->CompilerLog(LOCTEXT("MeshVarMissingDef", "Mesh variation node requires a default value."), Node);
 		}
 
-		MeshNode->SetVariationCount(TypedNodeMeshVar->Variations.Num());
-		for (int VariationIndex = 0; VariationIndex < TypedNodeMeshVar->Variations.Num(); ++VariationIndex)
+		const int32 NumVariations = TypedNodeMeshVar->GetNumVariations();
+		MeshNode->SetVariationCount(NumVariations);
+		for (int VariationIndex = 0; VariationIndex < NumVariations; ++VariationIndex)
 		{
 			const UEdGraphPin* VariationPin = TypedNodeMeshVar->VariationPin(VariationIndex);
 			if (!VariationPin) continue;
 
-			MeshNode->SetVariationTag(VariationIndex, StringCast<ANSICHAR>(*TypedNodeMeshVar->Variations[VariationIndex].Tag).Get());
+			MeshNode->SetVariationTag(VariationIndex, StringCast<ANSICHAR>(*TypedNodeMeshVar->GetVariation(VariationIndex).Tag).Get());
 			if (const UEdGraphPin* ConnectedPin = FollowInputPin(*VariationPin))
 			{
 				FMutableGraphMeshGenerationData VariationMeshData;
diff --git a/Engine/Plugins/Experimental/Mutable/Source/CustomizableObjectEditor/Private/MuCOE/GenerateMutableSource/GenerateMutableSourceSurface.cpp b/Engine/Plugins/Experimental/Mutable/Source/CustomizableObjectEditor/Private/MuCOE/GenerateMutableSource/GenerateMutableSourceSurface.cpp
--- a/Engine/Plugins/Experimental/Mutable/Source/CustomizableObjectEditor/Private/MuCOE/GenerateMutableSource/GenerateMutableSourceSurface.cpp
+++ b/Engine/Plugins/Experimental/Mutable/Source/CustomizableObjectEditor/Private/MuCOE/GenerateMutableSource/GenerateMutableSourceSurface.cpp
@@ -1561,14 +1561,15 @@ mu::NodeSurfacePtr GenerateMutableSourceSurface(const UEdGraphPin * Pin, FMutabl
 			}
 		}
 
-		SurfNode->SetVariationCount(TypedNodeVar->Variations.Num());
-		for (int VariationIndex = 0; VariationIndex < TypedNodeVar->Variations.Num(); ++VariationIndex)
+		const int32 NumVariations = TypedNodeVar->GetNumVariations();
+		SurfNode->SetVariationCount(NumVariations);
+		for (int VariationIndex = 0; VariationIndex < NumVariations; ++VariationIndex)
 		{
 			mu::NodeSurfacePtr VariationSurfaceNode;
 
 			if (UEdGraphPin* VariationPin = TypedNodeVar->VariationPin(VariationIndex))
 			{
-				SurfNode->SetVariationTag(VariationIndex, StringCast<ANSICHAR>(*TypedNodeVar->Variations[VariationIndex].Tag).Get());
+				SurfNode->SetVariationTag(VariationIndex, StringCast<ANSICHAR>(*TypedNodeVar->GetVariation(VariationIndex).Tag).Get());
 				for (const UEdGraphPin* ConnectedPin : FollowInputPinArray(*VariationPin))
 				{
 					// Is it a modifier?
diff --git a/Engine/Plugins/Experimental/Mutable/Source/CustomizableObjectEditor/Private/MuCOE/Nodes/CustomizableObjectNodeColorSwitch.cpp b/Engine/Plugins/Experimental/Mutable/Source/CustomizableObjectEditor/Private/MuCOE/Nodes/CustomizableObjectNodeColorSwitch.cpp
--- a/Engine/Plugins/Experimental/Mutable/Source/CustomizableObjectEditor/Private/MuCOE/Nodes/CustomizableObjectNodeColorSwitch.cpp
+++ b/Engine/Plugins/Experimental/Mutable/Source/CustomizableObjectEditor/Private/MuCOE/Nodes/CustomizableObjectNodeColorSwitch.cpp
@@ -7,34 +7,24 @@
 #define LOCTEXT_NAMESPACE "CustomizableObjectEditor"
 
 
-void UCustomizableObjectNodeColorSwitch::Serialize(FArchive& Ar)
+void UCustomizableObjectNodeColorSwitch::BackwardsCompatibleFixup()
 {
-	Super::Serialize(Ar);
+	Super::BackwardsCompatibleFixup();
 
-	if (Ar.CustomVer(FCustomizableObjectCustomVersion::GUID) < FCustomizableObjectCustomVersion::BugPinsSwitch)
+	const int32 CustomizableObjectCustomVersion = GetLinkerCustomVersion(FCustomizableObjectCustomVersion::GUID);
+
+	if (CustomizableObjectCustomVersion < FCustomizableObjectCustomVersion::BugPinsSwitch)
 	{
 		OutputPinReference = FindPin(TEXT("Color"));	
 	}
 }
 
 
-FString UCustomizableObjectNodeColorSwitch::GetOutputPinName() const
-{
-	return TEXT("Color");
-}
-	
-
 FName UCustomizableObjectNodeColorSwitch::GetCategory() const
 {
 	return UEdGraphSchema_CustomizableObject::PC_Color;
 }
 
 
-FString UCustomizableObjectNodeColorSwitch::GetPinPrefix() const
-{
-	return TEXT("Color ");
-}
-
-
 #undef LOCTEXT_NAMESPACE
 
diff --git a/Engine/Plugins/Experimental/Mutable/Source/CustomizableObjectEditor/Private/MuCOE/Nodes/CustomizableObjectNodeColorVariation.cpp b/Engine/Plugins/Experimental/Mutable/Source/CustomizableObjectEditor/Private/MuCOE/Nodes/CustomizableObjectNodeColorVariation.cpp
--- a/Engine/Plugins/Experimental/Mutable/Source/CustomizableObjectEditor/Private/MuCOE/Nodes/CustomizableObjectNodeColorVariation.cpp
+++ b/Engine/Plugins/Experimental/Mutable/Source/CustomizableObjectEditor/Private/MuCOE/Nodes/CustomizableObjectNodeColorVariation.cpp
@@ -2,74 +2,30 @@
 
 #include "MuCOE/Nodes/CustomizableObjectNodeColorVariation.h"
 
+#include "MuCOE/CustomizableObjectEditor_Deprecated.h"
 #include "MuCOE/EdGraphSchema_CustomizableObject.h"
-#include "MuCOE/UnrealEditorPortabilityHelpers.h"
 
-class UCustomizableObjectNodeRemapPins;
 
-#define LOCTEXT_NAMESPACE "CustomizableObjectEditor"
-
-
-UCustomizableObjectNodeColorVariation::UCustomizableObjectNodeColorVariation()
-	: Super()
+void UCustomizableObjectNodeColorVariation::BackwardsCompatibleFixup()
 {
+	const int32 CustomizableObjectCustomVersion = GetLinkerCustomVersion(FCustomizableObjectCustomVersion::GUID);
 
-}
-
-
-void UCustomizableObjectNodeColorVariation::PostEditChangeProperty(FPropertyChangedEvent& PropertyChangedEvent)
-{
-	FProperty* PropertyThatChanged = PropertyChangedEvent.Property;
-	if ( PropertyThatChanged  )
+	if (CustomizableObjectCustomVersion < FCustomizableObjectCustomVersion::DeduplicateNodeVariant)
 	{
-		ReconstructNode();
+		for (const FCustomizableObjectColorVariation& OldVariation : Variations_DEPRECATED)
+		{
+			FCustomizableObjectVariation Variation;
+			Variation.Tag = OldVariation.Tag;
+			
+			Variations.Add(Variation);
+		}
 	}
 
-	Super::PostEditChangeProperty(PropertyChangedEvent);
+	Super::BackwardsCompatibleFixup();
 }
 
 
-void UCustomizableObjectNodeColorVariation::AllocateDefaultPins(UCustomizableObjectNodeRemapPins* RemapPins)
+FName UCustomizableObjectNodeColorVariation::GetCategory() const
 {
-	const UEdGraphSchema_CustomizableObject* Schema = GetDefault<UEdGraphSchema_CustomizableObject>();
-
-	FString RemapIndex = Helper_GetPinName(OutputPin());
-
-	UEdGraphPin* OutputPin = CustomCreatePin(EGPD_Output, Schema->PC_Color, FName(TEXT("Color")));
-	OutputPin->bDefaultValueIsIgnored = true;
-
-	for ( int VariationIndex = Variations.Num() - 1; VariationIndex >= 0; --VariationIndex)
-	{
-		FString PinName = FString::Printf( TEXT("Variation %d"), VariationIndex);
-		UEdGraphPin* VariationPin = CustomCreatePin(EGPD_Input, Schema->PC_Color, FName(*PinName), false);
-		VariationPin->bDefaultValueIsIgnored = true;
-
-		FString FriendlyName = FString::Printf(TEXT("Variation %d [%s]"), VariationIndex, *Variations[VariationIndex].Tag);
-		VariationPin->PinFriendlyName = FText::FromString(FriendlyName);
-	}
-
-	UEdGraphPin* DefaultVariation = CustomCreatePin(EGPD_Input, Schema->PC_Color, FName(TEXT("Default")), false);
-	DefaultVariation->bDefaultValueIsIgnored = true;
+	return UEdGraphSchema_CustomizableObject::PC_Color;
 }
-
-
-FText UCustomizableObjectNodeColorVariation::GetNodeTitle(ENodeTitleType::Type TitleType) const
-{
-	return LOCTEXT("Color_Variation", "Color Variation");
-}
-
-
-FLinearColor UCustomizableObjectNodeColorVariation::GetNodeTitleColor() const
-{
-	const UEdGraphSchema_CustomizableObject* Schema = GetDefault<UEdGraphSchema_CustomizableObject>();
-	return Schema->GetPinTypeColor(Schema->PC_Color);
-}
-
-
-FText UCustomizableObjectNodeColorVariation::GetTooltipText() const
-{
-	return LOCTEXT("Color_Variation_Tooltip", "Select a color depending on what tags are active.");
-}
-
-#undef LOCTEXT_NAMESPACE
-
diff --git a/Engine/Plugins/Experimental/Mutable/Source/CustomizableObjectEditor/Private/MuCOE/Nodes/CustomizableObjectNodeFloatSwitch.cpp b/Engine/Plugins/Experimental/Mutable/Source/CustomizableObjectEditor/Private/MuCOE/Nodes/CustomizableObjectNodeFloatSwitch.cpp
--- a/Engine/Plugins/Experimental/Mutable/Source/CustomizableObjectEditor/Private/MuCOE/Nodes/CustomizableObjectNodeFloatSwitch.cpp
+++ b/Engine/Plugins/Experimental/Mutable/Source/CustomizableObjectEditor/Private/MuCOE/Nodes/CustomizableObjectNodeFloatSwitch.cpp
@@ -7,34 +7,24 @@
 #define LOCTEXT_NAMESPACE "CustomizableObjectEditor"
 
 
-void UCustomizableObjectNodeFloatSwitch::Serialize(FArchive& Ar)
+void UCustomizableObjectNodeFloatSwitch::BackwardsCompatibleFixup()
 {
-	Super::Serialize(Ar);
+	Super::BackwardsCompatibleFixup();
 
-	if (Ar.CustomVer(FCustomizableObjectCustomVersion::GUID) < FCustomizableObjectCustomVersion::BugPinsSwitch)
+	const int32 CustomizableObjectCustomVersion = GetLinkerCustomVersion(FCustomizableObjectCustomVersion::GUID);
+
+	if (CustomizableObjectCustomVersion < FCustomizableObjectCustomVersion::BugPinsSwitch)
 	{
 		OutputPinReference = FindPin(TEXT("Float"));	
 	}
 }
 
 
-FString UCustomizableObjectNodeFloatSwitch::GetOutputPinName() const
-{
-	return TEXT("Float");
-}
-
-
 FName UCustomizableObjectNodeFloatSwitch::GetCategory() const
 {
 	return UEdGraphSchema_CustomizableObject::PC_Float;
 }
 
 
-FString UCustomizableObjectNodeFloatSwitch::GetPinPrefix() const
-{
-	return TEXT("Float ");
-}
-
-
 #undef LOCTEXT_NAMESPACE
 
diff --git a/Engine/Plugins/Experimental/Mutable/Source/CustomizableObjectEditor/Private/MuCOE/Nodes/CustomizableObjectNodeFloatVariation.cpp b/Engine/Plugins/Experimental/Mutable/Source/CustomizableObjectEditor/Private/MuCOE/Nodes/CustomizableObjectNodeFloatVariation.cpp
--- a/Engine/Plugins/Experimental/Mutable/Source/CustomizableObjectEditor/Private/MuCOE/Nodes/CustomizableObjectNodeFloatVariation.cpp
+++ b/Engine/Plugins/Experimental/Mutable/Source/CustomizableObjectEditor/Private/MuCOE/Nodes/CustomizableObjectNodeFloatVariation.cpp
@@ -2,74 +2,30 @@
 
 #include "MuCOE/Nodes/CustomizableObjectNodeFloatVariation.h"
 
+#include "MuCOE/CustomizableObjectEditor_Deprecated.h"
 #include "MuCOE/EdGraphSchema_CustomizableObject.h"
-#include "MuCOE/UnrealEditorPortabilityHelpers.h"
 
-class UCustomizableObjectNodeRemapPins;
 
-#define LOCTEXT_NAMESPACE "CustomizableObjectEditor"
-
-
-UCustomizableObjectNodeFloatVariation::UCustomizableObjectNodeFloatVariation()
-	: Super()
+void UCustomizableObjectNodeFloatVariation::BackwardsCompatibleFixup()
 {
+	const int32 CustomizableObjectCustomVersion = GetLinkerCustomVersion(FCustomizableObjectCustomVersion::GUID);
 
-}
-
-
-void UCustomizableObjectNodeFloatVariation::PostEditChangeProperty(FPropertyChangedEvent& PropertyChangedEvent)
-{
-	FProperty* PropertyThatChanged = PropertyChangedEvent.Property;
-	if ( PropertyThatChanged  )
+	if (CustomizableObjectCustomVersion < FCustomizableObjectCustomVersion::DeduplicateNodeVariant)
 	{
-		ReconstructNode();
+		for (const FCustomizableObjectFloatVariation& OldVariation : Variations_DEPRECATED)
+		{
+			FCustomizableObjectVariation Variation;
+			Variation.Tag = OldVariation.Tag;
+			
+			Variations.Add(Variation);
+		}
 	}
 
-	Super::PostEditChangeProperty(PropertyChangedEvent);
+	Super::BackwardsCompatibleFixup();
 }
 
 
-void UCustomizableObjectNodeFloatVariation::AllocateDefaultPins(UCustomizableObjectNodeRemapPins* RemapPins)
+FName UCustomizableObjectNodeFloatVariation::GetCategory() const
 {
-	const UEdGraphSchema_CustomizableObject* Schema = GetDefault<UEdGraphSchema_CustomizableObject>();
-
-	UEdGraphPin* OutputPin = CustomCreatePin(EGPD_Output, Schema->PC_Float, FName(TEXT("Float")));
-	OutputPin->bDefaultValueIsIgnored = true;
-
-	for ( int VariationIndex = Variations.Num() - 1; VariationIndex >= 0; --VariationIndex)
-	{
-		FString RemapIndex = Helper_GetPinName(VariationPin(VariationIndex));
-
-		FString PinName = FString::Printf( TEXT("Variation %d"), VariationIndex);
-		UEdGraphPin* VariationPin = CustomCreatePin(EGPD_Input, Schema->PC_Float, FName(*PinName), false);
-		VariationPin->bDefaultValueIsIgnored = true;
-
-		FString FriendlyName = FString::Printf(TEXT("Variation %d [%s]"), VariationIndex, *Variations[VariationIndex].Tag);
-		VariationPin->PinFriendlyName = FText::FromString(FriendlyName);
-	}
-
-	UEdGraphPin* DefaultVariation = CustomCreatePin(EGPD_Input, Schema->PC_Float, FName(TEXT("Default")), false);
-	DefaultVariation->bDefaultValueIsIgnored = true;
+	return UEdGraphSchema_CustomizableObject::PC_Float;
 }
-
-
-FText UCustomizableObjectNodeFloatVariation::GetNodeTitle(ENodeTitleType::Type TitleType) const
-{
-	return LOCTEXT("Float_Variation", "Float Variation");
-}
-
-
-FLinearColor UCustomizableObjectNodeFloatVariation::GetNodeTitleColor() const
-{
-	const UEdGraphSchema_CustomizableObject* Schema = GetDefault<UEdGraphSchema_CustomizableObject>();
-	return Schema->GetPinTypeColor(Schema->PC_Float);
-}
-
-
-FText UCustomizableObjectNodeFloatVariation::GetTooltipText() const
-{
-	return LOCTEXT("Float_Variation_Tooltip", "Select a float depending on what tags are active.");
-}
-
-#undef LOCTEXT_NAMESPACE
-
diff --git a/Engine/Plugins/Experimental/Mutable/Source/CustomizableObjectEditor/Private/MuCOE/Nodes/CustomizableObjectNodeMaterialVariation.cpp b/Engine/Plugins/Experimental/Mutable/Source/CustomizableObjectEditor/Private/MuCOE/Nodes/CustomizableObjectNodeMaterialVariation.cpp
--- a/Engine/Plugins/Experimental/Mutable/Source/CustomizableObjectEditor/Private/MuCOE/Nodes/CustomizableObjectNodeMaterialVariation.cpp
+++ b/Engine/Plugins/Experimental/Mutable/Source/CustomizableObjectEditor/Private/MuCOE/Nodes/CustomizableObjectNodeMaterialVariation.cpp
@@ -2,75 +2,36 @@
 
 #include "MuCOE/Nodes/CustomizableObjectNodeMaterialVariation.h"
 
+#include "MuCOE/CustomizableObjectEditor_Deprecated.h"
 #include "MuCOE/EdGraphSchema_CustomizableObject.h"
 
-class UCustomizableObjectNodeRemapPins;
 
-#define LOCTEXT_NAMESPACE "CustomizableObjectEditor"
-
-
-UCustomizableObjectNodeMaterialVariation::UCustomizableObjectNodeMaterialVariation()
-	: Super()
-{
-
-}
-
-
-void UCustomizableObjectNodeMaterialVariation::PostEditChangeProperty(FPropertyChangedEvent& PropertyChangedEvent)
-{
-	FProperty* PropertyThatChanged = PropertyChangedEvent.Property;
-	if ( PropertyThatChanged 
-		//&& PropertyThatChanged->GetName() == TEXT("Variations")
-		)
-	{
-		ReconstructNode();
-	}
-
-	Super::PostEditChangeProperty(PropertyChangedEvent);
-}
-
-
-void UCustomizableObjectNodeMaterialVariation::AllocateDefaultPins(UCustomizableObjectNodeRemapPins* RemapPins)
+void UCustomizableObjectNodeMaterialVariation::BackwardsCompatibleFixup()
 {
-	const UEdGraphSchema_CustomizableObject* Schema = GetDefault<UEdGraphSchema_CustomizableObject>();
+	const int32 CustomizableObjectCustomVersion = GetLinkerCustomVersion(FCustomizableObjectCustomVersion::GUID);
 
-	FString PinName = TEXT("Material");
-	UEdGraphPin* OutputPin = CustomCreatePin(EGPD_Output, Schema->PC_Material, FName(*PinName));
-	OutputPin->bDefaultValueIsIgnored = true;
-
-	for ( int VariationIndex = Variations.Num() - 1; VariationIndex >= 0; --VariationIndex)
+	if (CustomizableObjectCustomVersion < FCustomizableObjectCustomVersion::DeduplicateNodeVariant)
 	{
-		PinName = FString::Printf( TEXT("Variation %d"), VariationIndex);
-		UEdGraphPin* VariationPin = CustomCreatePin(EGPD_Input, Schema->PC_Material, FName(*PinName), true);
-		VariationPin->bDefaultValueIsIgnored = true;
-
-		FString FriendlyName = FString::Printf(TEXT("Variation %d [%s]"), VariationIndex, *Variations[VariationIndex].Tag);
-		VariationPin->PinFriendlyName = FText::FromString(FriendlyName);
+		for (const FCustomizableObjectMaterialVariation& OldVariation : Variations_DEPRECATED)
+		{
+			FCustomizableObjectVariation Variation;
+			Variation.Tag = OldVariation.Tag;
+			
+			Variations.Add(Variation);
+		}
 	}
 
-	PinName = TEXT("Default");
-	UEdGraphPin* DefaultVariation = CustomCreatePin(EGPD_Input, Schema->PC_Material, FName(*PinName), true);
-	DefaultVariation->bDefaultValueIsIgnored = true;
+	Super::BackwardsCompatibleFixup();
 }
 
 
-FText UCustomizableObjectNodeMaterialVariation::GetNodeTitle(ENodeTitleType::Type TitleType) const
+FName UCustomizableObjectNodeMaterialVariation::GetCategory() const
 {
-	return LOCTEXT("Material_Variation", "Material Variation");
+	return UEdGraphSchema_CustomizableObject::PC_Material;
 }
 
 
-FLinearColor UCustomizableObjectNodeMaterialVariation::GetNodeTitleColor() const
+bool UCustomizableObjectNodeMaterialVariation::IsInputPinArray() const
 {
-	const UEdGraphSchema_CustomizableObject* Schema = GetDefault<UEdGraphSchema_CustomizableObject>();
-	return Schema->GetPinTypeColor(Schema->PC_Material);
+	return true;
 }
-
-
-FText UCustomizableObjectNodeMaterialVariation::GetTooltipText() const
-{
-	return LOCTEXT("Material_Variation_Tooltip", "Changes the materials given depending on what tags or state is active.");
-}
-
-#undef LOCTEXT_NAMESPACE
-
diff --git a/Engine/Plugins/Experimental/Mutable/Source/CustomizableObjectEditor/Private/MuCOE/Nodes/CustomizableObjectNodeMeshSwitch.cpp b/Engine/Plugins/Experimental/Mutable/Source/CustomizableObjectEditor/Private/MuCOE/Nodes/CustomizableObjectNodeMeshSwitch.cpp
--- a/Engine/Plugins/Experimental/Mutable/Source/CustomizableObjectEditor/Private/MuCOE/Nodes/CustomizableObjectNodeMeshSwitch.cpp
+++ b/Engine/Plugins/Experimental/Mutable/Source/CustomizableObjectEditor/Private/MuCOE/Nodes/CustomizableObjectNodeMeshSwitch.cpp
@@ -7,36 +7,24 @@
 #define LOCTEXT_NAMESPACE "CustomizableObjectEditor"
 
 
-void UCustomizableObjectNodeMeshSwitch::Serialize(FArchive& Ar)
+void UCustomizableObjectNodeMeshSwitch::BackwardsCompatibleFixup()
 {
-	Super::Serialize(Ar);
+	Super::BackwardsCompatibleFixup();
 
-	Ar.UsingCustomVersion(FCustomizableObjectCustomVersion::GUID);
+	const int32 CustomizableObjectCustomVersion = GetLinkerCustomVersion(FCustomizableObjectCustomVersion::GUID);
 
-	if (Ar.CustomVer(FCustomizableObjectCustomVersion::GUID) < FCustomizableObjectCustomVersion::BugPinsSwitch)
+	if (CustomizableObjectCustomVersion < FCustomizableObjectCustomVersion::BugPinsSwitch)
 	{
 		OutputPinReference = FindPin(TEXT("Mesh"));	
 	}
 }
 
 
-FString UCustomizableObjectNodeMeshSwitch::GetOutputPinName() const
-{
-	return TEXT("Mesh");
-}
-
-
 FName UCustomizableObjectNodeMeshSwitch::GetCategory() const
 {
 	return UEdGraphSchema_CustomizableObject::PC_Mesh;
 }
 
 
-FString UCustomizableObjectNodeMeshSwitch::GetPinPrefix() const
-{
-	return TEXT("Mesh ");
-}
-
-
 #undef LOCTEXT_NAMESPACE
 
diff --git a/Engine/Plugins/Experimental/Mutable/Source/CustomizableObjectEditor/Private/MuCOE/Nodes/CustomizableObjectNodeMeshVariation.cpp b/Engine/Plugins/Experimental/Mutable/Source/CustomizableObjectEditor/Private/MuCOE/Nodes/CustomizableObjectNodeMeshVariation.cpp
--- a/Engine/Plugins/Experimental/Mutable/Source/CustomizableObjectEditor/Private/MuCOE/Nodes/CustomizableObjectNodeMeshVariation.cpp
+++ b/Engine/Plugins/Experimental/Mutable/Source/CustomizableObjectEditor/Private/MuCOE/Nodes/CustomizableObjectNodeMeshVariation.cpp
@@ -2,71 +2,30 @@
 
 #include "MuCOE/Nodes/CustomizableObjectNodeMeshVariation.h"
 
+#include "MuCOE/CustomizableObjectEditor_Deprecated.h"
 #include "MuCOE/EdGraphSchema_CustomizableObject.h"
 
-class UCustomizableObjectNodeRemapPins;
 
-#define LOCTEXT_NAMESPACE "CustomizableObjectEditor"
-
-
-UCustomizableObjectNodeMeshVariation::UCustomizableObjectNodeMeshVariation()
-	: Super()
+void UCustomizableObjectNodeMeshVariation::BackwardsCompatibleFixup()
 {
+	const int32 CustomizableObjectCustomVersion = GetLinkerCustomVersion(FCustomizableObjectCustomVersion::GUID);
 
-}
-
-
-void UCustomizableObjectNodeMeshVariation::PostEditChangeProperty(FPropertyChangedEvent& PropertyChangedEvent)
-{
-	FProperty* PropertyThatChanged = PropertyChangedEvent.Property;
-	if ( PropertyThatChanged  )
+	if (CustomizableObjectCustomVersion < FCustomizableObjectCustomVersion::DeduplicateNodeVariant)
 	{
-		ReconstructNode();
+		for (const FCustomizableObjectMeshVariation& OldVariation : Variations_DEPRECATED)
+		{
+			FCustomizableObjectVariation Variation;
+			Variation.Tag = OldVariation.Tag;
+			
+			Variations.Add(Variation);
+		}
 	}
 
-	Super::PostEditChangeProperty(PropertyChangedEvent);
+	Super::BackwardsCompatibleFixup();
 }
 
 
-void UCustomizableObjectNodeMeshVariation::AllocateDefaultPins(UCustomizableObjectNodeRemapPins* RemapPins)
+FName UCustomizableObjectNodeMeshVariation::GetCategory() const
 {
-	const UEdGraphSchema_CustomizableObject* Schema = GetDefault<UEdGraphSchema_CustomizableObject>();
-	
-	UEdGraphPin* OutputPin = CustomCreatePin(EGPD_Output, Schema->PC_Mesh, FName(TEXT("Mesh")));
-	OutputPin->bDefaultValueIsIgnored = true;
-
-	for ( int VariationIndex = Variations.Num() - 1; VariationIndex >= 0; --VariationIndex)
-	{
-		FString PinName = FString::Printf( TEXT("Variation %d"), VariationIndex);
-		UEdGraphPin* VariationPin = CustomCreatePin(EGPD_Input, Schema->PC_Mesh, FName(*PinName), false);
-		VariationPin->bDefaultValueIsIgnored = true;
-
-		FString FriendlyName = FString::Printf(TEXT("Variation %d [%s]"), VariationIndex, *Variations[VariationIndex].Tag);
-		VariationPin->PinFriendlyName = FText::FromString(FriendlyName);
-	}
-
-	UEdGraphPin* DefaultVariation = CustomCreatePin(EGPD_Input, Schema->PC_Mesh, FName(TEXT("Default")), false);
-	DefaultVariation->bDefaultValueIsIgnored = true;
+	return UEdGraphSchema_CustomizableObject::PC_Mesh;
 }
-
-
-FText UCustomizableObjectNodeMeshVariation::GetNodeTitle(ENodeTitleType::Type TitleType) const
-{
-	return LOCTEXT("Mesh_Variation", "Mesh Variation");
-}
-
-
-FLinearColor UCustomizableObjectNodeMeshVariation::GetNodeTitleColor() const
-{
-	const UEdGraphSchema_CustomizableObject* Schema = GetDefault<UEdGraphSchema_CustomizableObject>();
-	return Schema->GetPinTypeColor(Schema->PC_Mesh);
-}
-
-
-FText UCustomizableObjectNodeMeshVariation::GetTooltipText() const
-{
-	return LOCTEXT("Mesh_Variation_Tooltip", "Select a mesh depending on what tags are active.");
-}
-
-#undef LOCTEXT_NAMESPACE
-
diff --git a/Engine/Plugins/Experimental/Mutable/Source/CustomizableObjectEditor/Private/MuCOE/Nodes/CustomizableObjectNodePassThroughTextureSwitch.cpp b/Engine/Plugins/Experimental/Mutable/Source/CustomizableObjectEditor/Private/MuCOE/Nodes/CustomizableObjectNodePassThroughTextureSwitch.cpp
--- a/Engine/Plugins/Experimental/Mutable/Source/CustomizableObjectEditor/Private/MuCOE/Nodes/CustomizableObjectNodePassThroughTextureSwitch.cpp
+++ b/Engine/Plugins/Experimental/Mutable/Source/CustomizableObjectEditor/Private/MuCOE/Nodes/CustomizableObjectNodePassThroughTextureSwitch.cpp
@@ -7,23 +7,11 @@
 #define LOCTEXT_NAMESPACE "CustomizableObjectEditor"
 
 
-FString UCustomizableObjectNodePassThroughTextureSwitch::GetOutputPinName() const
-{
-	return TEXT("PassThrough Texture");
-}
-
-
 FName UCustomizableObjectNodePassThroughTextureSwitch::GetCategory() const
 {
 	return UEdGraphSchema_CustomizableObject::PC_PassThroughImage;
 }
 
 
-FString UCustomizableObjectNodePassThroughTextureSwitch::GetPinPrefix() const
-{
-	return TEXT("PassThrough Texture ");
-}
-
-
 #undef LOCTEXT_NAMESPACE
 
diff --git a/Engine/Plugins/Experimental/Mutable/Source/CustomizableObjectEditor/Private/MuCOE/Nodes/CustomizableObjectNodeTextureSwitch.cpp b/Engine/Plugins/Experimental/Mutable/Source/CustomizableObjectEditor/Private/MuCOE/Nodes/CustomizableObjectNodeTextureSwitch.cpp
--- a/Engine/Plugins/Experimental/Mutable/Source/CustomizableObjectEditor/Private/MuCOE/Nodes/CustomizableObjectNodeTextureSwitch.cpp
+++ b/Engine/Plugins/Experimental/Mutable/Source/CustomizableObjectEditor/Private/MuCOE/Nodes/CustomizableObjectNodeTextureSwitch.cpp
@@ -7,25 +7,17 @@
 #define LOCTEXT_NAMESPACE "CustomizableObjectEditor"
 
 
-void UCustomizableObjectNodeTextureSwitch::Serialize(FArchive& Ar)
-{
-	Super::Serialize(Ar);
-	
-	Ar.UsingCustomVersion(FCustomizableObjectCustomVersion::GUID);
-
-	if (Ar.CustomVer(FCustomizableObjectCustomVersion::GUID) < FCustomizableObjectCustomVersion::PinsNamesImageToTexture)
-	{
-		OutputPinReference = FindPin(TEXT("Image"));
-	}
-}
-
-
 void UCustomizableObjectNodeTextureSwitch::BackwardsCompatibleFixup()
 {
 	Super::BackwardsCompatibleFixup();
 
 	const int32 CustomizableObjectCustomVersion = GetLinkerCustomVersion(FCustomizableObjectCustomVersion::GUID);
 
+	if (CustomizableObjectCustomVersion < FCustomizableObjectCustomVersion::PinsNamesImageToTexture)
+	{
+		OutputPinReference = FindPin(TEXT("Image"));
+	}
+
 	if (CustomizableObjectCustomVersion < FCustomizableObjectCustomVersion::FixPinsNamesImageToTexture2)
 	{
 		if (UEdGraphPin* TexturePin = FindPin(TEXT("Image")))
@@ -37,23 +29,11 @@ void UCustomizableObjectNodeTextureSwitch::BackwardsCompatibleFixup()
 }
 
 
-FString UCustomizableObjectNodeTextureSwitch::GetOutputPinName() const
-{
-	return TEXT("Texture");
-}
-
-
 FName UCustomizableObjectNodeTextureSwitch::GetCategory() const
 {
 	return UEdGraphSchema_CustomizableObject::PC_Image;
 }
 
 
-FString UCustomizableObjectNodeTextureSwitch::GetPinPrefix() const
-{
-	return TEXT("Texture ");
-}
-
-
 #undef LOCTEXT_NAMESPACE
 
diff --git a/Engine/Plugins/Experimental/Mutable/Source/CustomizableObjectEditor/Private/MuCOE/Nodes/CustomizableObjectNodeVariation.cpp b/Engine/Plugins/Experimental/Mutable/Source/CustomizableObjectEditor/Private/MuCOE/Nodes/CustomizableObjectNodeVariation.cpp
new file mode 100644
--- /dev/null
+++ b/Engine/Plugins/Experimental/Mutable/Source/CustomizableObjectEditor/Private/MuCOE/Nodes/CustomizableObjectNodeVariation.cpp
@@ -0,0 +1,113 @@
+// Copyright Epic Games, Inc. All Rights Reserved.
+
+#include "MuCOE/Nodes/CustomizableObjectNodeVariation.h"
+
+#include "MuCOE/EdGraphSchema_CustomizableObject.h"
+
+class UCustomizableObjectNodeRemapPins;
+
+#define LOCTEXT_NAMESPACE "CustomizableObjectEditor"
+
+
+void UCustomizableObjectNodeVariation::PostEditChangeProperty(FPropertyChangedEvent& PropertyChangedEvent)
+{
+	if (PropertyChangedEvent.Property)
+	{
+		ReconstructNode();
+	}
+
+	Super::PostEditChangeProperty(PropertyChangedEvent);
+}
+
+
+void UCustomizableObjectNodeVariation::AllocateDefaultPins(UCustomizableObjectNodeRemapPins* RemapPins)
+{
+	const UEdGraphSchema_CustomizableObject* Schema = GetDefault<UEdGraphSchema_CustomizableObject>();
+
+	const FName Category = GetCategory();
+	const bool bIsInputPinArray = IsInputPinArray();
+
+	{
+		const FName PinName = FName(UEdGraphSchema_CustomizableObject::GetPinCategoryName(GetCategory()).ToString());
+		CustomCreatePin(EGPD_Output, Category, PinName);
+	}
+	
+	VariationsPins.SetNum(Variations.Num());
+	for (int32 VariationIndex = Variations.Num() - 1; VariationIndex >= 0; --VariationIndex)
+	{
+		const FName PinName = FName(FString::Printf( TEXT("Variation %d"), VariationIndex));
+		UEdGraphPin* VariationPin = CustomCreatePin(EGPD_Input, Category, PinName, bIsInputPinArray);
+
+		VariationPin->PinFriendlyName = FText::Format(LOCTEXT("Variation_Pin_FriendlyName", "Variation {0} [{1}]"), VariationIndex, FText::FromString(*Variations[VariationIndex].Tag));
+		
+		VariationsPins[VariationIndex] = VariationPin;
+	}
+
+	CustomCreatePin(EGPD_Input, Category, FName(TEXT("Default")), bIsInputPinArray);
+}
+
+
+bool UCustomizableObjectNodeVariation::IsInputPinArray() const
+{
+	return false;
+}
+
+
+int32 UCustomizableObjectNodeVariation::GetNumVariations() const
+{
+	return Variations.Num();
+}
+
+
+const FCustomizableObjectVariation& UCustomizableObjectNodeVariation::GetVariation(int32 Index) const
+{
+	return Variations[Index];
+}
+
+
+UEdGraphPin* UCustomizableObjectNodeVariation::DefaultPin() const
+{
+	return FindPin(TEXT("Default"));
+}
+
+
+UEdGraphPin* UCustomizableObjectNodeVariation::VariationPin(int32 Index) const
+{
+	return VariationsPins[Index].Get();
+}
+
+
+FText UCustomizableObjectNodeVariation::GetNodeTitle(ENodeTitleType::Type TitleType) const
+{
+	return FText::Format(LOCTEXT("Variation_Node_Title", "{0} Variation"), UEdGraphSchema_CustomizableObject::GetPinCategoryName(GetCategory()));
+}
+
+
+FLinearColor UCustomizableObjectNodeVariation::GetNodeTitleColor() const
+{
+	const UEdGraphSchema_CustomizableObject* Schema = GetDefault<UEdGraphSchema_CustomizableObject>();
+	return Schema->GetPinTypeColor(GetCategory());
+}
+
+
+FText UCustomizableObjectNodeVariation::GetTooltipText() const
+{
+	return FText::Format(LOCTEXT("Variation_Tooltip", "Select a {0} depending on what tags are active."), UEdGraphSchema_CustomizableObject::GetPinCategoryName(GetCategory()));
+}
+
+
+void UCustomizableObjectNodeVariation::BackwardsCompatibleFixup()
+{
+	Super::BackwardsCompatibleFixup();
+
+	const int32 CustomizableObjectCustomVersion = GetLinkerCustomVersion(FCustomizableObjectCustomVersion::GUID);
+
+	if (CustomizableObjectCustomVersion < FCustomizableObjectCustomVersion::DeduplicateNodeVariant)
+	{
+		ReconstructNode();
+	}
+}
+
+
+#undef LOCTEXT_NAMESPACE
+
diff --git a/Engine/Plugins/Experimental/Mutable/Source/CustomizableObjectEditor/Private/MuCOE/SCustomizableObjectEditorTagExplorer.cpp b/Engine/Plugins/Experimental/Mutable/Source/CustomizableObjectEditor/Private/MuCOE/SCustomizableObjectEditorTagExplorer.cpp
--- a/Engine/Plugins/Experimental/Mutable/Source/CustomizableObjectEditor/Private/MuCOE/SCustomizableObjectEditorTagExplorer.cpp
+++ b/Engine/Plugins/Experimental/Mutable/Source/CustomizableObjectEditor/Private/MuCOE/SCustomizableObjectEditorTagExplorer.cpp
@@ -251,7 +251,7 @@ void SCustomizableObjectEditorTagExplorer::FillTagInformation(UCustomizableObjec
 		{
 			if (UCustomizableObjectNodeMaterial* TypedNodeMat = Cast<UCustomizableObjectNodeMaterial>(Node))
 			{
-				for (int i = 0; i < TypedNodeMat->Tags.Num(); ++i)
+				for (int32 i = 0; i < TypedNodeMat->Tags.Num(); ++i)
 				{
 					NodeTags.Add(TypedNodeMat->Tags[i], TypedNodeMat);
 					if (Tags.Find(TypedNodeMat->Tags[i]) == INDEX_NONE)
@@ -263,19 +263,20 @@ void SCustomizableObjectEditorTagExplorer::FillTagInformation(UCustomizableObjec
 
 			if (UCustomizableObjectNodeMaterialVariation* TypedNodeVariations = Cast<UCustomizableObjectNodeMaterialVariation>(Node))
 			{
-				for (int i = 0; i < TypedNodeVariations->Variations.Num(); ++i)
+				for (int32 i = 0; i < TypedNodeVariations->GetNumVariations(); ++i)
 				{
-					NodeTags.Add(TypedNodeVariations->Variations[i].Tag, TypedNodeVariations);
-					if (Tags.Find(TypedNodeVariations->Variations[i].Tag) == INDEX_NONE)
+					const FString& VairationTag = TypedNodeVariations->GetVariation(i).Tag; 
+					NodeTags.Add(VairationTag, TypedNodeVariations);
+					if (Tags.Find(VairationTag) == INDEX_NONE)
 					{
-						Tags.Add(TypedNodeVariations->Variations[i].Tag);
+						Tags.Add(VairationTag);
 					}
 				}
 			}
 
 			if (UCustomizableObjectNodeMeshClipMorph* TypedNodeClipMorph = Cast<UCustomizableObjectNodeMeshClipMorph>(Node))
 			{
-				for (int i = 0; i < TypedNodeClipMorph->Tags.Num(); ++i)
+				for (int32 i = 0; i < TypedNodeClipMorph->Tags.Num(); ++i)
 				{
 					NodeTags.Add(TypedNodeClipMorph->Tags[i], TypedNodeClipMorph);
 					if (Tags.Find(TypedNodeClipMorph->Tags[i]) == INDEX_NONE)
@@ -287,7 +288,7 @@ void SCustomizableObjectEditorTagExplorer::FillTagInformation(UCustomizableObjec
 
 			if (UCustomizableObjectNodeMeshClipWithMesh* TypedNodeClipMesh = Cast<UCustomizableObjectNodeMeshClipWithMesh>(Node))
 			{
-				for (int i = 0; i < TypedNodeClipMesh->Tags.Num(); ++i)
+				for (int32 i = 0; i < TypedNodeClipMesh->Tags.Num(); ++i)
 				{
 					NodeTags.Add(TypedNodeClipMesh->Tags[i], TypedNodeClipMesh);
 					if (Tags.Find(TypedNodeClipMesh->Tags[i]) == INDEX_NONE)
@@ -299,7 +300,7 @@ void SCustomizableObjectEditorTagExplorer::FillTagInformation(UCustomizableObjec
 
 			if (UCustomizableObjectNodeExtendMaterial* TypedNodeExtend = Cast<UCustomizableObjectNodeExtendMaterial>(Node))
 			{
-				for (int i = 0; i < TypedNodeExtend->Tags.Num(); ++i)
+				for (int32 i = 0; i < TypedNodeExtend->Tags.Num(); ++i)
 				{
 					NodeTags.Add(TypedNodeExtend->Tags[i], TypedNodeExtend);
 					if (Tags.Find(TypedNodeExtend->Tags[i]) == INDEX_NONE)
@@ -311,7 +312,7 @@ void SCustomizableObjectEditorTagExplorer::FillTagInformation(UCustomizableObjec
 
 			if (UCustomizableObjectNodeMeshClipDeform* TypedNodeClipDeform = Cast<UCustomizableObjectNodeMeshClipDeform>(Node))
 			{
-				for (int i = 0; i < TypedNodeClipDeform->Tags.Num(); ++i)
+				for (int32 i = 0; i < TypedNodeClipDeform->Tags.Num(); ++i)
 				{
 					NodeTags.Add(TypedNodeClipDeform->Tags[i], TypedNodeClipDeform);
 					if (Tags.Find(TypedNodeClipDeform->Tags[i]) == INDEX_NONE)
diff --git a/Engine/Plugins/Experimental/Mutable/Source/CustomizableObjectEditor/Public/MuCOE/CustomizableObjectEditor_Deprecated.h b/Engine/Plugins/Experimental/Mutable/Source/CustomizableObjectEditor/Public/MuCOE/CustomizableObjectEditor_Deprecated.h
--- a/Engine/Plugins/Experimental/Mutable/Source/CustomizableObjectEditor/Public/MuCOE/CustomizableObjectEditor_Deprecated.h
+++ b/Engine/Plugins/Experimental/Mutable/Source/CustomizableObjectEditor/Public/MuCOE/CustomizableObjectEditor_Deprecated.h
@@ -128,6 +128,49 @@ struct CUSTOMIZABLEOBJECTEDITOR_API FCustomizableObjectNodeSkeletalMeshLOD
 	TArray<FCustomizableObjectNodeSkeletalMeshMaterial> Materials;
 };
 
+// UCustomizableObjectNodeColorVariation
+// Deprecated, do not use!
+USTRUCT()
+struct CUSTOMIZABLEOBJECTEDITOR_API FCustomizableObjectColorVariation
+{
+	GENERATED_USTRUCT_BODY()
+
+	UPROPERTY(EditAnywhere, Category = CustomizableObject)
+	FString Tag;
+};
+
+// UCustomizableObjectFloatVariation
+// Deprecated, do not use!
+USTRUCT()
+struct CUSTOMIZABLEOBJECTEDITOR_API FCustomizableObjectFloatVariation
+{
+	GENERATED_USTRUCT_BODY()
+
+	UPROPERTY(EditAnywhere, Category = CustomizableObject)
+	FString Tag;
+};
+
+// UCustomizableObjectMaterialVariation
+// Deprecated, do not use!
+USTRUCT()
+struct CUSTOMIZABLEOBJECTEDITOR_API FCustomizableObjectMaterialVariation
+{
+	GENERATED_USTRUCT_BODY()
+
+	UPROPERTY(EditAnywhere, Category = CustomizableObject)
+	FString Tag;
+};
+
+// UCustomizableObjectMeshVariation
+// Deprecated, do not use!
+USTRUCT()
+struct CUSTOMIZABLEOBJECTEDITOR_API FCustomizableObjectMeshVariation
+{
+	GENERATED_USTRUCT_BODY()
+
+	UPROPERTY(EditAnywhere, Category = CustomizableObject)
+	FString Tag;
+};
 
 
 #if UE_ENABLE_INCLUDE_ORDER_DEPRECATED_IN_5_2
