[Mutable] Fix issues with old serialized Texture Invert nodes.

Removed duplicated code which was producing the error.

#rnx

diff --git a/Engine/Plugins/Experimental/Mutable/Source/CustomizableObjectEditor/Classes/MuCOE/Nodes/CustomizableObjectNodeTextureInvert.h b/Engine/Plugins/Experimental/Mutable/Source/CustomizableObjectEditor/Classes/MuCOE/Nodes/CustomizableObjectNodeTextureInvert.h
--- a/Engine/Plugins/Experimental/Mutable/Source/CustomizableObjectEditor/Classes/MuCOE/Nodes/CustomizableObjectNodeTextureInvert.h
+++ b/Engine/Plugins/Experimental/Mutable/Source/CustomizableObjectEditor/Classes/MuCOE/Nodes/CustomizableObjectNodeTextureInvert.h
@@ -17,16 +17,16 @@ class CUSTOMIZABLEOBJECTEDITOR_API UCustomizableObjectNodeTextureInvert : public
 public:
 	GENERATED_BODY()
 	
-	// Begin EdGraphNode interface
-	FText GetNodeTitle(ENodeTitleType::Type TittleType)const override;
-	FLinearColor GetNodeTitleColor() const override;
-	FText GetTooltipText() const override;
-	virtual void Serialize(FArchive& Ar) override;
+	// EdGraphNode interface
+	virtual FText GetNodeTitle(ENodeTitleType::Type TittleType)const override;
+	virtual FLinearColor GetNodeTitleColor() const override;
+	virtual FText GetTooltipText() const override;
 	
 	// UCustomizableObjectNode interface
 	virtual void AllocateDefaultPins(UCustomizableObjectNodeRemapPins* RemapPins) override;
 	virtual void BackwardsCompatibleFixup() override;
 
+	// Own interface
 	UEdGraphPin* GetBaseImagePin() const;
 
 private:
diff --git a/Engine/Plugins/Experimental/Mutable/Source/CustomizableObjectEditor/Private/MuCOE/Nodes/CustomizableObjectNodeTextureInvert.cpp b/Engine/Plugins/Experimental/Mutable/Source/CustomizableObjectEditor/Private/MuCOE/Nodes/CustomizableObjectNodeTextureInvert.cpp
--- a/Engine/Plugins/Experimental/Mutable/Source/CustomizableObjectEditor/Private/MuCOE/Nodes/CustomizableObjectNodeTextureInvert.cpp
+++ b/Engine/Plugins/Experimental/Mutable/Source/CustomizableObjectEditor/Private/MuCOE/Nodes/CustomizableObjectNodeTextureInvert.cpp
@@ -30,6 +30,11 @@ void UCustomizableObjectNodeTextureInvert::BackwardsCompatibleFixup()
 
 	const int32 CustomizableObjectCustomVersion = GetLinkerCustomVersion(FCustomizableObjectCustomVersion::GUID);
 
+	if (CustomizableObjectCustomVersion < FCustomizableObjectCustomVersion::PinsNamesImageToTexture)
+	{
+		BaseImagePinReference = FEdGraphPinReference(FindPin(TEXT("Base Image")));
+	}
+	
 	if (CustomizableObjectCustomVersion < FCustomizableObjectCustomVersion::FixPinsNamesImageToTexture2)
 	{
 		bool Replaced = false;
@@ -78,25 +83,4 @@ FText UCustomizableObjectNodeTextureInvert::GetTooltipText() const
 }
 
 
-void UCustomizableObjectNodeTextureInvert::Serialize(FArchive& Ar)
-{
-	Super::Serialize(Ar);
-
-	Ar.UsingCustomVersion(FCustomizableObjectCustomVersion::GUID);
-
-	if (Ar.CustomVer(FCustomizableObjectCustomVersion::GUID) < FCustomizableObjectCustomVersion::PinsNamesImageToTexture)
-	{
-		BaseImagePinReference = FEdGraphPinReference(FindPin(TEXT("Base Image")));
-	}
-
-	if (Ar.CustomVer(FCustomizableObjectCustomVersion::GUID) < FCustomizableObjectCustomVersion::FixPinsNamesImageToTexture2)
-	{
-		if (UEdGraphPin* TexturePin = FindPin(TEXT("Base Image"))) {
-			TexturePin->PinName = TEXT("Base Texture");
-			UCustomizableObjectNode::ReconstructNode();
-		}
-	}
-}
-
-
 #undef LOCTEXT_NAMESPACE
