[Mutable] Optimize simple average Mipmap generation for 4 byte pixels.

- An integer average without overflow scheme is used to compute multiple channels at once without explicit simd registers.

#rnx

diff --git a/Engine/Plugins/Experimental/Mutable/Source/MutableRuntime/Private/MuR/OpImageMipmap.cpp b/Engine/Plugins/Experimental/Mutable/Source/MutableRuntime/Private/MuR/OpImageMipmap.cpp
--- a/Engine/Plugins/Experimental/Mutable/Source/MutableRuntime/Private/MuR/OpImageMipmap.cpp
+++ b/Engine/Plugins/Experimental/Mutable/Source/MutableRuntime/Private/MuR/OpImageMipmap.cpp
@@ -3,6 +3,7 @@
 #include "MuR/ImagePrivate.h"
 #include "MuR/SystemPrivate.h"
 #include "Async/ParallelFor.h"
+#include "HAL/UnrealMemory.h"
 
 namespace mu
 {
@@ -332,10 +333,12 @@ namespace mu
 			const uint8* pSource, uint8* Dest,
 			FIntVector2 sourceSize)
 		{
+			
 			const uint8* pMipSource = pSource;
 			uint8* pMipDest = Dest;
 
 			FIntVector2 destSize = sourceSize;
+
 			for (int m = 0; m < mips; ++m)
 			{
 				check(destSize[0] > 1 || destSize[1] > 1);
@@ -363,11 +366,37 @@ namespace mu
 
 					for (int x = 0; x < fullColumns; ++x)
 					{
-						for (int c = 0; c < PIXEL_SIZE; ++c)
+						if constexpr (PIXEL_SIZE == 4)
+						{
+							//const uint64 Row0Bits = *reinterpret_cast<const uint64*>(pSourceRow0);
+							//const uint64 Row1Bits = *reinterpret_cast<const uint64*>(pSourceRow1);
+
+							// Use memcpy to avoid any possible but improbable UB. memcpy should be optimized away by the compiler.
+							uint64 Row0Bits; 
+							uint64 Row1Bits;
+
+							FMemory::Memcpy(&Row0Bits, pSourceRow0, sizeof(uint64));
+							FMemory::Memcpy(&Row1Bits, pSourceRow1, sizeof(uint64));
+							
+							const uint64 XorRow0Row1Bits = Row0Bits ^ Row1Bits;
+
+							// Average of 2 unsigned integers without overflow extended to work on multiple bytes.
+							constexpr uint64 ShiftMask = 0xFEFEFEFEFEFEFEFE;
+							const uint64 ErrorCorrection = XorRow0Row1Bits & XorRow0Row1Bits & 0x0101010101010101;
+							const uint64 AvgLowBits = (Row0Bits & Row1Bits) + ((XorRow0Row1Bits & ShiftMask) >> 1) + ErrorCorrection;
+							const uint64 AvgHighBits = AvgLowBits >> 32;
+							const uint32 Result = (AvgLowBits & AvgHighBits) + (((AvgLowBits ^ AvgHighBits) & ShiftMask) >> 1);
+						
+							//*reinterpret_cast<uint32*>(pDestRow) = Result;
+							FMemory::Memcpy(pDestRow, &Result, sizeof(uint32));
+						}
+						else
 						{
-							int p = pSourceRow0[c] + pSourceRow0[PIXEL_SIZE + c] + pSourceRow1[c] +
-								pSourceRow1[PIXEL_SIZE + c];
-							pDestRow[c] = (uint8)(p >> 2);
+							for (int32 C = 0; C < PIXEL_SIZE; ++C)
+							{
+								int32 PixelSum = pSourceRow0[C] + pSourceRow0[PIXEL_SIZE + C] + pSourceRow1[C] + pSourceRow1[PIXEL_SIZE + C];
+								pDestRow[C] = (uint8)(PixelSum >> 2);
+							}
 						}
 
 						pSourceRow0 += 2 * PIXEL_SIZE;
@@ -377,58 +406,109 @@ namespace mu
 
 					if (strayColumn)
 					{
-						for (int c = 0; c < PIXEL_SIZE; ++c)
+						if constexpr (PIXEL_SIZE == 4)
+						{	
+							//const uint32 Row0Bits = *reinterpret_cast<const uint32*>(pSourceRow0);
+							//const uint32 Row1Bits = *reinterpret_cast<const uint32*>(pSourceRow1);
+
+							uint32 Row0Bits; 
+							uint32 Row1Bits;
+
+							FMemory::Memcpy(&Row0Bits, pSourceRow0, sizeof(uint32));
+							FMemory::Memcpy(&Row1Bits, pSourceRow1, sizeof(uint32));
+
+							// Average of 2 unsigned integers without overflow extended to work on multiple bytes.
+							constexpr uint32 ShiftMask = 0xFEFEFEFE;
+							const uint32 Result = (Row0Bits & Row1Bits) + (((Row0Bits ^ Row1Bits) & ShiftMask) >> 1);
+
+							//*reinterpret_cast<uint32*>(pDestRow) = Result;
+							FMemory::Memcpy(pDestRow, &Result, sizeof(uint32));
+						}
+						else
 						{
-							int p = pSourceRow0[c] + pSourceRow1[c];
-							pDestRow[c] = (uint8)(p >> 1);
+							for (int32 C = 0; C < PIXEL_SIZE; ++C)
+							{
+								int32 PixelSum = pSourceRow0[C] + pSourceRow1[C];
+								pDestRow[C] = (uint8)(PixelSum >> 1);
+							}
 						}
 					}
 				};
 
-					constexpr int PixelConcurrencyThreshold = 0xffff;
-					if (destSize[0] * destSize[1] < PixelConcurrencyThreshold)
-					{
-						for (int y = 0; y < fullRows; ++y)
-						{
-							ProcessRow(y);
-						}
-					}
-					else
+				constexpr int PixelConcurrencyThreshold = 0xffff;
+				if (destSize[0] * destSize[1] < PixelConcurrencyThreshold)
+				{
+					for (int y = 0; y < fullRows; ++y)
 					{
-						ParallelFor(fullRows, ProcessRow);
+						ProcessRow(y);
 					}
+				}
+				else
+				{
+					ParallelFor(fullRows, ProcessRow);
+				}
 
-					if (strayRow)
+				if (strayRow)
+				{
+					const uint8* pSourceRow0 = pMipSource + 2 * fullRows * sourceStride;
+					const uint8* pSourceRow1 = pSourceRow0 + sourceStride;
+					uint8* pDestRow = pMipDest + fullRows * destStride;
+
+					for (int x = 0; x < fullColumns; ++x)
 					{
-						const uint8* pSourceRow0 = pMipSource + 2 * fullRows * sourceStride;
-						const uint8* pSourceRow1 = pSourceRow0 + sourceStride;
-						uint8* pDestRow = pMipDest + fullRows * destStride;
+						if constexpr (PIXEL_SIZE == 4)
+						{
+							//const uint32 Col0Bits = *reinterpret_cast<const uint32*>(pSourceRow0);
+							//const uint32 Col1Bits = *reinterpret_cast<const uint32*>(pSourceRow0 + 4);
+
+							uint32 Col0Bits; 
+							uint32 Col1Bits;
+
+							FMemory::Memcpy(&Col0Bits, pSourceRow0, sizeof(uint32));
+							FMemory::Memcpy(&Col1Bits, pSourceRow0 + 4, sizeof(uint32));
 
-						for (int x = 0; x < fullColumns; ++x)
+							// Average of 2 unsigned integers without overflow extended to work on multiple bytes. 
+							// In this case we use the ceil variant to be consistent with the method used for 4 pixel average. 
+							constexpr uint32 ShiftMask = 0xFEFEFEFE;
+							const uint32 Result = (Col0Bits & Col1Bits) + (((Col0Bits ^ Col1Bits) & ShiftMask) >> 1);
+
+							//*reinterpret_cast<uint32*>(pDestRow) = Result;
+							FMemory::Memcpy(pDestRow, &Result, sizeof(uint32));
+						}
+						else
 						{
-							for (int c = 0; c < PIXEL_SIZE; ++c)
+							for (int32 C = 0; C < PIXEL_SIZE; ++C)
 							{
-								int p = pSourceRow0[c] + pSourceRow0[PIXEL_SIZE + c];
-								pDestRow[c] = (uint8)(p >> 1);
+								int32 p = pSourceRow0[C] + pSourceRow0[PIXEL_SIZE + C];
+								pDestRow[C] = (uint8)(p >> 1);
 							}
-
-							pSourceRow0 += 2 * PIXEL_SIZE;
-							pDestRow += PIXEL_SIZE;
 						}
 
-						if (strayColumn)
+						pSourceRow0 += 2 * PIXEL_SIZE;
+						pDestRow += PIXEL_SIZE;
+					}
+
+					if (strayColumn)
+					{
+						if constexpr (PIXEL_SIZE == 4)
 						{
-							for (int c = 0; c < PIXEL_SIZE; ++c)
+							//*reinterpret_cast<uint32*>(pDestRow) = *reinterpret_cast<const uint32*>(pSourceRow0);
+							FMemory::Memcpy(pDestRow, pSourceRow0, 4);
+						}
+						else
+						{
+							for (int32 C = 0; C < PIXEL_SIZE; ++C)
 							{
-								pDestRow[c] = pSourceRow0[c];
+								pDestRow[C] = pSourceRow0[C];
 							}
 						}
 					}
+				}
 
-					// Reset the source pointer for the next mip, to use the dest that we have just
-					// generated.
-					pMipSource = pMipDest;
-					pMipDest += destSize[0] * destSize[1] * PIXEL_SIZE;
+				// Reset the source pointer for the next mip, to use the dest that we have just
+				// generated.
+				pMipSource = pMipDest;
+				pMipDest += destSize[0] * destSize[1] * PIXEL_SIZE;
 			}
 		}
 
