[mutable] Adjust code optimization for image layer masks to use a single channel color check for masks.
Fix default value for missing color channels to match between images and colours.
#rnx

diff --git a/Engine/Plugins/Experimental/Mutable/Source/MutableRuntime/Private/MuR/CodeRunner.cpp b/Engine/Plugins/Experimental/Mutable/Source/MutableRuntime/Private/MuR/CodeRunner.cpp
--- a/Engine/Plugins/Experimental/Mutable/Source/MutableRuntime/Private/MuR/CodeRunner.cpp
+++ b/Engine/Plugins/Experimental/Mutable/Source/MutableRuntime/Private/MuR/CodeRunner.cpp
@@ -5396,7 +5396,7 @@ namespace mu
 
             case 1:
             {
-				FVector4f Result = FVector4f(1, 1, 1, 1);
+				FVector4f Result = FVector4f(0, 0, 0, 1);
 
 				for (int32 t = 0; t < MUTABLE_OP_MAX_SWIZZLE_CHANNELS; ++t)
 				{
diff --git a/Engine/Plugins/Experimental/Mutable/Source/MutableTools/Private/MuT/ASTOpImageLayer.cpp b/Engine/Plugins/Experimental/Mutable/Source/MutableTools/Private/MuT/ASTOpImageLayer.cpp
--- a/Engine/Plugins/Experimental/Mutable/Source/MutableTools/Private/MuT/ASTOpImageLayer.cpp
+++ b/Engine/Plugins/Experimental/Mutable/Source/MutableTools/Private/MuT/ASTOpImageLayer.cpp
@@ -187,7 +187,7 @@ namespace mu
 		// Convert to image layer color if blend is plain
 		if (!at && blendAt && blendAt->GetOpType() == OP_TYPE::IM_PLAINCOLOUR)
 		{
-			// TODO: May some blags be supported?
+			// TODO: May some flags be supported?
 			if (Flags == 0)
 			{
 				const ASTOpFixed* BlendPlainColor = dynamic_cast<const ASTOpFixed*>(blendAt.get());
@@ -209,17 +209,18 @@ namespace mu
 			FVector4f colour;
 			if (maskAt->IsImagePlainConstant(colour))
 			{
-				if (colour.IsNearlyZero3(UE_SMALL_NUMBER))
+				// For masks we only use one channel
+				if (FMath::IsNearlyZero(colour[0]))
 				{
 					// If the mask is black, we can skip the entire operation
 					at = base.child();
 				}
-				else if (colour.Equals(FVector4f(1, 1, 1, 1), UE_SMALL_NUMBER))
+				else if (FMath::IsNearlyEqual(colour[0], 1, UE_SMALL_NUMBER))
 				{
 					// If the mask is white, we can remove it
-					Ptr<ASTOpImageLayer> nop = mu::Clone<ASTOpImageLayer>(this);
-					nop->mask = nullptr;
-					at = nop;
+					Ptr<ASTOpImageLayerColor> NewOp = mu::Clone<ASTOpImageLayerColor>(this);
+					NewOp->mask = nullptr;
+					at = NewOp;
 				}
 			}
 		}
@@ -464,6 +465,7 @@ namespace mu
 			if (validUsageRect)
 			{
 				// Adjust for compressed blocks (4), and some extra mips (2 more mips, which is 4)
+				// \TODO: block size may be different in ASTC
 				constexpr int blockSize = 4 * 4;
 
 				FImageRect maskUsage;
diff --git a/Engine/Plugins/Experimental/Mutable/Source/MutableTools/Private/MuT/ASTOpImageLayerColor.cpp b/Engine/Plugins/Experimental/Mutable/Source/MutableTools/Private/MuT/ASTOpImageLayerColor.cpp
--- a/Engine/Plugins/Experimental/Mutable/Source/MutableTools/Private/MuT/ASTOpImageLayerColor.cpp
+++ b/Engine/Plugins/Experimental/Mutable/Source/MutableTools/Private/MuT/ASTOpImageLayerColor.cpp
@@ -177,17 +177,18 @@ namespace mu
 		Ptr<ASTOp> at;
 
 		// Plain masks optimization
-		if (!at && mask.child())
+		if (mask.child() && !(Flags & OP::ImageLayerArgs::F_USE_MASK_FROM_BLENDED) )
 		{
 			FVector4f colour;
 			if (mask.child()->IsImagePlainConstant(colour))
 			{
-				if (colour.IsNearlyZero3(UE_SMALL_NUMBER))
+				// For masks we only use one channel
+				if (FMath::IsNearlyZero(colour[0]))
 				{
 					// If the mask is black, we can skip the entire operation
 					at = base.child();
 				}
-				else if (colour.Equals(FVector4f(1, 1, 1, 1), UE_SMALL_NUMBER))
+				else if (FMath::IsNearlyEqual(colour[0], 1, UE_SMALL_NUMBER))
 				{
 					// If the mask is white, we can remove it
 					Ptr<ASTOpImageLayerColor> NewOp = mu::Clone<ASTOpImageLayerColor>(this);
@@ -230,18 +231,35 @@ namespace mu
 					{
 						switch (blendType)
 						{
-						case EBlendType::BT_LIGHTEN: bRGBUnchanged = FMath::IsNearlyEqual(ColorConst[3],0.0f); break;
+						case EBlendType::BT_LIGHTEN: bRGBUnchanged = FMath::IsNearlyZero(ColorConst[3]); break;
 						case EBlendType::BT_MULTIPLY: bRGBUnchanged = FMath::IsNearlyEqual(ColorConst[3],1.0f); break;
 						default: break;
 						}
 					}
 					else
 					{
-						switch (blendType)
+						// How many channels are there in the base?
+						FImageDesc BaseDesc = base->GetImageDesc();
+						const FImageFormatData& FormatDesc = GetImageFormatData(BaseDesc.m_format);
+						if (FormatDesc.Channels == 1)
 						{
-						case EBlendType::BT_LIGHTEN: bRGBUnchanged = ColorConst.IsNearlyZero3(UE_SMALL_NUMBER); break;
-						case EBlendType::BT_MULTIPLY: bRGBUnchanged = ColorConst.Equals(FVector4f(1, 1, 1, 1)); break;
-						default: break;
+							// We only need to check R
+							switch (blendType)
+							{
+							case EBlendType::BT_LIGHTEN: bRGBUnchanged = FMath::IsNearlyZero(ColorConst[0]); break;
+							case EBlendType::BT_MULTIPLY: bRGBUnchanged = FMath::IsNearlyEqual(ColorConst[0],1.0); break;
+							default: break;
+							}
+						}
+						else
+						{
+							// Check RGB
+							switch (blendType)
+							{
+							case EBlendType::BT_LIGHTEN: bRGBUnchanged = ColorConst.IsNearlyZero3(UE_SMALL_NUMBER); break;
+							case EBlendType::BT_MULTIPLY: bRGBUnchanged = FVector3f(ColorConst).Equals(FVector3f(1, 1, 1)); break;
+							default: break;
+							}
 						}
 					}
 				}
