[Mutable] Fix crash in UCustomizableObjectInstance::BuildingMaterials if the map of ImageToPlatformData doesn't contain the ImageID.

#rnx

diff --git a/Engine/Plugins/Experimental/Mutable/Source/CustomizableObject/Private/MuCO/CustomizableObjectInstance.cpp b/Engine/Plugins/Experimental/Mutable/Source/CustomizableObject/Private/MuCO/CustomizableObjectInstance.cpp
--- a/Engine/Plugins/Experimental/Mutable/Source/CustomizableObject/Private/MuCO/CustomizableObjectInstance.cpp
+++ b/Engine/Plugins/Experimental/Mutable/Source/CustomizableObject/Private/MuCO/CustomizableObjectInstance.cpp
@@ -5791,7 +5791,11 @@ void UCustomizableInstancePrivateData::BuildMaterials(const TSharedPtr<FMutableO
 									}
 									else
 									{
-										UE_LOG(LogMutable, Error, TEXT("Required image was not generated in the mutable thread, and it is not cached."));
+										UE_LOG(LogMutable, Error, TEXT("Required image [%s] with ID [%lld] was not generated in the mutable thread, and it is not cached. LOD [%d]. Component [%d]"),
+											*Props.TextureParameterName,
+											Image.ImageID,
+											LODIndex, ComponentIndex);
+										continue;
 									}
 
 									if (bNeverStream)
