[Mutable] Fix crash in FSkeletalMeshSceneProxy::GetDynamicMeshElements caused by GCed materials used in the render thread.
Added CVar mutable.TicksUntilMaterialRelease to change the number of ticks to keep swapped materials alive.

#rnx

#localization none
#tests Frontend and Locker

diff --git a/Engine/Plugins/Experimental/Mutable/Source/CustomizableObject/Classes/MuCO/CustomizableObjectSystem.h b/Engine/Plugins/Experimental/Mutable/Source/CustomizableObject/Classes/MuCO/CustomizableObjectSystem.h
--- a/Engine/Plugins/Experimental/Mutable/Source/CustomizableObject/Classes/MuCO/CustomizableObjectSystem.h
+++ b/Engine/Plugins/Experimental/Mutable/Source/CustomizableObject/Classes/MuCO/CustomizableObjectSystem.h
@@ -85,6 +85,21 @@ public:
 };
 
 
+USTRUCT()
+struct FPendingReleaseMaterialsInfo
+{
+public:
+	GENERATED_USTRUCT_BODY()
+
+public:
+	UPROPERTY()
+	TArray<TObjectPtr<UMaterialInterface>> Materials;
+
+	UPROPERTY()
+	int32 TicksUntilRelease = 0;
+};
+
+
 
 /** Key to identify an image inside a generated mutable runtime instance. */
 struct FMutableImageReference
@@ -246,11 +261,17 @@ public:
 
 	void PurgePendingReleaseSkeletalMesh();
 
+	void AddPendingReleaseMaterials(TArray<TObjectPtr<UMaterialInterface>>& InMaterials);
+	void TickPendingReleaseMaterials();
+
 private:
 
 	UPROPERTY()
 	TArray<FPendingReleaseSkeletalMeshInfo> PendingReleaseSkeletalMesh;
 
+	UPROPERTY(Transient)
+	TArray<FPendingReleaseMaterialsInfo> PendingReleaseMaterials;
+
 public:
 
 	/** [Texture Parameters] Get a list of all the possible values for external texture parameters according to the various providers registered with RegisterImageProvider. */
diff --git a/Engine/Plugins/Experimental/Mutable/Source/CustomizableObject/Private/MuCO/CustomizableObjectInstance.cpp b/Engine/Plugins/Experimental/Mutable/Source/CustomizableObject/Private/MuCO/CustomizableObjectInstance.cpp
--- a/Engine/Plugins/Experimental/Mutable/Source/CustomizableObject/Private/MuCO/CustomizableObjectInstance.cpp
+++ b/Engine/Plugins/Experimental/Mutable/Source/CustomizableObject/Private/MuCO/CustomizableObjectInstance.cpp
@@ -5547,6 +5547,23 @@ void UCustomizableInstancePrivateData::BuildMaterials(const TSharedPtr<FMutableO
 			continue;
 		}
 
+		{
+			// TEMP: Keep a reference to the previous materials for n frames to avoid GC of materials in use in the render thread.
+			// TODO: MTBL-1632 - Implement a proper fix to replace this hotfix
+			const TArray<FSkeletalMaterial>& Materials = SkeletalMesh->GetMaterials();
+			const int32 NumMaterials = Materials.Num();
+			
+			TArray<TObjectPtr<UMaterialInterface>> MaterialsToRelease;
+			MaterialsToRelease.Reserve(NumMaterials);
+
+			for (const FSkeletalMaterial& Material : Materials)
+			{
+				MaterialsToRelease.Add(Material.MaterialInterface);
+			}
+			
+			UCustomizableObjectSystem::GetInstance()->AddPendingReleaseMaterials(MaterialsToRelease);
+		}
+
 		SkeletalMesh->GetMaterials().Reset();
 
 		// Maps serializations of FMutableMaterialPlaceholder to Created Dynamic Material instances, used to reuse materials across LODs
diff --git a/Engine/Plugins/Experimental/Mutable/Source/CustomizableObject/Private/MuCO/CustomizableObjectSystem.cpp b/Engine/Plugins/Experimental/Mutable/Source/CustomizableObject/Private/MuCO/CustomizableObjectSystem.cpp
--- a/Engine/Plugins/Experimental/Mutable/Source/CustomizableObject/Private/MuCO/CustomizableObjectSystem.cpp
+++ b/Engine/Plugins/Experimental/Mutable/Source/CustomizableObject/Private/MuCO/CustomizableObjectSystem.cpp
@@ -1030,6 +1030,35 @@ void UCustomizableObjectSystem::PurgePendingReleaseSkeletalMesh()
 	}
 }
 
+static int32 TicksUntilMaterialRelease = 4;
+
+static FAutoConsoleVariableRef CVarTicksUntilMaterialRelease(
+	TEXT("mutable.TicksUntilMaterialRelease"), TicksUntilMaterialRelease,
+	TEXT("If higher than zero, the number of ticks to keep discarded materials guarded from being GCed."),
+	ECVF_Default);
+
+void UCustomizableObjectSystem::AddPendingReleaseMaterials(TArray<TObjectPtr<UMaterialInterface>>& InMaterials)
+{
+	if (TicksUntilMaterialRelease > 0)
+	{
+		FPendingReleaseMaterialsInfo& PendingMaterials = PendingReleaseMaterials.AddDefaulted_GetRef();
+		PendingMaterials.Materials = InMaterials;
+		PendingMaterials.TicksUntilRelease = TicksUntilMaterialRelease;
+	}
+}
+
+
+void UCustomizableObjectSystem::TickPendingReleaseMaterials()
+{
+	for (TArray<FPendingReleaseMaterialsInfo>::TIterator PendingRelease = PendingReleaseMaterials.CreateIterator(); PendingRelease; PendingRelease++)
+	{
+		if (--PendingRelease->TicksUntilRelease <= 0)
+		{
+			PendingRelease.RemoveCurrent();
+		}
+	}
+}
+
 
 void UCustomizableObjectSystem::AddPendingReleaseSkeletalMesh(USkeletalMesh* SkeletalMesh)
 {
@@ -2437,6 +2466,8 @@ bool UCustomizableObjectSystem::Tick(float DeltaTime)
 	}
 #endif
 	
+	TickPendingReleaseMaterials();
+
 	// Get a new operation if we aren't working on one
 	if (!Private->CurrentMutableOperation)
 	{
