#rnx

diff --git a/Engine/Plugins/Experimental/Mutable/Source/CustomizableObject/Private/MuCO/UnrealMutableImageProvider.cpp b/Engine/Plugins/Experimental/Mutable/Source/CustomizableObject/Private/MuCO/UnrealMutableImageProvider.cpp
--- a/Engine/Plugins/Experimental/Mutable/Source/CustomizableObject/Private/MuCO/UnrealMutableImageProvider.cpp
+++ b/Engine/Plugins/Experimental/Mutable/Source/CustomizableObject/Private/MuCO/UnrealMutableImageProvider.cpp
@@ -699,8 +699,22 @@ mu::FImageDesc FUnrealMutableImageProvider::CreateDummyDesc()
 }
 
 
+TAutoConsoleVariable<bool> CVarMutableLockExternalImagesDuringGC(
+	TEXT("Mutable.LockExternalImagesDuringGC"),
+	true,
+	TEXT("If true, GlobalExternalImages where all texture parameters are stored will be locked from concurrent access during the AddReferencedObjects phase of GC."),
+	ECVF_Default);
+
+
 void FUnrealMutableImageProvider::AddReferencedObjects(FReferenceCollector& Collector)
 {
+	bool bDoLock = CVarMutableLockExternalImagesDuringGC.GetValueOnAnyThread();
+
+	if (bDoLock)
+	{
+		ExternalImagesLock.Lock();
+	}
+
 	for (TPair<FName, FUnrealMutableImageInfo>& Image : GlobalExternalImages)
 	{
 		if (Image.Value.TextureToLoad)
@@ -708,4 +722,9 @@ void FUnrealMutableImageProvider::AddReferencedObjects(FReferenceCollector& Coll
 			Collector.AddReferencedObject(Image.Value.TextureToLoad);
 		}
 	}
+
+	if (bDoLock)
+	{
+		ExternalImagesLock.Unlock();
+	}
 }
