#rnx

diff --git a/Engine/Plugins/Experimental/Mutable/Source/CustomizableObject/Private/MuCO/CustomizableObjectInstance.cpp b/Engine/Plugins/Experimental/Mutable/Source/CustomizableObject/Private/MuCO/CustomizableObjectInstance.cpp
--- a/Engine/Plugins/Experimental/Mutable/Source/CustomizableObject/Private/MuCO/CustomizableObjectInstance.cpp
+++ b/Engine/Plugins/Experimental/Mutable/Source/CustomizableObject/Private/MuCO/CustomizableObjectInstance.cpp
@@ -880,7 +880,7 @@ void UCustomizableObjectInstance::DoUpdateSkeletalMesh(bool bIsCloseDistTick, bo
 		const uint32 InstanceId = GetUniqueID();
 		const float Distance = FMath::Sqrt(GetPrivate()->LastMinSquareDistFromComponentToPlayer);
 		const bool bIsPlayerOrNearIt = GetPrivate()->HasCOInstanceFlags(UsedByPlayerOrNearIt);
-		UE_LOG(LogMutable, Verbose, TEXT("Started UpdateSkeletalMesh Async. of Instance %d with priority %d at dist %f bIsPlayerOrNearIt=%d, frame=%d"), InstanceId, static_cast<int32>(Priority), Distance, bIsPlayerOrNearIt, GFrameNumber);				
+		UE_LOG(LogMutable, Log, TEXT("Started UpdateSkeletalMesh Async. of Instance %d with priority %d at dist %f bIsPlayerOrNearIt=%d, frame=%d"), InstanceId, static_cast<int32>(Priority), Distance, bIsPlayerOrNearIt, GFrameNumber);				
 
 		UpdateDescriptorRuntimeHash = FDescriptorRuntimeHash(Descriptor);
 
diff --git a/Engine/Plugins/Experimental/Mutable/Source/CustomizableObject/Private/MuCO/CustomizableObjectSystem.cpp b/Engine/Plugins/Experimental/Mutable/Source/CustomizableObject/Private/MuCO/CustomizableObjectSystem.cpp
--- a/Engine/Plugins/Experimental/Mutable/Source/CustomizableObject/Private/MuCO/CustomizableObjectSystem.cpp
+++ b/Engine/Plugins/Experimental/Mutable/Source/CustomizableObject/Private/MuCO/CustomizableObjectSystem.cpp
@@ -637,6 +637,9 @@ void FinishUpdateGlobal(UCustomizableObjectInstance* Instance, EUpdateResult Upd
 	}
 
 	UCustomizableObjectSystem::GetInstance()->GetPrivate()->MutableTaskGraph.AllowLaunchingMutableTaskLowPriority(true, false);
+
+	const uint32 InstanceId = Instance ? Instance->GetUniqueID() : 0;
+	UE_LOG(LogMutable, Log, TEXT("Finished UpdateSkeletalMesh Async. of Instance %d , frame=%d"), InstanceId, GFrameNumber);
 }
 
 
