[Mutable] Fix crashes when accessing the ComponentData in platforms where, due to LOD Setting, the FirstLODToGenerate is greater than 0.

#rnx

#localization none
#tests PIE

diff --git a/Engine/Plugins/Experimental/Mutable/Source/CustomizableObject/Classes/MuCO/CustomizableObjectSystem.h b/Engine/Plugins/Experimental/Mutable/Source/CustomizableObject/Classes/MuCO/CustomizableObjectSystem.h
--- a/Engine/Plugins/Experimental/Mutable/Source/CustomizableObject/Classes/MuCO/CustomizableObjectSystem.h
+++ b/Engine/Plugins/Experimental/Mutable/Source/CustomizableObject/Classes/MuCO/CustomizableObjectSystem.h
@@ -24,6 +24,7 @@ class SNotificationItem;
 class UCustomizableObject;
 class UDefaultImageProvider;
 class USkeletalMesh;
+class UMaterialInterface;
 class UTexture2D;
 struct FFrame;
 struct FGuid;
diff --git a/Engine/Plugins/Experimental/Mutable/Source/CustomizableObject/Private/MuCO/CustomizableObjectInstance.cpp b/Engine/Plugins/Experimental/Mutable/Source/CustomizableObject/Private/MuCO/CustomizableObjectInstance.cpp
--- a/Engine/Plugins/Experimental/Mutable/Source/CustomizableObject/Private/MuCO/CustomizableObjectInstance.cpp
+++ b/Engine/Plugins/Experimental/Mutable/Source/CustomizableObject/Private/MuCO/CustomizableObjectInstance.cpp
@@ -63,20 +63,30 @@ UTexture2D* UCustomizableInstancePrivateData::CreateTexture()
 mu::FResourceID UCustomizableInstancePrivateData::GetLastMeshId(int32 ComponentIndex, int32 LODIndex) const
 {
 	const FCustomizableInstanceComponentData* ComponentData = GetComponentData(ComponentIndex);
+
+	if (ComponentData && ComponentData->LastMeshIdPerLOD.IsValidIndex(LODIndex))
+	{
+		return ComponentData->LastMeshIdPerLOD[LODIndex];
+	}
+
 	check(ComponentData);
 	check(ComponentData->LastMeshIdPerLOD.IsValidIndex(LODIndex));
 		
-	return ComponentsData[ComponentIndex].LastMeshIdPerLOD[LODIndex];
+	return MAX_uint64;
 }
 
 
 void UCustomizableInstancePrivateData::SetLastMeshId(int32 ComponentIndex, int32 LODIndex, mu::FResourceID MeshId)
 {
 	FCustomizableInstanceComponentData* ComponentData = GetComponentData(ComponentIndex);
-	check(ComponentData);
-	check(ComponentData->LastMeshIdPerLOD.IsValidIndex(LODIndex));
-
-	ComponentData->LastMeshIdPerLOD[LODIndex] = MeshId;
+	if (ComponentData && ComponentData->LastMeshIdPerLOD.IsValidIndex(LODIndex))
+	{
+		ComponentData->LastMeshIdPerLOD[LODIndex] = MeshId;
+	}
+	else
+	{
+		check(false);
+	}
 }
 
 
@@ -179,7 +189,7 @@ void UCustomizableInstancePrivateData::PrepareForUpdate(const TSharedPtr<FMutabl
 
 		for (uint16 ComponentIndex = 0; ComponentIndex < LOD.ComponentCount; ++ComponentIndex)
 		{
-			const FInstanceUpdateData::FComponent& Component = Components[ComponentIndex];
+			const FInstanceUpdateData::FComponent& Component = Components[LOD.FirstComponent + ComponentIndex];
 			ComponentIds.Add(Component.Id);
 
 			if (FCustomizableInstanceComponentData* ComponentData = GetComponentData(Component.Id))
diff --git a/Engine/Plugins/Experimental/Mutable/Source/CustomizableObject/Private/MuCO/CustomizableObjectSystem.cpp b/Engine/Plugins/Experimental/Mutable/Source/CustomizableObject/Private/MuCO/CustomizableObjectSystem.cpp
--- a/Engine/Plugins/Experimental/Mutable/Source/CustomizableObject/Private/MuCO/CustomizableObjectSystem.cpp
+++ b/Engine/Plugins/Experimental/Mutable/Source/CustomizableObject/Private/MuCO/CustomizableObjectSystem.cpp
@@ -43,7 +43,6 @@
 
 class AActor;
 class UAnimInstance;
-class UMaterialInterface;
 
 namespace impl
 {
@@ -603,6 +602,12 @@ FAutoConsoleVariableRef CVarMaxTextureSizeToGenerate(
 	TEXT("Max texture size on Mutable textures. Mip 0 will be the first mip with max size equal or less than MaxTextureSizeToGenerate."
 		"If a texture doesn't have small enough mips, mip 0 will be the last mip available."));
 
+static bool bApplyFixPrepareSkeletons = true;
+
+static FAutoConsoleVariableRef CVarApplyFixPrepareSkeletons(
+	TEXT("mutable.ApplyFixPrepareSkeletons"), bApplyFixPrepareSkeletons,
+	TEXT("If true, Fix missing SkeletonsData when FirstLODToGenerate is greater than 0. If false, There may be a crash when generating meshes on platform that skip LODs."),
+	ECVF_Default);
 
 void FinishUpdateGlobal(UCustomizableObjectInstance* Instance, EUpdateResult UpdateResult, FInstanceUpdateDelegate* UpdateCallback, const FDescriptorRuntimeHash InUpdatedHash)
 {
@@ -1585,6 +1590,61 @@ namespace impl
 	{
 		MUTABLE_CPUPROFILER_SCOPE(Subtask_Mutable_PrepareSkeletonData)
 
+		if (bApplyFixPrepareSkeletons)
+		{
+			for (FInstanceUpdateData::FComponent& Component : OperationData->InstanceUpdateData.Components)
+			{
+				if (!OperationData->InstanceUpdateData.Skeletons.IsValidIndex(Component.Id))
+				{
+					OperationData->InstanceUpdateData.Skeletons.SetNum(Component.Id + 1);
+				}
+
+				FInstanceUpdateData::FSkeletonData& SkeletonData = OperationData->InstanceUpdateData.Skeletons[Component.Id];
+				SkeletonData.ComponentIndex = Component.Id;
+
+				mu::MeshPtrConst Mesh = Component.Mesh;
+				if (!Mesh)
+				{
+					continue;
+				}
+
+				// Add SkeletonIds 
+				const int32 SkeletonIDsCount = Mesh->GetSkeletonIDsCount();
+				for (int32 SkeletonIndex = 0; SkeletonIndex < SkeletonIDsCount; ++SkeletonIndex)
+				{
+					SkeletonData.SkeletonIds.AddUnique(Mesh->GetSkeletonID(SkeletonIndex));
+				}
+
+				// Append BoneMap to the array of BoneMaps
+				const TArray<uint16>& BoneMap = Mesh->GetBoneMap();
+				Component.FirstBoneMap = OperationData->InstanceUpdateData.BoneMaps.Num();
+				Component.BoneMapCount = BoneMap.Num();
+				OperationData->InstanceUpdateData.BoneMaps.Append(BoneMap);
+
+				// Add active bone indices and poses
+				const int32 MaxBoneIndex = Mesh->GetBonePoseCount();
+				Component.ActiveBones.Reserve(MaxBoneIndex);
+				for (int32 BonePoseIndex = 0; BonePoseIndex < MaxBoneIndex; ++BonePoseIndex)
+				{
+					const uint16 BoneId = Mesh->GetBonePoseBoneId(BonePoseIndex);
+
+					Component.ActiveBones.Add(BoneId);
+
+					if (SkeletonData.BoneIds.Find(BoneId) == INDEX_NONE)
+					{
+						SkeletonData.BoneIds.Add(BoneId);
+
+						FTransform3f Transform;
+						Mesh->GetBoneTransform(BonePoseIndex, Transform);
+						SkeletonData.BoneMatricesWithScale.Emplace(Transform.Inverse().ToMatrixWithScale());
+					}
+				}
+			}
+		
+			return;
+		}
+
+
 		check(OperationData);
 		const int32 LODCount = OperationData->InstanceUpdateData.LODs.Num();
 		const FInstanceUpdateData::FLOD& MinLOD = OperationData->InstanceUpdateData.LODs[OperationData->CurrentMinLOD];
