[Mutable] Fix image transform rotation when keep aspect ratio enabled.

- Fixes cases were MinRelevancy could by higher than MaxRelevancy when crop rect outside dest rect.
- Fixes some possible divides by zero when the crop rect is outside the dest rect.

#rnx

diff --git a/Engine/Plugins/Experimental/Mutable/Source/MutableRuntime/Private/MuR/CodeRunner.cpp b/Engine/Plugins/Experimental/Mutable/Source/MutableRuntime/Private/MuR/CodeRunner.cpp
--- a/Engine/Plugins/Experimental/Mutable/Source/MutableRuntime/Private/MuR/CodeRunner.cpp
+++ b/Engine/Plugins/Experimental/Mutable/Source/MutableRuntime/Private/MuR/CodeRunner.cpp
@@ -4628,7 +4628,7 @@ namespace mu
 			{
             	MUTABLE_CPUPROFILER_SCOPE(IM_TRANSFORM_1)
 
-                FVector2f Scale = FVector2f(
+				FVector2f Scale = FVector2f(
                         Args.ScaleX ? LoadScalar(FCacheAddress(Args.ScaleX, item)) : 1.0f,
                         Args.ScaleY ? LoadScalar(FCacheAddress(Args.ScaleY, item)) : 1.0f);
 	
@@ -4648,23 +4648,25 @@ namespace mu
 						--MipsToDrop;
 					}
 
-					return Size;
+					return FUint16Vector2(FMath::Max(Size.X, uint16(1)), FMath::Max(Size.Y, uint16(1)));
 				});
 
-				const FVector2f DestSize = FVector2f(DestSizeI.X, DestSizeI.Y);
-				const FVector2f SourceSize = FVector2f(Args.SourceSizeX, Args.SourceSizeY);
-	
+				const FVector2f DestSize   = FVector2f(DestSizeI.X, DestSizeI.Y);
+				const FVector2f SourceSize = FVector2f(FMath::Max(Args.SourceSizeX, uint16(1)), FMath::Max(Args.SourceSizeY, uint16(1)));
+
+				FVector2f AspectCorrectionScale = FVector2f(1.0f, 1.0f);
 				if (Args.bKeepAspectRatio)
 				{
-					const float DestAspectOverSrcAspect = (DestSize.X * SourceSize.Y) / (DestSize.Y * SourceSize.X); 
+					const float DestAspectOverSrcAspect = (DestSize.X * SourceSize.Y) / (DestSize.Y * SourceSize.X);
 
-					Scale *= DestAspectOverSrcAspect > 1.0f 
-							? FVector2f(1.0f/DestAspectOverSrcAspect, 1.0f) 
-							: FVector2f(1.0f, DestAspectOverSrcAspect); 
+					AspectCorrectionScale = DestAspectOverSrcAspect > 1.0f 
+										  ? FVector2f(1.0f/DestAspectOverSrcAspect, 1.0f) 
+										  : FVector2f(1.0f, DestAspectOverSrcAspect); 
 				}
 			
 				const FTransform2f Transform = FTransform2f(FVector2f(-0.5f))
 					.Concatenate(FTransform2f(FScale2f(Scale)))
+					.Concatenate(FTransform2f(FScale2f(AspectCorrectionScale)))
 					.Concatenate(FTransform2f(FVector2f(0.5f)));
 
 				FBox2f NormalizedCropRect(ForceInit);
@@ -4682,8 +4684,8 @@ namespace mu
 				FScheduledOpData HeapData;
 				HeapData.ImageTransform.SizeX = DestSizeI.X;
 				HeapData.ImageTransform.SizeY = DestSizeI.Y;
-				FPlatformMath::StoreHalf(&HeapData.ImageTransform.ScaleXEncodedHalf, Scale.X),
-				FPlatformMath::StoreHalf(&HeapData.ImageTransform.ScaleYEncodedHalf, Scale.Y),
+				FPlatformMath::StoreHalf(&HeapData.ImageTransform.ScaleXEncodedHalf, Scale.X);
+				FPlatformMath::StoreHalf(&HeapData.ImageTransform.ScaleYEncodedHalf, Scale.Y);
 				HeapData.ImageTransform.MipValue = BestMip + GlobalImageTransformLodBias;
 
 				const int32 HeapDataAddress = m_heapData.Add(HeapData);
@@ -4714,12 +4716,25 @@ namespace mu
                         Args.OffsetX ? LoadScalar(FCacheAddress(Args.OffsetX, item)) : 0.0f,
                         Args.OffsetY ? LoadScalar(FCacheAddress(Args.OffsetY, item)) : 0.0f);
 
-                const FVector2f Scale = FVector2f(
+                FVector2f Scale = FVector2f(
 						FPlatformMath::LoadHalf(&HeapData.ImageTransform.ScaleXEncodedHalf),
 						FPlatformMath::LoadHalf(&HeapData.ImageTransform.ScaleYEncodedHalf));
 
-				// Map Range 0-1 to a full rotation
-                const float Rotation = LoadScalar(FCacheAddress(Args.Rotation, item)) * UE_TWO_PI;
+				FVector2f AspectCorrectionScale = FVector2f(1.0f, 1.0f);
+				if (Args.bKeepAspectRatio)
+				{
+					const FVector2f DestSize   = FVector2f(HeapData.ImageTransform.SizeX, HeapData.ImageTransform.SizeY);
+					const FVector2f SourceSize = FVector2f(FMath::Max(Args.SourceSizeX, uint16(1)), FMath::Max(Args.SourceSizeY, uint16(1)));
+					
+					const float DestAspectOverSrcAspect = (DestSize.X * SourceSize.Y) / (DestSize.Y * SourceSize.X);
+					
+					AspectCorrectionScale = DestAspectOverSrcAspect > 1.0f 
+										  ? FVector2f(1.0f/DestAspectOverSrcAspect, 1.0f) 
+										  : FVector2f(1.0f, DestAspectOverSrcAspect); 
+				}
+
+				// Map Range [0..1] to a full rotation
+                const float RotationRad = LoadScalar(FCacheAddress(Args.Rotation, item)) * UE_TWO_PI;
 	
 				EImageFormat SourceFormat = Source->GetFormat();
 				EImageFormat Format = GetUncompressedFormat(SourceFormat);
@@ -4751,12 +4766,34 @@ namespace mu
 					Source = NewImage;
 				}
 
+				Scale.X = FMath::IsNearlyZero(Scale.X, UE_KINDA_SMALL_NUMBER) ? UE_KINDA_SMALL_NUMBER : Scale.X;
+				Scale.Y = FMath::IsNearlyZero(Scale.Y, UE_KINDA_SMALL_NUMBER) ? UE_KINDA_SMALL_NUMBER : Scale.Y;
+
+				AspectCorrectionScale.X = FMath::IsNearlyZero(AspectCorrectionScale.X, UE_KINDA_SMALL_NUMBER) 
+									    ? UE_KINDA_SMALL_NUMBER 
+										: AspectCorrectionScale.X;
+
+				AspectCorrectionScale.Y = FMath::IsNearlyZero(AspectCorrectionScale.Y, UE_KINDA_SMALL_NUMBER) 
+										? UE_KINDA_SMALL_NUMBER 
+										: AspectCorrectionScale.Y;
+
+				const FTransform2f Transform = FTransform2f(FVector2f(-0.5f))
+						.Concatenate(FTransform2f(FScale2f(Scale)))
+						.Concatenate(FTransform2f(FQuat2f(RotationRad)))
+						.Concatenate(FTransform2f(FScale2f(AspectCorrectionScale)))
+						.Concatenate(FTransform2f(Offset + FVector2f(0.5f)));
+
 				const EAddressMode AddressMode = static_cast<EAddressMode>(Args.AddressMode);
+
+				const EInitializationType InitType = AddressMode == EAddressMode::ClampToBlack 
+											       ? EInitializationType::Black
+											       : EInitializationType::NotInitialized;
+
 				Ptr<Image> Result = CreateImage(
-						HeapData.ImageTransform.SizeX, HeapData.ImageTransform.SizeY, 1, Format, EInitializationType::Black);
+						HeapData.ImageTransform.SizeX, HeapData.ImageTransform.SizeY, 1, Format, InitType);
 
 				const float MipFactor = FMath::Frac(FMath::Max(0.0f, HeapData.ImageTransform.MipValue));
-				ImageTransform(Result.get(), Source.get(), Offset, Scale, Rotation, MipFactor, AddressMode, bUseImageTransformVectorImpl);
+				ImageTransform(Result.get(), Source.get(), Transform, MipFactor, AddressMode, bUseImageTransformVectorImpl);
 
 				Release(Source);
 				StoreImage(item, Result);
diff --git a/Engine/Plugins/Experimental/Mutable/Source/MutableRuntime/Private/MuR/OpImageTransform.h b/Engine/Plugins/Experimental/Mutable/Source/MutableRuntime/Private/MuR/OpImageTransform.h
--- a/Engine/Plugins/Experimental/Mutable/Source/MutableRuntime/Private/MuR/OpImageTransform.h
+++ b/Engine/Plugins/Experimental/Mutable/Source/MutableRuntime/Private/MuR/OpImageTransform.h
@@ -13,11 +13,8 @@
 namespace mu
 {
 	void ImageTransform(
-			Image* pDestImage, const Image* pImage, 
-			FVector2f Offset, FVector2f Scale, float RotationRad, 
-			float MipFactor,
-			EAddressMode AddressMode, 
-			bool bUseVectorImplCvar)
+			Image* pDestImage, const Image* pImage,
+			FTransform2f Transform, float MipFactor, EAddressMode AddressMode, bool bUseVectorImplCvar)
     {
 		MUTABLE_CPUPROFILER_SCOPE(ImageTransform);
 
@@ -26,14 +23,6 @@ namespace mu
 			return;
 		}
 
-		Scale.X = FMath::IsNearlyZero(Scale.X, UE_KINDA_SMALL_NUMBER) ? UE_KINDA_SMALL_NUMBER : Scale.X;
-		Scale.Y = FMath::IsNearlyZero(Scale.Y, UE_KINDA_SMALL_NUMBER) ? UE_KINDA_SMALL_NUMBER : Scale.Y;
-
-		const FTransform2f Transform = FTransform2f(FVector2f(-0.5f)).
-				Concatenate(FTransform2f(FScale2f(Scale))).
-				Concatenate(FTransform2f(FQuat2f(RotationRad))).
-				Concatenate(FTransform2f(Offset + FVector2f(0.5f)));
-
 		const FIntVector2 SrcSize  = FIntVector2(pImage->GetSizeX(), pImage->GetSizeY());
 		const FIntVector2 DestSize = FIntVector2(pDestImage->GetSizeX(), pDestImage->GetSizeY());
 
@@ -57,15 +46,15 @@ namespace mu
 							FMath::CeilToInt (NormalizedCropRect.Max.Y * (float)DestSize.Y)));
 
 			return FIntRect(
-				FMath::Max(CropRect.Min.X, 0), FMath::Max(CropRect.Min.Y, 0),
-				FMath::Min(CropRect.Max.X, DestSize.X), FMath::Min(CropRect.Max.Y, DestSize.Y));
+				FMath::Clamp(CropRect.Min.X, 0, DestSize.X), FMath::Clamp(CropRect.Min.Y, 0, DestSize.Y),
+				FMath::Clamp(CropRect.Max.X, 0, DestSize.X), FMath::Clamp(CropRect.Max.Y, 0, DestSize.Y));
 		});
 
 		if (AddressMode == EAddressMode::ClampToBlack)
 		{
 			pDestImage->m_flags |= Image::IF_HAS_RELEVANCY_MAP;
 			pDestImage->RelevancyMinY = DestCropRect.Min.Y;
-			pDestImage->RelevancyMaxY = DestCropRect.Max.Y - 1;
+			pDestImage->RelevancyMaxY = FMath::Max(DestCropRect.Max.Y - 1, DestCropRect.Min.Y);
 		}
 
 		uint8* DestData       = pDestImage->GetData();
diff --git a/Engine/Plugins/Experimental/Mutable/Source/MutableRuntime/Private/MuR/OpImageTransformPrivate.cpp b/Engine/Plugins/Experimental/Mutable/Source/MutableRuntime/Private/MuR/OpImageTransformPrivate.cpp
--- a/Engine/Plugins/Experimental/Mutable/Source/MutableRuntime/Private/MuR/OpImageTransformPrivate.cpp
+++ b/Engine/Plugins/Experimental/Mutable/Source/MutableRuntime/Private/MuR/OpImageTransformPrivate.cpp
@@ -28,12 +28,9 @@ namespace mu
 			alignas(8) uint16 Data[NumChannels];
 		};
 
-
-		const FVector2f BaseSizeF = FVector2f(Src0Size.X, Src1Size.Y);
 		const FVector2f DestSizeF = FVector2f(DestSize.X, DestSize.Y);
 
 		const FVector2f DestNormFactor = FVector2f(1.0f) / DestSizeF;
-		const FVector2f BaseNormFactor = FVector2f(1.0f) / BaseSizeF;
 
 		check(DestCropRect.Min.X >= 0 && DestCropRect.Max.X <= DestSize.X);
 		check(DestCropRect.Min.Y >= 0 && DestCropRect.Max.Y <= DestSize.Y);
@@ -66,17 +63,18 @@ namespace mu
 					{
 						if ((DestUv.X < 0.0f) | (DestUv.X > 1.0f) | (DestUv.Y < 0.0f) | (DestUv.Y > 1.0f))
 						{
-							if constexpr (NumChannels == 4)
-							{
-								*reinterpret_cast<uint32*>(&DestData[(Y * DestSize.X + X) * NumChannels]) = 0;
-							}
-							else
-							{
-								for (uint32 Channel = 0; Channel < NumChannels; ++Channel)
-								{
-									DestData[(Y * DestSize.X + X) * NumChannels + Channel] = 0;
-								}
-							}
+							// Black texture init is mandatory for ClampToBlack Address mode, no need to set black again. 
+							//if constexpr (NumChannels == 4)
+							//{
+							//	FMemory::Memzero(&DestData[(Y * DestSize.X + X) * NumChannels], 4);
+							//}
+							//else
+							//{
+							//	for (uint32 Channel = 0; Channel < NumChannels; ++Channel)
+							//	{
+							//		DestData[(Y * DestSize.X + X) * NumChannels + Channel] = 0;
+							//	}
+							//}
 							continue;
 						}
 					}
@@ -144,7 +142,8 @@ namespace mu
 						FPixelData Result;
 						if constexpr (NumChannels == 4)
 						{
-							const uint32 PackedData = *reinterpret_cast<const uint32*>(Ptr);
+							uint32 PackedData;
+							FMemory::Memcpy(&PackedData, Ptr, sizeof(uint32));
 
 							Result.Data[0] = static_cast<uint16>((PackedData >> (8 * 0)) & 0xFF);
 							Result.Data[1] = static_cast<uint16>((PackedData >> (8 * 1)) & 0xFF);
@@ -228,6 +227,11 @@ namespace mu
 		float MipFactor,
 		const FTransform2f& Transform)
     {
+		if (DestCropRect.Area() == 0)
+		{
+			return;
+		}
+
 		const FVector2f DestSizeF = FVector2f(DestSize.X, DestSize.Y);
 		const FVector2f DestNormFactor = FVector2f(1.0f) / DestSizeF;
 
@@ -268,17 +272,18 @@ namespace mu
 					{
 						if ((DestUv.X < 0.0f) | (DestUv.X > 1.0f) | (DestUv.Y < 0.0f) | (DestUv.Y > 1.0f))
 						{
-							if constexpr (NumChannels == 4)
-							{
-								*reinterpret_cast<uint32*>(&DestData[(Y * DestSize.X + X) * NumChannels]) = 0;
-							}
-							else
-							{
-								for (uint32 Channel = 0; Channel < NumChannels; ++Channel)
-								{
-									DestData[(Y * DestSize.X + X) * NumChannels + Channel] = 0;
-								}
-							}
+							// Black texture init is mandatory for ClampToBlack Address mode, no need to set black again. 
+							//if constexpr (NumChannels == 4)
+							//{
+							//	FMemory::Memzero(&DestData[(Y * DestSize.X + X) * NumChannels], 4);
+							//}
+							//else
+							//{
+							//	for (uint32 Channel = 0; Channel < NumChannels; ++Channel)
+							//	{
+							//		DestData[(Y * DestSize.X + X) * NumChannels + Channel] = 0;
+							//	}
+							//}
 
 							continue;
 						}
@@ -321,7 +326,9 @@ namespace mu
 					{
 						if constexpr (NumChannels == 4)
 						{
-							const uint32 PackedData = *reinterpret_cast<const uint32*>(DataPtr);
+							uint32 PackedData;
+							FMemory::Memcpy(&PackedData, DataPtr, sizeof(uint32));
+
 							return VectorIntToFloat(MakeVectorRegisterInt(
 										(PackedData >> (8 * 0)) & 0xFF, 
 										(PackedData >> (8 * 1)) & 0xFF,
