[Mutable] Change behaviour of image transform size to be consistent on all code parts when one only one of the sizes are 0.

#rnx

diff --git a/Engine/Plugins/Experimental/Mutable/Source/MutableTools/Private/MuT/ASTOpImageTransform.cpp b/Engine/Plugins/Experimental/Mutable/Source/MutableTools/Private/MuT/ASTOpImageTransform.cpp
--- a/Engine/Plugins/Experimental/Mutable/Source/MutableTools/Private/MuT/ASTOpImageTransform.cpp
+++ b/Engine/Plugins/Experimental/Mutable/Source/MutableTools/Private/MuT/ASTOpImageTransform.cpp
@@ -202,11 +202,7 @@ namespace mu
 	{
 		if (Base)
 		{
-			if (SizeX == 0 && SizeY == 0)
-			{
-				return Base->GetImageSizeExpression();
-			}
-			else
+			if (!(SizeX == 0 && SizeY == 0))
 			{
 				Ptr<ImageSizeExpression> SizeExpr = new ImageSizeExpression;
 				SizeExpr->type = ImageSizeExpression::ISET_CONSTANT;
@@ -215,6 +211,10 @@ namespace mu
 
 				return SizeExpr;
 			}
+			else
+			{
+				return Base->GetImageSizeExpression();
+			}
 		}
 
 		return nullptr;
diff --git a/Engine/Plugins/Experimental/Mutable/Source/MutableTools/Private/MuT/ASTOpImageTransform.h b/Engine/Plugins/Experimental/Mutable/Source/MutableTools/Private/MuT/ASTOpImageTransform.h
--- a/Engine/Plugins/Experimental/Mutable/Source/MutableTools/Private/MuT/ASTOpImageTransform.h
+++ b/Engine/Plugins/Experimental/Mutable/Source/MutableTools/Private/MuT/ASTOpImageTransform.h
@@ -28,11 +28,11 @@ struct FProgram;
 		ASTChild ScaleY;
 		ASTChild Rotation;
 
-		uint16 SizeX;
-		uint16 SizeY;
+		uint16 SizeX = 0;
+		uint16 SizeY = 0;
 
-		uint16 SourceSizeX;
-		uint16 SourceSizeY;
+		uint16 SourceSizeX = 0;
+		uint16 SourceSizeY = 0;
 
 		EAddressMode AddressMode = EAddressMode::Wrap;
 		bool bKeepAspectRatio = false;
