[Mutable] Reduce PARAMETER_VALUE memory footprint by moving FProjector and string storage outside the parameter variant.

- It uses a self managed heap allocated object for FProjectors and string so the parameters union size is reduced.

- The managed object should behave similarly to a reference and supports underling type serialization so minimal changes are needed in other parts of the code, for instance, serialization does not need backward compatibility fixes.

- Improved Map serialization. Now it can reuse the serialized Key and Value temporaries.

#rnx

diff --git a/Engine/Plugins/Experimental/Mutable/Source/MutableRuntime/Private/MuR/Model.cpp b/Engine/Plugins/Experimental/Mutable/Source/MutableRuntime/Private/MuR/Model.cpp
--- a/Engine/Plugins/Experimental/Mutable/Source/MutableRuntime/Private/MuR/Model.cpp
+++ b/Engine/Plugins/Experimental/Mutable/Source/MutableRuntime/Private/MuR/Model.cpp
@@ -345,7 +345,7 @@ namespace mu
             return;
         }
 
-        const ParamProjectorType& Projector = m_pD->m_program.m_parameters[Index].m_defaultValue.Get<ParamProjectorType>();
+        const FProjector& Projector = m_pD->m_program.m_parameters[Index].m_defaultValue.Get<ParamProjectorType>();
         if (OutProjectionType) *OutProjectionType = Projector.type;
     	if (OutPos) *OutPos = Projector.position;
 		if (OutDir) *OutDir = Projector.direction;
diff --git a/Engine/Plugins/Experimental/Mutable/Source/MutableRuntime/Private/MuR/Parameters.cpp b/Engine/Plugins/Experimental/Mutable/Source/MutableRuntime/Private/MuR/Parameters.cpp
--- a/Engine/Plugins/Experimental/Mutable/Source/MutableRuntime/Private/MuR/Parameters.cpp
+++ b/Engine/Plugins/Experimental/Mutable/Source/MutableRuntime/Private/MuR/Parameters.cpp
@@ -844,7 +844,8 @@ namespace mu
         if ( !pos )
         {
             // Return the single value
-            return m_pD->m_values[index].Get<ParamStringType>().c_str();
+            const string& String = m_pD->m_values[index].Get<ParamStringType>();
+            return String.c_str();
         }
 
         // Multivalue case
@@ -856,12 +857,14 @@ namespace mu
 			const PARAMETER_VALUE* it = m.Find( pos->m_pD->m_values );
             if ( it )
             {
-                return it->Get<ParamStringType>().c_str();
+                const string& String = it->Get<ParamStringType>();
+                return String.c_str();
             }
         }
 
         // Multivalue parameter, but no multivalue set. Return single value.
-        return m_pD->m_values[index].Get<ParamStringType>().c_str();
+        const string& String = m_pD->m_values[index].Get<ParamStringType>();
+        return String.c_str();
     }
 
 
@@ -1023,12 +1026,13 @@ namespace mu
         }
 
 	    // Parameters cannot change the projector type anymore
-		ParamProjectorType Value;
+		FProjector Value;
         Value.type = PROJECTOR_TYPE::COUNT;
         if (m_pD->m_pModel)
         {
             const FProgram& program = m_pD->m_pModel->GetPrivate()->m_program;
-            Value.type = program.m_parameters[ParameterIndex].m_defaultValue.Get<ParamProjectorType>().type;
+            const FProjector& Projector = program.m_parameters[ParameterIndex].m_defaultValue.Get<ParamProjectorType>();
+            Value.type = Projector.type;
         }
 
         Value.position = pos;
diff --git a/Engine/Plugins/Experimental/Mutable/Source/MutableRuntime/Private/MuR/ParametersPrivate.h b/Engine/Plugins/Experimental/Mutable/Source/MutableRuntime/Private/MuR/ParametersPrivate.h
--- a/Engine/Plugins/Experimental/Mutable/Source/MutableRuntime/Private/MuR/ParametersPrivate.h
+++ b/Engine/Plugins/Experimental/Mutable/Source/MutableRuntime/Private/MuR/ParametersPrivate.h
@@ -15,6 +15,124 @@
 
 namespace mu
 {
+
+namespace Private
+{
+
+    template<class T>
+    class TIndirectObject
+    {
+        TUniquePtr<T> StoragePtr;
+
+    public:
+        template<typename... TArgs>
+        TIndirectObject(TArgs&&... Args) : StoragePtr(MakeUnique<T>(Forward<TArgs>(Args)...)) 
+        {
+        }
+        
+        TIndirectObject(const TIndirectObject<T>& Other) : StoragePtr(MakeUnique<T>()) 
+        { 
+            Get() = Other; 
+        }
+
+		TIndirectObject(TIndirectObject<T>&& Other) : StoragePtr(MoveTemp(Other.StoragePtr)) 
+        {
+        }
+
+        TIndirectObject(const T& Object) : StoragePtr(MakeUnique<T>()) 
+        { 
+            Get() = Object; 
+        }
+        
+        TIndirectObject(T&& Object) : StoragePtr(MakeUnique<T>()) 
+        { 
+            Get() = MoveTemp(Object); 
+        }
+
+        TIndirectObject& operator=(TIndirectObject<T>&&) = default;
+
+        TIndirectObject& operator=(const TIndirectObject<T>& Other) 
+        { 
+            Get() = Other; 
+        }
+
+        TIndirectObject& operator=(const T& Object) 
+        { 
+            Get() = Object; 
+        }
+
+        TIndirectObject& operator=(const T&& Object) 
+        { 
+            Get() = MoveTemp(Object); 
+        }
+
+        const T& Get() const 
+        { 
+            check(StoragePtr); 
+            return *StoragePtr; 
+        }
+
+        T& Get() 
+        { 
+            check(StoragePtr); 
+            return *StoragePtr; 
+        }
+
+        operator T&() 
+        { 
+            return Get(); 
+        }
+
+        operator const T&() const 
+        { 
+            return Get(); 
+        }
+
+        T* operator &() 
+        { 
+            return StoragePtr.Get(); 
+        }
+
+        const T* operator &() const 
+        { 
+            return StoragePtr.Get(); 
+        }
+
+        T& operator *() 
+        { 
+            return Get(); 
+        }
+
+        const T& operator *() const 
+        { 
+            return Get(); 
+        }
+
+        bool operator==(const TIndirectObject<T>& Other) const 
+        { 
+            return *StoragePtr == Other; 
+        }
+
+        bool operator==(const T& Object) const 
+        { 
+            return *StoragePtr == Object; 
+        }
+
+		//!
+		void Serialise(OutputArchive& Arch) const
+		{
+			Arch << Get();
+		}
+
+		//!
+		void Unserialise(InputArchive& Arch)
+		{
+			Arch >> Get();
+		} 
+    };
+
+}
+
     MUTABLE_DEFINE_ENUM_SERIALISABLE(PARAMETER_TYPE)
     MUTABLE_DEFINE_ENUM_SERIALISABLE(PARAMETER_DETAILED_TYPE)
     MUTABLE_DEFINE_ENUM_SERIALISABLE(PROJECTOR_TYPE)
@@ -248,19 +366,20 @@ namespace mu
             arch >> defaultValue;
         }
 	};
-	
 
 	using ParamBoolType = bool;
 	using ParamIntType = int32;
 	using ParamFloatType = float;
 	using ParamColorType = FVector3f;
-	using ParamProjectorType = FProjector;
+	using ParamProjectorType = Private::TIndirectObject<FProjector>;
 	using ParamImageType = FName;
-	using ParamStringType = string;
-
+	using ParamStringType = Private::TIndirectObject<string>;
 	
-	using PARAMETER_VALUE = TVariant<ParamBoolType, ParamIntType, ParamFloatType, ParamColorType, ParamProjectorType, ParamImageType, ParamStringType>;
+	using PARAMETER_VALUE = TVariant<
+            ParamBoolType, ParamIntType, ParamFloatType, ParamColorType, ParamProjectorType, ParamImageType, ParamStringType>;
 
+    // static_assert to track PARAMETER_VALUE size changes. It is ok to change if needed.
+    static_assert(sizeof(PARAMETER_VALUE) == 8*3, "PARAMETER_VALUE size has changed.");
 
 	// TVariant currently does not support this operator. Once supported remove it.
 	inline bool operator==(const PARAMETER_VALUE& ValueA, const PARAMETER_VALUE& ValueB)
@@ -446,10 +565,11 @@ namespace mu
 
         //! Run-time data
         TSharedPtr<const Model> m_pModel;
-
+ 
         //! Values for the parameters if they are not multidimensional.
 		TArray<PARAMETER_VALUE> m_values;
 
+
         //! If the parameter is multidemensional, the values are stored here.
         //! The key of the map is the vector of values stored in a RangeIndex
 		TArray< TMap< TArray<int32>, PARAMETER_VALUE > > m_multiValues;
@@ -458,7 +578,7 @@ namespace mu
         //!
         void Serialise( OutputArchive& arch ) const
         {
-            const uint32 ver = 2;
+            const uint32 ver = 3;
             arch << ver;
 
             arch << m_values;
@@ -470,9 +590,9 @@ namespace mu
         {
             uint32 ver;
             arch >> ver;
-			check(ver == 2);
-        	
-			arch >> m_values;
+			check(ver <= 3);
+       
+		    arch >> m_values;
 			arch >> m_multiValues;
         }
 
@@ -494,7 +614,6 @@ namespace mu
 
 			return m_multiValues[ParamIndex].Num()>0;
 		}
-
     };
 
 }
diff --git a/Engine/Plugins/Experimental/Mutable/Source/MutableRuntime/Private/MuR/SerialisationPrivate.h b/Engine/Plugins/Experimental/Mutable/Source/MutableRuntime/Private/MuR/SerialisationPrivate.h
--- a/Engine/Plugins/Experimental/Mutable/Source/MutableRuntime/Private/MuR/SerialisationPrivate.h
+++ b/Engine/Plugins/Experimental/Mutable/Source/MutableRuntime/Private/MuR/SerialisationPrivate.h
@@ -309,29 +309,30 @@ namespace mu
 
 	//---------------------------------------------------------------------------------------------
 	template< typename K, typename T >
-	inline void operator<< (OutputArchive& arch, const TMap<K, T>& v)
+	inline void operator<< (OutputArchive& Arch, const TMap<K, T>& V)
 	{
-		arch << (uint32)v.Num();
-		for (const auto& p : v)
+		Arch << (uint32)V.Num();
+		for (const TPair<K, T>& Element : V)
 		{
-			arch << p.Key;
-			arch << p.Value;
+			Arch << Element.Key;
+			Arch << Element.Value;
 		}
 	}
 
 	template< typename K, typename T >
-	inline void operator>> (InputArchive& arch, TMap<K, T>& v)
+	inline void operator>> (InputArchive& Arch, TMap<K, T>& V)
 	{
-		uint32 size;
-		arch >> size;
+		uint32 Num;
+		Arch >> Num;
 
-		for (uint32 i = 0; i < size; ++i)
+		for (uint32 Index = 0; Index < Num; ++Index)
 		{
-			K k;
-			T t;
-			arch >> k;
-			arch >> t;
-			v.Add(k, t);
+			K Key;
+			T Element;
+			Arch >> Key;
+			Arch >> Element;
+
+			V.Emplace(MoveTemp(Key), MoveTemp(Element));
 		}
 	}
 
diff --git a/Engine/Plugins/Experimental/Mutable/Source/MutableRuntime/Public/MuR/Serialisation.h b/Engine/Plugins/Experimental/Mutable/Source/MutableRuntime/Public/MuR/Serialisation.h
--- a/Engine/Plugins/Experimental/Mutable/Source/MutableRuntime/Public/MuR/Serialisation.h
+++ b/Engine/Plugins/Experimental/Mutable/Source/MutableRuntime/Public/MuR/Serialisation.h
@@ -382,31 +382,30 @@ namespace mu
 
 	//---------------------------------------------------------------------------------------------
 	template<typename T, typename Alloc> 
-	void operator<<(OutputArchive& arch, const TArray<T, Alloc>& v)
+	void operator<<(OutputArchive& Arch, const TArray<T, Alloc>& V)
 	{
-		const uint32 Num = (uint32)v.Num();
-		arch << Num;
+		const uint32 Num = (uint32)V.Num();
+		Arch << Num;
 		
-		for (SIZE_T i = 0; i < Num; ++i)
+		for (SIZE_T Index = 0; Index < Num; ++Index)
 		{
-			arch << v[i];
+			Arch << V[Index];
 		}
 	}
 
 	template<typename T, typename Alloc> 
-	void operator>>(InputArchive& arch, TArray<T, Alloc>& v)
+	void operator>>(InputArchive& Arch, TArray<T, Alloc>& V)
 	{
 		uint32 Num;
-		arch >> Num;
-		v.SetNum(Num);
+		Arch >> Num;
+		V.SetNum(Num);
 
-		for (SIZE_T i = 0; i < Num; ++i)
+		for (SIZE_T Index = 0; Index < Num; ++Index)
 		{
-			arch >> v[i];
+			Arch >> V[Index];
 		}
 	}
 
-
 	//---------------------------------------------------------------------------------------------
 	template<> 
 	inline void operator<<(OutputArchive& arch, const FName& v)
diff --git a/Engine/Plugins/Experimental/Mutable/Source/MutableTools/Private/MuT/CodeGenerator_Projector.cpp b/Engine/Plugins/Experimental/Mutable/Source/MutableTools/Private/MuT/CodeGenerator_Projector.cpp
--- a/Engine/Plugins/Experimental/Mutable/Source/MutableTools/Private/MuT/CodeGenerator_Projector.cpp
+++ b/Engine/Plugins/Experimental/Mutable/Source/MutableTools/Private/MuT/CodeGenerator_Projector.cpp
@@ -140,7 +140,8 @@ namespace mu
 			op = it->second;
 		}
 
-		result.type = op->parameter.m_defaultValue.Get<ParamProjectorType>().type;
+		const FProjector& Projector = op->parameter.m_defaultValue.Get<ParamProjectorType>();
+		result.type = Projector.type;
 		result.op = op;
 	}
 
