[Mutable] Allow calling FUnrealMutableImageProvider API from all threads.

#rnx

diff --git a/Engine/Plugins/Experimental/Mutable/Source/CustomizableObject/Private/MuCO/UnrealMutableImageProvider.cpp b/Engine/Plugins/Experimental/Mutable/Source/CustomizableObject/Private/MuCO/UnrealMutableImageProvider.cpp
--- a/Engine/Plugins/Experimental/Mutable/Source/CustomizableObject/Private/MuCO/UnrealMutableImageProvider.cpp
+++ b/Engine/Plugins/Experimental/Mutable/Source/CustomizableObject/Private/MuCO/UnrealMutableImageProvider.cpp
@@ -467,8 +467,6 @@ mu::FImageDesc FUnrealMutableImageProvider::GetImageDesc(FName Id, uint8 Mipmaps
 
 void FUnrealMutableImageProvider::CacheImage(FName Id, bool bUser)
 {
-	check( IsInGameThread() );
-
 	if (Id == NAME_None)
 	{
 		return;	
@@ -561,9 +559,6 @@ void FUnrealMutableImageProvider::CacheImage(FName Id, bool bUser)
 
 void FUnrealMutableImageProvider::UnCacheImage(FName Id, bool bUser)
 {
-	// TODO: Review GM
-	//check(IsInGameThread());
-
 	if (Id == NAME_None)
 	{
 		return;	
@@ -596,8 +591,6 @@ void FUnrealMutableImageProvider::UnCacheImage(FName Id, bool bUser)
 
 void FUnrealMutableImageProvider::ClearCache(bool bUser)
 {
-	check(IsInGameThread());
-
 	FScopeLock Lock(&ExternalImagesLock);
 	for (TTuple<FName, FUnrealMutableImageInfo> Tuple : GlobalExternalImages)
 	{
