#rnx

diff --git a/Engine/Plugins/Experimental/Mutable/Source/CustomizableObject/Classes/MuCO/CustomizableObject.h b/Engine/Plugins/Experimental/Mutable/Source/CustomizableObject/Classes/MuCO/CustomizableObject.h
--- a/Engine/Plugins/Experimental/Mutable/Source/CustomizableObject/Classes/MuCO/CustomizableObject.h
+++ b/Engine/Plugins/Experimental/Mutable/Source/CustomizableObject/Classes/MuCO/CustomizableObject.h
@@ -1215,6 +1215,10 @@ public:
 	UPROPERTY(EditAnywhere, Category = CompileOptions)
 	bool bEnableAnimBpPhysicsAssetsManipualtion = false;
 
+	// When this is enabled generated meshes will merge the AssetUserData from all of its constituent mesh parts
+	UPROPERTY(EditAnywhere, Category = CompileOptions)
+	bool bEnableAssetUserDataMerge = false;
+
 	// Options when compiling this customizable object (see EMutableCompileMeshType declaration for info)
 	UPROPERTY(EditAnywhere, Category = CompileOptions)
 	EMutableCompileMeshType MeshCompileType = EMutableCompileMeshType::LocalAndChildren;
@@ -1587,6 +1591,10 @@ public:
 	UPROPERTY()
 	TArray<FAnimBpOverridePhysicsAssetsInfo> AnimBpOverridePhysiscAssetsInfo;
 
+	/** Stores the UAssetUserData assets gathered from the SkeletalMesh nodes during compilation, to be used in mesh generation in-game */
+	UPROPERTY()
+	TMap<FString, TSoftObjectPtr<UAssetUserData>> AssetUserDataAssetsMap;
+
 	UPROPERTY()
 	/** Stores the sockets provided by the part skeletal meshes, to be merged in the generated meshes */
 	TArray<FMutableRefSocket> SocketArray;
diff --git a/Engine/Plugins/Experimental/Mutable/Source/CustomizableObject/Classes/MuCO/CustomizableObjectInstance.h b/Engine/Plugins/Experimental/Mutable/Source/CustomizableObject/Classes/MuCO/CustomizableObjectInstance.h
--- a/Engine/Plugins/Experimental/Mutable/Source/CustomizableObject/Classes/MuCO/CustomizableObjectInstance.h
+++ b/Engine/Plugins/Experimental/Mutable/Source/CustomizableObject/Classes/MuCO/CustomizableObjectInstance.h
@@ -267,6 +267,12 @@ public:
 	/**  Set random values to the parameters using a seed. Useful for testing only. */
 	void SetRandomValues(const int32 InRandomizationSeed);
 
+	/** Returns the AssetUserData that was gathered from all the constituent mesh parts during the last update. 
+	 *  It requires that the CustomizableObject had the bEnableAssetUserDataMerge set to true during compilation. 
+	*/
+	UFUNCTION(BlueprintCallable, Category = CustomizableObjectInstance)
+	TSet<UAssetUserData*> GetMergedAssetUserData(int32 ComponentIndex) const;
+
 	// Utilities to manage saving and loading parameters from profiles.
 	bool LoadParametersFromProfile(int32 ProfileIndex);
 	bool SaveParametersToProfile(int32 ProfileIndex);
diff --git a/Engine/Plugins/Experimental/Mutable/Source/CustomizableObject/Private/MuCO/CustomizableInstancePrivateData.h b/Engine/Plugins/Experimental/Mutable/Source/CustomizableObject/Private/MuCO/CustomizableInstancePrivateData.h
--- a/Engine/Plugins/Experimental/Mutable/Source/CustomizableObject/Private/MuCO/CustomizableInstancePrivateData.h
+++ b/Engine/Plugins/Experimental/Mutable/Source/CustomizableObject/Private/MuCO/CustomizableInstancePrivateData.h
@@ -126,6 +126,13 @@ struct FCustomizableInstanceComponentData
 	UPROPERTY(Transient, Category = CustomizableObjectInstance, editfixedsize, VisibleAnywhere)
 	TMap<FName, TSoftClassPtr<UAnimInstance>> AnimSlotToBP;
 
+	// AssetUserData gathered for a component from its constituent meshes
+	UPROPERTY(Transient, Category = CustomizableObjectInstance, editfixedsize, VisibleAnywhere)
+	TSet<TObjectPtr<UAssetUserData>> AssetUserDataArray;
+
+	// Temp structure to store the AssetUserData that will be async loaded during an update
+	TSet<TSoftObjectPtr<UAssetUserData>> AssetUserDataToStream;
+
 #if WITH_EDITORONLY_DATA
 	// Just used for mutable.EnableMutableAnimInfoDebugging command
 	TArray<FString> MeshPartPaths;
diff --git a/Engine/Plugins/Experimental/Mutable/Source/CustomizableObject/Private/MuCO/CustomizableObjectInstance.cpp b/Engine/Plugins/Experimental/Mutable/Source/CustomizableObject/Private/MuCO/CustomizableObjectInstance.cpp
--- a/Engine/Plugins/Experimental/Mutable/Source/CustomizableObject/Private/MuCO/CustomizableObjectInstance.cpp
+++ b/Engine/Plugins/Experimental/Mutable/Source/CustomizableObject/Private/MuCO/CustomizableObjectInstance.cpp
@@ -200,6 +200,8 @@ void UCustomizableInstancePrivateData::PrepareForUpdate(const TSharedPtr<FMutabl
 				}
 
 				ComponentData->AnimSlotToBP.Empty();
+				ComponentData->AssetUserDataArray.Empty();
+				ComponentData->AssetUserDataToStream.Empty();
 
 #if WITH_EDITORONLY_DATA
 				ComponentData->MeshPartPaths.Empty();
@@ -5041,6 +5043,7 @@ FAutoConsoleVariableRef CVarMutableHighPriorityLoading(
 	bEnableHighPriorityLoading,
 	TEXT("If enabled, the request to load additional assets will have high priority."));
 
+
 FGraphEventRef UCustomizableInstancePrivateData::LoadAdditionalAssetsAsync(const TSharedPtr<FMutableOperationData>& OperationData, UCustomizableObjectInstance* Public, FStreamableManager& StreamableManager)
 {
 	MUTABLE_CPUPROFILER_SCOPE(UCustomizableInstancePrivateData::LoadAdditionalAssetsAsync);
@@ -5264,6 +5267,25 @@ FGraphEventRef UCustomizableInstancePrivateData::LoadAdditionalAssetsAsync(const
 				{
 					AnimBPGameplayTags.AddTag(FGameplayTag::RequestGameplayTag(*Tag));
 				}
+				else if (Tag.RemoveFromStart("__AssetUserData:"))
+				{
+					const FString& AssetPath = Tag;
+
+					TSoftObjectPtr<UAssetUserData>* AssetUserData = CustomizableObject->AssetUserDataAssetsMap.Find(AssetPath);
+
+					if (AssetUserData && !AssetUserData->IsNull())
+					{
+						if (AssetUserData->Get())
+						{
+							ComponentData->AssetUserDataArray.Add(AssetUserData->Get());
+						}
+						else
+						{
+							ComponentData->AssetUserDataToStream.Add(*AssetUserData);
+							AssetsToStream.Add(AssetUserData->ToSoftObjectPath());
+						}
+					}
+				}
 #if WITH_EDITORONLY_DATA
 				else if (Tag.RemoveFromStart("__MeshPath:"))
 				{
@@ -5295,6 +5317,7 @@ FGraphEventRef UCustomizableInstancePrivateData::LoadAdditionalAssetsAsync(const
 	return Result;
 }
 
+
 void UCustomizableObjectInstance::AdditionalAssetsAsyncLoaded( FGraphEventRef CompletionEvent )
 {
 	// TODO: Do we need this separated?
@@ -5423,6 +5446,28 @@ void UCustomizableInstancePrivateData::AdditionalAssetsAsyncLoaded(UCustomizable
 #endif
 		}
 
+		for (TSoftObjectPtr<UAssetUserData> LoadedAssetUserData : ComponentData.AssetUserDataToStream)
+		{
+			if (LoadedAssetUserData.IsValid())
+			{
+				ComponentData.AssetUserDataArray.Add(LoadedAssetUserData.Get());
+			}
+#if WITH_EDITOR
+			else
+			{
+				FMessageLogModule& MessageLogModule = FModuleManager::LoadModuleChecked<FMessageLogModule>("MessageLog");
+				MessageLogModule.RegisterLogListing(FName("Mutable"), FText::FromString(FString("Mutable")));
+				FMessageLog MessageLog("Mutable");
+
+				FString ErrorMsg = FString::Printf(TEXT("Mutable couldn't load the AssetUserData [%s]. If it has been deleted or renamed, please recompile all the mutable objects that use it."), *LoadedAssetUserData.GetAssetName());
+				UE_LOG(LogMutable, Error, TEXT("%s"), *ErrorMsg);
+				MessageLog.Notify(FText::FromString(ErrorMsg), EMessageSeverity::Error, true);
+			}
+#endif
+		}
+
+		ComponentData.AssetUserDataToStream.Empty();
+
 		const int32 AdditionalPhysicsNum = ComponentData.PhysicsAssets.AdditionalPhysicsAssetsToLoad.Num();
 		ComponentData.PhysicsAssets.AdditionalPhysicsAssets.Reserve(AdditionalPhysicsNum);
 		for (int32 I = 0; I < AdditionalPhysicsNum; ++I)
@@ -6763,3 +6808,28 @@ void UCustomizableObjectInstance::ForEachAnimInstance(int32 ComponentIndex, FEac
 		}
 	}
 }
+
+
+TSet<UAssetUserData*> UCustomizableObjectInstance::GetMergedAssetUserData(int32 ComponentIndex) const
+{
+	UCustomizableInstancePrivateData* PrivateInstanceData = GetPrivate();
+
+	if (PrivateInstanceData)
+	{
+		check(PrivateInstanceData->ComponentsData.IsValidIndex(ComponentIndex));
+
+		TSet<UAssetUserData*> Set;
+		
+		// Have to convert to UAssetUserData* because BP functions don't support TObjectPtr
+		for (const TObjectPtr<UAssetUserData>& Elem : PrivateInstanceData->ComponentsData[ComponentIndex].AssetUserDataArray)
+		{
+			Set.Add(Elem);
+		}
+
+		return Set;
+	}
+	else
+	{
+		return TSet<UAssetUserData*>();
+	}
+}
diff --git a/Engine/Plugins/Experimental/Mutable/Source/CustomizableObject/Private/MuCO/CustomizableObjectSystem.cpp b/Engine/Plugins/Experimental/Mutable/Source/CustomizableObject/Private/MuCO/CustomizableObjectSystem.cpp
--- a/Engine/Plugins/Experimental/Mutable/Source/CustomizableObject/Private/MuCO/CustomizableObjectSystem.cpp
+++ b/Engine/Plugins/Experimental/Mutable/Source/CustomizableObject/Private/MuCO/CustomizableObjectSystem.cpp
@@ -1989,6 +1989,18 @@ namespace impl
 		FCustomizableObjectSystemPrivate* CustomizableObjectSystemPrivateData = System->GetPrivate();
 		check(CustomizableObjectSystemPrivateData != nullptr);
 
+		for (int32 ComponentIndex = 0; ComponentIndex < CustomizableObjectInstance->GetPrivate()->ComponentsData.Num(); ++ComponentIndex)
+		{
+			ensure(CustomizableObjectInstance->SkeletalMeshes.IsValidIndex(ComponentIndex));
+
+			FCustomizableInstanceComponentData& ComponentData = CustomizableObjectInstance->GetPrivate()->ComponentsData[ComponentIndex];
+			
+			for (TObjectPtr<UAssetUserData> AssetUserData : ComponentData.AssetUserDataArray)
+			{
+				CustomizableObjectInstance->SkeletalMeshes[ComponentIndex]->AddAssetUserData(AssetUserData);
+			}
+		}
+
 		// Next Task: Release Mutable. We need this regardless if we cancel or not
 		//-------------------------------------------------------------		
 		{
diff --git a/Engine/Plugins/Experimental/Mutable/Source/CustomizableObjectEditor/Private/MuCOE/CustomizableObjectCompiler.cpp b/Engine/Plugins/Experimental/Mutable/Source/CustomizableObjectEditor/Private/MuCOE/CustomizableObjectCompiler.cpp
--- a/Engine/Plugins/Experimental/Mutable/Source/CustomizableObjectEditor/Private/MuCOE/CustomizableObjectCompiler.cpp
+++ b/Engine/Plugins/Experimental/Mutable/Source/CustomizableObjectEditor/Private/MuCOE/CustomizableObjectCompiler.cpp
@@ -1175,6 +1175,17 @@ void FCustomizableObjectCompiler::CompileInternal(UCustomizableObject* Object, c
 	
 		Object->AnimBPAssetsMap = GenerationContext.AnimBPAssetsMap;
 
+		// Mark the object as modified, used to avoid missing assets in packages. 
+		if (!Object->AssetUserDataAssetsMap.OrderIndependentCompareEqual(GenerationContext.AssetUserDataAssetsMap))
+		{
+			if (!ParamNamesToSelectedOptions.Num()) // Don't mark the objects as modified because of a partial compilation
+			{
+				Object->MarkPackageDirty();
+			}
+		}
+
+		Object->AssetUserDataAssetsMap = GenerationContext.AssetUserDataAssetsMap;
+
 		// Mark the object as modified, used to avoid missing assets in packages.
 		if (Object->SocketArray != GenerationContext.SocketArray)
 		{
diff --git a/Engine/Plugins/Experimental/Mutable/Source/CustomizableObjectEditor/Private/MuCOE/GenerateMutableSource/GenerateMutableSource.cpp b/Engine/Plugins/Experimental/Mutable/Source/CustomizableObjectEditor/Private/MuCOE/GenerateMutableSource/GenerateMutableSource.cpp
--- a/Engine/Plugins/Experimental/Mutable/Source/CustomizableObjectEditor/Private/MuCOE/GenerateMutableSource/GenerateMutableSource.cpp
+++ b/Engine/Plugins/Experimental/Mutable/Source/CustomizableObjectEditor/Private/MuCOE/GenerateMutableSource/GenerateMutableSource.cpp
@@ -1302,6 +1302,12 @@ FString GenerateGameplayTag(const FString& GameplayTag)
 }
 
 
+FString GenerateAssetUserDataTag(const FString& AssetUserData)
+{
+	return FString("__AssetUserData:") + AssetUserData;
+}
+
+
 void PopulateReferenceSkeletalMeshesData(FMutableGraphGenerationContext& GenerationContext)
 {
 	const FString PlatformName = GenerationContext.Options.TargetPlatform->IniPlatformName();
diff --git a/Engine/Plugins/Experimental/Mutable/Source/CustomizableObjectEditor/Private/MuCOE/GenerateMutableSource/GenerateMutableSource.h b/Engine/Plugins/Experimental/Mutable/Source/CustomizableObjectEditor/Private/MuCOE/GenerateMutableSource/GenerateMutableSource.h
--- a/Engine/Plugins/Experimental/Mutable/Source/CustomizableObjectEditor/Private/MuCOE/GenerateMutableSource/GenerateMutableSource.h
+++ b/Engine/Plugins/Experimental/Mutable/Source/CustomizableObjectEditor/Private/MuCOE/GenerateMutableSource/GenerateMutableSource.h
@@ -691,6 +691,9 @@ struct FMutableGraphGenerationContext
 	/** Stores the anim BP assets gathered from the SkeletalMesh nodes during compilation, to be used in mesh generation in-game */
 	TMap<FString, TSoftClassPtr<UAnimInstance>> AnimBPAssetsMap;
 
+	/** Stores the AssetUserData assets gathered from the SkeletalMesh nodes during compilation, to be used in mesh generation in-game */
+	TMap<FString, TSoftObjectPtr<UAssetUserData>> AssetUserDataAssetsMap;
+
 	/** Stores the sockets provided by the part skeletal meshes, to be merged in the generated meshes */
 	TArray<FMutableRefSocket> SocketArray;
 
@@ -796,6 +799,8 @@ FString GenerateAnimationInstanceTag(const FString& AnimInstance, const FName& S
 
 FString GenerateGameplayTag(const FString& GameplayTag);
 
+FString GenerateAssetUserDataTag(const FString& AssetUserData);
+
 // Computes the LOD bias for a texture given the current mesh LOD and automatic LOD settings, the reference texture settings
 // and whether it's being built for a server or not
 int32 ComputeLODBias(const FMutableGraphGenerationContext& GenerationContext, const UTexture2D* ReferenceTexture, int32 MaxTextureSize,
diff --git a/Engine/Plugins/Experimental/Mutable/Source/CustomizableObjectEditor/Private/MuCOE/GenerateMutableSource/GenerateMutableSourceMesh.cpp b/Engine/Plugins/Experimental/Mutable/Source/CustomizableObjectEditor/Private/MuCOE/GenerateMutableSource/GenerateMutableSourceMesh.cpp
--- a/Engine/Plugins/Experimental/Mutable/Source/CustomizableObjectEditor/Private/MuCOE/GenerateMutableSource/GenerateMutableSourceMesh.cpp
+++ b/Engine/Plugins/Experimental/Mutable/Source/CustomizableObjectEditor/Private/MuCOE/GenerateMutableSource/GenerateMutableSourceMesh.cpp
@@ -2869,6 +2869,26 @@ mu::NodeMeshPtr GenerateMutableSourceMesh(const UEdGraphPin* Pin,
 				MeshUniqueTags += AnimBPTag;
 			}
 
+			TArray<FString> AssetUserDataTags;
+
+			if (GenerationContext.Object->bEnableAssetUserDataMerge)
+			{
+				const TArray<UAssetUserData*>* AssetUserDataArray = TypedNodeSkel->SkeletalMesh->GetAssetUserDataArray();
+
+				if (AssetUserDataArray && !AssetUserDataArray->IsEmpty())
+				{
+					for (const UAssetUserData* AssetUserData : *AssetUserDataArray)
+					{
+						FString AuxString = AssetUserData->GetPathName();
+						GenerationContext.AssetUserDataAssetsMap.Add(AuxString, TSoftObjectPtr<UAssetUserData>(AssetUserData));
+
+						FString AssetUserDataTag = GenerateAssetUserDataTag(AuxString);
+						AssetUserDataTags.Add(AssetUserDataTag);
+						MeshUniqueTags += AssetUserDataTag;
+					}
+				}
+			}
+
 			FSkeletalMeshModel* ImportedModel = TypedNodeSkel->SkeletalMesh->GetImportedModel();
 			
 			mu::MeshPtr MutableMesh = GenerateMutableMesh(TypedNodeSkel->SkeletalMesh, TypedNodeSkel->AnimInstance, LODIndexConnected, SectionIndexConnected, LODIndex, SectionIndex, MeshUniqueTags, GenerationContext, TypedNodeSkel);
@@ -2937,6 +2957,11 @@ mu::NodeMeshPtr GenerateMutableSourceMesh(const UEdGraphPin* Pin,
 					AddTagToMutableMeshUnique(*MutableMesh, GamePlayTag);
 				}
 
+				for (const FString& AssetUserDataTag : AssetUserDataTags)
+				{
+					AddTagToMutableMeshUnique(*MutableMesh, AssetUserDataTag);
+				}
+
 				AddSocketTagsToMesh(TypedNodeSkel->SkeletalMesh, MutableMesh, GenerationContext);
 
 				if (UCustomizableObjectSystem::GetInstance()->IsMutableAnimInfoDebuggingEnabled())
diff --git a/Engine/Plugins/Experimental/Mutable/Source/CustomizableObjectEditor/Private/MuCOE/GenerateMutableSource/GenerateMutableSourceTable.cpp b/Engine/Plugins/Experimental/Mutable/Source/CustomizableObjectEditor/Private/MuCOE/GenerateMutableSource/GenerateMutableSourceTable.cpp
--- a/Engine/Plugins/Experimental/Mutable/Source/CustomizableObjectEditor/Private/MuCOE/GenerateMutableSource/GenerateMutableSourceTable.cpp
+++ b/Engine/Plugins/Experimental/Mutable/Source/CustomizableObjectEditor/Private/MuCOE/GenerateMutableSource/GenerateMutableSourceTable.cpp
@@ -258,6 +258,23 @@ bool FillTableColumn(const UCustomizableObjectNodeTable* TableNode,	mu::TablePtr
 					AddTagToMutableMeshUnique(*MutableMesh, GenerateGameplayTag(Tag.ToString()));
 				}
 
+				if (GenerationContext.Object->bEnableAssetUserDataMerge)
+				{
+					const TArray<UAssetUserData*>* AssetUserDataArray = SkeletalMesh->GetAssetUserDataArray();
+
+					if (AssetUserDataArray && !AssetUserDataArray->IsEmpty())
+					{
+						for (const UAssetUserData* AssetUserData : *AssetUserDataArray)
+						{
+							FString AuxString = AssetUserData->GetPathName();
+							GenerationContext.AssetUserDataAssetsMap.Add(AuxString, TSoftObjectPtr<UAssetUserData>(AssetUserData));
+
+							FString AssetUserDataTag = GenerateAssetUserDataTag(AuxString);
+							AddTagToMutableMeshUnique(*MutableMesh, AssetUserDataTag);
+						}
+					}
+				}
+
 				AddSocketTagsToMesh(SkeletalMesh, MutableMesh, GenerationContext);
 
 				if (UCustomizableObjectSystem::GetInstance()->IsMutableAnimInfoDebuggingEnabled())
