[mutable] Add option to use Unreal's texture compressor at compile and debug time. Enable it by default for cook. Add option to override the internal pixel format function.
Change the method to find out the platform texture format, so that we have less ad-hoc code and more support for Unreal's project settings affecting it.
Cleanup image conversion code related to instance baking, removing duplicated hacky code.
#rnx

diff --git a/Engine/Plugins/Experimental/Mutable/Source/CustomizableObject/Classes/MuCO/CustomizableObject.h b/Engine/Plugins/Experimental/Mutable/Source/CustomizableObject/Classes/MuCO/CustomizableObject.h
--- a/Engine/Plugins/Experimental/Mutable/Source/CustomizableObject/Classes/MuCO/CustomizableObject.h
+++ b/Engine/Plugins/Experimental/Mutable/Source/CustomizableObject/Classes/MuCO/CustomizableObject.h
@@ -152,6 +152,17 @@ enum class ECustomizableObjectNumBoneInfluences : uint8
 	Twelve = 12 // This is essentially the same as "Unlimited", but UE ultimately limits to 12
 };
 
+UENUM()
+enum class ECustomizableObjectTextureCompression : uint8
+{
+	// Don't use texture compression
+	None = 0,
+	// Use Mutable's fast low-quality compression
+	Fast,
+	// Use Unreal's highest quality compression (100x slower to compress)
+	HighQuality
+};
+
 USTRUCT()
 struct FCompilationOptions
 {
@@ -159,7 +170,11 @@ struct FCompilationOptions
 
 	// Flag to know if texture compression should be enabled
 	UPROPERTY()
-	bool bTextureCompression = true;
+	bool bTextureCompression_DEPRECATED = true;
+	
+	// Enum to know what texture compression should be used
+	UPROPERTY()
+	ECustomizableObjectTextureCompression TextureCompression = ECustomizableObjectTextureCompression::Fast;
 
 	// From 0 to 3
 	UPROPERTY()
diff --git a/Engine/Plugins/Experimental/Mutable/Source/CustomizableObject/Classes/MuCO/CustomizableObjectInstance.h b/Engine/Plugins/Experimental/Mutable/Source/CustomizableObject/Classes/MuCO/CustomizableObjectInstance.h
--- a/Engine/Plugins/Experimental/Mutable/Source/CustomizableObject/Classes/MuCO/CustomizableObjectInstance.h
+++ b/Engine/Plugins/Experimental/Mutable/Source/CustomizableObject/Classes/MuCO/CustomizableObjectInstance.h
@@ -843,10 +843,6 @@ private:
 };
 
 
-#if WITH_EDITOR
-CUSTOMIZABLEOBJECT_API void CopyTextureProperties(UTexture2D* Texture, const UTexture2D* SourceTexture);
-#endif	
-
 #if UE_ENABLE_INCLUDE_ORDER_DEPRECATED_IN_5_2
 #include "Async/TaskGraphInterfaces.h"
 #include "MuCO/CustomizableObjectParameterTypeDefinitions.h"
diff --git a/Engine/Plugins/Experimental/Mutable/Source/CustomizableObject/Classes/MuCO/CustomizableObjectSystem.h b/Engine/Plugins/Experimental/Mutable/Source/CustomizableObject/Classes/MuCO/CustomizableObjectSystem.h
--- a/Engine/Plugins/Experimental/Mutable/Source/CustomizableObject/Classes/MuCO/CustomizableObjectSystem.h
+++ b/Engine/Plugins/Experimental/Mutable/Source/CustomizableObject/Classes/MuCO/CustomizableObjectSystem.h
@@ -8,6 +8,7 @@
 #include "MuCO/CustomizableObjectInstance.h"
 #include "MuCO/CustomizableObjectInstanceDescriptor.h"
 #include "MuR/Parameters.h"
+#include "MuR/Image.h"
 #include "MuR/Types.h"
 
 #if WITH_EDITOR
@@ -258,6 +259,10 @@ public:
 	bool IsOnlyGenerateRequestedLODsEnabled() const;
 	void SetOnlyGenerateRequestedLODsEnabled(bool bIsEnabled);
 
+#if WITH_EDITOR
+	void SetImagePixelFormatOverride(const mu::FImageOperator::FImagePixelFormatFunc&);
+#endif
+
 	void AddPendingReleaseSkeletalMesh( USkeletalMesh* SkeletalMesh );
 
 	void PurgePendingReleaseSkeletalMesh();
diff --git a/Engine/Plugins/Experimental/Mutable/Source/CustomizableObject/Private/MuCO/CustomizableInstancePrivateData.h b/Engine/Plugins/Experimental/Mutable/Source/CustomizableObject/Private/MuCO/CustomizableInstancePrivateData.h
--- a/Engine/Plugins/Experimental/Mutable/Source/CustomizableObject/Private/MuCO/CustomizableInstancePrivateData.h
+++ b/Engine/Plugins/Experimental/Mutable/Source/CustomizableObject/Private/MuCO/CustomizableInstancePrivateData.h
@@ -4,6 +4,8 @@
 
 
 #include "Materials/MaterialInterface.h"
+#include "Rendering/SkeletalMeshLODModel.h"
+#include "Rendering/SkeletalMeshModel.h"
 #include "MuCO/CustomizableObjectInstance.h"
 #include "MuCO/CustomizableObjectSystemPrivate.h"
 
@@ -271,6 +273,11 @@ public:
 	/** Set OnlyLOD to -1 to generate all mips */
 	CUSTOMIZABLEOBJECT_API static FTexturePlatformData* MutableCreateImagePlatformData(mu::Ptr<const mu::Image> MutableImage, int32 OnlyLOD, uint16 FullSizeX, uint16 FullSizeY);
 
+	/** */
+#if WITH_EDITORONLY_DATA
+	static void RegenerateImportedModel(USkeletalMesh* SkeletalMesh);
+#endif
+
 private:
 
 	void InitLastUpdateData(const TSharedPtr<FMutableOperationData>& OperationData);
diff --git a/Engine/Plugins/Experimental/Mutable/Source/CustomizableObject/Private/MuCO/CustomizableObject.cpp b/Engine/Plugins/Experimental/Mutable/Source/CustomizableObject/Private/MuCO/CustomizableObject.cpp
--- a/Engine/Plugins/Experimental/Mutable/Source/CustomizableObject/Private/MuCO/CustomizableObject.cpp
+++ b/Engine/Plugins/Experimental/Mutable/Source/CustomizableObject/Private/MuCO/CustomizableObject.cpp
@@ -156,6 +156,17 @@ void UCustomizableObject::PostLoad()
 		ReferenceSkeletalMeshes.Add(ReferenceSkeletalMesh_DEPRECATED);
 		ReferenceSkeletalMesh_DEPRECATED = nullptr;
 	}
+
+	const int32 CustomizableObjectCustomVersion = GetLinkerCustomVersion(FCustomizableObjectCustomVersion::GUID);
+
+	// Convert texture compression option
+	if (CustomizableObjectCustomVersion < FCustomizableObjectCustomVersion::TextureCompressionEnum)
+	{
+		CompileOptions.TextureCompression = CompileOptions.bTextureCompression_DEPRECATED
+			? ECustomizableObjectTextureCompression::Fast
+			: ECustomizableObjectTextureCompression::None
+			;
+	}
 #endif
 }
 
@@ -702,10 +713,9 @@ void UCustomizableObject::CompileForTargetPlatform(const ITargetPlatform* Target
 
 	if (bIsRootObject && bIsRelevantForThisTarget)
 	{
-
 		FCompilationOptions Options;
 		Options.OptimizationLevel = 3;	// max optimization when packaging.
-		Options.bTextureCompression = true;
+		Options.TextureCompression = ECustomizableObjectTextureCompression::HighQuality;
 		Options.bIsCooking = true;
 		Options.bSaveCookedDataToDisk = bUsesOnCookStart;
 		Options.TargetPlatform = TargetPlatform;
diff --git a/Engine/Plugins/Experimental/Mutable/Source/CustomizableObject/Private/MuCO/CustomizableObjectInstance.cpp b/Engine/Plugins/Experimental/Mutable/Source/CustomizableObject/Private/MuCO/CustomizableObjectInstance.cpp
--- a/Engine/Plugins/Experimental/Mutable/Source/CustomizableObject/Private/MuCO/CustomizableObjectInstance.cpp
+++ b/Engine/Plugins/Experimental/Mutable/Source/CustomizableObject/Private/MuCO/CustomizableObjectInstance.cpp
@@ -12,6 +12,7 @@
 #include "ClothConfig.h"
 #include "ClothingAsset.h"
 #include "Engine/SkeletalMeshSocket.h"
+#include "Engine/SkeletalMeshLODSettings.h"
 #include "Engine/Texture2DArray.h"
 #include "MaterialDomain.h"
 #include "Materials/Material.h"
@@ -2967,27 +2968,6 @@ void UCustomizableObjectInstance::CommitMinMaxLOD()
 }
 
 
-void CopyTextureProperties(UTexture2D* Texture, const UTexture2D* SourceTexture)
-{
-	MUTABLE_CPUPROFILER_SCOPE(CopyTextureProperties)
-		
-	Texture->NeverStream = SourceTexture->NeverStream;
-
-	Texture->SRGB = SourceTexture->SRGB;
-	Texture->Filter = SourceTexture->Filter;
-	Texture->LODBias = SourceTexture->LODBias;
-
-#if WITH_EDITOR
-	Texture->MipGenSettings = SourceTexture->MipGenSettings;
-	Texture->CompressionNone = SourceTexture->CompressionNone;
-#endif
-
-	Texture->LODGroup = SourceTexture->LODGroup;
-	Texture->AddressX = SourceTexture->AddressX;
-	Texture->AddressY = SourceTexture->AddressY;
-}
-
-
 // The memory allocated in the function and pointed by the returned pointer is owned by the caller and must be freed. 
 // If assigned to a UTexture2D, it will be freed by that UTexture2D
 FTexturePlatformData* UCustomizableInstancePrivateData::MutableCreateImagePlatformData(mu::Ptr<const mu::Image> MutableImage, int32 OnlyLOD, uint16 FullSizeX, uint16 FullSizeY)
@@ -3081,7 +3061,7 @@ FTexturePlatformData* UCustomizableInstancePrivateData::MutableCreateImagePlatfo
 		//return nullptr;
 	}
 
-	mu::FImageOperator ImOp = mu::FImageOperator::GetDefault();
+	mu::FImageOperator ImOp = mu::FImageOperator::GetDefault(mu::FImageOperator::FImagePixelFormatFunc());
 
 	EPixelFormat PlatformFormat = PF_Unknown;
 	switch (MutableFormat)
@@ -3276,7 +3256,7 @@ void UCustomizableInstancePrivateData::ConvertImage(UTexture2D* Texture, mu::Ima
 	// Extract a single channel, if requested.
 	if (ExtractChannel >= 0)
 	{
-		mu::FImageOperator ImOp = mu::FImageOperator::GetDefault();
+		mu::FImageOperator ImOp = mu::FImageOperator::GetDefault(mu::FImageOperator::FImagePixelFormatFunc());
 
 		MutableImage = ImOp.ImagePixelFormat( 4, MutableImage.get(), mu::EImageFormat::IF_RGBA_UBYTE );
 
@@ -6833,3 +6813,178 @@ TSet<UAssetUserData*> UCustomizableObjectInstance::GetMergedAssetUserData(int32
 		return TSet<UAssetUserData*>();
 	}
 }
+
+
+#if WITH_EDITORONLY_DATA
+
+void UCustomizableInstancePrivateData::RegenerateImportedModel(USkeletalMesh* SkeletalMesh)
+{
+	FSkeletalMeshRenderData* SkelResource = SkeletalMesh->GetResourceForRendering();
+	if (!SkelResource)
+	{
+		return;
+	}
+
+	for (UClothingAssetBase* ClothingAssetBase : SkeletalMesh->GetMeshClothingAssets())
+	{
+		if (!ClothingAssetBase)
+		{
+			continue;
+		}
+
+		UClothingAssetCommon* ClothAsset = Cast<UClothingAssetCommon>(ClothingAssetBase);
+
+		if (!ClothAsset)
+		{
+			continue;
+		}
+
+		if (!ClothAsset->LodData.Num())
+		{
+			continue;
+		}
+
+		for (FClothLODDataCommon& ClothLodData : ClothAsset->LodData)
+		{
+			ClothLodData.PointWeightMaps.Empty(16);
+			for (TPair<uint32, FPointWeightMap>& WeightMap : ClothLodData.PhysicalMeshData.WeightMaps)
+			{
+				if (WeightMap.Value.Num())
+				{
+					FPointWeightMap& PointWeightMap = ClothLodData.PointWeightMaps.AddDefaulted_GetRef();
+					PointWeightMap.Initialize(WeightMap.Value, WeightMap.Key);
+				}
+			}
+		}
+	}
+
+	FSkeletalMeshModel* ImportedModel = SkeletalMesh->GetImportedModel();
+	ImportedModel->bGuidIsHash = false;
+	ImportedModel->SkeletalMeshModelGUID = FGuid::NewGuid();
+
+	ImportedModel->LODModels.Empty();
+
+	int32 OriginalIndex = 0;
+	for (int32 LODIndex = 0; LODIndex < SkelResource->LODRenderData.Num(); ++LODIndex)
+	{
+		ImportedModel->LODModels.Add(new FSkeletalMeshLODModel());
+
+		FSkeletalMeshLODRenderData& LODModel = SkelResource->LODRenderData[LODIndex];
+		int32 CurrentSectionInitialVertex = 0;
+
+		ImportedModel->LODModels[LODIndex].ActiveBoneIndices = LODModel.ActiveBoneIndices;
+		ImportedModel->LODModels[LODIndex].NumTexCoords = LODModel.GetNumTexCoords();
+		ImportedModel->LODModels[LODIndex].RequiredBones = LODModel.RequiredBones;
+		ImportedModel->LODModels[LODIndex].NumVertices = LODModel.GetNumVertices();
+
+		// Indices
+		int indexCount = LODModel.MultiSizeIndexContainer.GetIndexBuffer()->Num();
+		ImportedModel->LODModels[LODIndex].IndexBuffer.SetNum(indexCount);
+		for (int i = 0; i < indexCount; ++i)
+		{
+			ImportedModel->LODModels[LODIndex].IndexBuffer[i] = LODModel.MultiSizeIndexContainer.GetIndexBuffer()->Get(i);
+		}
+
+		ImportedModel->LODModels[LODIndex].Sections.SetNum(LODModel.RenderSections.Num());
+
+		for (int SectionIndex = 0; SectionIndex < LODModel.RenderSections.Num(); ++SectionIndex)
+		{
+			const FSkelMeshRenderSection& RenderSection = LODModel.RenderSections[SectionIndex];
+			FSkelMeshSection& ImportedSection = ImportedModel->LODModels[LODIndex].Sections[SectionIndex];
+
+			ImportedSection.CorrespondClothAssetIndex = RenderSection.CorrespondClothAssetIndex;
+			ImportedSection.ClothingData = RenderSection.ClothingData;
+
+			if (RenderSection.ClothMappingDataLODs.Num())
+			{
+				ImportedSection.ClothMappingDataLODs.SetNum(1);
+				ImportedSection.ClothMappingDataLODs[0] = RenderSection.ClothMappingDataLODs[0];
+			}
+
+			// Vertices
+			ImportedSection.NumVertices = RenderSection.NumVertices;
+			ImportedSection.SoftVertices.Empty(RenderSection.NumVertices);
+			ImportedSection.SoftVertices.AddUninitialized(RenderSection.NumVertices);
+			ImportedSection.bUse16BitBoneIndex = LODModel.DoesVertexBufferUse16BitBoneIndex();
+
+			for (uint32 i = 0; i < RenderSection.NumVertices; ++i)
+			{
+				const FPositionVertex* PosPtr = static_cast<const FPositionVertex*>(LODModel.StaticVertexBuffers.PositionVertexBuffer.GetVertexData());
+				PosPtr += (CurrentSectionInitialVertex + i);
+
+				check(!LODModel.StaticVertexBuffers.StaticMeshVertexBuffer.GetUseHighPrecisionTangentBasis());
+				const FPackedNormal* TangentPtr = static_cast<const FPackedNormal*>(LODModel.StaticVertexBuffers.StaticMeshVertexBuffer.GetTangentData());
+				TangentPtr += ((CurrentSectionInitialVertex + i) * 2);
+
+				check(LODModel.StaticVertexBuffers.StaticMeshVertexBuffer.GetUseFullPrecisionUVs());
+
+				using UVsVectorType = typename TDecay<decltype(DeclVal<FSoftSkinVertex>().UVs[0])>::Type;
+
+				const UVsVectorType* TexCoordPosPtr = static_cast<const UVsVectorType*>(LODModel.StaticVertexBuffers.StaticMeshVertexBuffer.GetTexCoordData());
+				const uint32 NumTexCoords = LODModel.StaticVertexBuffers.StaticMeshVertexBuffer.GetNumTexCoords();
+				TexCoordPosPtr += ((CurrentSectionInitialVertex + i) * NumTexCoords);
+
+				FSoftSkinVertex& Vertex = ImportedSection.SoftVertices[i];
+				for (int32 j = 0; j < RenderSection.MaxBoneInfluences; ++j)
+				{
+					Vertex.InfluenceBones[j] = LODModel.SkinWeightVertexBuffer.GetBoneIndex(CurrentSectionInitialVertex + i, j);
+					Vertex.InfluenceWeights[j] = LODModel.SkinWeightVertexBuffer.GetBoneWeight(CurrentSectionInitialVertex + i, j);
+				}
+
+				for (int32 j = RenderSection.MaxBoneInfluences; j < MAX_TOTAL_INFLUENCES; ++j)
+				{
+					Vertex.InfluenceBones[j] = 0;
+					Vertex.InfluenceWeights[j] = 0;
+				}
+
+
+				Vertex.Color = FColor::White;
+
+				Vertex.Position = PosPtr->Position;
+
+				Vertex.TangentX = TangentPtr[0].ToFVector3f();
+				Vertex.TangentZ = TangentPtr[1].ToFVector3f();
+				float TangentSign = TangentPtr[1].Vector.W == 0 ? -1.f : 1.f;
+				Vertex.TangentY = FVector3f::CrossProduct(Vertex.TangentZ, Vertex.TangentX) * TangentSign;
+
+				Vertex.UVs[0] = TexCoordPosPtr[0];
+				Vertex.UVs[1] = NumTexCoords > 1 ? TexCoordPosPtr[1] : UVsVectorType::ZeroVector;
+				Vertex.UVs[2] = NumTexCoords > 2 ? TexCoordPosPtr[2] : UVsVectorType::ZeroVector;
+				Vertex.UVs[3] = NumTexCoords > 3 ? TexCoordPosPtr[3] : UVsVectorType::ZeroVector;
+			}
+
+			CurrentSectionInitialVertex += RenderSection.NumVertices;
+
+			// Triangles
+			ImportedSection.NumTriangles = RenderSection.NumTriangles;
+			ImportedSection.BaseIndex = RenderSection.BaseIndex;
+			ImportedSection.BaseVertexIndex = RenderSection.BaseVertexIndex;
+			ImportedSection.BoneMap = RenderSection.BoneMap;
+			ImportedSection.MaterialIndex = RenderSection.MaterialIndex;
+			ImportedSection.MaxBoneInfluences = RenderSection.MaxBoneInfluences;
+			ImportedSection.OriginalDataSectionIndex = OriginalIndex++;
+
+			FSkelMeshSourceSectionUserData& SectionUserData = ImportedModel->LODModels[LODIndex].UserSectionsData.FindOrAdd(ImportedSection.OriginalDataSectionIndex);
+
+			SectionUserData.CorrespondClothAssetIndex = RenderSection.CorrespondClothAssetIndex;
+			SectionUserData.ClothingData.AssetGuid = RenderSection.ClothingData.AssetGuid;
+			SectionUserData.ClothingData.AssetLodIndex = RenderSection.ClothingData.AssetLodIndex;
+		}
+
+		ImportedModel->LODModels[LODIndex].SyncronizeUserSectionsDataArray();
+
+		// DDC keys
+		const USkeletalMeshLODSettings* LODSettings = SkeletalMesh->GetLODSettings();
+		const bool bValidLODSettings = LODSettings && LODSettings->GetNumberOfSettings() > LODIndex;
+		const FSkeletalMeshLODGroupSettings* SkeletalMeshLODGroupSettings = bValidLODSettings ? &LODSettings->GetSettingsForLODLevel(LODIndex) : nullptr;
+
+		FSkeletalMeshLODInfo* LODInfo = SkeletalMesh->GetLODInfo(LODIndex);
+		LODInfo->BuildGUID = LODInfo->ComputeDeriveDataCacheKey(SkeletalMeshLODGroupSettings);
+
+		ImportedModel->LODModels[LODIndex].BuildStringID = ImportedModel->LODModels[LODIndex].GetLODModelDeriveDataKey();
+
+	}
+
+}
+
+#endif
diff --git a/Engine/Plugins/Experimental/Mutable/Source/CustomizableObject/Private/MuCO/CustomizableObjectSystem.cpp b/Engine/Plugins/Experimental/Mutable/Source/CustomizableObject/Private/MuCO/CustomizableObjectSystem.cpp
--- a/Engine/Plugins/Experimental/Mutable/Source/CustomizableObject/Private/MuCO/CustomizableObjectSystem.cpp
+++ b/Engine/Plugins/Experimental/Mutable/Source/CustomizableObject/Private/MuCO/CustomizableObjectSystem.cpp
@@ -8,6 +8,7 @@
 #include "Engine/SkeletalMesh.h"
 #include "Engine/SkinnedAsset.h"
 #include "Engine/SkinnedAssetCommon.h"
+#include "Engine/SkeletalMeshLODSettings.h"
 #include "GameFramework/PlayerController.h"
 #include "Interfaces/ITargetPlatform.h"
 #include "MuCO/CustomizableInstanceLODManagement.h"
@@ -18,7 +19,6 @@
 #include "MuCO/ICustomizableObjectModule.h"
 #include "MuCO/LogBenchmarkUtil.h"
 #include "MuCO/LogInformationUtil.h"
-#include "MuCO/UnrealBakeHelpers.h"
 #include "MuCO/UnrealExtensionDataStreamer.h"
 #include "MuCO/UnrealMutableImageProvider.h"
 #include "MuCO/UnrealMutableModelDiskStreamer.h"
@@ -44,11 +44,6 @@
 class AActor;
 class UAnimInstance;
 
-namespace impl
-{
-	void Task_Game_Callbacks_Work(UCustomizableObjectInstance* CustomizableObjectInstance);
-}
-
 
 DECLARE_CYCLE_STAT(TEXT("MutablePendingRelease Time"), STAT_MutablePendingRelease, STATGROUP_Game);
 DECLARE_CYCLE_STAT(TEXT("MutableTask"), STAT_MutableTask, STATGROUP_Game);
@@ -660,7 +655,7 @@ void UpdateSkeletalMesh(UCustomizableObjectInstance& CustomizableObjectInstance,
 		if (TObjectPtr<USkeletalMesh> SkeletalMesh = CustomizableObjectInstance.SkeletalMeshes[ComponentIndex])
 		{
 #if WITH_EDITOR
-			FUnrealBakeHelpers::BakeHelper_RegenerateImportedModel(SkeletalMesh);
+			UCustomizableInstancePrivateData::RegenerateImportedModel(SkeletalMesh);
 #else
 			SkeletalMesh->RebuildSocketMap();
 #endif
@@ -1197,6 +1192,11 @@ namespace impl
 
 		const mu::Instance* Instance = nullptr;
 
+		if (OperationData->PixelFormatOverride)
+		{
+			System->SetImagePixelConversionOverride( OperationData->PixelFormatOverride );
+		}
+
 		// Main instance generation step
 		{
 			// LOD mask, set to all ones to build  all LODs
@@ -1786,11 +1786,11 @@ namespace impl
 		MUTABLE_CPUPROFILER_SCOPE(Task_Mutable_ReleaseInstance)
 
 		check(OperationData.IsValid());
+		check(MutableSystem);
 
 		if (OperationData->InstanceID > 0)
 		{
 			MUTABLE_CPUPROFILER_SCOPE(EndUpdate);
-			check(MutableSystem);
 			MutableSystem->EndUpdate(OperationData->InstanceID);
 			OperationData->InstanceUpdateData.Clear();
 
@@ -1801,6 +1801,8 @@ namespace impl
 			}
 		}
 
+		MutableSystem->SetImagePixelConversionOverride(nullptr);
+
 		if (CVarClearWorkingMemoryOnUpdateEnd.GetValueOnAnyThread())
 		{
 			MutableSystem->ClearWorkingMemory();
@@ -2397,6 +2399,9 @@ namespace impl
 		CurrentOperationData->State = CandidateInstance->GetState();
 		CurrentOperationData->UpdateResult = EUpdateResult::Success;
 		CurrentOperationData->UpdateCallback = Operation->UpdateCallback;
+#if WITH_EDITOR
+		CurrentOperationData->PixelFormatOverride = SystemPrivateData->ImageFormatOverrideFunc;
+#endif
 
 		if (!CandidateInstancePrivateData->HasCOInstanceFlags(ForceGenerateMipTail))
 		{
@@ -3076,6 +3081,17 @@ void UCustomizableObjectSystem::SetOnlyGenerateRequestedLODsEnabled(bool bIsEnab
 }
 
 
+#if WITH_EDITOR
+void UCustomizableObjectSystem::SetImagePixelFormatOverride(const mu::FImageOperator::FImagePixelFormatFunc& InFunc)
+{
+	if (Private != nullptr)
+	{
+		Private->ImageFormatOverrideFunc = InFunc;
+	}
+}
+#endif
+
+
 void UCustomizableObjectSystem::AddUncompiledCOWarning(const UCustomizableObject& InObject, FString const* OptionalLogInfo)
 {
 	FString Msg;
diff --git a/Engine/Plugins/Experimental/Mutable/Source/CustomizableObject/Private/MuCO/CustomizableObjectSystemPrivate.h b/Engine/Plugins/Experimental/Mutable/Source/CustomizableObject/Private/MuCO/CustomizableObjectSystemPrivate.h
--- a/Engine/Plugins/Experimental/Mutable/Source/CustomizableObject/Private/MuCO/CustomizableObjectSystemPrivate.h
+++ b/Engine/Plugins/Experimental/Mutable/Source/CustomizableObject/Private/MuCO/CustomizableObjectSystemPrivate.h
@@ -597,6 +597,8 @@ struct FMutableOperationData
 	EUpdateResult UpdateResult;
 	FInstanceUpdateDelegate UpdateCallback;
 
+	mu::FImageOperator::FImagePixelFormatFunc PixelFormatOverride;
+
 #if WITH_EDITOR
 	/** Used for profiling in the editor. */
 	uint32 MutableRuntimeCycles = 0;
@@ -690,6 +692,10 @@ public:
 	static int32 SkeletalMeshMinLodQualityLevel;
 	static bool bEnableMutableReusePreviousUpdateData;
 
+#if WITH_EDITOR
+	mu::FImageOperator::FImagePixelFormatFunc ImageFormatOverrideFunc;
+#endif
+
 	/** */
 	inline void AddGameThreadTask(const FMutableTask& Task)
 	{
diff --git a/Engine/Plugins/Experimental/Mutable/Source/CustomizableObject/Private/MuCO/UnrealBakeHelpers.cpp b/Engine/Plugins/Experimental/Mutable/Source/CustomizableObject/Private/MuCO/UnrealBakeHelpers.cpp
deleted file mode 100644
diff --git a/Engine/Plugins/Experimental/Mutable/Source/CustomizableObject/Public/MuCO/CustomizableObjectCustomVersion.h b/Engine/Plugins/Experimental/Mutable/Source/CustomizableObject/Public/MuCO/CustomizableObjectCustomVersion.h
--- a/Engine/Plugins/Experimental/Mutable/Source/CustomizableObject/Public/MuCO/CustomizableObjectCustomVersion.h
+++ b/Engine/Plugins/Experimental/Mutable/Source/CustomizableObject/Public/MuCO/CustomizableObjectCustomVersion.h
@@ -94,6 +94,8 @@ struct CUSTOMIZABLEOBJECT_API FCustomizableObjectCustomVersion
 
 		DeduplicateNodeVariant,
 
+		TextureCompressionEnum,
+
 		// -----<new versions can be added above this line>-------------------------------------------------
 		VersionPlusOne,
 		LatestVersion = VersionPlusOne - 1
diff --git a/Engine/Plugins/Experimental/Mutable/Source/CustomizableObject/Public/MuCO/UnrealBakeHelpers.h b/Engine/Plugins/Experimental/Mutable/Source/CustomizableObject/Public/MuCO/UnrealBakeHelpers.h
deleted file mode 100644
diff --git a/Engine/Plugins/Experimental/Mutable/Source/CustomizableObjectEditor/Private/MuCOE/CustomizableObjectCompileRunnable.cpp b/Engine/Plugins/Experimental/Mutable/Source/CustomizableObjectEditor/Private/MuCOE/CustomizableObjectCompileRunnable.cpp
--- a/Engine/Plugins/Experimental/Mutable/Source/CustomizableObjectEditor/Private/MuCOE/CustomizableObjectCompileRunnable.cpp
+++ b/Engine/Plugins/Experimental/Mutable/Source/CustomizableObjectEditor/Private/MuCOE/CustomizableObjectCompileRunnable.cpp
@@ -7,12 +7,12 @@
 #include "MuR/Model.h"
 #include "MuT/Compiler.h"
 #include "MuT/ErrorLog.h"
+#include "MuT/UnrealPixelFormatOverride.h"
 #include "Serialization/MemoryWriter.h"
 #include "Trace/Trace.inl"
 
 class ITargetPlatform;
 
-
 #define LOCTEXT_NAMESPACE "CustomizableObjectEditor"
 
 
@@ -21,6 +21,7 @@ FCustomizableObjectCompileRunnable::FCustomizableObjectCompileRunnable(mu::Ptr<m
 	, bThreadCompleted(false)
 	, MutableIsDisabled(false)
 {
+	PrepareUnrealCompression();
 }
 
 
@@ -38,7 +39,8 @@ uint32 FCustomizableObjectCompileRunnable::Run()
 		return true;
 	}
 
-	mu::CompilerOptionsPtr CompilerOptions = new mu::CompilerOptions();
+	// Translate CO compile options into mu::CompilerOptions
+	mu::Ptr<mu::CompilerOptions> CompilerOptions = new mu::CompilerOptions();
 
 	CompilerOptions->SetUseDiskCache(Options.bUseDiskCompilation);
 
@@ -81,6 +83,13 @@ uint32 FCustomizableObjectCompileRunnable::Run()
 		break;
 	}
 
+	// Texture compression override, if necessary
+	if (Options.TextureCompression== ECustomizableObjectTextureCompression::HighQuality)
+	{
+		CompilerOptions->SetImagePixelFormatOverride( UnrealPixelFormatFunc );
+	}
+
+
 	// Minimum resident mip count.
 	const int MinResidentMips = UTexture::GetStaticMinTextureResidentMipCount();
 	// Data smaller than this will always be loaded, as part of the customizable object compiled model.
@@ -93,7 +102,7 @@ uint32 FCustomizableObjectCompileRunnable::Run()
 	
 	CompilerOptions->SetImageTiling(Options.ImageTiling);
 
-	mu::CompilerPtr Compiler = new mu::Compiler(CompilerOptions);
+	mu::Ptr<mu::Compiler> Compiler = new mu::Compiler(CompilerOptions);
 
 	UE_LOG(LogMutable, Verbose, TEXT("PROFILE: [ %16.8f ] FCustomizableObjectCompileRunnable Compile start."), FPlatformTime::Seconds());
 	Model = Compiler->Compile(MutableRoot);
diff --git a/Engine/Plugins/Experimental/Mutable/Source/CustomizableObjectEditor/Private/MuCOE/CustomizableObjectEditor.cpp b/Engine/Plugins/Experimental/Mutable/Source/CustomizableObjectEditor/Private/MuCOE/CustomizableObjectEditor.cpp
--- a/Engine/Plugins/Experimental/Mutable/Source/CustomizableObjectEditor/Private/MuCOE/CustomizableObjectEditor.cpp
+++ b/Engine/Plugins/Experimental/Mutable/Source/CustomizableObjectEditor/Private/MuCOE/CustomizableObjectEditor.cpp
@@ -441,7 +441,7 @@ void FCustomizableObjectEditor::CreatePreviewInstance()
 			}
 			else
 			{
-				FNotificationInfo Info(NSLOCTEXT("CustomizableObject", "CustomizableObjectCompileTryLater", "Please wait until asset registry loads all assets"));
+				FNotificationInfo Info(LOCTEXT("CustomizableObjectCompileTryLater", "Please wait until asset registry loads all assets"));
 				Info.bFireAndForget = true;
 				Info.bUseThrobber = true;
 				Info.FadeOutDuration = 1.0f;
@@ -808,12 +808,6 @@ void FCustomizableObjectEditor::BindCommands()
 		FCanExecuteAction(),
 		FIsActionChecked());
 
-	ToolkitCommands->MapAction(
-		Commands.CompileOptions_EnableTextureCompression,
-		FExecuteAction::CreateSP(this, &FCustomizableObjectEditor::CompileOptions_TextureCompression_Toggled),
-		FCanExecuteAction(),
-		FIsActionChecked::CreateSP(this, &FCustomizableObjectEditor::CompileOptions_TextureCompression_IsChecked));
-
 	ToolkitCommands->MapAction(
 		Commands.CompileOptions_UseDiskCompilation,
 		FExecuteAction::CreateSP(this, &FCustomizableObjectEditor::CompileOptions_UseDiskCompilation_Toggled),
@@ -1020,10 +1014,10 @@ TSharedRef<SWidget> FCustomizableObjectEditor::GenerateCompileOptionsMenuContent
 	{
 		// Level
 		CompileOptimizationStrings.Empty();
-		CompileOptimizationStrings.Add(MakeShareable(new FString(NSLOCTEXT("UnrealEd", "OptimizationNone", "None").ToString())));
-		CompileOptimizationStrings.Add(MakeShareable(new FString(NSLOCTEXT("UnrealEd", "OptimizationMin", "Minimal").ToString())));
-		CompileOptimizationStrings.Add(MakeShareable(new FString(NSLOCTEXT("UnrealEd", "OptimizationMed", "Medium").ToString())));
-		CompileOptimizationStrings.Add(MakeShareable(new FString(NSLOCTEXT("UnrealEd", "OptimizationMax", "Maximum").ToString())));
+		CompileOptimizationStrings.Add(MakeShareable(new FString(LOCTEXT("OptimizationNone", "None").ToString())));
+		CompileOptimizationStrings.Add(MakeShareable(new FString(LOCTEXT("OptimizationMin", "Minimal").ToString())));
+		CompileOptimizationStrings.Add(MakeShareable(new FString(LOCTEXT("OptimizationMed", "Medium").ToString())));
+		CompileOptimizationStrings.Add(MakeShareable(new FString(LOCTEXT("OptimizationMax", "Maximum").ToString())));
 
 		if (CustomizableObject)
 		{
@@ -1038,6 +1032,23 @@ TSharedRef<SWidget> FCustomizableObjectEditor::GenerateCompileOptionsMenuContent
 			MenuBuilder.AddWidget(CompileOptimizationCombo.ToSharedRef(), LOCTEXT("MutableCompileOptimizationLevel", "Optimization Level"));
 		}
 
+		{
+			CompileTextureCompressionStrings.Empty();
+			CompileTextureCompressionStrings.Add(MakeShareable(new FString(LOCTEXT("MutableTextureCompressionNone", "None").ToString())));
+			CompileTextureCompressionStrings.Add(MakeShareable(new FString(LOCTEXT("MutableTextureCompressionFast", "Fast").ToString())));
+			CompileTextureCompressionStrings.Add(MakeShareable(new FString(LOCTEXT("MutableTextureCompressionHighQuality", "High Quality").ToString())));
+
+			int32 SelectedCompression = FMath::Clamp(int32(CustomizableObject->CompileOptions.TextureCompression), 0, CompileTextureCompressionStrings.Num() - 1);
+			CompileTextureCompressionCombo =
+				SNew(STextComboBox)
+				.OptionsSource(&CompileTextureCompressionStrings)
+				.InitiallySelectedItem(CompileTextureCompressionStrings[SelectedCompression])
+				.OnSelectionChanged(this, &FCustomizableObjectEditor::OnChangeCompileTextureCompressionType)
+				;
+
+			MenuBuilder.AddWidget(CompileTextureCompressionCombo.ToSharedRef(), LOCTEXT("MutableCompileTextureCompressionType", "Texture Compression"));
+		}
+
 		// Image tiling
 		// Unfortunately SNumericDropDown doesn't work with integers at the time of writing.
 		TArray<SNumericDropDown<float>::FNamedValue> TilingOptions;
@@ -1291,7 +1302,7 @@ void FCustomizableObjectEditor::CompileObject()
 
 	if (!AssetRegistryLoaded)
 	{
-		FNotificationInfo Info(NSLOCTEXT("CustomizableObject", "CustomizableObjectCompileTryLater", "Please wait until asset registry loads all assets"));
+		FNotificationInfo Info(LOCTEXT("CustomizableObjectCompileTryLater", "Please wait until asset registry loads all assets"));
 		Info.bFireAndForget = true;
 		Info.bUseThrobber = true;
 		Info.FadeOutDuration = 1.0f;
@@ -1357,16 +1368,11 @@ bool FCustomizableObjectEditor::CompileOptions_UseDiskCompilation_IsChecked()
 }
 
 
-void FCustomizableObjectEditor::CompileOptions_TextureCompression_Toggled()
+void FCustomizableObjectEditor::OnChangeCompileTextureCompressionType(TSharedPtr<FString> NewSelection, ESelectInfo::Type)
 {
-	CustomizableObject->CompileOptions.bTextureCompression = !CustomizableObject->CompileOptions.bTextureCompression;
+	const FScopedTransaction Transaction(LOCTEXT("ChangedTextureCompressionTransaction", "Changed Texture Compression Type"));
 	CustomizableObject->Modify();
-}
-
-
-bool FCustomizableObjectEditor::CompileOptions_TextureCompression_IsChecked()
-{
-	return CustomizableObject->CompileOptions.bTextureCompression;
+	CustomizableObject->CompileOptions.TextureCompression = ECustomizableObjectTextureCompression(CompileTextureCompressionStrings.Find(NewSelection));
 }
 
 
@@ -1399,7 +1405,7 @@ void FCustomizableObjectEditor::DeleteSelectedNodes()
 		return;
 	}
 
-	const FScopedTransaction Transaction(NSLOCTEXT("UnrealEd", "UEdGraphSchema_CustomizableObject", "Delete Nodes"));
+	const FScopedTransaction Transaction(LOCTEXT("UEdGraphSchema_CustomizableObject", "Delete Nodes"));
 
 	const FGraphPanelSelectionSet SelectedNodes = GraphEditor->GetSelectedNodes();
 	GraphEditor->ClearSelectionSet();
@@ -1764,7 +1770,7 @@ void FCustomizableObjectEditor::PasteNodes()
 void FCustomizableObjectEditor::PasteNodesHere(const FVector2D& Location)
 {
 	// Undo/Redo support
-	const FScopedTransaction Transaction( NSLOCTEXT("UnrealEd", "CustomizableObjectEditorPaste", "Customizable Object Editor Editor: Paste") );
+	const FScopedTransaction Transaction( LOCTEXT("CustomizableObjectEditorPaste", "Customizable Object Editor Editor: Paste") );
 	CustomizableObject->Source->Modify();
 	CustomizableObject->Modify();
 
@@ -1901,7 +1907,7 @@ void FCustomizableObjectEditor::CreateCommentBoxFromKey()
 
 UEdGraphNode* FCustomizableObjectEditor::CreateCommentBox(const FVector2D& InTargetPosition)
 {
-	//const FScopedTransaction Transaction(NSLOCTEXT("UnrealEd", "UEdGraphSchema_CustomizableObject", "Add Comment Box"));
+	//const FScopedTransaction Transaction(LOCTEXT("UEdGraphSchema_CustomizableObject", "Add Comment Box"));
 
 	UEdGraphNode_Comment* CommentTemplate = NewObject<UEdGraphNode_Comment>();
 
diff --git a/Engine/Plugins/Experimental/Mutable/Source/CustomizableObjectEditor/Private/MuCOE/CustomizableObjectEditor.h b/Engine/Plugins/Experimental/Mutable/Source/CustomizableObjectEditor/Private/MuCOE/CustomizableObjectEditor.h
--- a/Engine/Plugins/Experimental/Mutable/Source/CustomizableObjectEditor/Private/MuCOE/CustomizableObjectEditor.h
+++ b/Engine/Plugins/Experimental/Mutable/Source/CustomizableObjectEditor/Private/MuCOE/CustomizableObjectEditor.h
@@ -282,13 +282,14 @@ private:
 	void ResetCompileOptions();
 	TSharedPtr<STextComboBox> CompileOptimizationCombo;
 	TArray< TSharedPtr<FString> > CompileOptimizationStrings;
+	TSharedPtr<STextComboBox> CompileTextureCompressionCombo;
+	TArray< TSharedPtr<FString> > CompileTextureCompressionStrings;
 	TSharedPtr<SNumericDropDown<float>> CompileTilingCombo;
 
 	void CompileOptions_UseDiskCompilation_Toggled();
 	bool CompileOptions_UseDiskCompilation_IsChecked();
-	void CompileOptions_TextureCompression_Toggled();
-	bool CompileOptions_TextureCompression_IsChecked();
-	void OnChangeCompileOptimizationLevel(TSharedPtr<FString> NewSelection, ESelectInfo::Type SelectInfo);
+	void OnChangeCompileOptimizationLevel(TSharedPtr<FString> NewSelection, ESelectInfo::Type);
+	void OnChangeCompileTextureCompressionType(TSharedPtr<FString> NewSelection, ESelectInfo::Type);
 
 	/** Save Customizable Object open in editor */
 	void SaveAsset_Execute() override;
diff --git a/Engine/Plugins/Experimental/Mutable/Source/CustomizableObjectEditor/Private/MuCOE/CustomizableObjectEditorViewportClient.cpp b/Engine/Plugins/Experimental/Mutable/Source/CustomizableObjectEditor/Private/MuCOE/CustomizableObjectEditorViewportClient.cpp
--- a/Engine/Plugins/Experimental/Mutable/Source/CustomizableObjectEditor/Private/MuCOE/CustomizableObjectEditorViewportClient.cpp
+++ b/Engine/Plugins/Experimental/Mutable/Source/CustomizableObjectEditor/Private/MuCOE/CustomizableObjectEditorViewportClient.cpp
@@ -26,7 +26,7 @@
 #include "MuCO/CustomizableObjectInstance.h"
 #include "MuCO/CustomizableObjectSystem.h"
 #include "MuCO/CustomizableObjectMipDataProvider.h"
-#include "MuCO/UnrealBakeHelpers.h"
+#include "MuCOE/UnrealBakeHelpers.h"
 #include "MuCOE/CustomizableObjectPreviewScene.h"
 #include "MuCOE/CustomizableObjectWidget.h"
 #include "MuCOE/ICustomizableObjectInstanceEditor.h"
@@ -35,6 +35,7 @@
 #include "MuCOE/Nodes/CustomizableObjectNodeProjectorConstant.h"
 #include "MuCOE/Nodes/CustomizableObjectNodeProjectorParameter.h"
 #include "MuCOE/UnrealEditorPortabilityHelpers.h"
+#include "MuT/UnrealPixelFormatOverride.h"
 #include "ObjectTools.h"
 #include "Preferences/PersonaOptions.h"
 #include "ScopedTransaction.h"
@@ -2267,13 +2268,16 @@ void FCustomizableObjectEditorViewportClient::BakeInstance(UCustomizableObjectIn
 	UCustomizableObjectSystem* System = UCustomizableObjectSystem::GetInstance();
 	check(System);
 
-	if (System->IsProgressiveMipStreamingEnabled())
+	if (!InInstance)
 	{
 		// The instance in the editor viewport does not have high quality mips in the platform data because streaming is enabled.
 		// Disable streaming and retry with a newly generated temp instance.
 		System->SetProgressiveMipStreamingEnabled(false);
 		// Disable requested LOD generation as it will prevent the new instance from having all the LODs
 		System->SetOnlyGenerateRequestedLODsEnabled(false);
+		// Force high quality texture compression for this instance
+		PrepareUnrealCompression();
+		System->SetImagePixelFormatOverride(UnrealPixelFormatFunc);
 
 		BakeTempInstance = Instance->Clone();
 		BakeTempInstance->SkeletalMeshes.Empty();
@@ -2283,7 +2287,26 @@ void FCustomizableObjectEditorViewportClient::BakeInstance(UCustomizableObjectIn
 		return;
 	}
 
-	FString ObjectName = Instance->GetCustomizableObject()->GetName();
+	// Warn if the CO looks like it wasn't compiled with high-quality texture compression.
+	// \TODO: This is a weak test, we should store the setting with the compiled data and check that instead.
+	UCustomizableObject* CO = Instance->GetCustomizableObject();
+	if (!CO)
+	{
+		// Something is very wrong
+		return;
+	}
+
+	if (CO->CompileOptions.TextureCompression!=ECustomizableObjectTextureCompression::HighQuality)
+	{
+		FNotificationInfo Info(NSLOCTEXT("CustomizableObjectEditor", "CustomizableObjectBakeLowQuality", "The Customizable Object wasn't compiled with high quality textures. For the best baking results, change the Texture Compression setting and recompile it."));
+		Info.bFireAndForget = true;
+		Info.bUseThrobber = true;
+		Info.FadeOutDuration = 1.0f;
+		Info.ExpireDuration = 6.0f;
+		FSlateNotificationManager::Get().AddNotification(Info);
+	}
+
+	FString ObjectName = CO->GetName();
 	FText DefaultFileName = FText::Format(LOCTEXT("DefaultFileNameForBakeInstance", "{0}"), FText::AsCultureInvariant(ObjectName));
 	bool bExportAllResources = false;
 
@@ -2321,7 +2344,7 @@ void FCustomizableObjectEditorViewportClient::BakeInstance(UCustomizableObjectIn
 		bExportAllResources = FolderDlg->GetExportAllResources();
 
 		BakingOverwritePermission = false;
-		FString CustomObjectPath = Instance->GetCustomizableObject()->GetPathName();
+		FString CustomObjectPath = CO->GetPathName();
 		FString AssetPath = FolderDlg->GetAssetPath();
 		FString FullAssetPath = AssetPath + FString("/") + ObjectName + FString(".") + ObjectName;
 
@@ -2684,9 +2707,6 @@ void FCustomizableObjectEditorViewportClient::BakeInstance(UCustomizableObjectIn
 					}
 				}
 
-				// Make sure source data is present in the mesh before we duplicate:
-				FUnrealBakeHelpers::BakeHelper_RegenerateImportedModel(Mesh);
-
 				// Skeletal Mesh
 				if (!ManageBakingAction(AssetPath, ObjectName))
 				{
@@ -2744,6 +2764,7 @@ void FCustomizableObjectEditorViewportClient::BakeInstance(UCustomizableObjectIn
 		// Reenable Mutable texture streaming and requested LOD generation as they had been disabled to bake the textures
 		System->SetProgressiveMipStreamingEnabled(true);
 		System->SetOnlyGenerateRequestedLODsEnabled(true);
+		System->SetImagePixelFormatOverride(nullptr);
 		BakeTempInstance = nullptr;
 	}
 }
diff --git a/Engine/Plugins/Experimental/Mutable/Source/CustomizableObjectEditor/Private/MuCOE/GenerateMutableSource/GenerateMutableSourceSurface.cpp b/Engine/Plugins/Experimental/Mutable/Source/CustomizableObjectEditor/Private/MuCOE/GenerateMutableSource/GenerateMutableSourceSurface.cpp
--- a/Engine/Plugins/Experimental/Mutable/Source/CustomizableObjectEditor/Private/MuCOE/GenerateMutableSource/GenerateMutableSourceSurface.cpp
+++ b/Engine/Plugins/Experimental/Mutable/Source/CustomizableObjectEditor/Private/MuCOE/GenerateMutableSource/GenerateMutableSourceSurface.cpp
@@ -6,6 +6,11 @@
 #include "Engine/StaticMesh.h"
 #include "GPUSkinPublicDefs.h"
 #include "Interfaces/ITargetPlatform.h"
+#include "Interfaces/ITextureFormat.h"
+#include "Interfaces/ITextureFormatModule.h"
+#include "Interfaces/ITextureFormatManagerModule.h"
+#include "Modules/ModuleManager.h"
+#include "TextureCompressorModule.h"
 #include "Materials/MaterialInstance.h"
 #include "GPUSkinVertexFactory.h"
 
@@ -44,10 +49,10 @@
 #include "MuT/NodeScalarConstant.h"
 #include "MuT/NodeSurfaceEdit.h"
 #include "MuT/NodeSurfaceVariation.h"
+#include "MuT/UnrealPixelFormatOverride.h"
 
 #define LOCTEXT_NAMESPACE "CustomizableObjectEditor"
 
-
 /**
  * Show a compilation warning when any of the meshes behind the pin InMeshPin has a UV that is not normalized.
  * 
@@ -632,11 +637,6 @@ mu::NodeSurfacePtr GenerateMutableSourceSurface(const UEdGraphPin * Pin, FMutabl
 						{
 							int LayerIndex = 0;
 
-							TArray< TArray<FName> > PlatformFormats;
-							GenerationContext.Options.TargetPlatform->GetTextureFormats(ReferenceTexture, PlatformFormats);
-
-							bool bHaveFormatPlatformFormats = PlatformFormats.Num() > 0;
-
 							mu::NodeImagePtr LastImage = ImageNode;
 
 							// To apply LOD bias						
@@ -697,103 +697,59 @@ mu::NodeSurfacePtr GenerateMutableSourceSurface(const UEdGraphPin * Pin, FMutabl
 								LastImage = CompositedImage;
 							}
 
-							if (bHaveFormatPlatformFormats)
-							{
-								mu::NodeImageFormatPtr FormatImage = new mu::NodeImageFormat();
-								FormatImage->SetSource(LastImage.get());
-								FormatImage->SetFormat(mu::EImageFormat::IF_RGBA_UBYTE);
-								FormatImage->SetMessageContext(Node);
-								LastImage = FormatImage;
+							mu::NodeImageFormatPtr FormatImage = new mu::NodeImageFormat();
+							FormatImage->SetSource(LastImage.get());
+							FormatImage->SetFormat(mu::EImageFormat::IF_RGBA_UBYTE);
+							FormatImage->SetMessageContext(Node);
+							LastImage = FormatImage;
 
-								if (GenerationContext.Options.bTextureCompression)
+							TArray<FTextureBuildSettings> BuildSettingsPerLayer;
+							ReferenceTexture->GetTargetPlatformBuildSettings(GenerationContext.Options.TargetPlatform, BuildSettingsPerLayer);
+							if (!BuildSettingsPerLayer.IsEmpty())
+							{
+								if (GenerationContext.Options.TextureCompression!=ECustomizableObjectTextureCompression::None)
 								{
-									check(PlatformFormats[0].Num() > LayerIndex);
-
-									const FString PlatformFormat = PlatformFormats[0][LayerIndex].ToString();
-
-									// Remove platform prefix
-									FString FormatWithoutPrefix = PlatformFormat;
-									PlatformFormat.Split(TEXT("_"), nullptr, &FormatWithoutPrefix, ESearchCase::IgnoreCase, ESearchDir::FromEnd);
-
-									mu::EImageFormat mutableFormat = mu::EImageFormat::IF_NONE;
-									mu::EImageFormat mutableFormatIfAlpha = mu::EImageFormat::IF_NONE;
-
-									if (FormatWithoutPrefix == TEXT("AutoDXT"))
+									static ITextureFormatManagerModule* TextureFormatManager = nullptr;
+									if (!TextureFormatManager)
 									{
-										mutableFormat = mu::EImageFormat::IF_BC1;
-										mutableFormatIfAlpha = mu::EImageFormat::IF_BC3;
+										TextureFormatManager = &FModuleManager::LoadModuleChecked<ITextureFormatManagerModule>("TextureFormat");
+										check(TextureFormatManager);
 									}
-									else if (FormatWithoutPrefix == TEXT("DXT1")) mutableFormat = mu::EImageFormat::IF_BC1;
-									else if (FormatWithoutPrefix == TEXT("DXT3")) mutableFormat = mu::EImageFormat::IF_BC2;
-									else if (FormatWithoutPrefix == TEXT("DXT5")) mutableFormat = mu::EImageFormat::IF_BC3;
-									else if (FormatWithoutPrefix == TEXT("BC1")) mutableFormat = mu::EImageFormat::IF_BC1;
-									else if (FormatWithoutPrefix == TEXT("BC2")) mutableFormat = mu::EImageFormat::IF_BC2;
-									else if (FormatWithoutPrefix == TEXT("BC3")) mutableFormat = mu::EImageFormat::IF_BC3;
-									else if (FormatWithoutPrefix == TEXT("BC4")) mutableFormat = mu::EImageFormat::IF_BC4;
-									else if (FormatWithoutPrefix == TEXT("BC5")) mutableFormat = mu::EImageFormat::IF_BC5;
-									else if (FormatWithoutPrefix == TEXT("G8")) mutableFormat = mu::EImageFormat::IF_L_UBYTE;
-									else if (FormatWithoutPrefix == TEXT("BGRA8")) mutableFormat = mu::EImageFormat::IF_RGBA_UBYTE;
-									else if (PlatformFormat.Contains(TEXT("ASTC")))
+									const ITextureFormat* TextureFormat = TextureFormatManager->FindTextureFormat(BuildSettingsPerLayer[0].TextureFormatName);
+									check(TextureFormat);
+									EPixelFormat UnrealTargetPlatformFormat = TextureFormat->GetEncodedPixelFormat(BuildSettingsPerLayer[0], false);
+									EPixelFormat UnrealTargetPlatformFormatAlpha = TextureFormat->GetEncodedPixelFormat(BuildSettingsPerLayer[0], true);
+
+									// \TODO: The QualityFix filter is used while the internal mutable runtime compression doesn't provide enough quality for some large block formats.
+									mu::EImageFormat MutableFormat = QualityFix(UnrealToMutablePixelFormat(UnrealTargetPlatformFormat,false));
+									mu::EImageFormat MutableFormatIfAlpha = QualityFix(UnrealToMutablePixelFormat(UnrealTargetPlatformFormatAlpha,true));
+
+									// Unsupported format: look for something generic
+									if (MutableFormat == mu::EImageFormat::IF_NONE)
 									{
-										// \TODO: ASTC block size to use depends on several project and platform settings. Try to get it.
-										
-										// This is based on GetQualityFormat in TextureFormatASTC.cpp but it would be much better to find
-										// an way of doing it without dependant code.
-										if ((FormatWithoutPrefix == TEXT("AutoASTC")) || (FormatWithoutPrefix == TEXT("RGBAuto")))
-										{
-											mutableFormat = mu::EImageFormat::IF_ASTC_4x4_RGB_LDR;
-											mutableFormatIfAlpha = mu::EImageFormat::IF_ASTC_4x4_RGBA_LDR;
-										}
-										else if (FormatWithoutPrefix == TEXT("RGB")) mutableFormat = mu::EImageFormat::IF_ASTC_4x4_RGB_LDR;
-										else if (FormatWithoutPrefix == TEXT("RGBA")) mutableFormat = mu::EImageFormat::IF_ASTC_4x4_RGBA_LDR;
-										else if (FormatWithoutPrefix == TEXT("RGBA_HQ")) mutableFormat = mu::EImageFormat::IF_ASTC_4x4_RGBA_LDR;
-										else if (FormatWithoutPrefix == TEXT("NormalRG")) mutableFormat = mu::EImageFormat::IF_ASTC_4x4_RG_LDR;
-										else if (PlatformFormat.Contains(TEXT("ASTC_NormalRG_Precise")))
-										{
-											// TODO: This is just a workaround to prevent the "Unexpected image format" warning below. ASTC_NormalRG_Precise is
-											// not supported yet by Mutable so it forces IF_ASTC_4x4_RG_LDR as a replacement. It should be changed with a more
-											// appropriate format or directly implement it
-											mutableFormat = mu::EImageFormat::IF_ASTC_4x4_RG_LDR;
-
-											const FString ReplacedImageFormatMsg = FString::Printf(TEXT("In object [%s] the unsupported ASTC_NormalRG_Precise image format is used, ASTC_4x4_RG_LDR will be used instead."), *GenerationContext.Object->GetName());
-											const FText ReplacedImageFormatText = FText::FromString(ReplacedImageFormatMsg);
-											GenerationContext.Compiler->CompilerLog(ReplacedImageFormatText, Node, EMessageSeverity::Info);
-											UE_LOG(LogMutable, Log, TEXT("%s"), *ReplacedImageFormatMsg);
-										}
-										else if (PlatformFormat.Contains(TEXT("ASTC_NormalLA")))
-										{
-											// TODO: This is just a workaround to prevent the "Unexpected image format" warning below. ASTC_NormalLA is
-											// not supported yet by Mutable so it forces IF_ASTC_4x4_RG_LDR as a replacement. It should be changed with a more
-											// appropriate format or directly implement it
-											mutableFormat = mu::EImageFormat::IF_ASTC_4x4_RG_LDR;
-
-											const FString ReplacedImageFormatMsg2 = FString::Printf(TEXT("In object [%s] the unsupported ASTC_NormalLA image format is used, ASTC_4x4_RG_LDR will be used instead."), *GenerationContext.Object->GetName());
-											const FText ReplacedImageFormatText2 = FText::FromString(ReplacedImageFormatMsg2);
-											GenerationContext.Compiler->CompilerLog(ReplacedImageFormatText2, Node, EMessageSeverity::Info);
-											UE_LOG(LogMutable, Log, TEXT("%s"), *ReplacedImageFormatMsg2);
-										}
+										const FString ReplacedImageFormatMsg = FString::Printf(TEXT("In object [%s] the unsupported image format %d is used, IF_RGBA_UBYTE will be used instead."), *GenerationContext.Object->GetName(), UnrealTargetPlatformFormat );
+										const FText ReplacedImageFormatText = FText::FromString(ReplacedImageFormatMsg);
+										GenerationContext.Compiler->CompilerLog(ReplacedImageFormatText, Node, EMessageSeverity::Info);
+										UE_LOG(LogMutable, Log, TEXT("%s"), *ReplacedImageFormatMsg);
+										MutableFormat = mu::EImageFormat::IF_RGBA_UBYTE;
 									}
-
-									if (mutableFormat == mu::EImageFormat::IF_NONE)
+									if (MutableFormatIfAlpha == mu::EImageFormat::IF_NONE)
 									{
-										// Format not supported by Mutable, use RBGA_UBYTE as default.
-										mutableFormat = mu::EImageFormat::IF_RGBA_UBYTE;
-
-										const FString UnexpectedImageFormatMsg = FString::Printf(TEXT("In object [%s] Unexpected image format [%s], RGBA_UBYTE will be used instead."), *GenerationContext.Object->GetName(), *PlatformFormat);
-										const FText UnexpectedImageFormatText = FText::FromString(UnexpectedImageFormatMsg);
-										GenerationContext.Compiler->CompilerLog(UnexpectedImageFormatText, Node);
-										UE_LOG(LogMutable, Warning, TEXT("%s"), *UnexpectedImageFormatMsg);
+										const FString ReplacedImageFormatMsg = FString::Printf(TEXT("In object [%s] the unsupported image format %d is used, IF_RGBA_UBYTE will be used instead."), *GenerationContext.Object->GetName(), UnrealTargetPlatformFormatAlpha);
+										const FText ReplacedImageFormatText = FText::FromString(ReplacedImageFormatMsg);
+										GenerationContext.Compiler->CompilerLog(ReplacedImageFormatText, Node, EMessageSeverity::Info);
+										UE_LOG(LogMutable, Log, TEXT("%s"), *ReplacedImageFormatMsg);
+										MutableFormatIfAlpha = mu::EImageFormat::IF_RGBA_UBYTE;
 									}
 
-									FormatImage->SetFormat(mutableFormat, mutableFormatIfAlpha);
+									FormatImage->SetFormat(MutableFormat, MutableFormatIfAlpha);
 								}
 							}
 
 							ImageNode = LastImage;
 						}
 
-						mu::NodeImagePtr ImageNodePtr = ImageNode;
-						SurfNode->SetImage(ImageIndex, ImageNodePtr);
+						SurfNode->SetImage(ImageIndex, ImageNode);
 
 						FString SurfNodeImageName = FString::Printf(TEXT("%d"), GenerationContext.ImageProperties.Num() - 1);
 
@@ -815,7 +771,7 @@ mu::NodeSurfacePtr GenerateMutableSourceSurface(const UEdGraphPin * Pin, FMutabl
 							ensure(LOD == GenerationContext.FirstLODAvailable);
 							float* AlternateProjectionResFactor = TextureNameToProjectionResFactor.Find(ImageName);
 							GenerationContext.GroupProjectorLODCache.Add(MaterialImageId,
-								FGroupProjectorImageInfo(ImageNodePtr, ImageName, ImageName, TypedNodeMat,
+								FGroupProjectorImageInfo(ImageNode, ImageName, ImageName, TypedNodeMat,
 									AlternateProjectionResFactor ? *AlternateProjectionResFactor : 0.f, AlternateResStateName, SurfNode, UVLayout));
 						}
 					}
@@ -1605,6 +1561,4 @@ mu::NodeSurfacePtr GenerateMutableSourceSurface(const UEdGraphPin * Pin, FMutabl
 	return Result;
 }
 
-
 #undef LOCTEXT_NAMESPACE
-
diff --git a/Engine/Plugins/Experimental/Mutable/Source/CustomizableObjectEditor/Private/MuCOE/SMutableGraphViewer.cpp b/Engine/Plugins/Experimental/Mutable/Source/CustomizableObjectEditor/Private/MuCOE/SMutableGraphViewer.cpp
--- a/Engine/Plugins/Experimental/Mutable/Source/CustomizableObjectEditor/Private/MuCOE/SMutableGraphViewer.cpp
+++ b/Engine/Plugins/Experimental/Mutable/Source/CustomizableObjectEditor/Private/MuCOE/SMutableGraphViewer.cpp
@@ -17,6 +17,7 @@
 #include "Widgets/Input/STextComboBox.h"
 #include "Widgets/Input/SNumericDropDown.h"
 #include "Widgets/Views/STreeView.h"
+#include "ScopedTransaction.h"
 
 class FExtender;
 class FReferenceCollector;
@@ -302,6 +303,23 @@ TSharedRef<SWidget> SMutableGraphViewer::GenerateCompileOptionsMenuContent()
 			;
 		MenuBuilder.AddWidget(CompileOptimizationCombo.ToSharedRef(), LOCTEXT("MutableCompileOptimizationLevel", "Optimization Level"));
 
+		{
+			CompileTextureCompressionStrings.Empty();
+			CompileTextureCompressionStrings.Add(MakeShareable(new FString(LOCTEXT("MutableTextureCompressionNone", "None").ToString())));
+			CompileTextureCompressionStrings.Add(MakeShareable(new FString(LOCTEXT("MutableTextureCompressionFast", "Fast").ToString())));
+			CompileTextureCompressionStrings.Add(MakeShareable(new FString(LOCTEXT("MutableTextureCompressionHighQuality", "High Quality").ToString())));
+
+			int32 SelectedCompression = FMath::Clamp(int32(CompileOptions.TextureCompression), 0, CompileTextureCompressionStrings.Num() - 1);
+			CompileTextureCompressionCombo =
+				SNew(STextComboBox)
+				.OptionsSource(&CompileTextureCompressionStrings)
+				.InitiallySelectedItem(CompileTextureCompressionStrings[SelectedCompression])
+				.OnSelectionChanged(this, &SMutableGraphViewer::OnChangeCompileTextureCompressionType)
+				;
+
+			MenuBuilder.AddWidget(CompileTextureCompressionCombo.ToSharedRef(), LOCTEXT("MutableCompileTextureCompressionType", "Texture Compression"));
+		}
+
 		// Image tiling
 		// Unfortunately SNumericDropDown doesn't work with integers at the time of writing.
 		TArray<SNumericDropDown<float>::FNamedValue> TilingOptions;
@@ -337,12 +355,19 @@ TSharedRef<SWidget> SMutableGraphViewer::GenerateCompileOptionsMenuContent()
 }
 
 
-void SMutableGraphViewer::OnChangeCompileOptimizationLevel(TSharedPtr<FString> NewSelection, ESelectInfo::Type SelectInfo)
+void SMutableGraphViewer::OnChangeCompileOptimizationLevel(TSharedPtr<FString> NewSelection, ESelectInfo::Type)
 {
 	CompileOptions.OptimizationLevel = CompileOptimizationStrings.Find(NewSelection);
 }
 
 
+void SMutableGraphViewer::OnChangeCompileTextureCompressionType(TSharedPtr<FString> NewSelection, ESelectInfo::Type)
+{
+	const FScopedTransaction Transaction(LOCTEXT("ChangedOptimizationLevelTransaction", "Changed Optimization Level"));
+	CompileOptions.TextureCompression = ECustomizableObjectTextureCompression(CompileTextureCompressionStrings.Find(NewSelection));
+}
+
+
 void SMutableGraphViewer::RebuildTree()
 {
 	RootNodes.Reset();
diff --git a/Engine/Plugins/Experimental/Mutable/Source/CustomizableObjectEditor/Private/MuCOE/SMutableGraphViewer.h b/Engine/Plugins/Experimental/Mutable/Source/CustomizableObjectEditor/Private/MuCOE/SMutableGraphViewer.h
--- a/Engine/Plugins/Experimental/Mutable/Source/CustomizableObjectEditor/Private/MuCOE/SMutableGraphViewer.h
+++ b/Engine/Plugins/Experimental/Mutable/Source/CustomizableObjectEditor/Private/MuCOE/SMutableGraphViewer.h
@@ -104,8 +104,11 @@ private:
 	TSharedRef<SWidget> GenerateCompileOptionsMenuContent();
 	TSharedPtr<STextComboBox> CompileOptimizationCombo;
 	TArray< TSharedPtr<FString> > CompileOptimizationStrings;
+	TSharedPtr<STextComboBox> CompileTextureCompressionCombo;
+	TArray< TSharedPtr<FString> > CompileTextureCompressionStrings;
 	TSharedPtr<SNumericDropDown<float>> CompileTilingCombo;
 	void OnChangeCompileOptimizationLevel(TSharedPtr<FString> NewSelection, ESelectInfo::Type SelectInfo);
+	void OnChangeCompileTextureCompressionType(TSharedPtr<FString> NewSelection, ESelectInfo::Type);
 
 	/** Callbacks from the tree widget. */
  	TSharedRef<ITableRow> GenerateRowForNodeTree(TSharedPtr<FMutableGraphTreeElement> InTreeNode, const TSharedRef<STableViewBase>& InOwnerTable);
diff --git a/Engine/Plugins/Experimental/Mutable/Source/CustomizableObjectEditor/Private/MuCOE/SMutableImageViewer.cpp b/Engine/Plugins/Experimental/Mutable/Source/CustomizableObjectEditor/Private/MuCOE/SMutableImageViewer.cpp
--- a/Engine/Plugins/Experimental/Mutable/Source/CustomizableObjectEditor/Private/MuCOE/SMutableImageViewer.cpp
+++ b/Engine/Plugins/Experimental/Mutable/Source/CustomizableObjectEditor/Private/MuCOE/SMutableImageViewer.cpp
@@ -266,7 +266,7 @@ void SMutableImageViewer::Tick(const FGeometry& AllottedGeometry, const double I
 				&&
 				MutableImage->GetFormat() <= mu::EImageFormat::IF_ASTC_10x10_RG_LDR)
 			{
-				mu::FImageOperator ImOp = mu::FImageOperator::GetDefault();
+				mu::FImageOperator ImOp = mu::FImageOperator::GetDefault(mu::FImageOperator::FImagePixelFormatFunc());
 				int32 Quality = 4;
 				ImageToConvert = ImOp.ImagePixelFormat(Quality, MutableImage.get(), mu::GetUncompressedFormat(MutableImage->GetFormat()) );
 			}
diff --git a/Engine/Plugins/Experimental/Mutable/Source/CustomizableObjectEditor/Private/MuCOE/SMutableObjectViewer.cpp b/Engine/Plugins/Experimental/Mutable/Source/CustomizableObjectEditor/Private/MuCOE/SMutableObjectViewer.cpp
--- a/Engine/Plugins/Experimental/Mutable/Source/CustomizableObjectEditor/Private/MuCOE/SMutableObjectViewer.cpp
+++ b/Engine/Plugins/Experimental/Mutable/Source/CustomizableObjectEditor/Private/MuCOE/SMutableObjectViewer.cpp
@@ -139,7 +139,7 @@ void SMutableObjectViewer::Construct(const FArguments& InArgs, UCustomizableObje
 	ParentNewTabId = InParentNewTabId;
 
 	// Initialize the debugger compile options
-	CompileOptions.bTextureCompression = true;
+	CompileOptions.TextureCompression = ECustomizableObjectTextureCompression::Fast;
 	CompileOptions.OptimizationLevel = 3;
 	{
 		ITargetPlatformManagerModule* TPM = GetTargetPlatformManager();
@@ -340,19 +340,6 @@ TSharedRef<SWidget> SMutableObjectViewer::GenerateCompileOptionsMenuContent()
 
 		MenuBuilder.AddWidget(DebugPlatformCombo.ToSharedRef(), LOCTEXT("MutableDebugPlatform", "Target Platform"));
 
-		// Texture compresssion
-		MenuBuilder.AddMenuEntry(
-			LOCTEXT("Generate_MutableEnableTextureCompression", "Enable texture compression."),
-			LOCTEXT("Generate_MutableEnableTextureCompressionTooltip", "Enable or disable the compression of textures in the mutable code compilation."),
-			FSlateIcon(),
-			FUIAction(
-				FExecuteAction::CreateLambda([this]() { CompileOptions.bTextureCompression = !CompileOptions.bTextureCompression; }),
-				FCanExecuteAction(),
-				FIsActionChecked::CreateLambda([this]() { return CompileOptions.bTextureCompression; })),
-			NAME_None,
-			EUserInterfaceActionType::ToggleButton
-		);
-
 		// Compilation options
 		//-----------------------------------
 
@@ -373,6 +360,23 @@ TSharedRef<SWidget> SMutableObjectViewer::GenerateCompileOptionsMenuContent()
 			;
 		MenuBuilder.AddWidget(CompileOptimizationCombo.ToSharedRef(), LOCTEXT("MutableCompileOptimizationLevel", "Optimization Level"));
 
+		{
+			CompileTextureCompressionStrings.Empty();
+			CompileTextureCompressionStrings.Add(MakeShareable(new FString(LOCTEXT("MutableTextureCompressionNone", "None").ToString())));
+			CompileTextureCompressionStrings.Add(MakeShareable(new FString(LOCTEXT("MutableTextureCompressionFast", "Fast").ToString())));
+			CompileTextureCompressionStrings.Add(MakeShareable(new FString(LOCTEXT("MutableTextureCompressionHighQuality", "High Quality").ToString())));
+
+			int32 SelectedCompression = FMath::Clamp(int32(CompileOptions.TextureCompression), 0, CompileTextureCompressionStrings.Num() - 1);
+			CompileTextureCompressionCombo =
+				SNew(STextComboBox)
+				.OptionsSource(&CompileTextureCompressionStrings)
+				.InitiallySelectedItem(CompileTextureCompressionStrings[SelectedCompression])
+				.OnSelectionChanged(this, &SMutableObjectViewer::OnChangeCompileTextureCompressionType)
+				;
+
+			MenuBuilder.AddWidget(CompileTextureCompressionCombo.ToSharedRef(), LOCTEXT("MutableCompileTextureCompressionType", "Texture Compression"));
+		}
+
 		// Image tiling
 		// Unfortunately SNumericDropDown doesn't work with integers at the time of writing.
 		TArray<SNumericDropDown<float>::FNamedValue> TilingOptions;
@@ -427,6 +431,12 @@ void SMutableObjectViewer::OnChangeCompileOptimizationLevel(TSharedPtr<FString>
 }
 
 
+void SMutableObjectViewer::OnChangeCompileTextureCompressionType(TSharedPtr<FString> NewSelection, ESelectInfo::Type)
+{
+	CompileOptions.TextureCompression = ECustomizableObjectTextureCompression(CompileTextureCompressionStrings.Find(NewSelection));
+}
+
+
 void SMutableObjectViewer::OnChangeDebugPlatform(TSharedPtr<FString> NewSelection, ESelectInfo::Type SelectInfo)
 {
 	if (!CustomizableObject) return;
diff --git a/Engine/Plugins/Experimental/Mutable/Source/CustomizableObjectEditor/Private/MuCOE/SMutableObjectViewer.h b/Engine/Plugins/Experimental/Mutable/Source/CustomizableObjectEditor/Private/MuCOE/SMutableObjectViewer.h
--- a/Engine/Plugins/Experimental/Mutable/Source/CustomizableObjectEditor/Private/MuCOE/SMutableObjectViewer.h
+++ b/Engine/Plugins/Experimental/Mutable/Source/CustomizableObjectEditor/Private/MuCOE/SMutableObjectViewer.h
@@ -68,8 +68,11 @@ private:
 	TSharedRef<SWidget> GenerateCompileOptionsMenuContent();
 	TSharedPtr<STextComboBox> CompileOptimizationCombo;
 	TArray< TSharedPtr<FString> > CompileOptimizationStrings;
+	TSharedPtr<STextComboBox> CompileTextureCompressionCombo;
+	TArray< TSharedPtr<FString> > CompileTextureCompressionStrings;
 	TSharedPtr<SNumericDropDown<float>> CompileTilingCombo;
 	void OnChangeCompileOptimizationLevel(TSharedPtr<FString> NewSelection, ESelectInfo::Type SelectInfo);
+	void OnChangeCompileTextureCompressionType(TSharedPtr<FString> NewSelection, ESelectInfo::Type);
 	void OnChangeDebugPlatform(TSharedPtr<FString> NewSelection, ESelectInfo::Type SelectInfo);
 
 	/** Tree callbacks. */
diff --git a/Engine/Plugins/Experimental/Mutable/Source/CustomizableObjectEditor/Private/MuCOE/UnrealBakeHelpers.cpp b/Engine/Plugins/Experimental/Mutable/Source/CustomizableObjectEditor/Private/MuCOE/UnrealBakeHelpers.cpp
new file mode 100644
--- /dev/null
+++ b/Engine/Plugins/Experimental/Mutable/Source/CustomizableObjectEditor/Private/MuCOE/UnrealBakeHelpers.cpp
@@ -0,0 +1,220 @@
+// Copyright Epic Games, Inc. All Rights Reserved.
+
+#include "MuCOE/UnrealBakeHelpers.h"
+
+#include "MuCO/CustomizableObject.h"
+#include "MuCO/CustomizableObjectMipDataProvider.h"
+#include "MuT/UnrealPixelFormatOverride.h"
+
+#include "Engine/SkeletalMesh.h"
+#include "Engine/Texture2D.h"
+
+#include "AssetRegistry/AssetRegistryModule.h"
+#include "AssetToolsModule.h"
+#include "IAssetTools.h"
+#include "Serialization/ArchiveReplaceObjectRef.h"
+
+#include "TextureResource.h"
+
+
+UObject* FUnrealBakeHelpers::BakeHelper_DuplicateAsset(UObject* Object, const FString& ObjName, const FString& PkgName, bool ResetDuplicatedFlags, TMap<UObject*, UObject*>& ReplacementMap, bool OverwritePackage)
+{
+	FString FinalObjName = ObjName;
+	FString FinalPkgName = PkgName;
+
+	if (!OverwritePackage)
+	{
+		FAssetToolsModule& AssetToolsModule = FModuleManager::GetModuleChecked<FAssetToolsModule>("AssetTools");
+		AssetToolsModule.Get().CreateUniqueAssetName(PkgName, "", FinalPkgName, FinalObjName);
+	}
+
+	UPackage* Package = CreatePackage(*FinalPkgName);
+	Package->FullyLoad();
+
+	FObjectDuplicationParameters Params = InitStaticDuplicateObjectParams(Object, Package, *FinalObjName, RF_AllFlags, nullptr, EDuplicateMode::Normal);
+
+	UObject* DupObject = StaticDuplicateObjectEx(Params);
+	if (DupObject)
+	{
+		DupObject->SetFlags(RF_Public | RF_Standalone);
+		DupObject->ClearFlags(RF_Transient);
+
+		// The garbage collector is called in the middle of the bake process, and this can destroy this temporary objects. 
+		// We add them to the garbage root to prevent this. This will avoid them being unloaded while the editor is running, but this
+		// action is not used often.
+		DupObject->AddToRoot();
+		DupObject->MarkPackageDirty();
+
+		// Notify the asset registry
+		FAssetRegistryModule::AssetCreated(DupObject);
+
+		// Replace all references
+		ReplacementMap.Add(Object, DupObject);
+
+		constexpr EArchiveReplaceObjectFlags ReplaceFlags = (EArchiveReplaceObjectFlags::IgnoreOuterRef | EArchiveReplaceObjectFlags::IgnoreArchetypeRef);
+		FArchiveReplaceObjectRef<UObject> ReplaceAr(DupObject, ReplacementMap, ReplaceFlags);
+	}
+
+	return DupObject;
+}
+
+namespace
+{
+
+	void CopyTextureProperties(UTexture2D* Texture, const UTexture2D* SourceTexture)
+	{
+		MUTABLE_CPUPROFILER_SCOPE(CopyTextureProperties)
+
+			Texture->NeverStream = SourceTexture->NeverStream;
+
+		Texture->SRGB = SourceTexture->SRGB;
+		Texture->Filter = SourceTexture->Filter;
+		Texture->LODBias = SourceTexture->LODBias;
+
+		Texture->MipGenSettings = SourceTexture->MipGenSettings;
+		Texture->CompressionNone = SourceTexture->CompressionNone;
+
+		Texture->LODGroup = SourceTexture->LODGroup;
+		Texture->AddressX = SourceTexture->AddressX;
+		Texture->AddressY = SourceTexture->AddressY;
+	}
+}
+
+
+UTexture2D* FUnrealBakeHelpers::BakeHelper_CreateAssetTexture(UTexture2D* SrcTex, const FString& TexObjName, const FString& TexPkgName, const UTexture* OrgTex, bool ResetDuplicatedFlags, TMap<UObject*, UObject*>& ReplacementMap, bool OverwritePackage)
+{
+	const bool bIsMutableTexture = !SrcTex->Source.IsValid();
+	if (!bIsMutableTexture)
+	{
+		return Cast<UTexture2D>(BakeHelper_DuplicateAsset(SrcTex, TexObjName, TexPkgName, ResetDuplicatedFlags, ReplacementMap, OverwritePackage));
+	}
+
+	int32 sx = SrcTex->GetPlatformData()->SizeX;
+	int32 sy = SrcTex->GetPlatformData()->SizeY;
+	EPixelFormat SrcPixelFormat = SrcTex->GetPlatformData()->PixelFormat;
+	ETextureSourceFormat PixelFormat = (SrcPixelFormat == PF_BC4 || SrcPixelFormat == PF_G8) ? TSF_G8 : TSF_BGRA8;
+
+	// Begin duplicate texture
+	FString FinalObjName = TexObjName;
+	FString FinalPkgName = TexPkgName;
+
+	if (!OverwritePackage)
+	{
+		FAssetToolsModule& AssetToolsModule = FModuleManager::GetModuleChecked<FAssetToolsModule>("AssetTools");
+		AssetToolsModule.Get().CreateUniqueAssetName(TexPkgName, "", FinalPkgName, FinalObjName);
+	}
+
+	UPackage* Package = CreatePackage(*FinalPkgName);
+	Package->FullyLoad();
+
+	UTexture2D* DupTex = NewObject<UTexture2D>(Package, *FinalObjName, RF_Public | RF_Standalone);
+
+	ReplacementMap.Add(SrcTex, DupTex);
+
+	// The garbage collector is called in the middle of the bake process, and this can destroy this temporary object. 
+	// We add them to the garbage root to prevent this. This will avoid them being unloaded while the editor is running, but this
+	// action is not used often.
+	DupTex->AddToRoot();
+	DupTex->MarkPackageDirty();
+
+	// Notify the asset registry
+	FAssetRegistryModule::AssetCreated(DupTex);
+
+	// Replace all references
+	constexpr EArchiveReplaceObjectFlags ReplaceFlags = (EArchiveReplaceObjectFlags::IgnoreOuterRef | EArchiveReplaceObjectFlags::IgnoreArchetypeRef);
+	FArchiveReplaceObjectRef<UObject> ReplaceAr(DupTex, ReplacementMap, ReplaceFlags);
+
+	// End duplicate texture
+
+	CopyTextureProperties(DupTex, SrcTex);
+
+	DupTex->RemoveUserDataOfClass(UMutableTextureMipDataProviderFactory::StaticClass());
+	if (OrgTex)
+	{
+		DupTex->CompressionSettings = OrgTex->CompressionSettings;
+	}
+
+	// Mutable textures only have platform data. We need to build the source data for them to be assets.
+	DupTex->Source.Init(sx, sy, 1, 1, PixelFormat);
+
+	int32 MipCount = SrcTex->GetPlatformData()->Mips.Num();
+	if (!MipCount)
+	{
+		UE_LOG(LogMutable, Warning, TEXT("Bake Instances: Empty texture found [%s]."), *SrcTex->GetName());
+		return DupTex;
+	}
+
+	// Create a mutable image from the platform data.
+	mu::EImageFormat MutableFormat = UnrealToMutablePixelFormat(SrcTex->GetPlatformData()->PixelFormat, SrcTex->HasAlphaChannel());
+	mu::Ptr<mu::Image> PlatformImage = new mu::Image(sx,sy,1,MutableFormat, mu::EInitializationType::NotInitialized);
+
+	constexpr int32 MipIndex = 0;
+	const uint8_t* SourceData = reinterpret_cast<const uint8_t*>(SrcTex->GetPlatformData()->Mips[MipIndex].BulkData.LockReadOnly());
+	check(SourceData); // A mutable-generated texture should always contain platform data
+	int32 PlatformDataSize = SrcTex->GetPlatformData()->Mips[MipIndex].BulkData.GetBulkDataSize();
+	check(PlatformImage->GetDataSize()== PlatformDataSize);
+	FMemory::Memcpy( PlatformImage->GetData(), SourceData, PlatformDataSize);
+	SourceData = nullptr;
+	SrcTex->GetPlatformData()->Mips[MipIndex].BulkData.Unlock();
+
+	// Reformat the mutable image
+	mu::EImageFormat UncompressedMutableFormat = mu::EImageFormat::IF_RGBA_UBYTE;
+	switch (SrcTex->GetPlatformData()->PixelFormat)
+	{
+	case PF_G8:
+	case PF_L8:
+	case PF_A8:
+	case PF_BC4:
+		UncompressedMutableFormat = mu::EImageFormat::IF_L_UBYTE;
+		break;
+
+	default:
+		break;
+	}
+
+	mu::FImageOperator ImOp = mu::FImageOperator::GetDefault(mu::FImageOperator::FImagePixelFormatFunc());
+	int32 Quality = 4; // Doesn't matter for decompression.
+	mu::Ptr<mu::Image> UncompressedImage = ImOp.ImagePixelFormat(Quality, PlatformImage.get(), UncompressedMutableFormat);
+
+	// Copy the decompressed data to the texture source data
+	int32 SourceDataSize = DupTex->Source.CalcMipSize(MipIndex);
+
+	// If this doesn't match, more cases have to be added to the switch above.
+	check(UncompressedImage->GetDataSize() == SourceDataSize);
+
+	uint8_t* Dest = DupTex->Source.LockMip(MipIndex);
+	check(Dest);
+	FMemory::Memcpy(Dest, UncompressedImage->GetData(), SourceDataSize);
+
+	// Probably can be integrated in the pixel format
+	const bool bNeedsRBSwizzle = PixelFormat == TSF_BGRA8;
+	if (bNeedsRBSwizzle)
+	{
+		for (int x = 0; x < sx * sy; ++x)
+		{
+			uint8_t temp = Dest[0];
+			Dest[0] = Dest[2];
+			Dest[2] = temp;
+			Dest += 4;
+		}
+	}
+
+	Dest = nullptr;
+	DupTex->Source.UnlockMip(MipIndex);
+
+	if (PixelFormat == TSF_G8 || PixelFormat == TSF_G16)
+	{
+		// If compression settings are not set to TC_Grayscale the texture will get a DXT format
+		// instead of G8 or G16.
+		FTextureFormatSettings Settings;
+		Settings.CompressionSettings = TC_Grayscale;
+		DupTex->SetLayerFormatSettings(0, Settings);
+
+		DupTex->CompressionSettings = TC_Grayscale;
+	}
+
+	DupTex->UpdateResource();
+
+	return DupTex;
+}
+
diff --git a/Engine/Plugins/Experimental/Mutable/Source/CustomizableObjectEditor/Public/MuCOE/UnrealBakeHelpers.h b/Engine/Plugins/Experimental/Mutable/Source/CustomizableObjectEditor/Public/MuCOE/UnrealBakeHelpers.h
new file mode 100644
--- /dev/null
+++ b/Engine/Plugins/Experimental/Mutable/Source/CustomizableObjectEditor/Public/MuCOE/UnrealBakeHelpers.h
@@ -0,0 +1,28 @@
+// Copyright Epic Games, Inc. All Rights Reserved.
+
+#pragma once
+
+#include "HAL/Platform.h"
+#include "Containers/Map.h"
+
+class UObject;
+class USkeletalMesh;
+class UTexture2D;
+class UTexture;
+
+
+class CUSTOMIZABLEOBJECTEDITOR_API FUnrealBakeHelpers
+{
+public:
+
+	static UObject* BakeHelper_DuplicateAsset(UObject* Object, const FString& ObjName, const FString& PkgName, bool ResetDuplicatedFlags, TMap<UObject*, UObject*>& ReplacementMap, bool OverwritePackage);
+
+	/**
+	 * Duplicates a texture asset. Duplicates Mutable and non Mutable textures.
+	 *
+	 * @param OrgTex Original source texture from which a Mutable texture has been generated. Only required when SrcTex is a Mutable texture.
+	 */
+	static UTexture2D* BakeHelper_CreateAssetTexture(UTexture2D* SrcTex, const FString& TexObjName, const FString& TexPkgName, const UTexture* OrgTex, bool ResetDuplicatedFlags, TMap<UObject*, UObject*>& ReplacementMap, bool OverwritePackage);
+
+};
+
diff --git a/Engine/Plugins/Experimental/Mutable/Source/MutableRuntime/Private/MuR/CodeRunner.h b/Engine/Plugins/Experimental/Mutable/Source/MutableRuntime/Private/MuR/CodeRunner.h
--- a/Engine/Plugins/Experimental/Mutable/Source/MutableRuntime/Private/MuR/CodeRunner.h
+++ b/Engine/Plugins/Experimental/Mutable/Source/MutableRuntime/Private/MuR/CodeRunner.h
@@ -759,7 +759,9 @@ namespace  mu
 				Ptr<Image> New = Runner->CreateImage(i->GetSizeX(), i->GetSizeY(), i->GetLODCount(), i->GetFormat(), EInitializationType::NotInitialized);
 				New->Copy(i);
 				return New;
-			}
+			},
+
+			Runner->m_pSystem->ImagePixelFormatOverride
 		);
 	}
 
diff --git a/Engine/Plugins/Experimental/Mutable/Source/MutableRuntime/Private/MuR/CodeRunnerTasked.cpp b/Engine/Plugins/Experimental/Mutable/Source/MutableRuntime/Private/MuR/CodeRunnerTasked.cpp
--- a/Engine/Plugins/Experimental/Mutable/Source/MutableRuntime/Private/MuR/CodeRunnerTasked.cpp
+++ b/Engine/Plugins/Experimental/Mutable/Source/MutableRuntime/Private/MuR/CodeRunnerTasked.cpp
@@ -778,7 +778,7 @@ namespace mu
 		if (!Mask 
 			&& 
 			// Flags have to match exactly for this optimize case. Other flags are not supported.
-			Args.flags== OP::ImageLayerArgs::F_USE_MASK_FROM_BLENDED
+			Args.flags == OP::ImageLayerArgs::F_USE_MASK_FROM_BLENDED
 			&&
 			Args.blendType == uint8(EBlendType::BT_BLEND)
 			&&
@@ -1094,7 +1094,7 @@ namespace mu
 					case EBlendType::BT_BLEND: 
 					{
 						// In this case we know it is already an uncompressed image, and we won't need additional allocations;
-						FImageOperator ImOp = FImageOperator::GetDefault();
+						FImageOperator ImOp = FImageOperator::GetDefault(FImageOperator::FImagePixelFormatFunc());
 						ImOp.FillColor( Result.get(), Color); 
 						break;
 					}
@@ -1184,6 +1184,7 @@ namespace mu
 		EImageFormat TargetFormat = EImageFormat::IF_NONE;
 		Ptr<const Image> Base;
 		Ptr<Image> Result;
+		FImageOperator::FImagePixelFormatFunc ImagePixelFormatFunc;
 	};
 
 
@@ -1226,6 +1227,8 @@ namespace mu
 			return false;
 		}
 
+		ImagePixelFormatFunc = Runner->m_pSystem->ImagePixelFormatOverride;
+
 		// Create destination data
 		Result = Runner->CreateImage(Base->GetSizeX(), Base->GetSizeY(), Base->GetLODCount(), TargetFormat, EInitializationType::NotInitialized);
 		return true;
@@ -1238,7 +1241,7 @@ namespace mu
 		MUTABLE_CPUPROFILER_SCOPE(FImagePixelFormatTask);
 		
 		bool bSuccess = false;
-		FImageOperator ImOp = FImageOperator::GetDefault();
+		FImageOperator ImOp = FImageOperator::GetDefault(ImagePixelFormatFunc);
 		ImOp.ImagePixelFormat(bSuccess, ImageCompressionQuality, Result.get(), Base.get(), -1);
 
 		int32 OriginalDataSize = FMath::Max(Result->m_data.Num(), Base->m_data.Num());
@@ -1292,6 +1295,7 @@ namespace mu
 		Ptr<const Image> Base;
 		Ptr<Image> Result;
 		FImageOperator::FScratchImageMipmap Scratch;
+		FImageOperator::FImagePixelFormatFunc ImagePixelFormatFunc;
 	};
 
 
@@ -1350,6 +1354,8 @@ namespace mu
 		FImageOperator ImOp = MakeImageOperator(Runner);
 		ImOp.ImageMipmap_PrepareScratch( Result.get(), Base.get(), LevelCount, Scratch);
 
+		ImagePixelFormatFunc = Runner->m_pSystem->ImagePixelFormatOverride;
+
 		return true;
 	}
 
@@ -1360,8 +1366,8 @@ namespace mu
 		// This runs in a worker thread
 		MUTABLE_CPUPROFILER_SCOPE(FImageMipmapTask);
 
-		FImageOperator ImOp = FImageOperator::GetDefault();
 		FMipmapGenerationSettings settings{};
+		FImageOperator ImOp = FImageOperator::GetDefault(ImagePixelFormatFunc);
 		ImOp.ImageMipmap(Scratch, ImageCompressionQuality, Result.get(), Base.get(), Result->GetLODCount(), settings);
 	}
 
@@ -1580,6 +1586,7 @@ namespace mu
 		OP::ImageResizeArgs Args;
 		Ptr<const Image> Base;
 		Ptr<Image> Result;
+		FImageOperator::FImagePixelFormatFunc ImagePixelFormatFunc;
 	};
 
 
@@ -1622,6 +1629,8 @@ namespace mu
 
 		Result = Runner->CreateImage( destSize[0], destSize[1], Lods, Base->GetFormat(), EInitializationType::NotInitialized );
 
+		ImagePixelFormatFunc = Runner->m_pSystem->ImagePixelFormatOverride;
+
 		return true;
 	}
 
@@ -1641,7 +1650,7 @@ namespace mu
 
 		// Warning: This will actually allocate temp memory that may exceed the budget.
 		// \TODO: Fix it.
-		FImageOperator ImOp = FImageOperator::GetDefault();
+		FImageOperator ImOp = FImageOperator::GetDefault(ImagePixelFormatFunc);
 		ImOp.ImageResizeLinear( Result.get(), ImageCompressionQuality, Base.get());
 
 		int32 LodCount = Result->GetLODCount();
@@ -1688,6 +1697,7 @@ namespace mu
 		Ptr<const Image> Base;
 		Ptr<Image> Result;
 		FImageSize DestSize;
+		FImageOperator::FImagePixelFormatFunc ImagePixelFormatFunc;
 	};
 
 
@@ -1730,6 +1740,8 @@ namespace mu
 
 		Result = Runner->CreateImage(DestSize[0], DestSize[1], Lods, Base->GetFormat(), EInitializationType::NotInitialized);
 
+		ImagePixelFormatFunc = Runner->m_pSystem->ImagePixelFormatOverride;
+
 		return true;
 	}
 
@@ -1740,7 +1752,7 @@ namespace mu
 		MUTABLE_CPUPROFILER_SCOPE(FImageResizeRelTask);
 
 		// \TODO: Track allocs
-		FImageOperator ImOp = FImageOperator::GetDefault();
+		FImageOperator ImOp = FImageOperator::GetDefault(ImagePixelFormatFunc);
 		ImOp.ImageResizeLinear(Result.get(), ImageCompressionQuality, Base.get());
 
 		int32 LodCount = Result->GetLODCount();
@@ -1842,6 +1854,7 @@ namespace mu
 		Ptr<const Image> Mask;
 		Ptr<Image> Result;
 		box< UE::Math::TIntVector2<uint16> > Rect;
+		FImageOperator::FImagePixelFormatFunc ImagePixelFormatFunc;
 	};
 
 
@@ -1950,6 +1963,8 @@ namespace mu
 			}
 		}
 
+		ImagePixelFormatFunc = Runner->m_pSystem->ImagePixelFormatOverride;
+
 		return true;
 	}
 
@@ -1966,7 +1981,7 @@ namespace mu
 
 			// Compose without a mask
 			// \TODO: track allocs
-			FImageOperator ImOp = FImageOperator::GetDefault();
+			FImageOperator ImOp = FImageOperator::GetDefault(ImagePixelFormatFunc);
 			ImOp.ImageCompose(Result.get(), Block.get(), Rect);
 		}
 		else
diff --git a/Engine/Plugins/Experimental/Mutable/Source/MutableRuntime/Private/MuR/ModelPrivate.h b/Engine/Plugins/Experimental/Mutable/Source/MutableRuntime/Private/MuR/ModelPrivate.h
--- a/Engine/Plugins/Experimental/Mutable/Source/MutableRuntime/Private/MuR/ModelPrivate.h
+++ b/Engine/Plugins/Experimental/Mutable/Source/MutableRuntime/Private/MuR/ModelPrivate.h
@@ -375,99 +375,6 @@ namespace mu
 			return RomSize;
 		}
 
-        //! Adds a constant image data and returns its constant index.
-		int32 AddConstant(Ptr<const Image> pImage, int32 MinTextureResidentMipCount)
-		{
-			check(pImage->GetSizeX()*pImage->GetSizeY()>0);
-
-			// Mips to store
-			int32 MipsToStore = 1;
-			if (MinTextureResidentMipCount > 0 && MinTextureResidentMipCount<255)
-			{
-				int32 MaxMipmaps = Image::GetMipmapCount( pImage->GetSizeX(), pImage->GetSizeY() );
-				
-				// This is not true. We may want the full mipmaps for fragments of images, regardless of the resident mip size.
-				//MipsToStore = FMath::Max(1, MaxMipmaps-MinTextureResidentMipCount);
-				// \TODO: Calculate the mip ranges that makes sense to store.
-				MipsToStore = MaxMipmaps;
-			}
-
-			int32 FirstLODIndexIndex = m_constantImageLODIndices.Num();
-
-			// Some images cannot be resized or mipmaped
-			bool bCannotBeScaled = pImage->m_flags & Image::IF_CANNOT_BE_SCALED;
-			if (bCannotBeScaled)
-			{
-				// Store only the mips that we have already calculated. We assume we have calculated them correctly.
-				//MipsToStore = 1;
-				MipsToStore = pImage->GetLODCount();
-			}
-
-			// TODO:
-			int32 CompressionQuality = 4;
-
-			FImageOperator ImOp = FImageOperator::GetDefault();
-
-			// TODO: Not efficient if we don't make mips (no need to clone base)
-			// TODO: If the image already has mips, we will be duplicating them...
-			Ptr<Image> pMip = ImOp.ExtractMip(pImage.get(),0);
-			for ( int Mip=0; Mip<MipsToStore; ++Mip )
-			{
-				check(pMip->GetFormat() == pImage->GetFormat());
-
-				// Shrink the buffer to the minimum necessary size.
-				uint32 CalculatedSize = pMip->CalculateDataSize();
-				if (CalculatedSize)
-				{
-					pMip->m_data.SetNum(CalculatedSize);
-				}
-
-				// Ensure unique at mip level
-				int32 MipIndex = -1;
-				for (int32 CandidateMipIndex = 0; CandidateMipIndex < m_constantImageLODs.Num(); ++CandidateMipIndex)
-				{
-					const Image* pCandidate = m_constantImageLODs[CandidateMipIndex].Value.get();
-					if (*pCandidate == *pMip)
-					{
-						// Reuse mip
-						MipIndex = CandidateMipIndex;
-					}
-				}
-
-				if (MipIndex < 0)
-				{
-					check(pMip->GetLODCount()==1);
-					MipIndex = m_constantImageLODs.Add(TPair<int32, Ptr<const Image>>(-1, pMip ));
-				}
-
-				m_constantImageLODIndices.Add(uint32(MipIndex));
-
-				// Generate next mip if necessary
-				if (Mip + 1 < MipsToStore)
-				{
-					if (Mip > pImage->GetLODCount())
-					{
-						// Generate from the last mip.
-						pMip = ImOp.ExtractMip(pMip.get(), 1);
-					}
-					else
-					{
-						pMip = ImOp.ExtractMip(pImage.get(),Mip+1);
-					}
-					check(pMip);
-				}
-			}
-
-			FImageLODRange LODRange;
-			LODRange.FirstIndex = FirstLODIndexIndex;
-			LODRange.LODCount = MipsToStore;
-			LODRange.ImageFormat = pImage->GetFormat();
-			LODRange.ImageSizeX = pImage->GetSizeX();
-			LODRange.ImageSizeY = pImage->GetSizeY();
-			int32 ImageIndex = m_constantImages.Add(LODRange);
-			return ImageIndex;
-		}
-
 
 		int32 AddConstant(Ptr<const Mesh> pMesh)
 		{
diff --git a/Engine/Plugins/Experimental/Mutable/Source/MutableRuntime/Private/MuR/OpImagePixelFormat.cpp b/Engine/Plugins/Experimental/Mutable/Source/MutableRuntime/Private/MuR/OpImagePixelFormat.cpp
--- a/Engine/Plugins/Experimental/Mutable/Source/MutableRuntime/Private/MuR/OpImagePixelFormat.cpp
+++ b/Engine/Plugins/Experimental/Mutable/Source/MutableRuntime/Private/MuR/OpImagePixelFormat.cpp
@@ -68,6 +68,16 @@ void FImageOperator::ImagePixelFormat( bool& bOutSuccess, int32 CompressionQuali
 {
 	MUTABLE_CPUPROFILER_SCOPE(ImagePixelFormatInPlace);
 
+	if (FormatImageOverride)
+	{
+		MUTABLE_CPUPROFILER_SCOPE(ImagePixelFormatOverride);
+		FormatImageOverride(bOutSuccess, CompressionQuality, pResult, Base, OnlyLOD);
+		if (bOutSuccess)
+		{
+			return;
+		}
+	}
+
 	bOutSuccess = true;
 
 	FIntVector2 resultSize;
diff --git a/Engine/Plugins/Experimental/Mutable/Source/MutableRuntime/Private/MuR/System.cpp b/Engine/Plugins/Experimental/Mutable/Source/MutableRuntime/Private/MuR/System.cpp
--- a/Engine/Plugins/Experimental/Mutable/Source/MutableRuntime/Private/MuR/System.cpp
+++ b/Engine/Plugins/Experimental/Mutable/Source/MutableRuntime/Private/MuR/System.cpp
@@ -199,6 +199,15 @@ namespace mu
     }
 
 
+	//---------------------------------------------------------------------------------------------
+	void System::SetImagePixelConversionOverride(const FImageOperator::FImagePixelFormatFunc& FormatFunc)
+	{
+		LLM_SCOPE_BYNAME(TEXT("MutableRuntime"));
+
+		m_pD->ImagePixelFormatOverride = FormatFunc;
+	}
+
+
     //---------------------------------------------------------------------------------------------
     Instance::ID System::NewInstance( const TSharedPtr<const Model>& InModel )
     {
diff --git a/Engine/Plugins/Experimental/Mutable/Source/MutableRuntime/Private/MuR/SystemPrivate.h b/Engine/Plugins/Experimental/Mutable/Source/MutableRuntime/Private/MuR/SystemPrivate.h
--- a/Engine/Plugins/Experimental/Mutable/Source/MutableRuntime/Private/MuR/SystemPrivate.h
+++ b/Engine/Plugins/Experimental/Mutable/Source/MutableRuntime/Private/MuR/SystemPrivate.h
@@ -1862,6 +1862,9 @@ namespace mu
 
 		TSharedPtr<ImageParameterGenerator> ImageParameterGenerator;
 
+		/** If this is set, it will be tried first instead of the internal formatting function. */
+		FImageOperator::FImagePixelFormatFunc ImagePixelFormatOverride;
+
 		/** */
 		FWorkingMemoryManager WorkingMemoryManager;
 
diff --git a/Engine/Plugins/Experimental/Mutable/Source/MutableRuntime/Public/MuR/Image.h b/Engine/Plugins/Experimental/Mutable/Source/MutableRuntime/Public/MuR/Image.h
--- a/Engine/Plugins/Experimental/Mutable/Source/MutableRuntime/Public/MuR/Image.h
+++ b/Engine/Plugins/Experimental/Mutable/Source/MutableRuntime/Public/MuR/Image.h
@@ -537,15 +537,22 @@ namespace mu
 		FImageReleaseFunc ReleaseImage;
 		FImageCloneFunc CloneImage;
 
-		FImageOperator(FImageCreateFunc Create, FImageReleaseFunc Release, FImageCloneFunc Clone) : CreateImage(Create), ReleaseImage(Release), CloneImage(Clone) {};
+		/** Interface to override the internal mutable image pixel format conversion functions.
+		* Arguments match the FImageOperator::ImagePixelFormat function.
+		*/
+		typedef TFunction<void(bool&, int32, Image*, const Image*, int32)> FImagePixelFormatFunc;
+		FImagePixelFormatFunc FormatImageOverride;
+
+		FImageOperator(FImageCreateFunc Create, FImageReleaseFunc Release, FImageCloneFunc Clone, FImagePixelFormatFunc FormatOverride) : CreateImage(Create), ReleaseImage(Release), CloneImage(Clone), FormatImageOverride(FormatOverride) {};
 
 		/** Create an default version for untracked resources. */
-		static inline FImageOperator GetDefault()
+		static FImageOperator GetDefault( const FImagePixelFormatFunc& InFormatOverride )
 		{
 			return FImageOperator(
-				[](int32 x, int32 y, int32 m, EImageFormat f, EInitializationType i) { return new Image(x, y, m, f, EInitializationType::NotInitialized); },
+				[](int32 x, int32 y, int32 m, EImageFormat f, EInitializationType i) { return new Image(x, y, m, f, i); },
 				[](Ptr<Image>& i) {i = nullptr; },
-				[](const Image* i) { return i->Clone(); }
+				[](const Image* i) { return i->Clone(); },
+				InFormatOverride
 			);
 		}
 
diff --git a/Engine/Plugins/Experimental/Mutable/Source/MutableRuntime/Public/MuR/System.h b/Engine/Plugins/Experimental/Mutable/Source/MutableRuntime/Public/MuR/System.h
--- a/Engine/Plugins/Experimental/Mutable/Source/MutableRuntime/Public/MuR/System.h
+++ b/Engine/Plugins/Experimental/Mutable/Source/MutableRuntime/Public/MuR/System.h
@@ -80,8 +80,7 @@ namespace mu
 	};
 
 
-    //! \brief Interface to request external images used as parameters.
-    //! \ingroup runtime
+    /** Interface to request external images used as parameters. */
     class MUTABLERUNTIME_API ImageParameterGenerator : public Base
     {
     public:
@@ -133,9 +132,15 @@ namespace mu
 		/** Set the amount of generated resources keys that will be stored for resource reusal. */
 		void SetGeneratedCacheSize(uint32 InCount);
 
-        //! Set a new provider for external image data. This is only necessary if image parameters
-        //! are used in the models.
-        void SetImageParameterGenerator(const TSharedPtr<ImageParameterGenerator>& );
+        /** Set a new provider for external image data. This is only necessary if image parameters are used in the models. */
+        void SetImageParameterGenerator(const TSharedPtr<ImageParameterGenerator>&);
+
+		/** Set a function that will be used to convert image pixel formats instead of the internal conversion. 
+		* \warning The provided function can be called from any thread, and also concurrently.
+		* If this function fails (returns false in the first parameter) the internal function is attempted next.
+		* This is useful to provide higher-quality external compressors in editor or when cooking.
+		*/
+		void SetImagePixelConversionOverride(const FImageOperator::FImagePixelFormatFunc&);
 
         //! Create a new instance from the given model. The instance can then be configured through
         //! calls to BeginUpdate/EndUpdate.
diff --git a/Engine/Plugins/Experimental/Mutable/Source/MutableTools/MutableTools.Build.cs b/Engine/Plugins/Experimental/Mutable/Source/MutableTools/MutableTools.Build.cs
--- a/Engine/Plugins/Experimental/Mutable/Source/MutableTools/MutableTools.Build.cs
+++ b/Engine/Plugins/Experimental/Mutable/Source/MutableTools/MutableTools.Build.cs
@@ -27,8 +27,11 @@ namespace UnrealBuildTool.Rules
 					"MutableRuntime", 
 					"Core",
 					"GeometryCore",
+					"ImageCore",
+					"TextureCompressor",
+					"TextureBuildUtilities",
 				}
-            );
+			);
 		}
 	}
 }
diff --git a/Engine/Plugins/Experimental/Mutable/Source/MutableTools/Private/MuT/AST.cpp b/Engine/Plugins/Experimental/Mutable/Source/MutableTools/Private/MuT/AST.cpp
--- a/Engine/Plugins/Experimental/Mutable/Source/MutableTools/Private/MuT/AST.cpp
+++ b/Engine/Plugins/Experimental/Mutable/Source/MutableTools/Private/MuT/AST.cpp
@@ -904,7 +904,7 @@ FImageDesc ASTOp::GetImageDesc( bool, FGetImageDescContext* ) const
 
 bool ASTOp::IsImagePlainConstant(FVector4f&) const
 {
-	// Some image operations don't have this implmented and hit here.
+	// Some image operations don't have this implemented and hit here.
 	// \TODO: Optimize for those cases.
     //check(false);
     return false;
diff --git a/Engine/Plugins/Experimental/Mutable/Source/MutableTools/Private/MuT/AST.h b/Engine/Plugins/Experimental/Mutable/Source/MutableTools/Private/MuT/AST.h
--- a/Engine/Plugins/Experimental/Mutable/Source/MutableTools/Private/MuT/AST.h
+++ b/Engine/Plugins/Experimental/Mutable/Source/MutableTools/Private/MuT/AST.h
@@ -211,6 +211,11 @@ namespace mu
 
 	struct FLinkerOptions
 	{
+		FLinkerOptions(FImageOperator& InImOp)
+			: ImageOperator(InImOp)
+		{
+		}
+
 		int32 MinTextureResidentMipCount = 0;
 
 		/** */
@@ -232,6 +237,9 @@ namespace mu
 		};
 
 		TMap<mu::Ptr<const mu::Mesh>, int32, FDefaultSetAllocator, FDeduplicationMeshFuncs> MeshConstantMap;
+
+		/** Image operation functions, so that they can be overriden. */
+		FImageOperator& ImageOperator;
 	};
 
 
diff --git a/Engine/Plugins/Experimental/Mutable/Source/MutableTools/Private/MuT/ASTOpConstantResource.cpp b/Engine/Plugins/Experimental/Mutable/Source/MutableTools/Private/MuT/ASTOpConstantResource.cpp
--- a/Engine/Plugins/Experimental/Mutable/Source/MutableTools/Private/MuT/ASTOpConstantResource.cpp
+++ b/Engine/Plugins/Experimental/Mutable/Source/MutableTools/Private/MuT/ASTOpConstantResource.cpp
@@ -64,6 +64,103 @@ namespace mu
 	}
 
 
+	namespace
+	{
+		/** Adds a constant image data to a program and returns its constant index. */
+		int32 AddConstantImage(FProgram& Program, Ptr<const Image> pImage, int32 MinTextureResidentMipCount, FLinkerOptions* Options)
+		{
+			check(pImage->GetSizeX() * pImage->GetSizeY() > 0);
+
+			// Mips to store
+			int32 MipsToStore = 1;
+			if (MinTextureResidentMipCount > 0 && MinTextureResidentMipCount < 255)
+			{
+				int32 MaxMipmaps = Image::GetMipmapCount(pImage->GetSizeX(), pImage->GetSizeY());
+
+				// This is not true. We may want the full mipmaps for fragments of images, regardless of the resident mip size.
+				//MipsToStore = FMath::Max(1, MaxMipmaps-MinTextureResidentMipCount);
+				// \TODO: Calculate the mip ranges that makes sense to store.
+				MipsToStore = MaxMipmaps;
+			}
+
+			int32 FirstLODIndexIndex = Program.m_constantImageLODIndices.Num();
+
+			// Some images cannot be resized or mipmaped
+			bool bCannotBeScaled = pImage->m_flags & Image::IF_CANNOT_BE_SCALED;
+			if (bCannotBeScaled)
+			{
+				// Store only the mips that we have already calculated. We assume we have calculated them correctly.
+				//MipsToStore = 1;
+				MipsToStore = pImage->GetLODCount();
+			}
+
+			// TODO:
+			int32 CompressionQuality = 4;
+
+			FImageOperator& ImOp = Options->ImageOperator;
+
+			// TODO: Not efficient if we don't make mips (no need to clone base)
+			// TODO: If the image already has mips, we will be duplicating them...
+			Ptr<Image> pMip = ImOp.ExtractMip(pImage.get(), 0);
+			for (int Mip = 0; Mip < MipsToStore; ++Mip)
+			{
+				check(pMip->GetFormat() == pImage->GetFormat());
+
+				// Shrink the buffer to the minimum necessary size.
+				uint32 CalculatedSize = pMip->CalculateDataSize();
+				if (CalculatedSize)
+				{
+					pMip->m_data.SetNum(CalculatedSize);
+				}
+
+				// Ensure unique at mip level
+				int32 MipIndex = -1;
+				for (int32 CandidateMipIndex = 0; CandidateMipIndex < Program.m_constantImageLODs.Num(); ++CandidateMipIndex)
+				{
+					const Image* pCandidate = Program.m_constantImageLODs[CandidateMipIndex].Value.get();
+					if (*pCandidate == *pMip)
+					{
+						// Reuse mip
+						MipIndex = CandidateMipIndex;
+					}
+				}
+
+				if (MipIndex < 0)
+				{
+					check(pMip->GetLODCount() == 1);
+					MipIndex = Program.m_constantImageLODs.Add(TPair<int32, Ptr<const Image>>(-1, pMip));
+				}
+
+				Program.m_constantImageLODIndices.Add(uint32(MipIndex));
+
+				// Generate next mip if necessary
+				if (Mip + 1 < MipsToStore)
+				{
+					if (Mip > pImage->GetLODCount())
+					{
+						// Generate from the last mip.
+						pMip = ImOp.ExtractMip(pMip.get(), 1);
+					}
+					else
+					{
+						pMip = ImOp.ExtractMip(pImage.get(), Mip + 1);
+					}
+					check(pMip);
+				}
+			}
+
+			FImageLODRange LODRange;
+			LODRange.FirstIndex = FirstLODIndexIndex;
+			LODRange.LODCount = MipsToStore;
+			LODRange.ImageFormat = pImage->GetFormat();
+			LODRange.ImageSizeX = pImage->GetSizeX();
+			LODRange.ImageSizeY = pImage->GetSizeY();
+			int32 ImageIndex = Program.m_constantImages.Add(LODRange);
+			return ImageIndex;
+		}
+	}
+
+
 	//-------------------------------------------------------------------------------------------------
 	void ASTOpConstantResource::Link(FProgram& program, FLinkerOptions* Options)
 	{
@@ -133,7 +230,7 @@ namespace mu
 					else
 					{
 						int32 MinTextureResidentMipCount = Options->MinTextureResidentMipCount;
-						args.value = program.AddConstant(pTyped, MinTextureResidentMipCount);
+						args.value = AddConstantImage( program, pTyped, MinTextureResidentMipCount, Options);
 					}
 					break;
 				}
@@ -189,8 +286,7 @@ namespace mu
 
 
 	//-------------------------------------------------------------------------------------------------
-	void ASTOpConstantResource::GetBlockLayoutSize(int blockIndex, int* pBlockX, int* pBlockY,
-		FBlockLayoutSizeCache*)
+	void ASTOpConstantResource::GetBlockLayoutSize(int blockIndex, int* pBlockX, int* pBlockY, FBlockLayoutSizeCache*)
 	{
 		switch (type)
 		{
diff --git a/Engine/Plugins/Experimental/Mutable/Source/MutableTools/Private/MuT/CodeGenerator_Image.cpp b/Engine/Plugins/Experimental/Mutable/Source/MutableTools/Private/MuT/CodeGenerator_Image.cpp
--- a/Engine/Plugins/Experimental/Mutable/Source/MutableTools/Private/MuT/CodeGenerator_Image.cpp
+++ b/Engine/Plugins/Experimental/Mutable/Source/MutableTools/Private/MuT/CodeGenerator_Image.cpp
@@ -227,13 +227,7 @@ namespace mu
 			// Do we need to crop?
 			if (rect.min[0]!=0 || rect.min[1]!=0 || pImage->GetSizeX() != rect.size[0] || pImage->GetSizeY() != rect.size[1])
 			{
-				FImageOperator ImOp
-				(
-					[](int32 x, int32 y, int32 m, EImageFormat f, EInitializationType i) { return new Image(x, y, m, f, i); },
-					[](Ptr<Image>& i) {i = nullptr; },
-					[](const Image* i) { return i->Clone(); }
-				);
-
+				FImageOperator ImOp = FImageOperator::GetDefault(m_compilerOptions->ImageFormatFunc);
 
 				Ptr<Image> pCropped = new Image(rect.size[0], rect.size[1], 1, pImage->GetFormat(), EInitializationType::NotInitialized);
 				ImOp.ImageCrop(pCropped.get(), m_compilerOptions->ImageCompressionQuality, pImage.get(), rect);
diff --git a/Engine/Plugins/Experimental/Mutable/Source/MutableTools/Private/MuT/CodeOptimiser.cpp b/Engine/Plugins/Experimental/Mutable/Source/MutableTools/Private/MuT/CodeOptimiser.cpp
--- a/Engine/Plugins/Experimental/Mutable/Source/MutableTools/Private/MuT/CodeOptimiser.cpp
+++ b/Engine/Plugins/Experimental/Mutable/Source/MutableTools/Private/MuT/CodeOptimiser.cpp
@@ -429,21 +429,19 @@ namespace mu
 	//---------------------------------------------------------------------------------------------
 	class ConstantTask 
 	{
-	public:
-		// input
-		mu::Ptr<ASTOp> m_source;
-		bool bUseDiskCache = false;
-		int ImageCompressionQuality = 0;
-
 	private:
 
 		//
 		mu::Ptr<ASTOp> m_result;
 
 	public:
+		// input
+		mu::Ptr<ASTOp> m_source;
+		bool bUseDiskCache = false;
+		int ImageCompressionQuality = 0;
+
 
-		// mu::Task interface
-		void Run()
+		void Run(FImageOperator ImOp)
 		{
 			MUTABLE_CPUPROFILER_SCOPE(ConstantTask_Run);
 
@@ -458,8 +456,10 @@ namespace mu
 			pSettings->SetImageCompressionQuality( ImageCompressionQuality );
 			SystemPtr pSystem = new System( pSettings );
 
+			pSystem->GetPrivate()->ImagePixelFormatOverride = ImOp.FormatImageOverride;
+
 			// Don't generate mips suring linking here.
-			FLinkerOptions LinkerOptions;
+			FLinkerOptions LinkerOptions(ImOp);
 			LinkerOptions.MinTextureResidentMipCount = 255;
 
 			TSharedPtr<const Model> model = MakeShared<Model>();
@@ -751,7 +751,8 @@ namespace mu
 					constantTask->m_source = n;
 					constantTask->ImageCompressionQuality = options->ImageCompressionQuality;
 
-					constantTask->Run();
+					FImageOperator ImOp = FImageOperator::GetDefault(options->ImageFormatFunc);
+					constantTask->Run(ImOp);
 					constantTask->Complete();
 					delete constantTask;
 
@@ -785,7 +786,7 @@ namespace mu
 	//-------------------------------------------------------------------------------------------------
 	//-------------------------------------------------------------------------------------------------
 	//-------------------------------------------------------------------------------------------------
-	CodeOptimiser::CodeOptimiser( CompilerOptionsPtr options, vector<STATE_COMPILATION_DATA>& states )
+	CodeOptimiser::CodeOptimiser(Ptr<CompilerOptions> options, vector<STATE_COMPILATION_DATA>& states )
 		: m_states( states )
 	{
 		m_options = options;
diff --git a/Engine/Plugins/Experimental/Mutable/Source/MutableTools/Private/MuT/CodeOptimiser.h b/Engine/Plugins/Experimental/Mutable/Source/MutableTools/Private/MuT/CodeOptimiser.h
--- a/Engine/Plugins/Experimental/Mutable/Source/MutableTools/Private/MuT/CodeOptimiser.h
+++ b/Engine/Plugins/Experimental/Mutable/Source/MutableTools/Private/MuT/CodeOptimiser.h
@@ -20,14 +20,14 @@ namespace mu
     class CodeOptimiser : public Base
     {
     public:
-        CodeOptimiser( CompilerOptionsPtr options, vector<STATE_COMPILATION_DATA>& states );
+        CodeOptimiser( Ptr<CompilerOptions> options, vector<STATE_COMPILATION_DATA>& states );
 
         //! Optimise the virtual machine code, using several transforms.
         void OptimiseAST();
 
     private:
 
-        CompilerOptionsPtr m_options;
+		Ptr<CompilerOptions> m_options;
 
         //!
         vector<STATE_COMPILATION_DATA>& m_states;
diff --git a/Engine/Plugins/Experimental/Mutable/Source/MutableTools/Private/MuT/CodeOptimiser_Parameters.cpp b/Engine/Plugins/Experimental/Mutable/Source/MutableTools/Private/MuT/CodeOptimiser_Parameters.cpp
--- a/Engine/Plugins/Experimental/Mutable/Source/MutableTools/Private/MuT/CodeOptimiser_Parameters.cpp
+++ b/Engine/Plugins/Experimental/Mutable/Source/MutableTools/Private/MuT/CodeOptimiser_Parameters.cpp
@@ -2187,8 +2187,7 @@ namespace mu
         {
             MUTABLE_CPUPROFILER_SCOPE(DataFormats);
 
-            DataOptimiseAST( m_options->GetPrivate()->ImageCompressionQuality, roots,
-                             m_options->GetPrivate()->OptimisationOptions );
+            DataOptimise( m_options.get(), roots);
 
             // After optimising the data formats, we may remove more constants
             DuplicatedDataRemoverAST( roots );
diff --git a/Engine/Plugins/Experimental/Mutable/Source/MutableTools/Private/MuT/Compiler.cpp b/Engine/Plugins/Experimental/Mutable/Source/MutableTools/Private/MuT/Compiler.cpp
--- a/Engine/Plugins/Experimental/Mutable/Source/MutableTools/Private/MuT/Compiler.cpp
+++ b/Engine/Plugins/Experimental/Mutable/Source/MutableTools/Private/MuT/Compiler.cpp
@@ -111,7 +111,7 @@ namespace mu
 
 
     //---------------------------------------------------------------------------------------------
-    void CompilerOptions::SetOptimisationMaxIteration( int maxIterations )
+    void CompilerOptions::SetOptimisationMaxIteration( int32 maxIterations )
     {
         m_pD->OptimisationOptions.MaxOptimisationLoopCount = maxIterations;
     }
@@ -153,10 +153,17 @@ namespace mu
 	}
 
 
+	//---------------------------------------------------------------------------------------------
+	void CompilerOptions::SetImagePixelFormatOverride(const FImageOperator::FImagePixelFormatFunc& InFunc)
+	{
+		m_pD->ImageFormatFunc = InFunc;
+	}
+
+
     //---------------------------------------------------------------------------------------------
     //---------------------------------------------------------------------------------------------
     //---------------------------------------------------------------------------------------------
-    Compiler::Compiler( CompilerOptionsPtr options )
+    Compiler::Compiler( Ptr<CompilerOptions> options )
     {
         m_pD = new Private();
         m_pD->m_options = options;
@@ -228,7 +235,8 @@ namespace mu
 
 		// Keep the link options outside the scope because it is also used to cache constant data that has already been 
 		// added and could be reused across states.
-		FLinkerOptions LinkerOptions;
+		FImageOperator ImOp = FImageOperator::GetDefault(m_pD->m_options->GetPrivate()->ImageFormatFunc);
+		FLinkerOptions LinkerOptions(ImOp);
 
 		for( STATE_COMPILATION_DATA& s: states )
         {
diff --git a/Engine/Plugins/Experimental/Mutable/Source/MutableTools/Private/MuT/CompilerPrivate.h b/Engine/Plugins/Experimental/Mutable/Source/MutableTools/Private/MuT/CompilerPrivate.h
--- a/Engine/Plugins/Experimental/Mutable/Source/MutableTools/Private/MuT/CompilerPrivate.h
+++ b/Engine/Plugins/Experimental/Mutable/Source/MutableTools/Private/MuT/CompilerPrivate.h
@@ -97,6 +97,8 @@ namespace mu
 		int32 ImageTiling=0 ;
 
         bool bLog = false;
+
+		FImageOperator::FImagePixelFormatFunc ImageFormatFunc;
     };
 
 
@@ -163,10 +165,10 @@ namespace mu
             m_pErrorLog = new ErrorLog();
         }
 
-        ErrorLogPtr m_pErrorLog;
+        Ptr<ErrorLog> m_pErrorLog;
 
         //! Detailed options
-        CompilerOptionsPtr m_options;
+        Ptr<CompilerOptions> m_options;
 
 
 		//!
diff --git a/Engine/Plugins/Experimental/Mutable/Source/MutableTools/Private/MuT/DataPacker.cpp b/Engine/Plugins/Experimental/Mutable/Source/MutableTools/Private/MuT/DataPacker.cpp
--- a/Engine/Plugins/Experimental/Mutable/Source/MutableTools/Private/MuT/DataPacker.cpp
+++ b/Engine/Plugins/Experimental/Mutable/Source/MutableTools/Private/MuT/DataPacker.cpp
@@ -26,6 +26,7 @@
 #include "MuT/ASTOpMeshDifference.h"
 #include "MuT/ASTOpMeshMorph.h"
 #include "MuT/ASTOpLayoutFromMesh.h"
+#include "MuT/CompilerPrivate.h"
 
 
 namespace mu
@@ -710,9 +711,11 @@ namespace mu
 
 
     //---------------------------------------------------------------------------------------------
-    void DataOptimiseAST( int imageCompressionQuality, ASTOpList& roots,
-                          const FModelOptimizationOptions& options )
+    void DataOptimise( const CompilerOptions* Options, ASTOpList& roots )
     {
+		int32 ImageCompressionQuality = Options->GetPrivate()->ImageCompressionQuality;
+		const FModelOptimizationOptions& OptimizeOptions = Options->GetPrivate()->OptimisationOptions;
+
         // Images
         AccumulateImageFormatsAST accFormat;
         accFormat.Run( roots );
@@ -725,7 +728,7 @@ namespace mu
 				ASTOpConstantResource* typed = dynamic_cast<ASTOpConstantResource*>(n.get());
                 Ptr<const Image> pOld = static_cast<const Image*>(typed->GetValue().get());
 
-				FImageOperator ImOp = FImageOperator::GetDefault();
+				FImageOperator ImOp = FImageOperator::GetDefault( Options->GetPrivate()->ImageFormatFunc );
 
 				// See if there is a better format for this image
 				FVector4f PlainColor;
@@ -754,26 +757,26 @@ namespace mu
 				}
 				else if ( accFormat.m_supportedFormats[typed][(size_t)EImageFormat::IF_L_UBIT_RLE] )
                 {
-                    ImagePtr pNew = ImOp.ImagePixelFormat( imageCompressionQuality, pOld.get(), EImageFormat::IF_L_UBIT_RLE );
+                    ImagePtr pNew = ImOp.ImagePixelFormat( ImageCompressionQuality, pOld.get(), EImageFormat::IF_L_UBIT_RLE );
 
                     // Only replace if the compression was worth!
                     size_t oldSize = pOld->GetDataSize();
                     size_t newSize = pNew->GetDataSize();
-                    if (float(oldSize) > float(newSize) * options.MinRLECompressionGain)
+                    if (float(oldSize) > float(newSize) * OptimizeOptions.MinRLECompressionGain)
                     {
-                        typed->SetValue(pNew, options.bUseDiskCache);
+                        typed->SetValue(pNew, OptimizeOptions.bUseDiskCache);
                     }
                 }
                 else if ( accFormat.m_supportedFormats[typed][(size_t)EImageFormat::IF_L_UBYTE_RLE] )
                 {
-                    ImagePtr pNew = ImOp.ImagePixelFormat( imageCompressionQuality, pOld.get(), EImageFormat::IF_L_UBYTE_RLE );
+                    ImagePtr pNew = ImOp.ImagePixelFormat( ImageCompressionQuality, pOld.get(), EImageFormat::IF_L_UBYTE_RLE );
 
                     // Only replace if the compression was worth!
                     size_t oldSize = pOld->GetDataSize();
                     size_t newSize = pNew->GetDataSize();
-                    if (float(oldSize) > float(newSize) * options.MinRLECompressionGain)
+                    if (float(oldSize) > float(newSize) * OptimizeOptions.MinRLECompressionGain)
                     {
-                        typed->SetValue(pNew, options.bUseDiskCache);
+                        typed->SetValue(pNew, OptimizeOptions.bUseDiskCache);
                     }
                 }
             }
@@ -792,7 +795,7 @@ namespace mu
 				ASTOpConstantResource* typed = dynamic_cast<ASTOpConstantResource*>(n.get());
                 Ptr<Mesh> pMesh = static_cast<const Mesh*>(typed->GetValue().get())->Clone();
                 MeshRemoveUnusedBufferSemantics( pMesh.get(), meshSemanticsVisitor.m_requiredSemantics[typed]);
-                typed->SetValue(pMesh, options.bUseDiskCache);
+                typed->SetValue(pMesh, OptimizeOptions.bUseDiskCache);
             }
         });
 
diff --git a/Engine/Plugins/Experimental/Mutable/Source/MutableTools/Private/MuT/DataPacker.h b/Engine/Plugins/Experimental/Mutable/Source/MutableTools/Private/MuT/DataPacker.h
--- a/Engine/Plugins/Experimental/Mutable/Source/MutableTools/Private/MuT/DataPacker.h
+++ b/Engine/Plugins/Experimental/Mutable/Source/MutableTools/Private/MuT/DataPacker.h
@@ -15,12 +15,10 @@ namespace mu { struct FProgram; }
 
 namespace mu
 {
+	class CompilerOptions;
 
-    //---------------------------------------------------------------------------------------------
-    //! Convert constant data to different formats, based on their usage
-    //---------------------------------------------------------------------------------------------
-    extern void DataOptimiseAST( int imageCompressionQuality, ASTOpList& roots,
-                                 const FModelOptimizationOptions& );
+    /** Convert constant data to different formats, based on their usage. */
+    extern void DataOptimise( const CompilerOptions*, ASTOpList& roots );
 
     //---------------------------------------------------------------------------------------------
     //! Find the given constant in a subtree
diff --git a/Engine/Plugins/Experimental/Mutable/Source/MutableTools/Private/MuT/UnrealPixelFormatOverride.cpp b/Engine/Plugins/Experimental/Mutable/Source/MutableTools/Private/MuT/UnrealPixelFormatOverride.cpp
new file mode 100644
--- /dev/null
+++ b/Engine/Plugins/Experimental/Mutable/Source/MutableTools/Private/MuT/UnrealPixelFormatOverride.cpp
@@ -0,0 +1,276 @@
+// Copyright Epic Games, Inc. All Rights Reserved.
+
+#include "MuT/UnrealPixelFormatOverride.h"
+
+#include "Interfaces/ITextureFormatManagerModule.h"
+#include "Interfaces/ITextureFormatModule.h"
+#include "Modules/ModuleManager.h"
+#include "Engine/TextureDefines.h"
+#include "TextureCompressorModule.h"
+#include "TextureBuildUtilities.h"
+#include "ImageCore.h"
+
+
+static ITextureFormatManagerModule* STextureFormatManager = nullptr;
+
+
+void PrepareUnrealCompression()
+{
+	check(IsInGameThread());
+
+	if (!STextureFormatManager)
+	{
+		STextureFormatManager = &FModuleManager::LoadModuleChecked<ITextureFormatManagerModule>("TextureFormat");
+		check(STextureFormatManager);
+	}
+}
+
+
+void FillBuildSettingsFromMutableFormat(FTextureBuildSettings& Settings, bool& bOutHasAlpha, mu::EImageFormat Format)
+{
+	Settings.MipGenSettings = TMGS_NoMipmaps;
+
+	switch (Format)
+	{
+	case mu::EImageFormat::IF_ASTC_4x4_RGBA_LDR:
+		Settings.TextureFormatName = TEXT("ASTC_RGBA_HQ");
+		Settings.CompressionQuality = 4; // See GetQualityFormat in TextureFormatASTC.cpp
+		bOutHasAlpha = true;
+		break;
+	case mu::EImageFormat::IF_ASTC_6x6_RGBA_LDR:
+		Settings.TextureFormatName = TEXT("ASTC_RGBA");
+		Settings.CompressionQuality = 3; // See GetQualityFormat in TextureFormatASTC.cpp
+		bOutHasAlpha = true;
+		break;
+	case mu::EImageFormat::IF_ASTC_8x8_RGBA_LDR:
+		Settings.TextureFormatName = TEXT("ASTC_RGBA");
+		Settings.CompressionQuality = 2; // See GetQualityFormat in TextureFormatASTC.cpp
+		bOutHasAlpha = true;
+		break;
+	case mu::EImageFormat::IF_ASTC_10x10_RGBA_LDR:
+		Settings.TextureFormatName = TEXT("ASTC_RGBA");
+		Settings.CompressionQuality = 1; // See GetQualityFormat in TextureFormatASTC.cpp
+		bOutHasAlpha = true;
+		break;
+	case mu::EImageFormat::IF_ASTC_12x12_RGBA_LDR:
+		Settings.TextureFormatName = TEXT("ASTC_RGBA");
+		Settings.CompressionQuality = 0; // See GetQualityFormat in TextureFormatASTC.cpp
+		bOutHasAlpha = true;
+		break;
+
+	case mu::EImageFormat::IF_ASTC_4x4_RGB_LDR:
+		Settings.TextureFormatName = TEXT("ASTC_RGBA_HQ");
+		Settings.CompressionQuality = 4; // See GetQualityFormat in TextureFormatASTC.cpp
+		bOutHasAlpha = false;
+		break;
+	case mu::EImageFormat::IF_ASTC_6x6_RGB_LDR:
+		Settings.TextureFormatName = TEXT("ASTC_RGBA");
+		Settings.CompressionQuality = 3; // See GetQualityFormat in TextureFormatASTC.cpp
+		bOutHasAlpha = false;
+		break;
+	case mu::EImageFormat::IF_ASTC_8x8_RGB_LDR:
+		Settings.TextureFormatName = TEXT("ASTC_RGBA");
+		Settings.CompressionQuality = 2; // See GetQualityFormat in TextureFormatASTC.cpp
+		bOutHasAlpha = false;
+		break;
+	case mu::EImageFormat::IF_ASTC_10x10_RGB_LDR:
+		Settings.TextureFormatName = TEXT("ASTC_RGBA");
+		Settings.CompressionQuality = 1; // See GetQualityFormat in TextureFormatASTC.cpp
+		bOutHasAlpha = false;
+		break;
+	case mu::EImageFormat::IF_ASTC_12x12_RGB_LDR:
+		Settings.TextureFormatName = TEXT("ASTC_RGBA");
+		Settings.CompressionQuality = 0; // See GetQualityFormat in TextureFormatASTC.cpp
+		bOutHasAlpha = false;
+		break;
+
+	case mu::EImageFormat::IF_ASTC_4x4_RG_LDR:
+		Settings.TextureFormatName = TEXT("ASTC_NormalRG_Precise");
+		Settings.CompressionQuality = 4; // See GetQualityFormat in TextureFormatASTC.cpp
+		bOutHasAlpha = false;
+		break;
+	case mu::EImageFormat::IF_ASTC_6x6_RG_LDR:
+		Settings.TextureFormatName = TEXT("ASTC_NormalRG");
+		Settings.CompressionQuality = 3; // See GetQualityFormat in TextureFormatASTC.cpp
+		bOutHasAlpha = false;
+		break;
+	case mu::EImageFormat::IF_ASTC_8x8_RG_LDR:
+		Settings.TextureFormatName = TEXT("ASTC_NormalLA");
+		Settings.CompressionQuality = 2; // See GetQualityFormat in TextureFormatASTC.cpp
+		bOutHasAlpha = false;
+		break;
+	case mu::EImageFormat::IF_ASTC_10x10_RG_LDR:
+		Settings.TextureFormatName = TEXT("ASTC_NormalRG");
+		Settings.CompressionQuality = 1; // See GetQualityFormat in TextureFormatASTC.cpp
+		bOutHasAlpha = false;
+		break;
+	case mu::EImageFormat::IF_ASTC_12x12_RG_LDR:
+		Settings.TextureFormatName = TEXT("ASTC_NormalRG");
+		Settings.CompressionQuality = 0; // See GetQualityFormat in TextureFormatASTC.cpp
+		bOutHasAlpha = false;
+		break;
+
+	case mu::EImageFormat::IF_BC1:
+		Settings.TextureFormatName = TEXT("DXT1");
+		bOutHasAlpha = false;
+		break;
+	case mu::EImageFormat::IF_BC2:
+		Settings.TextureFormatName = TEXT("DXT3");
+		bOutHasAlpha = true;
+		break;
+	case mu::EImageFormat::IF_BC3:
+		Settings.TextureFormatName = TEXT("DXT5");
+		bOutHasAlpha = true;
+		break;
+	case mu::EImageFormat::IF_BC4:
+		Settings.TextureFormatName = TEXT("BC4");
+		bOutHasAlpha = false;
+		break;
+	case mu::EImageFormat::IF_BC5:
+		Settings.TextureFormatName = TEXT("BC5");
+		bOutHasAlpha = true;
+		break;
+
+	default:
+		Settings.TextureFormatName = NAME_None;
+		bOutHasAlpha = false;
+		break;
+	}
+}
+
+
+void MutableToImageCore(const mu::Image* InMutable, FImage& CoreImage, int32 LOD)
+{
+	MUTABLE_CPUPROFILER_SCOPE(MutableToImageCore);
+
+	mu::Ptr<const mu::Image> Mutable = InMutable;
+
+	ERawImageFormat::Type CoreImageFormat;
+	switch (Mutable->GetFormat())
+	{
+	case mu::EImageFormat::IF_BGRA_UBYTE: CoreImageFormat = ERawImageFormat::BGRA8; break;
+		//case mu::EImageFormat::IF_L_UBYTE: CoreImageFormat = ERawImageFormat::G8; break;
+
+	default:
+	{
+		// Unsupported format: force conversion
+		mu::FImageOperator ImOp = mu::FImageOperator::GetDefault(mu::FImageOperator::FImagePixelFormatFunc());
+		Mutable = ImOp.ImagePixelFormat(4, InMutable, mu::EImageFormat::IF_BGRA_UBYTE);
+		CoreImageFormat = ERawImageFormat::BGRA8;
+		break;
+	}
+
+	}
+
+	FIntVector2 MipSize = Mutable->CalculateMipSize(LOD);
+	CoreImage.Init(MipSize.X, MipSize.Y, CoreImageFormat, EGammaSpace::Linear);
+	FMemory::Memcpy(CoreImage.GetPixelPointer(0, 0), Mutable->GetMipData(LOD), CoreImage.GetImageSizeBytes());
+}
+
+
+void ImageCoreToMutable(const FCompressedImage2D& Compressed, mu::Image* Mutable, int32 LOD)
+{
+	check(Compressed.RawData.Num()==Mutable->CalculateDataSize(LOD));
+	FMemory::Memcpy(Mutable->GetMipData(LOD), Compressed.RawData.GetData(), Compressed.RawData.Num());
+}
+
+
+mu::EImageFormat UnrealToMutablePixelFormat(EPixelFormat PlatformFormat, bool bHasAlpha)
+{
+	switch (PlatformFormat)
+	{
+	case PF_ASTC_4x4: return bHasAlpha ? mu::EImageFormat::IF_ASTC_4x4_RGBA_LDR : mu::EImageFormat::IF_ASTC_4x4_RGB_LDR;
+	case PF_ASTC_6x6: return bHasAlpha ? mu::EImageFormat::IF_ASTC_6x6_RGBA_LDR : mu::EImageFormat::IF_ASTC_6x6_RGB_LDR;
+	case PF_ASTC_8x8: return bHasAlpha ? mu::EImageFormat::IF_ASTC_8x8_RGBA_LDR : mu::EImageFormat::IF_ASTC_8x8_RGB_LDR;
+	case PF_ASTC_10x10: return bHasAlpha ? mu::EImageFormat::IF_ASTC_10x10_RGBA_LDR : mu::EImageFormat::IF_ASTC_10x10_RGB_LDR;
+	case PF_ASTC_12x12: return bHasAlpha ? mu::EImageFormat::IF_ASTC_12x12_RGBA_LDR : mu::EImageFormat::IF_ASTC_12x12_RGB_LDR;
+	case PF_ASTC_4x4_NORM_RG: return mu::EImageFormat::IF_ASTC_4x4_RG_LDR;
+	case PF_ASTC_6x6_NORM_RG: return mu::EImageFormat::IF_ASTC_6x6_RG_LDR;
+	case PF_ASTC_8x8_NORM_RG: return mu::EImageFormat::IF_ASTC_8x8_RG_LDR;
+	case PF_ASTC_10x10_NORM_RG: return mu::EImageFormat::IF_ASTC_10x10_RG_LDR;
+	case PF_ASTC_12x12_NORM_RG: return mu::EImageFormat::IF_ASTC_12x12_RG_LDR;
+	case PF_DXT1: return mu::EImageFormat::IF_BC1;
+	case PF_DXT3: return mu::EImageFormat::IF_BC2;
+	case PF_DXT5: return mu::EImageFormat::IF_BC3;
+	case PF_BC4: return mu::EImageFormat::IF_BC4;
+	case PF_BC5: return mu::EImageFormat::IF_BC5;
+	case PF_G8: return mu::EImageFormat::IF_L_UBYTE;
+	case PF_L8: return mu::EImageFormat::IF_L_UBYTE;
+	case PF_A8: return mu::EImageFormat::IF_L_UBYTE;
+	case PF_R8G8B8A8: return mu::EImageFormat::IF_RGBA_UBYTE;
+	case PF_A8R8G8B8: return mu::EImageFormat::IF_RGBA_UBYTE;
+	case PF_B8G8R8A8: return mu::EImageFormat::IF_BGRA_UBYTE;
+	default:
+		return mu::EImageFormat::IF_NONE;
+	}
+}
+
+
+mu::EImageFormat QualityFix(mu::EImageFormat Format)
+{
+	switch (Format)
+	{
+	case mu::EImageFormat::IF_ASTC_8x8_RGB_LDR:		return mu::EImageFormat::IF_ASTC_4x4_RGB_LDR; break;
+	case mu::EImageFormat::IF_ASTC_8x8_RGBA_LDR:	return mu::EImageFormat::IF_ASTC_4x4_RGBA_LDR; break;
+	case mu::EImageFormat::IF_ASTC_8x8_RG_LDR:		return mu::EImageFormat::IF_ASTC_4x4_RG_LDR; break;
+	case mu::EImageFormat::IF_ASTC_12x12_RGB_LDR:	return mu::EImageFormat::IF_ASTC_4x4_RGB_LDR; break;
+	case mu::EImageFormat::IF_ASTC_12x12_RGBA_LDR:	return mu::EImageFormat::IF_ASTC_4x4_RGBA_LDR; break;
+	case mu::EImageFormat::IF_ASTC_12x12_RG_LDR:	return mu::EImageFormat::IF_ASTC_4x4_RG_LDR; break;
+	case mu::EImageFormat::IF_ASTC_6x6_RGB_LDR:		return mu::EImageFormat::IF_ASTC_4x4_RGB_LDR; break;
+	case mu::EImageFormat::IF_ASTC_6x6_RGBA_LDR:	return mu::EImageFormat::IF_ASTC_4x4_RGBA_LDR; break;
+	case mu::EImageFormat::IF_ASTC_6x6_RG_LDR:		return mu::EImageFormat::IF_ASTC_4x4_RG_LDR; break;
+	case mu::EImageFormat::IF_ASTC_10x10_RGB_LDR:	return mu::EImageFormat::IF_ASTC_4x4_RGB_LDR; break;
+	case mu::EImageFormat::IF_ASTC_10x10_RGBA_LDR:	return mu::EImageFormat::IF_ASTC_4x4_RGBA_LDR; break;
+	case mu::EImageFormat::IF_ASTC_10x10_RG_LDR:	return mu::EImageFormat::IF_ASTC_4x4_RG_LDR; break;
+	default:
+		break;
+	}
+
+	return Format;
+}
+
+
+void UnrealPixelFormatFunc(bool& bOutSuccess, int32 Quality, mu::Image* Target, const mu::Image* Source, int32 OnlyLOD)
+{
+	// If this fails, PrepareUnrealCompression wasn't called before.
+	check(STextureFormatManager);
+
+	bOutSuccess = true;
+
+	FTextureBuildSettings Settings;
+	bool bHasAlpha = false;
+	FillBuildSettingsFromMutableFormat(Settings, bHasAlpha, Target->GetFormat());
+
+	if (Settings.TextureFormatName == NAME_None)
+	{
+		// Unsupported format in the override: use standard mutable compression.
+		return;
+	}
+
+	const ITextureFormat* TextureFormat = STextureFormatManager->FindTextureFormat(Settings.TextureFormatName);
+	check(TextureFormat);
+
+	int32 FirstLOD = 0;
+	int32 LODCount = Source->GetLODCount();
+	if (OnlyLOD >= 0)
+	{
+		FirstLOD = OnlyLOD;
+		LODCount = 1;
+	}
+
+	for (int32 LOD = FirstLOD; LOD < LODCount; ++LOD)
+	{
+		FImage SourceUnreal;
+		MutableToImageCore(Source, SourceUnreal, LOD);
+
+		FCompressedImage2D CompressedUnreal;
+		TextureFormat->CompressImage(SourceUnreal, Settings,
+			FIntVector3(SourceUnreal.SizeX, SourceUnreal.SizeY, 1),
+			0, 0, 1,
+			FString(),
+			bHasAlpha, CompressedUnreal);
+
+		ImageCoreToMutable(CompressedUnreal, Target, LOD);
+	}
+}
+
diff --git a/Engine/Plugins/Experimental/Mutable/Source/MutableTools/Public/MuT/Compiler.h b/Engine/Plugins/Experimental/Mutable/Source/MutableTools/Public/MuT/Compiler.h
--- a/Engine/Plugins/Experimental/Mutable/Source/MutableTools/Public/MuT/Compiler.h
+++ b/Engine/Plugins/Experimental/Mutable/Source/MutableTools/Public/MuT/Compiler.h
@@ -5,6 +5,7 @@
 #include "HAL/PlatformMath.h"
 #include "MuR/Ptr.h"
 #include "MuR/RefCounted.h"
+#include "MuR/Image.h"
 #include "Templates/SharedPointer.h"
 
 namespace mu
@@ -14,10 +15,6 @@ namespace mu
     typedef Ptr<Compiler> CompilerPtr;
     typedef Ptr<const Compiler> CompilerPtrConst;
 
-    class CompilerOptions;
-    using CompilerOptionsPtr=Ptr<CompilerOptions>;
-    using CompilerOptionsPtrConst=Ptr<const CompilerOptions>;
-
     class Model;
     using ModelPtr=Ptr<Model>;
     using ModelPtrConst=Ptr<const Model>;
@@ -75,6 +72,9 @@ namespace mu
 		/** If enabled it will make sure that the object is compile to generate smaller mips of the images. */
 		void SetEnableProgressiveImages(bool bEnabled);
 
+		/** Set an optional pixel conversion function that will be called before any pixel format conversion. */
+		void SetImagePixelFormatOverride(const FImageOperator::FImagePixelFormatFunc&);
+
         //! Different data packing strategies
         enum class TextureLayoutStrategy : uint8
         {
@@ -88,7 +88,7 @@ namespace mu
             Count
         };
 
-        //! Return a readable string for a GPU.
+        //! 
         static const char* GetTextureLayoutStrategyName( TextureLayoutStrategy s );
 
         //-----------------------------------------------------------------------------------------
@@ -122,7 +122,7 @@ namespace mu
     {
     public:
 
-        Compiler( CompilerOptionsPtr options=nullptr );
+        Compiler( Ptr<CompilerOptions> Options=nullptr );
 
         //! Compile the expression into a run-time model.
 		TSharedPtr<Model> Compile( const Ptr<Node>& pNode );
diff --git a/Engine/Plugins/Experimental/Mutable/Source/MutableTools/Public/MuT/UnrealPixelFormatOverride.h b/Engine/Plugins/Experimental/Mutable/Source/MutableTools/Public/MuT/UnrealPixelFormatOverride.h
new file mode 100644
--- /dev/null
+++ b/Engine/Plugins/Experimental/Mutable/Source/MutableTools/Public/MuT/UnrealPixelFormatOverride.h
@@ -0,0 +1,35 @@
+// Copyright Epic Games, Inc. All Rights Reserved.
+
+#pragma once
+
+#include "MuR/Image.h"
+#include "Interfaces/ITextureFormatManagerModule.h"
+#include "Interfaces/ITextureFormatModule.h"
+#include "Engine/TextureDefines.h"
+#include "TextureCompressorModule.h"
+#include "TextureBuildUtilities.h"
+#include "ImageCore.h"
+
+/** This function needs to be called from the game thread before any call to UnrealPixelFormatFunc can be done. */
+MUTABLETOOLS_API extern void PrepareUnrealCompression();
+
+/** */
+void FillBuildSettingsFromMutableFormat(FTextureBuildSettings& Settings, bool& bOutHasAlpha, mu::EImageFormat Format);
+
+/** */
+void MutableToImageCore(const mu::Image* InMutable, FImage& CoreImage, int32 LOD);
+
+/** */
+void ImageCoreToMutable(const FCompressedImage2D& Compressed, mu::Image* Mutable, int32 LOD);
+
+/** Convert an Unreal platform pixel format to Mutable. */
+MUTABLETOOLS_API extern mu::EImageFormat UnrealToMutablePixelFormat(EPixelFormat PlatformFormat, bool bHasAlpha);
+
+/** Function that remaps some formats to the ones that provide more acceptable quality in Mutable. */
+MUTABLETOOLS_API extern mu::EImageFormat QualityFix(mu::EImageFormat Format);
+
+/** Try to convert an image using unreal's compression. If the format is not supported, it will return false in bOutSuccess.
+* Can be called from any thread.
+* PrepareUnrealCompression needs to be called from the game thread once prior to using this function. 
+*/
+MUTABLETOOLS_API extern void UnrealPixelFormatFunc(bool& bOutSuccess, int32 Quality, mu::Image* Target, const mu::Image* Source, int32 OnlyLOD);
