[Mutable] Minor CustomizableObjectSystem code improvements.

- Added consts.
- Changed function copies to references.
- Moved checks at the beggining of functions.
- Removed parameters passed outside the FMutableOperationData.
- Default initialization of members.
- Removed unused code.
- Moved implementations to source file.
- Fixed typos.

#rnx

diff --git a/Engine/Plugins/Experimental/Mutable/Source/CustomizableObject/Classes/MuCO/CustomizableObjectSystem.h b/Engine/Plugins/Experimental/Mutable/Source/CustomizableObject/Classes/MuCO/CustomizableObjectSystem.h
--- a/Engine/Plugins/Experimental/Mutable/Source/CustomizableObject/Classes/MuCO/CustomizableObjectSystem.h
+++ b/Engine/Plugins/Experimental/Mutable/Source/CustomizableObject/Classes/MuCO/CustomizableObjectSystem.h
@@ -27,6 +27,7 @@ class UDefaultImageProvider;
 class USkeletalMesh;
 class UMaterialInterface;
 class UTexture2D;
+class FCustomizableObjectSystemPrivate; // This is used to hide Mutable SDK members in the public headers.
 struct FFrame;
 struct FGuid;
 
@@ -37,9 +38,6 @@ struct FGuid;
 // In case of async file operations, what priority to use
 #define MUTABLE_SYSTEM_ASYNC_STREAMING_PRIORITY			EAsyncIOPriorityAndFlags::AIOP_Normal
 
-// This is used to hide Mutable SDK members in the public headers.
-class FCustomizableObjectSystemPrivate;
-
 extern TAutoConsoleVariable<bool> CVarClearWorkingMemoryOnUpdateEnd;
 
 extern TAutoConsoleVariable<bool> CVarReuseImagesBetweenInstances;
@@ -61,24 +59,21 @@ struct FEditorCompileSettings
 
 #endif
 
-//
 namespace EMutableProfileMetric
 {
 	typedef uint8 Type;
 
-	const Type BuiltInstances = 1;
-	const Type UpdateOperations = 2;
-	const Type Count = 4;
+	constexpr Type BuiltInstances = 1;
+	constexpr Type UpdateOperations = 2;
+	constexpr Type Count = 4;
 
 };
 
 USTRUCT()
 struct FPendingReleaseSkeletalMeshInfo
 {
-public:
 	GENERATED_USTRUCT_BODY()
 
-public:
 	UPROPERTY()
 	TObjectPtr<USkeletalMesh> SkeletalMesh = nullptr;
 
@@ -90,10 +85,8 @@ public:
 USTRUCT()
 struct FPendingReleaseMaterialsInfo
 {
-public:
 	GENERATED_USTRUCT_BODY()
 
-public:
 	UPROPERTY()
 	TArray<TObjectPtr<UMaterialInterface>> Materials;
 
@@ -108,19 +101,15 @@ struct FMutableImageReference
 {
 	/** Original image ID. Once generated it will be unique. However, future updates of the image may return a different ID for the
 	* same image, if many other resources have been built in the middle. For this reason the rest of the data in the struct is what
-	* must be used to request the additional mips.
-	*/
+	* must be used to request the additional mips.	*/
 	uint32 ImageID = 0;
 
-	/** */
 	uint32 SurfaceId = 0;
 
-	/** */
 	uint8 LOD = 0;
 	uint8 Component = 0;
 	uint8 Image = 0;
 
-	/** */
 	uint8 BaseMip = 0;
 };
 
@@ -177,13 +166,13 @@ public:
 
 	// Used in the editor to show the list of available options.
 	// Only necessary if the images are required in editor previews.
-	virtual void GetTextureParameterValues(TArray<FCustomizableObjectExternalTexture>& OutValues) {};
+	virtual void GetTextureParameterValues(TArray<FCustomizableObjectExternalTexture>& OutValues) {}
 };
 
 
 // Before the Mutable Queue rework this made sense, but this is no longer the case. Remove this when doing MTBL-1409.
 /** End a Customizable Object Instance Update. All code paths of an update have to end here. */
-void FinishUpdateGlobal(UCustomizableObjectInstance* Instance, EUpdateResult UpdateResult, FInstanceUpdateDelegate* UpdateCallback, const FDescriptorRuntimeHash InUpdatedHash = FDescriptorRuntimeHash());
+void FinishUpdateGlobal(UCustomizableObjectInstance* Instance, EUpdateResult UpdateResult, const FInstanceUpdateDelegate* UpdateCallback, const FDescriptorRuntimeHash InUpdatedHash = FDescriptorRuntimeHash());
 
 
 UCLASS(Blueprintable, BlueprintType)
@@ -195,11 +184,13 @@ public:
 	UCustomizableObjectSystem() = default;
 	void InitSystem();
 
-	/** Get the singleton object. It will be created if it doesn't exist yet.
-	 * @param bCreate Create a system if it does not has been created yet. */
+	/** Get the singleton object. It will be created if it doesn't exist yet. */
 	UFUNCTION(BlueprintCallable, BlueprintPure, Category = Status)
 	static UCustomizableObjectSystem* GetInstance();
 	
+	UFUNCTION(BlueprintCallable, BlueprintPure, Category = Status)
+	static UCustomizableObjectSystem* GetInstanceChecked();
+	
 	// Return true if the singleton has been created. It is different than GetInstance in that GetInstance will create it if it doesn't exist.
 	static bool IsCreated();
 
@@ -209,7 +200,7 @@ public:
 	// End UObject interface.
 
 	// Creates a new Customizable Object Compiler (Only does real work in editor builds). The caller is responsible for freeing the new compiler
-	class FCustomizableObjectCompilerBase* GetNewCompiler();
+	FCustomizableObjectCompilerBase* GetNewCompiler();
 	void SetNewCompilerFunc(FCustomizableObjectCompilerBase* (*NewCompilerFunc)());
 
 	bool IsReplaceDiscardedWithReferenceMeshEnabled() const;
@@ -223,10 +214,10 @@ public:
 	// Lock a CustomizableObjects, preventing the generation or update of any of its instances
 	// Will return true if successful, false if it fails to lock because an update is already underway
 	// This is usually only used in the editor
-	bool LockObject(const class UCustomizableObject*);
-	void UnlockObject(const class UCustomizableObject*);
+	bool LockObject(const UCustomizableObject*);
+	void UnlockObject(const UCustomizableObject*);
 
-	/** Checks if there are any outstading disk or mip update operations in flight for the parameter Customizable Object that may
+	/** Checks if there are any outstanding disk or mip update operations in flight for the parameter Customizable Object that may
 	* make it unsafe to compile at the moment.
 	* @return true if there are operations in flight and it's not safe to compile */
 	bool CheckIfDiskOrMipUpdateOperationsPending(const UCustomizableObject& Object) const;
@@ -245,7 +236,6 @@ public:
 
 	// Copy of the Mutable Editor Settings tied to CO compilation. They are updated whenever changed
 	FEditorCompileSettings EditorSettings;
-
 #endif
 	
 	// Return the current MinLodQualityLevel for skeletal meshes.
@@ -271,7 +261,6 @@ public:
 	void TickPendingReleaseMaterials();
 
 private:
-
 	UPROPERTY()
 	TArray<FPendingReleaseSkeletalMeshInfo> PendingReleaseSkeletalMesh;
 
@@ -279,7 +268,6 @@ private:
 	TArray<FPendingReleaseMaterialsInfo> PendingReleaseMaterials;
 
 public:
-
 	/** [Texture Parameters] Get a list of all the possible values for external texture parameters according to the various providers registered with RegisterImageProvider. */
 	TArray<FCustomizableObjectExternalTexture> GetTextureParameterValues();
 
@@ -311,7 +299,6 @@ private:
 	TObjectPtr<UDefaultImageProvider> DefaultImageProvider = nullptr;
 
 public:
-    
 	// Show a warning on-screen and via a notification (if in Editor) and log an error when a CustomizableObject is
 	// being used and it's not compiled.  Callers can add additional information to the error log.
 	void AddUncompiledCOWarning(const UCustomizableObject& InObject, FString const* OptionalLogInfo = nullptr);
@@ -324,7 +311,6 @@ public:
 	// Show data about all UCustomizableObjectInstance existing elements
 	void LogShowData(bool bFullInfo, bool ShowMaterialInfo) const;
 
-
 	// Give access to the internal object data.
 	FCustomizableObjectSystemPrivate* GetPrivate();
 	const FCustomizableObjectSystemPrivate* GetPrivate() const;
@@ -332,13 +318,12 @@ public:
 	FCustomizableObjectSystemPrivate* GetPrivateChecked();
 	const FCustomizableObjectSystemPrivate* GetPrivateChecked() const;
 	
-	FStreamableManager& GetStreamableManager() { return StreamableManager; }
+	FStreamableManager& GetStreamableManager();
 
-	UCustomizableInstanceLODManagementBase* GetInstanceLODManagement() { return CurrentInstanceLODManagement.Get(); }
+	UCustomizableInstanceLODManagementBase* GetInstanceLODManagement() const;
 
 	// Pass a null ptr to reset to the default InstanceLODManagement
-	void SetInstanceLODManagement(UCustomizableInstanceLODManagementBase* NewInstanceLODManagement) 
-		{ CurrentInstanceLODManagement = NewInstanceLODManagement ? NewInstanceLODManagement : ToRawPtr(DefaultInstanceLODManagement); }
+	void SetInstanceLODManagement(UCustomizableInstanceLODManagementBase* NewInstanceLODManagement);
 
 	// Find out the version of the plugin
 	UFUNCTION(BlueprintCallable, Category = Status)
@@ -352,7 +337,7 @@ public:
 	UFUNCTION(BlueprintCallable, Category = Status)
 	int32 GetNumPendingInstances() const;
 
-	// Get the total number of instances includingbuilt and not built.
+	// Get the total number of instances including built and not built.
 	UFUNCTION(BlueprintCallable, Category = Status)
 	int32 GetTotalInstances() const;
 
@@ -392,10 +377,9 @@ public:
 	// Array where textures are added temporarily while the mutable thread may want to
 	// reused them for some instance under construction.
 	UPROPERTY(Transient)
-	TArray< TObjectPtr<UTexture2D> > ProtectedCachedTextures;
+	TArray<TObjectPtr<UTexture2D>> ProtectedCachedTextures;
 
 private:
-
 	TSharedPtr<FCustomizableObjectSystemPrivate> Private = nullptr;
 
 	// For async material loading
@@ -436,7 +420,6 @@ private:
 
 	/** Map used to cache per platform MaxChunkSize. If MaxChunkSize > 0, streamed data will be split in multiple files */
 	TMap<FString, int64> PlatformMaxChunkSize;
-	
 #endif
 	
 	// Friends
diff --git a/Engine/Plugins/Experimental/Mutable/Source/CustomizableObject/Private/MuCO/CustomizableObjectSystem.cpp b/Engine/Plugins/Experimental/Mutable/Source/CustomizableObject/Private/MuCO/CustomizableObjectSystem.cpp
--- a/Engine/Plugins/Experimental/Mutable/Source/CustomizableObject/Private/MuCO/CustomizableObjectSystem.cpp
+++ b/Engine/Plugins/Experimental/Mutable/Source/CustomizableObject/Private/MuCO/CustomizableObjectSystem.cpp
@@ -44,7 +44,6 @@
 class AActor;
 class UAnimInstance;
 
-
 DECLARE_CYCLE_STAT(TEXT("MutablePendingRelease Time"), STAT_MutablePendingRelease, STATGROUP_Game);
 DECLARE_CYCLE_STAT(TEXT("MutableTask"), STAT_MutableTask, STATGROUP_Game);
 
@@ -91,6 +90,7 @@ static void CVarMutableSinkFunction()
 		PrivateSystem->SkeletalMeshMinLodQualityLevel = CVarSkeletalMeshMinLodQualityLevelCVarName ? CVarSkeletalMeshMinLodQualityLevelCVarName->GetInt() : INDEX_NONE;
 	}
 }
+
 static FAutoConsoleVariableSink CVarMutableSink(FConsoleCommandDelegate::CreateStatic(&CVarMutableSinkFunction));
 
 
@@ -114,7 +114,7 @@ void FMutablePendingInstanceWork::SetLODUpdatesLastTick(int32 NumLODUpdates)
 
 void FMutablePendingInstanceWork::AddUpdate(const FMutablePendingInstanceUpdate& UpdateToAdd)
 {
-	if (FMutablePendingInstanceUpdate* ExistingUpdate = PendingInstanceUpdates.Find(UpdateToAdd.CustomizableObjectInstance))
+	if (const FMutablePendingInstanceUpdate* ExistingUpdate = PendingInstanceUpdates.Find(UpdateToAdd.CustomizableObjectInstance))
 	{
 		FMutablePendingInstanceUpdate TaskToEnqueue = UpdateToAdd;
 
@@ -132,7 +132,7 @@ void FMutablePendingInstanceWork::AddUpdate(const FMutablePendingInstanceUpdate&
 		PendingInstanceUpdates.Add(UpdateToAdd);
 	}
 
-	if (FMutablePendingInstanceDiscard* ExistingDiscard = PendingInstanceDiscards.Find(UpdateToAdd.CustomizableObjectInstance))
+	if (const FMutablePendingInstanceDiscard* ExistingDiscard = PendingInstanceDiscards.Find(UpdateToAdd.CustomizableObjectInstance))
 	{
 		FinishUpdateGlobal(ExistingDiscard->CustomizableObjectInstance.Get(), EUpdateResult::ErrorReplaced, nullptr);
 
@@ -155,7 +155,7 @@ const FMutablePendingInstanceUpdate* FMutablePendingInstanceWork::GetUpdate(cons
 
 void FMutablePendingInstanceWork::AddDiscard(const FMutablePendingInstanceDiscard& TaskToEnqueue)
 {
-	if (FMutablePendingInstanceUpdate* ExistingUpdate = PendingInstanceUpdates.Find(TaskToEnqueue.CustomizableObjectInstance.Get()))
+	if (const FMutablePendingInstanceUpdate* ExistingUpdate = PendingInstanceUpdates.Find(TaskToEnqueue.CustomizableObjectInstance.Get()))
 	{
 		FinishUpdateGlobal(ExistingUpdate->CustomizableObjectInstance.Get(), EUpdateResult::ErrorReplaced, ExistingUpdate->Callback);
 		PendingInstanceUpdates.Remove(ExistingUpdate->CustomizableObjectInstance);
@@ -203,6 +203,27 @@ UCustomizableObjectSystem* UCustomizableObjectSystem::GetInstance()
 }
 
 
+UCustomizableObjectSystem* UCustomizableObjectSystem::GetInstanceChecked()
+{
+	UCustomizableObjectSystem* System = GetInstance();
+	check(System);
+	
+	return System;
+}
+
+
+UCustomizableInstanceLODManagementBase* UCustomizableObjectSystem::GetInstanceLODManagement() const
+{
+	return CurrentInstanceLODManagement.Get();
+}
+
+
+void UCustomizableObjectSystem::SetInstanceLODManagement(UCustomizableInstanceLODManagementBase* NewInstanceLODManagement)
+{
+	CurrentInstanceLODManagement = NewInstanceLODManagement ? NewInstanceLODManagement : ToRawPtr(DefaultInstanceLODManagement);
+}
+
+
 FString UCustomizableObjectSystem::GetPluginVersion() const
 {
 	// Bridge the call from the module. This implementation is available from blueprint.
@@ -214,16 +235,15 @@ void UCustomizableObjectSystem::LogShowData(bool bFullInfo, bool ShowMaterialInf
 {
 	LogInformationUtil::ResetCounters();
 
-	AActor* ParentActor;
 	TArray<UCustomizableObjectInstance*> ArrayData;
 
 	for (TObjectIterator<UCustomizableSkeletalComponent> It; It; ++It)
 	{
-		UCustomizableSkeletalComponent* CustomizableSkeletalComponent = *It;
+		const UCustomizableSkeletalComponent* CustomizableSkeletalComponent = *It;
 
-		if ((CustomizableSkeletalComponent != nullptr) && (CustomizableSkeletalComponent->CustomizableObjectInstance != nullptr))
+		if (CustomizableSkeletalComponent && CustomizableSkeletalComponent->CustomizableObjectInstance)
 		{
-			ParentActor = CustomizableSkeletalComponent->GetAttachmentRootActor();
+			const AActor* ParentActor = CustomizableSkeletalComponent->GetAttachmentRootActor();
 
 			if (ParentActor != nullptr)
 			{
@@ -239,8 +259,8 @@ void UCustomizableObjectSystem::LogShowData(bool bFullInfo, bool ShowMaterialInf
 		return (A.GetPrivate()->LastMinSquareDistFromComponentToPlayer < B.GetPrivate()->LastMinSquareDistFromComponentToPlayer);
 	});
 
-	int i;
-	const int Max = ArrayData.Num();
+	int32 i;
+	const int32 Max = ArrayData.Num();
 
 	if (bFullInfo)
 	{
@@ -298,6 +318,12 @@ const FCustomizableObjectSystemPrivate* UCustomizableObjectSystem::GetPrivateChe
 }
 
 
+FStreamableManager& UCustomizableObjectSystem::GetStreamableManager()
+{
+	return StreamableManager;
+}
+
+
 bool UCustomizableObjectSystem::IsCreated()
 {
 	return FCustomizableObjectSystemPrivate::SSystem != 0;
@@ -341,7 +367,7 @@ void UCustomizableObjectSystem::InitSystem()
 	Private->LastWorkingMemoryBytes = CVarWorkingMemory.GetValueOnGameThread() * 1024;
 	Private->LastGeneratedResourceCacheSize = CVarGeneratedResourcesCacheSize.GetValueOnGameThread();
 
-	mu::Ptr<mu::Settings> pSettings = new mu::Settings;
+	const mu::Ptr<mu::Settings> pSettings = new mu::Settings;
 	check(pSettings);
 	pSettings->SetProfile(false);
 	pSettings->SetWorkingMemoryBytes(Private->LastWorkingMemoryBytes);
@@ -447,17 +473,16 @@ FCustomizableObjectCompilerBase* UCustomizableObjectSystem::GetNewCompiler()
 
 void UCustomizableObjectSystem::SetNewCompilerFunc(FCustomizableObjectCompilerBase* (*InNewCompilerFunc)())
 {
-	check(Private != nullptr);
-	Private->NewCompilerFunc = InNewCompilerFunc;
+	GetPrivateChecked()->NewCompilerFunc = InNewCompilerFunc;
 }
 
 
-void FCustomizableObjectSystemPrivate::CreatedTexture(UTexture2D* Texture)
+void FCustomizableObjectSystemPrivate::CreatedTexture(UTexture2D* Texture) const
 {
 #if !(UE_BUILD_SHIPPING || UE_BUILD_TEST)
-	bool bLogEnabled = true;
+	const bool bLogEnabled = true;
 #else
-	bool bLogEnabled = LogBenchmarkUtil::IsLoggingActive();
+	const bool bLogEnabled = LogBenchmarkUtil::IsLoggingActive();
 #endif
 	if (bLogEnabled)
 	{
@@ -485,6 +510,13 @@ static FAutoConsoleVariableRef CVarEnableMutableAnimInfoDebugging(
 	ECVF_Default);
 
 
+void FCustomizableObjectSystemPrivate::AddGameThreadTask(const FMutableTask& Task)
+{
+	check(IsInGameThread())
+	PendingTasks.Enqueue(Task);
+}
+
+
 void FCustomizableObjectSystemPrivate::AddReferencedObjects(FReferenceCollector& Collector)
 {
 	if (CurrentInstanceBeingUpdated)
@@ -494,6 +526,12 @@ void FCustomizableObjectSystemPrivate::AddReferencedObjects(FReferenceCollector&
 }
 
 
+FString FCustomizableObjectSystemPrivate::GetReferencerName() const
+{
+	return TEXT("FCustomizableObjectSystemPrivate");
+}
+
+
 TAutoConsoleVariable<bool> CVarCleanupTextureCache(
 	TEXT("mutable.EnableCleanupCache"),
 	true,
@@ -505,7 +543,7 @@ void FCustomizableObjectSystemPrivate::CleanupCache()
 {
 	check(IsInGameThread());
 
-	bool bCleanupEnabled = CVarCleanupTextureCache.GetValueOnGameThread();
+	const bool bCleanupEnabled = CVarCleanupTextureCache.GetValueOnGameThread();
 
 	for (int32 ModelIndex = 0; ModelIndex < ModelResourcesCache.Num();)
 	{
@@ -543,6 +581,28 @@ void FCustomizableObjectSystemPrivate::CleanupCache()
 }
 
 
+FMutableResourceCache& FCustomizableObjectSystemPrivate::GetObjectCache(const UCustomizableObject* Object)
+{
+	check(IsInGameThread());
+
+	// Not mandatory, but a good place for a cleanup
+	CleanupCache();
+
+	for (int ModelIndex = 0; ModelIndex < ModelResourcesCache.Num(); ++ModelIndex)
+	{
+		if (ModelResourcesCache[ModelIndex].Object==Object)
+		{
+			return ModelResourcesCache[ModelIndex];
+		}
+	}
+		
+	// Not found, create and add it.
+	ModelResourcesCache.Push(FMutableResourceCache());
+	ModelResourcesCache.Last().Object = Object;
+	return ModelResourcesCache.Last();
+}
+
+
 int32 FCustomizableObjectSystemPrivate::EnableMutableProgressiveMipStreaming = 1;
 
 // Warning! If this is enabled, do not get references to the textures generated by Mutable! They are owned by Mutable and could become invalid at any moment
@@ -604,7 +664,7 @@ static FAutoConsoleVariableRef CVarApplyFixPrepareSkeletons(
 	TEXT("If true, Fix missing SkeletonsData when FirstLODToGenerate is greater than 0. If false, There may be a crash when generating meshes on platform that skip LODs."),
 	ECVF_Default);
 
-void FinishUpdateGlobal(UCustomizableObjectInstance* Instance, EUpdateResult UpdateResult, FInstanceUpdateDelegate* UpdateCallback, const FDescriptorRuntimeHash InUpdatedHash)
+void FinishUpdateGlobal(UCustomizableObjectInstance* Instance, EUpdateResult UpdateResult, const FInstanceUpdateDelegate* UpdateCallback, const FDescriptorRuntimeHash InUpdatedHash)
 {
 	check(IsInGameThread())
 	
@@ -722,6 +782,70 @@ void FCustomizableObjectSystemPrivate::GetMipStreamingConfig(const UCustomizable
 }
 
 
+bool FCustomizableObjectSystemPrivate::IsReplaceDiscardedWithReferenceMeshEnabled() const
+{
+	return bReplaceDiscardedWithReferenceMesh;
+}
+
+
+void FCustomizableObjectSystemPrivate::SetReplaceDiscardedWithReferenceMeshEnabled(bool bIsEnabled)
+{
+	bReplaceDiscardedWithReferenceMesh = bIsEnabled;
+}
+
+
+int32 FCustomizableObjectSystemPrivate::GetCountAllocatedSkeletalMesh() const
+{
+	return MutableStats.CountAllocatedSkeletalMesh;
+}
+
+
+void FCustomizableObjectSystemPrivate::AddTextureReference(const FMutableImageCacheKey& TextureId)
+{
+	uint32& CountRef = TextureReferenceCount.FindOrAdd(TextureId);
+
+	CountRef++;
+}
+
+
+bool FCustomizableObjectSystemPrivate::RemoveTextureReference(const FMutableImageCacheKey& TextureId)
+{
+	uint32* CountPtr = TextureReferenceCount.Find(TextureId);
+
+	if (CountPtr && *CountPtr > 0)
+	{
+		(*CountPtr)--;
+
+		if (*CountPtr == 0)
+		{
+			TextureReferenceCount.Remove(TextureId);
+
+			return true;
+		}
+	}
+	else
+	{
+		ensure(false); // Mutable texture reference count is incorrect
+		TextureReferenceCount.Remove(TextureId);
+	}
+
+	return false;
+}
+
+
+bool FCustomizableObjectSystemPrivate::TextureHasReferences(const FMutableImageCacheKey& TextureId) const
+{
+	const uint32* CountPtr = TextureReferenceCount.Find(TextureId);
+
+	if (CountPtr && *CountPtr > 0)
+	{
+		return true;
+	}
+
+	return false;
+}
+
+
 void FCustomizableObjectSystemPrivate::InitUpdateSkeletalMesh(UCustomizableObjectInstance& Instance, EQueuePriorityType Priority, bool bIsCloseDistTick, FInstanceUpdateDelegate* UpdateCallback)
 {
 	MUTABLE_CPUPROFILER_SCOPE(FCustomizableObjectSystemPrivate::InitUpdateSkeletalMesh);
@@ -1006,7 +1130,7 @@ void UCustomizableObjectSystem::PurgePendingReleaseSkeletalMesh()
 {
 	MUTABLE_CPUPROFILER_SCOPE(UCustomizableObjectSystem::PurgePendingReleaseSkeletalMesh);
 
-	double	CurTime = FPlatformTime::Seconds();
+	const double CurTime = FPlatformTime::Seconds();
 	static double TimeToDelete = 1.0;
 
 	for (int32 InfoIndex = PendingReleaseSkeletalMesh.Num() - 1; InfoIndex >= 0; --InfoIndex)
@@ -1053,7 +1177,7 @@ void UCustomizableObjectSystem::AddPendingReleaseMaterials(TArray<TObjectPtr<UMa
 
 void UCustomizableObjectSystem::TickPendingReleaseMaterials()
 {
-	for (TArray<FPendingReleaseMaterialsInfo>::TIterator PendingRelease = PendingReleaseMaterials.CreateIterator(); PendingRelease; PendingRelease++)
+	for (TArray<FPendingReleaseMaterialsInfo>::TIterator PendingRelease = PendingReleaseMaterials.CreateIterator(); PendingRelease; ++PendingRelease)
 	{
 		if (--PendingRelease->TicksUntilRelease <= 0)
 		{
@@ -1115,30 +1239,30 @@ void FCustomizableObjectSystemPrivate::UpdateMemoryLimit()
 namespace impl
 {
 
-	void Subtask_Mutable_UpdateParameterRelevancy(const TSharedPtr<FMutableOperationData>& OperationData, const mu::Parameters* MutableParameters)
+	void Subtask_Mutable_UpdateParameterRelevancy(const TSharedPtr<FMutableOperationData>& OperationData)
 	{
 		MUTABLE_CPUPROFILER_SCOPE(Subtask_Mutable_UpdateParameterRelevancy)
 
 		check(OperationData);
-		OperationData->RelevantParametersInProgress.Empty();
-
-		check(MutableParameters);
+		check(OperationData->MutableParameters);
 		check(OperationData->InstanceID != 0);
 
+		OperationData->RelevantParametersInProgress.Empty();
+
 		// This must run in the mutable thread.
 		check(UCustomizableObjectSystem::GetInstance() != nullptr);
 		check(UCustomizableObjectSystem::GetInstance()->GetPrivate() != nullptr);
-		mu::SystemPtr MutableSystem = UCustomizableObjectSystem::GetInstance()->GetPrivate()->MutableSystem;
-
-		int32 NumParameters = MutableParameters->GetCount();
+		const mu::Ptr<mu::System> MutableSystem = UCustomizableObjectSystem::GetInstance()->GetPrivate()->MutableSystem;
 
 		// Update the parameter relevancy.
 		{
 			MUTABLE_CPUPROFILER_SCOPE(ParameterRelevancy)
 
+			const int32 NumParameters = OperationData->MutableParameters->GetCount();
+
 			TArray<bool> Relevant;
 			Relevant.SetNumZeroed(NumParameters);
-			MutableSystem->GetParameterRelevancy(OperationData->InstanceID, MutableParameters, Relevant.GetData());
+			MutableSystem->GetParameterRelevancy(OperationData->InstanceID, OperationData->MutableParameters, Relevant.GetData());
 
 			for (int32 ParamIndex = 0; ParamIndex < NumParameters; ++ParamIndex)
 			{
@@ -1152,7 +1276,7 @@ namespace impl
 
 
 	// This runs in the mutable thread.
-	void Subtask_Mutable_BeginUpdate_GetMesh(const TSharedPtr<FMutableOperationData>& OperationData, TSharedPtr<mu::Model, ESPMode::ThreadSafe> Model, const mu::Parameters* MutableParameters, int32 State)
+	void Subtask_Mutable_BeginUpdate_GetMesh(const TSharedPtr<FMutableOperationData>& OperationData, TSharedPtr<mu::Model, ESPMode::ThreadSafe> Model, mu::Ptr<const mu::Parameters> MutableParameters, int32 State)
 	{
 		MUTABLE_CPUPROFILER_SCOPE(Subtask_Mutable_BeginUpdate_GetMesh)
 
@@ -1190,25 +1314,18 @@ namespace impl
 
 		check(OperationData->InstanceID != 0);
 
-		const mu::Instance* Instance = nullptr;
-
 		if (OperationData->PixelFormatOverride)
 		{
 			System->SetImagePixelConversionOverride( OperationData->PixelFormatOverride );
 		}
 
 		// Main instance generation step
+		// LOD mask, set to all ones to build  all LODs
+		const mu::Instance* Instance = System->BeginUpdate(OperationData->InstanceID, MutableParameters, State, mu::System::AllLODs);
+		if (!Instance)
 		{
-			// LOD mask, set to all ones to build  all LODs
-			uint32 LODMask = mu::System::AllLODs;
-
-			Instance = System->BeginUpdate(OperationData->InstanceID, MutableParameters, State, LODMask);
-
-			if (!Instance)
-			{
-				UE_LOG(LogMutable, Warning, TEXT("An Instace update has failed."));
-				return;
-			}
+			UE_LOG(LogMutable, Warning, TEXT("An Instace update has failed."));
+			return;
 		}
 
 		OperationData->NumLODsAvailable = Instance->GetLODCount();
@@ -1256,7 +1373,7 @@ namespace impl
 		OperationData->InstanceUpdateData.LODs.SetNum(Instance->GetLODCount());
 		for (int32 MutableLODIndex = 0; MutableLODIndex < Instance->GetLODCount(); ++MutableLODIndex)
 		{
-			// Skip LODs outside the reange we want to generate
+			// Skip LODs outside the range we want to generate
 			if (MutableLODIndex < OperationData->CurrentMinLOD || MutableLODIndex > OperationData->CurrentMaxLOD)
 			{
 				continue;
@@ -1319,8 +1436,8 @@ namespace impl
 				// Materials and images
 				for (int32 MeshSurfaceIndex = 0; MeshSurfaceIndex < SurfaceCount; ++MeshSurfaceIndex)
 				{
-					uint32 SurfaceId = Component.Mesh ? Component.Mesh->GetSurfaceId(MeshSurfaceIndex) : OperationData->LastUpdateData.SurfaceIds[PrevUpdateComponent.FirstSurface + MeshSurfaceIndex];
-					int32 InstanceSurfaceIndex = Instance->FindSurfaceById(MutableLODIndex, ComponentIndex, SurfaceId);
+					const uint32 SurfaceId = Component.Mesh ? Component.Mesh->GetSurfaceId(MeshSurfaceIndex) : OperationData->LastUpdateData.SurfaceIds[PrevUpdateComponent.FirstSurface + MeshSurfaceIndex];
+					const int32 InstanceSurfaceIndex = Instance->FindSurfaceById(MutableLODIndex, ComponentIndex, SurfaceId);
 					check(Component.bReuseMesh || Component.Mesh->GetVertexCount() == 0 || InstanceSurfaceIndex >= 0);
 
 					int32 BaseSurfaceIndex = InstanceSurfaceIndex;
@@ -1390,15 +1507,15 @@ namespace impl
 						{
 							MUTABLE_CPUPROFILER_SCOPE(GetScalar)
 
-							FString ScalarName = Instance->GetScalarName(BaseLODIndex, ComponentIndex, InstanceSurfaceIndex, ScalarIndex);
-							float ScalarValue = Instance->GetScalar(BaseLODIndex, ComponentIndex, InstanceSurfaceIndex, ScalarIndex);
+							const FString ScalarName = Instance->GetScalarName(BaseLODIndex, ComponentIndex, InstanceSurfaceIndex, ScalarIndex);
+							const float ScalarValue = Instance->GetScalar(BaseLODIndex, ComponentIndex, InstanceSurfaceIndex, ScalarIndex);
 							
 							FString EncodingMaterialIdString = "__MutableMaterialId";
 							
 							// Decoding Material Switch from Mutable parameter name
 							if (ScalarName.Equals(EncodingMaterialIdString))
 							{
-								Surface.MaterialIndex = (uint32)ScalarValue;
+								Surface.MaterialIndex = static_cast<uint32>(ScalarValue);
 							
 								// This parameter is not needed in the final material instance
 								Surface.ScalarCount -= 1;
@@ -1432,16 +1549,13 @@ namespace impl
 
 
 	// This runs in the mutable thread.
-	void Subtask_Mutable_GetImages(const TSharedPtr<FMutableOperationData>& OperationData, TSharedPtr<mu::Model, ESPMode::ThreadSafe> Model, const mu::Parameters* MutableParameters, int32 State)
+	void Subtask_Mutable_GetImages(const TSharedPtr<FMutableOperationData>& OperationData)
 	{
 		MUTABLE_CPUPROFILER_SCOPE(Subtask_Mutable_GetImages)
 
 		check(OperationData);
-		check(MutableParameters);
 
-		check(UCustomizableObjectSystem::GetInstance() != nullptr);
-		FCustomizableObjectSystemPrivate* CustomizableObjectSystemPrivateData = UCustomizableObjectSystem::GetInstance()->GetPrivate();
-		check(CustomizableObjectSystemPrivateData != nullptr);
+		const FCustomizableObjectSystemPrivate* CustomizableObjectSystemPrivateData = UCustomizableObjectSystem::GetInstanceChecked()->GetPrivateChecked();
 		mu::System* System = CustomizableObjectSystemPrivateData->MutableSystem.get();
 		check(System != nullptr);
 
@@ -1457,8 +1571,8 @@ namespace impl
 			{
 				System->GetImageDesc(OperationData->InstanceID, Image.ImageID, ImageDesc);
 
-				uint16 MaxTextureSizeToGenerate = uint16(CustomizableObjectSystemPrivateData->MaxTextureSizeToGenerate);
-				uint16 MaxSize = FMath::Max(ImageDesc.m_size[0], ImageDesc.m_size[1]);
+				const uint16 MaxTextureSizeToGenerate = static_cast<uint16>(CustomizableObjectSystemPrivateData->MaxTextureSizeToGenerate);
+				const uint16 MaxSize = FMath::Max(ImageDesc.m_size[0], ImageDesc.m_size[1]);
 				uint16 Reduction = 1;
 
 				if (MaxTextureSizeToGenerate > 0 && MaxSize > MaxTextureSizeToGenerate)
@@ -1478,15 +1592,15 @@ namespace impl
 
 			if (bCached)
 			{
-				UE_LOG(LogMutable, VeryVerbose, TEXT("Texture resource with id [%d] is cached."), Image.ImageID);
+				UE_LOG(LogMutable, VeryVerbose, TEXT("Texture resource with id [%llu] is cached."), Image.ImageID);
 				INC_DWORD_STAT(STAT_MutableNumCachedTextures);
 			}
 			else
 			{
-				int32 MaxSize = FMath::Max(Image.FullImageSizeX, Image.FullImageSizeY);
-				int32 FullLODCount = FMath::CeilLogTwo(MaxSize) + 1;
-				int32 MinMipsInImage = FMath::Min(FullLODCount, UTexture::GetStaticMinTextureResidentMipCount());
-				int32 MaxMipsToSkip = FullLODCount - MinMipsInImage;
+				const int32 MaxSize = FMath::Max(Image.FullImageSizeX, Image.FullImageSizeY);
+				const int32 FullLODCount = FMath::CeilLogTwo(MaxSize) + 1;
+				const int32 MinMipsInImage = FMath::Min(FullLODCount, UTexture::GetStaticMinTextureResidentMipCount());
+				const int32 MaxMipsToSkip = FullLODCount - MinMipsInImage;
 				int32 MipsToSkip = FMath::Min(MaxMipsToSkip, OperationData->MipsToSkip);
 
 				if (!FMath::IsPowerOfTwo(Image.FullImageSizeX) || !FMath::IsPowerOfTwo(Image.FullImageSizeY))
@@ -1515,18 +1629,18 @@ namespace impl
 				}
 
 
-				// We should have genrated exactly this size.
-				bool bSizeMissmatch = Image.Image->GetSizeX() != MipSizeX || Image.Image->GetSizeY() != MipSizeY;
+				// We should have generated exactly this size.
+				const bool bSizeMissmatch = Image.Image->GetSizeX() != MipSizeX || Image.Image->GetSizeY() != MipSizeY;
 				if (bSizeMissmatch)
 				{
 					// Generate a correctly-sized but empty image instead, to avoid crashes.
-					UE_LOG(LogMutable, Warning, TEXT("Mutable generated a wrongly-sized image %d."), Image.ImageID);
+					UE_LOG(LogMutable, Warning, TEXT("Mutable generated a wrongly-sized image %llu."), Image.ImageID);
 					Image.Image = new mu::Image(MipSizeX, MipSizeY, FullLODCount - MipsToSkip, Image.Image->GetFormat(), mu::EInitializationType::Black);
 				}
 
 				// We need one mip or the complete chain. Otherwise there was a bug.
-				int32 FullMipCount = Image.Image->GetMipmapCount(Image.Image->GetSizeX(), Image.Image->GetSizeY());
-				int32 RealMipCount = Image.Image->GetLODCount();
+				const int32 FullMipCount = Image.Image->GetMipmapCount(Image.Image->GetSizeX(), Image.Image->GetSizeY());
+				const int32 RealMipCount = Image.Image->GetLODCount();
 
 				bool bForceMipchain = 
 					// Did we fail to generate the entire mipchain (if we have mips at all)?
@@ -1536,10 +1650,10 @@ namespace impl
 				{
 					MUTABLE_CPUPROFILER_SCOPE(GetImage_MipFix);
 
-					UE_LOG(LogMutable, Warning, TEXT("Mutable generated an incomplete mip chain for image %d."), Image.ImageID);
+					UE_LOG(LogMutable, Warning, TEXT("Mutable generated an incomplete mip chain for image %llu."), Image.ImageID);
 
 					// Force the right number of mips. The missing data will be black.
-					mu::Ptr<mu::Image> NewImage = new mu::Image(Image.Image->GetSizeX(), Image.Image->GetSizeY(), FullMipCount, Image.Image->GetFormat(), mu::EInitializationType::Black);
+					const mu::Ptr<mu::Image> NewImage = new mu::Image(Image.Image->GetSizeX(), Image.Image->GetSizeY(), FullMipCount, Image.Image->GetFormat(), mu::EInitializationType::Black);
 					check(NewImage);
 					if (NewImage->GetDataSize() >= Image.Image->GetDataSize())
 					{
@@ -1566,13 +1680,13 @@ namespace impl
 			{
 				const FInstanceUpdateData::FImage& Image = OperationData->InstanceUpdateData.Images[Surface.FirstImage+ImageIndex];
 
-				FString KeyName = Image.Name;
+				const FString KeyName = Image.Name;
 				mu::ImagePtrConst MutableImage = Image.Image;
 
 				// If the image is null, it must be in the cache (or repeated in this instance), and we don't need to do anything here.
 				if (MutableImage)
 				{
-					// Image refences are just references to texture assets and require no work at all
+					// Image references are just references to texture assets and require no work at all
 					if (!MutableImage->IsReference())
 					{
 						FTexturePlatformData* PlatformData = UCustomizableInstancePrivateData::MutableCreateImagePlatformData(MutableImage, -1, Image.FullImageSizeX, Image.FullImageSizeY);
@@ -1725,24 +1839,24 @@ namespace impl
 
 
 	// This runs in a worker thread.
-	void Task_Mutable_Update_GetMesh(TSharedPtr<FMutableOperationData> OperationData, TSharedPtr<mu::Model, ESPMode::ThreadSafe> Model, mu::ParametersPtrConst Parameters, bool bBuildParameterRelevancy, int32 State)
+	void Task_Mutable_Update_GetMesh(const TSharedPtr<FMutableOperationData>& OperationData, const TSharedPtr<mu::Model>& Model)
 	{
 		MUTABLE_CPUPROFILER_SCOPE(Task_Mutable_Update_GetMesh)
 
+		check(OperationData);
+
 #if WITH_EDITOR
-		uint32 StartCycles = FPlatformTime::Cycles();
+		const uint32 StartCycles = FPlatformTime::Cycles();
 #endif
 
-		check(OperationData.IsValid());
-
-		Subtask_Mutable_BeginUpdate_GetMesh(OperationData, Model, Parameters.get(), State);
+		Subtask_Mutable_BeginUpdate_GetMesh(OperationData, Model, OperationData->MutableParameters, OperationData->State);
 
 		// TODO: Not strictly mutable: move to another worker thread task to free mutable access?
 		Subtask_Mutable_PrepareSkeletonData(OperationData);
 
-		if (bBuildParameterRelevancy)
+		if (OperationData->bBuildParameterRelevancy)
 		{
-			Subtask_Mutable_UpdateParameterRelevancy(OperationData, Parameters.get());
+			Subtask_Mutable_UpdateParameterRelevancy(OperationData);
 		}
 		else
 		{
@@ -1750,42 +1864,41 @@ namespace impl
 		}
 
 #if WITH_EDITOR
-		uint32 EndCycles = FPlatformTime::Cycles();
+		const uint32 EndCycles = FPlatformTime::Cycles();
 		OperationData->MutableRuntimeCycles = EndCycles - StartCycles;
 #endif
-
 	}
 
-
-	void Task_Mutable_Update_GetImages(TSharedPtr<FMutableOperationData> OperationData, TSharedPtr<mu::Model, ESPMode::ThreadSafe> Model, mu::ParametersPtrConst Parameters, int32 State)
+	
+	// This runs in a worker thread.
+	void Task_Mutable_Update_GetImages(const TSharedPtr<FMutableOperationData>& OperationData)
 	{
-		// This runs in a worker thread.
 		MUTABLE_CPUPROFILER_SCOPE(Task_Mutable_GetImages)
 
+		check(OperationData);
+
 #if WITH_EDITOR
 		uint32 StartCycles = FPlatformTime::Cycles();
 #endif
 
-		check(OperationData.IsValid());
-
-		Subtask_Mutable_GetImages(OperationData, Model, Parameters.get(), State);
+		Subtask_Mutable_GetImages(OperationData);
 
 		// TODO: Not strictly mutable: move to another worker thread task to free mutable access?
 		Subtask_Mutable_PrepareTextures(OperationData);
 
 #if WITH_EDITOR
-		uint32 EndCycles = FPlatformTime::Cycles();
+		const uint32 EndCycles = FPlatformTime::Cycles();
 		OperationData->MutableRuntimeCycles += EndCycles - StartCycles;
 #endif
 	}
 
 
 	// This runs in a worker thread.
-	void Task_Mutable_ReleaseInstance(TSharedPtr<FMutableOperationData> OperationData, mu::SystemPtr MutableSystem)
+	void Task_Mutable_ReleaseInstance(const TSharedPtr<FMutableOperationData>& OperationData, mu::Ptr<mu::System> MutableSystem)
 	{
 		MUTABLE_CPUPROFILER_SCOPE(Task_Mutable_ReleaseInstance)
 
-		check(OperationData.IsValid());
+		check(OperationData);
 		check(MutableSystem);
 
 		if (OperationData->InstanceID > 0)
@@ -1812,12 +1925,12 @@ namespace impl
 	}
 
 
-	void Task_Mutable_ReleaseInstanceID(TSharedPtr<FMutableOperationData> OperationData, mu::SystemPtr MutableSystem)
+	// This runs in a worker thread.
+	void Task_Mutable_ReleaseInstanceID(const TSharedPtr<FMutableOperationData>& OperationData, const mu::Ptr<mu::System>& MutableSystem)
 	{
 		MUTABLE_CPUPROFILER_SCOPE(Task_Mutable_ReleaseInstanceID)
 
-		// This runs in a worker thread.
-		check(OperationData.IsValid());
+		check(OperationData);
 
 		if (OperationData->InstanceID > 0)
 		{
@@ -1832,13 +1945,14 @@ namespace impl
 	}
 
 
-	void Task_Game_ReleasePlatformData(TSharedPtr<FMutableReleasePlatformOperationData> OperationData)
+	void Task_Game_ReleasePlatformData(const TSharedPtr<FMutableReleasePlatformOperationData>& OperationData)
 	{
 		MUTABLE_CPUPROFILER_SCOPE(Task_Game_ReleasePlatformData)
 
-		check(OperationData.IsValid());
+		check(OperationData);
+		
 		TMap<uint32, FTexturePlatformData*>& ImageToPlatformDataMap = OperationData->ImageToPlatformDataMap;
-		for (TPair<uint32, FTexturePlatformData*> Pair : ImageToPlatformDataMap)
+		for (const TPair<uint32, FTexturePlatformData*>& Pair : ImageToPlatformDataMap)
 		{
 			delete Pair.Value; // If this is not null then it must mean it hasn't been used, otherwise they would have taken ownership and nulled it
 		}
@@ -1846,20 +1960,21 @@ namespace impl
 	}
 
 	
-	void Task_Game_Callbacks(TSharedPtr<FMutableOperationData> OperationData, const TWeakObjectPtr<UCustomizableObjectInstance>& CustomizableObjectInstancePtr)
+	void Task_Game_Callbacks(const TSharedPtr<FMutableOperationData>& OperationData)
 	{
 		MUTABLE_CPUPROFILER_SCOPE(Task_Game_Callbacks)
 
 		check(IsInGameThread());
+		check(OperationData);
 
 		UCustomizableObjectSystem* System = UCustomizableObjectSystem::GetInstance();
 		if (!System || !System->IsValidLowLevel() || System->HasAnyFlags(RF_BeginDestroyed))
 		{
-			FinishUpdateGlobal(CustomizableObjectInstancePtr.Get(), EUpdateResult::Error, &OperationData->UpdateCallback);
+			FinishUpdateGlobal(OperationData->Instance.Get(), EUpdateResult::Error, &OperationData->UpdateCallback);
 			return;
 		}
 
-		UCustomizableObjectInstance* CustomizableObjectInstance = CustomizableObjectInstancePtr.Get();
+		UCustomizableObjectInstance* CustomizableObjectInstance = OperationData->Instance.Get();
 
 		// TODO: Review checks.
 		if (!CustomizableObjectInstance || !CustomizableObjectInstance->IsValidLowLevel() )
@@ -1869,8 +1984,7 @@ namespace impl
 			return;
 		}
 
-		FCustomizableObjectSystemPrivate * CustomizableObjectSystemPrivateData = System->GetPrivate();
-		check(CustomizableObjectSystemPrivateData != nullptr);
+		FCustomizableObjectSystemPrivate * CustomizableObjectSystemPrivateData = System->GetPrivateChecked();
 
 		// Actual work
 		// TODO MTBL-391: Review This hotfix
@@ -1898,22 +2012,21 @@ namespace impl
 	}
 
 
-	void Task_Game_ConvertResources(TSharedPtr<FMutableOperationData> OperationData, const TWeakObjectPtr<UCustomizableObjectInstance>& CustomizableObjectInstancePtr)
+	void Task_Game_ConvertResources(TSharedPtr<FMutableOperationData> OperationData)
 	{
 		MUTABLE_CPUPROFILER_SCOPE(Task_Game_ConvertResources)
 
 		check(IsInGameThread());
+		check(OperationData);
 
 		UCustomizableObjectSystem* System = UCustomizableObjectSystem::GetInstance();
 		if (!System || !System->IsValidLowLevel() || System->HasAnyFlags(RF_BeginDestroyed))
 		{
-			FinishUpdateGlobal(CustomizableObjectInstancePtr.Get(), EUpdateResult::Error, &OperationData->UpdateCallback);
+			FinishUpdateGlobal(OperationData->Instance.Get(), EUpdateResult::Error, &OperationData->UpdateCallback);
 			return;
 		}
 
-		check(OperationData.IsValid());
-
-		UCustomizableObjectInstance* CustomizableObjectInstance = CustomizableObjectInstancePtr.Get();
+		UCustomizableObjectInstance* CustomizableObjectInstance = OperationData->Instance.Get();
 
 		// Actual work
 		// TODO: Review checks.
@@ -1988,8 +2101,7 @@ namespace impl
 			} // END - Process texture coverage queries
 		} // if (!bInstanceValid)
 
-		FCustomizableObjectSystemPrivate* CustomizableObjectSystemPrivateData = System->GetPrivate();
-		check(CustomizableObjectSystemPrivateData != nullptr);
+		FCustomizableObjectSystemPrivate* CustomizableObjectSystemPrivateData = System->GetPrivateChecked();
 
 		for (int32 ComponentIndex = 0; ComponentIndex < CustomizableObjectInstance->GetPrivate()->ComponentsData.Num(); ++ComponentIndex)
 		{
@@ -2006,7 +2118,7 @@ namespace impl
 		// Next Task: Release Mutable. We need this regardless if we cancel or not
 		//-------------------------------------------------------------		
 		{
-			mu::SystemPtr MutableSystem = CustomizableObjectSystemPrivateData->MutableSystem;
+			const mu::Ptr<mu::System> MutableSystem = CustomizableObjectSystemPrivateData->MutableSystem;
 			CustomizableObjectSystemPrivateData->MutableTaskGraph.AddMutableThreadTask(
 				TEXT("Task_Mutable_ReleaseInstance"),
 				[OperationData, MutableSystem]() {Task_Mutable_ReleaseInstance(OperationData, MutableSystem); });
@@ -2043,9 +2155,9 @@ namespace impl
 			CustomizableObjectSystemPrivateData->AddGameThreadTask(
 				{
 				FMutableTaskDelegate::CreateLambda(
-					[OperationData,CustomizableObjectInstancePtr]()
+					[OperationData]()
 					{
-						Task_Game_Callbacks(OperationData,CustomizableObjectInstancePtr);
+						Task_Game_Callbacks(OperationData);
 					}),
 					{}
 				});
@@ -2057,8 +2169,8 @@ namespace impl
 	}
 
 
-	/** "Lock Cached Resources" */
-	void Task_Game_LockCache(TSharedPtr<FMutableOperationData> OperationData, const TWeakObjectPtr<UCustomizableObjectInstance>& CustomizableObjectInstancePtr, mu::Ptr<const mu::Parameters> Parameters, bool bBuildParameterRelevancy)
+	/** Lock Cached Resources. */
+	void Task_Game_LockCache(TSharedPtr<FMutableOperationData> OperationData)
 	{
 		MUTABLE_CPUPROFILER_SCOPE(Task_Game_LockCache)
 
@@ -2071,8 +2183,7 @@ namespace impl
 			return;
 		}
 
-		UCustomizableObjectInstance* ObjectInstance = CustomizableObjectInstancePtr.Get();
-
+		UCustomizableObjectInstance* ObjectInstance = OperationData->Instance.Get();
 		if (!ObjectInstance)
 		{
 			System->ClearCurrentMutableOperation();
@@ -2098,13 +2209,11 @@ namespace impl
 		if (!CustomizableObject)
 		{
 			System->ClearCurrentMutableOperation();
-			FinishUpdateGlobal(CustomizableObjectInstancePtr.Get(), EUpdateResult::Error, &OperationData->UpdateCallback);
+			FinishUpdateGlobal(ObjectInstance, EUpdateResult::Error, &OperationData->UpdateCallback);
 			return;
 		}
-
-		check(OperationData.IsValid());
 		
-		if (bBuildParameterRelevancy)
+		if (OperationData->bBuildParameterRelevancy)
 		{
 			// Relevancy
 			ObjectInstancePrivateData->RelevantParameters = OperationData->RelevantParametersInProgress;
@@ -2115,8 +2224,7 @@ namespace impl
 		// When protecting textures there mustn't be any left from a previous update
 		check(System->ProtectedCachedTextures.Num() == 0);
 
-		FCustomizableObjectSystemPrivate* SystemPrivateData = System->GetPrivate();
-		check(SystemPrivateData != nullptr);
+		FCustomizableObjectSystemPrivate* SystemPrivateData = System->GetPrivateChecked();
 
 		// TODO: If this is the first code that runs after the CO program has finished AND if it's
 		// guaranteed that the next CO program hasn't started yet, we need to call ClearActiveObject
@@ -2135,7 +2243,7 @@ namespace impl
 		for (const FInstanceUpdateData::FImage& Image : OperationData->InstanceUpdateData.Images)
 		{
 			FMutableImageCacheKey Key(Image.ImageID, OperationData->MipsToSkip);
-			TWeakObjectPtr<UTexture2D>* TexturePtr = Cache.Images.Find(Key);
+			const TWeakObjectPtr<UTexture2D>* TexturePtr = Cache.Images.Find(Key);
 
 			if (TexturePtr && TexturePtr->Get() && SystemPrivateData->TextureHasReferences(Key))
 			{
@@ -2159,15 +2267,12 @@ namespace impl
 #endif
 		{
 			// Task inputs
-			TSharedPtr<mu::Model, ESPMode::ThreadSafe> Model = CustomizableObject->GetPrivate()->GetModel();
-			int32 State = ObjectInstance->GetState();
-
 			Mutable_GetImagesTask = SystemPrivateData->MutableTaskGraph.AddMutableThreadTask(
-					TEXT("Task_Mutable_GetImages"),
-					[OperationData, Parameters, Model, State]()
-					{
-						impl::Task_Mutable_Update_GetImages(OperationData, Model, Parameters, State);
-					});
+				TEXT("Task_Mutable_GetImages"),
+				[OperationData]()
+				{
+					impl::Task_Mutable_Update_GetImages(OperationData);
+				});
 		}
 
 
@@ -2184,9 +2289,9 @@ namespace impl
 		SystemPrivateData->AddGameThreadTask(
 			{
 			FMutableTaskDelegate::CreateLambda(
-				[OperationData,CustomizableObjectInstancePtr]()
+				[OperationData]()
 				{
-					Task_Game_ConvertResources(OperationData,CustomizableObjectInstancePtr);
+					Task_Game_ConvertResources(OperationData);
 				}),
 #ifdef MUTABLE_USE_NEW_TASKGRAPH
 				{},
@@ -2202,13 +2307,10 @@ namespace impl
 	{
 		MUTABLE_CPUPROFILER_SCOPE(Task_Game_ReleaseInstanceID)
 
-		UCustomizableObjectSystem* System = UCustomizableObjectSystem::GetInstance();
-		check(System != nullptr);
+		UCustomizableObjectSystem* System = UCustomizableObjectSystem::GetInstanceChecked();
+		FCustomizableObjectSystemPrivate* SystemPrivateData = System->GetPrivateChecked();
 
-		FCustomizableObjectSystemPrivate* SystemPrivateData = System->GetPrivate();
-		check(SystemPrivateData != nullptr);
-
-		mu::SystemPtr MutableSystem = SystemPrivateData->MutableSystem;
+		const mu::Ptr<mu::System> MutableSystem = SystemPrivateData->MutableSystem;
 
 		// Task: Release Instance ID
 		//-------------------------------------------------------------
@@ -2216,17 +2318,9 @@ namespace impl
 		check(CurrentOperationData);
 		CurrentOperationData->InstanceID = IDToRelease;
 
-#ifdef MUTABLE_USE_NEW_TASKGRAPH
-		UE::Tasks::FTask Mutable_GetMeshTask;
-#else
-		FGraphEventRef Mutable_GetMeshTask;
-#endif
 		{
 			// Task inputs
-			TSharedPtr<FMutableOperation> CurrentMutableOperation = SystemPrivateData->CurrentMutableOperation;
-			check(CurrentMutableOperation);
-
-			Mutable_GetMeshTask = SystemPrivateData->MutableTaskGraph.AddMutableThreadTask(
+			SystemPrivateData->MutableTaskGraph.AddMutableThreadTask(
 				TEXT("Task_Mutable_ReleaseInstanceID"),
 				[CurrentOperationData, MutableSystem]()
 				{
@@ -2237,14 +2331,14 @@ namespace impl
 
 
 	/** "Start Update" */
-	void Task_Game_StartUpdate(TSharedPtr<FMutableOperation> Operation)
+	void Task_Game_StartUpdate(const TSharedPtr<FMutableOperation>& Operation)
 	{
 		MUTABLE_CPUPROFILER_SCOPE(Task_Game_StartUpdate)
 
+		check(Operation);
+		
 		UCustomizableObjectSystem::GetInstance()->GetPrivate()->MutableTaskGraph.AllowLaunchingMutableTaskLowPriority(false, false);
 		
-		check(Operation);
-
 		UCustomizableObjectSystem* System = UCustomizableObjectSystem::GetInstance();
 		check(System != nullptr);
 
@@ -2255,7 +2349,7 @@ namespace impl
 			return;
 		}
 
-		TObjectPtr<UCustomizableObjectInstance> CandidateInstance = Operation->CustomizableObjectInstance.Get();
+		UCustomizableObjectInstance* CandidateInstance = Operation->CustomizableObjectInstance.Get();
 		
 		UCustomizableInstancePrivateData* CandidateInstancePrivateData = CandidateInstance->GetPrivate();
 		if (!CandidateInstancePrivateData)
@@ -2328,8 +2422,7 @@ namespace impl
 			Operation->StartUpdateTime = FPlatformTime::Seconds();
 		}
 
-		FCustomizableObjectSystemPrivate* SystemPrivateData = System->GetPrivate();
-		check(SystemPrivateData != nullptr);
+		FCustomizableObjectSystemPrivate* SystemPrivateData = System->GetPrivateChecked();
 
 		SystemPrivateData->CurrentInstanceBeingUpdated = CandidateInstance;
 
@@ -2384,6 +2477,7 @@ namespace impl
 		//-------------------------------------------------------------
 		TSharedPtr<FMutableOperationData> CurrentOperationData = MakeShared<FMutableOperationData>();
 		check(CurrentOperationData);
+		CurrentOperationData->Instance = Operation->CustomizableObjectInstance;
 		CurrentOperationData->TextureCoverageQueries_MutableThreadParams = CandidateInstancePrivateData->TextureCoverageQueries;
 		CurrentOperationData->TextureCoverageQueries_MutableThreadResults.Empty();
 		CurrentOperationData->bCanReuseGeneratedData = SystemPrivateData->bEnableMutableReusePreviousUpdateData;
@@ -2399,6 +2493,7 @@ namespace impl
 		CurrentOperationData->State = CandidateInstance->GetState();
 		CurrentOperationData->UpdateResult = EUpdateResult::Success;
 		CurrentOperationData->UpdateCallback = Operation->UpdateCallback;
+		CurrentOperationData->bBuildParameterRelevancy = Operation->IsBuildParameterRelevancy();
 #if WITH_EDITOR
 		CurrentOperationData->PixelFormatOverride = SystemPrivateData->ImageFormatOverrideFunc;
 #endif
@@ -2421,19 +2516,13 @@ namespace impl
 #endif
 		{
 			// Task inputs
-			TWeakObjectPtr<UCustomizableObjectInstance> CustomizableObjectInstancePtr = CandidateInstance;
-			check(CustomizableObjectInstancePtr.IsValid());
-			TSharedPtr<FMutableOperation> CurrentMutableOperation = SystemPrivateData->CurrentMutableOperation;
-			check(CurrentMutableOperation);
-			bool bBuildParameterRelevancy = CurrentMutableOperation->IsBuildParameterRelevancy();
-			TSharedPtr<mu::Model, ESPMode::ThreadSafe> Model = CustomizableObject->GetPrivate()->GetModel();
-			int32 State = CustomizableObjectInstancePtr->GetState();
+			TSharedPtr<mu::Model> Model = CustomizableObject->GetPrivate()->GetModel();
 
 			Mutable_GetMeshTask = SystemPrivateData->MutableTaskGraph.AddMutableThreadTask(
 				TEXT("Task_Mutable_Update_GetMesh"),
-				[CurrentOperationData, bBuildParameterRelevancy, Parameters, Model, State]()
+				[CurrentOperationData, Model]()
 				{
-					impl::Task_Mutable_Update_GetMesh(CurrentOperationData, Model, Parameters, bBuildParameterRelevancy, State);
+					impl::Task_Mutable_Update_GetMesh(CurrentOperationData, Model);
 				});
 		}
 
@@ -2442,18 +2531,12 @@ namespace impl
 		//-------------------------------------------------------------
 		{
 			// Task inputs
-			TWeakObjectPtr<UCustomizableObjectInstance> CustomizableObjectInstancePtr = CandidateInstance;
-			check(CustomizableObjectInstancePtr.IsValid());
-			TSharedPtr<FMutableOperation> CurrentMutableOperation = SystemPrivateData->CurrentMutableOperation;
-			check(CurrentMutableOperation);
-			bool bBuildParameterRelevancy = CurrentMutableOperation->IsBuildParameterRelevancy();
-
 			SystemPrivateData->AddGameThreadTask(
 				{
 				FMutableTaskDelegate::CreateLambda(
-					[CurrentOperationData, CustomizableObjectInstancePtr, bBuildParameterRelevancy, Parameters]()
+					[CurrentOperationData]()
 					{
-						impl::Task_Game_LockCache(CurrentOperationData, CustomizableObjectInstancePtr, Parameters, bBuildParameterRelevancy);
+						impl::Task_Game_LockCache(CurrentOperationData);
 					}),
 				Mutable_GetMeshTask
 				});
@@ -2519,10 +2602,9 @@ bool UCustomizableObjectSystem::Tick(float DeltaTime)
 		return true;
 	}
 
-	UWorld* World = GWorld;
-	if (World)
+	if (GWorld)
 	{
-		EWorldType::Type WorldType = World->WorldType;
+		const EWorldType::Type WorldType = GWorld->WorldType;
 
 		if (WorldType != EWorldType::PIE && WorldType != EWorldType::Game && WorldType != EWorldType::Editor && WorldType != EWorldType::GamePreview)
 		{
@@ -2532,7 +2614,7 @@ bool UCustomizableObjectSystem::Tick(float DeltaTime)
 
 	// \TODO: Review: We should never compile an object from this tick, so this could be removed
 #if WITH_EDITOR
-	FAssetRegistryModule& AssetRegistryModule = FModuleManager::LoadModuleChecked<FAssetRegistryModule>("AssetRegistry");
+	const FAssetRegistryModule& AssetRegistryModule = FModuleManager::LoadModuleChecked<FAssetRegistryModule>("AssetRegistry");
 
 	if (AssetRegistryModule.Get().IsLoadingAssets())
 	{
@@ -2777,7 +2859,9 @@ void UCustomizableObjectSystem::DiscardInstances()
 	int32 NumInstancesDiscarded = 0;
 	const int32 DiscardLimitPerTick = CVarMaxNumInstancesToDiscardPerTick.GetValueOnGameThread();
 
-	for (auto Iterator = Private->MutablePendingInstanceWork.GetDiscardIterator(); Iterator && NumInstancesDiscarded < DiscardLimitPerTick; ++Iterator)
+	for (TSet<FMutablePendingInstanceDiscard, FPendingInstanceDiscardKeyFuncs>::TIterator Iterator = Private->MutablePendingInstanceWork.GetDiscardIterator();
+		Iterator && NumInstancesDiscarded < DiscardLimitPerTick;
+		++Iterator)
 	{
 		MUTABLE_CPUPROFILER_SCOPE(OperationDiscard);
 
@@ -2943,6 +3027,8 @@ FMutableOperation::FMutableOperation(const FMutableOperation& Other)
 	Parameters = Other.Parameters;
 	TextureParameters = Other.TextureParameters;
 	UpdateCallback = Other.UpdateCallback;
+	State = Other.State;
+	IDToRelease = Other.IDToRelease;
 
 	CacheTexturesParameters(TextureParameters);
 }
@@ -3010,74 +3096,63 @@ FMutableOperation FMutableOperation::CreateInstanceUpdate(UCustomizableObjectIns
 
 int32 UCustomizableObjectSystem::GetNumInstances() const
 {
-	check(Private != nullptr);
-	return Private->MutableStats.NumInstances;
+	return GetPrivateChecked()->MutableStats.NumInstances;
 }
 
 int32 UCustomizableObjectSystem::GetNumPendingInstances() const
 {
-	check(Private != nullptr);
-	return Private->MutableStats.NumPendingInstances;
+	return GetPrivateChecked()->MutableStats.NumPendingInstances;
 }
 
 int32 UCustomizableObjectSystem::GetTotalInstances() const
 {
-	check(Private != nullptr);
-	return Private->MutableStats.TotalInstances;
+	return GetPrivateChecked()->MutableStats.TotalInstances;
 }
 
 int32 UCustomizableObjectSystem::GetTextureMemoryUsed() const
 {
-	check(Private != nullptr);
-	return int32(Private->MutableStats.TextureMemoryUsed);
+	return static_cast<int32>(GetPrivateChecked()->MutableStats.TextureMemoryUsed);
 }
 
 int32 UCustomizableObjectSystem::GetAverageBuildTime() const
 {
-	check(Private != nullptr);
-	return Private->MutableStats.TotalBuiltInstances == 0 ? 0 : Private->MutableStats.TotalBuildMs / Private->MutableStats.TotalBuiltInstances;
+	return GetPrivateChecked()->MutableStats.TotalBuiltInstances == 0 ? 0 : Private->MutableStats.TotalBuildMs / Private->MutableStats.TotalBuiltInstances;
 }
 
 
 int32 UCustomizableObjectSystem::GetSkeletalMeshMinLODQualityLevel() const
 {
-	check(Private != nullptr);
-	return Private->SkeletalMeshMinLodQualityLevel;
+	return GetPrivateChecked()->SkeletalMeshMinLodQualityLevel;
 }
 
 
 bool UCustomizableObjectSystem::IsSupport16BitBoneIndexEnabled() const
 {
-	check(Private != nullptr);
-	return Private->bSupport16BitBoneIndex;
+	return GetPrivateChecked()->bSupport16BitBoneIndex;
 }
 
 
 bool UCustomizableObjectSystem::IsProgressiveMipStreamingEnabled() const
 {
-	check(Private != nullptr);
-	return Private->EnableMutableProgressiveMipStreaming != 0;
+	return GetPrivateChecked()->EnableMutableProgressiveMipStreaming != 0;
 }
 
 
 void UCustomizableObjectSystem::SetProgressiveMipStreamingEnabled(bool bIsEnabled)
 {
-	check(Private != nullptr);
-	Private->EnableMutableProgressiveMipStreaming = bIsEnabled ? 1 : 0;
+	GetPrivateChecked()->EnableMutableProgressiveMipStreaming = bIsEnabled ? 1 : 0;
 }
 
 
 bool UCustomizableObjectSystem::IsOnlyGenerateRequestedLODsEnabled() const
 {
-	check(Private != nullptr);
-	return Private->EnableOnlyGenerateRequestedLODs != 0;
+	return GetPrivateChecked()->EnableOnlyGenerateRequestedLODs != 0;
 }
 
 
 void UCustomizableObjectSystem::SetOnlyGenerateRequestedLODsEnabled(bool bIsEnabled)
 {
-	check(Private != nullptr);
-	Private->EnableOnlyGenerateRequestedLODs = bIsEnabled ? 1 : 0;
+	GetPrivateChecked()->EnableOnlyGenerateRequestedLODs = bIsEnabled ? 1 : 0;
 }
 
 
@@ -3151,8 +3226,7 @@ void UCustomizableObjectSystem::EndBenchmark()
 
 void UCustomizableObjectSystem::SetReleaseMutableTexturesImmediately(bool bReleaseTextures)
 {
-	check(Private != nullptr);
-	Private->bReleaseTexturesImmediately = bReleaseTextures;
+	GetPrivateChecked()->bReleaseTexturesImmediately = bReleaseTextures;
 }
 
 
@@ -3195,7 +3269,7 @@ void UCustomizableObjectSystem::OnPreBeginPIE(const bool bIsSimulatingInEditor)
 
 	if (!TempObjectsToRecompile.IsEmpty())
 	{
-		FText Msg = FText::FromString(TEXT("Warning: one or more Customizable Objects used in PIE are uncompiled.\n\nDo you want to compile them?"));
+		const FText Msg = FText::FromString(TEXT("Warning: one or more Customizable Objects used in PIE are uncompiled.\n\nDo you want to compile them?"));
 		if (FMessageDialog::Open(EAppMsgType::OkCancel, Msg) == EAppReturnType::Ok)
 		{
 			ObjectsToRecompile.Empty(TempObjectsToRecompile.Num());
@@ -3212,11 +3286,10 @@ void UCustomizableObjectSystem::StartNextRecompile()
 	}
 
 	FAssetData Itr = ObjectsToRecompile.Pop();
-	UCustomizableObject* CustomizableObject = Cast<UCustomizableObject>(Itr.GetAsset());
 
-	if (CustomizableObject)
+	if (UCustomizableObject* CustomizableObject = Cast<UCustomizableObject>(Itr.GetAsset()))
 	{
-		FText UpdateMsg = FText::FromString(FString::Printf(TEXT("Compiling Customizable Objects:\n%s"), *CustomizableObject->GetName()));
+		const FText UpdateMsg = FText::FromString(FString::Printf(TEXT("Compiling Customizable Objects:\n%s"), *CustomizableObject->GetName()));
 		FSlateNotificationManager::Get().UpdateProgressNotification(RecompileNotificationHandle, NumObjectsCompiled, TotalNumObjectsToRecompile, UpdateMsg);
 
 		// Use default options
diff --git a/Engine/Plugins/Experimental/Mutable/Source/CustomizableObject/Private/MuCO/CustomizableObjectSystemPrivate.h b/Engine/Plugins/Experimental/Mutable/Source/CustomizableObject/Private/MuCO/CustomizableObjectSystemPrivate.h
--- a/Engine/Plugins/Experimental/Mutable/Source/CustomizableObject/Private/MuCO/CustomizableObjectSystemPrivate.h
+++ b/Engine/Plugins/Experimental/Mutable/Source/CustomizableObject/Private/MuCO/CustomizableObjectSystemPrivate.h
@@ -76,7 +76,7 @@ public:
 	double StartUpdateTime = 0.0;
 
 	/** Instance optimization state. */
-	int32 State;
+	int32 State = -1;
 
 	/** Only used in the IDRelease operation type */
 	mu::Instance::ID IDToRelease;
@@ -559,6 +559,8 @@ struct FPendingTextureCoverageQuery
 /** Runtime data used during a mutable instance update */
 struct FMutableOperationData
 {
+	TWeakObjectPtr<UCustomizableObjectInstance> Instance;
+	
 	bool bCanReuseGeneratedData = false;
 	FInstanceGeneratedData LastUpdateData;
 
@@ -590,10 +592,12 @@ struct FMutableOperationData
 	/** This list of queries is generated in the update mutable task, and consumed later in the game thread. */
 	TArray<FPendingTextureCoverageQuery> PendingTextureCoverageQueries;
 
-	/** */
 	mu::Ptr<const mu::Parameters> MutableParameters;
+
 	int32 State = 0;
 
+	bool bBuildParameterRelevancy = false;
+
 	EUpdateResult UpdateResult;
 	FInstanceUpdateDelegate UpdateCallback;
 
@@ -639,14 +643,13 @@ struct FMutableStats
 	uint32 CountAllocatedSkeletalMesh = 0;
 
 	// \TODO: Remove this array if we are not gathering stats!
-	TArray<TWeakObjectPtr<class UTexture2D>> TextureTrackerArray;
+	TArray<TWeakObjectPtr<UTexture2D>> TextureTrackerArray;
 };
 
 
 class FCustomizableObjectSystemPrivate : public FGCObject
 {
 public:
-
 	// Singleton for the unreal mutable system.
 	static UCustomizableObjectSystem* SSystem;
 
@@ -660,7 +663,6 @@ public:
 	// This object is responsible for streaming data to the MutableSystem.
 	TSharedPtr<class FUnrealMutableModelBulkReader> Streamer;
 
-	// 
 	TSharedPtr<class FUnrealExtensionDataStreamer> ExtensionDataStreamer;
 
 	// This object is responsible for providing custom images to mutable models (for image parameters)
@@ -696,97 +698,26 @@ public:
 	mu::FImageOperator::FImagePixelFormatFunc ImageFormatOverrideFunc;
 #endif
 
-	/** */
-	inline void AddGameThreadTask(const FMutableTask& Task)
-	{
-		check(IsInGameThread())
-		PendingTasks.Enqueue(Task);
-	}
-
+	void AddGameThreadTask(const FMutableTask& Task);
 
 	/** FSerializableObject interface */
 	virtual void AddReferencedObjects(FReferenceCollector& Collector) override;
-	virtual FString GetReferencerName() const override
-	{
-		return TEXT("FCustomizableObjectSystemPrivate");
-	}
-
+	virtual FString GetReferencerName() const override;
 
 	// Remove references to cached objects that have been deleted in the unreal
 	// side, and cannot be cached anyway.
 	// This should only happen in the game thread
 	void CleanupCache();
 
-
 	// This should only happen in the game thread
-	FMutableResourceCache& GetObjectCache(const UCustomizableObject* Object)
-	{
-		check(IsInGameThread());
-
-		// Not mandatory, but a good place for a cleanup
-		CleanupCache();
-
-		for (int ModelIndex = 0; ModelIndex < ModelResourcesCache.Num(); ++ModelIndex)
-		{
-			if (ModelResourcesCache[ModelIndex].Object==Object)
-			{
-				return ModelResourcesCache[ModelIndex];
-			}
-		}
-		
-		// Not found, create and add it.
-		ModelResourcesCache.Push(FMutableResourceCache());
-		ModelResourcesCache.Last().Object = Object;
-		return ModelResourcesCache.Last();
-	}
-
-
-	void AddTextureReference(const FMutableImageCacheKey& TextureId)
-	{
-		uint32& CountRef = TextureReferenceCount.FindOrAdd(TextureId);
+	FMutableResourceCache& GetObjectCache(const UCustomizableObject* Object);
 
-		CountRef++;
-	}
+	void AddTextureReference(const FMutableImageCacheKey& TextureId);
 
-	
 	// Returns true if the texture's references become zero
-	bool RemoveTextureReference(const FMutableImageCacheKey& TextureId)
-	{
-		uint32* CountPtr = TextureReferenceCount.Find(TextureId);
-
-		if (CountPtr && *CountPtr > 0)
-		{
-			(*CountPtr)--;
-
-			if (*CountPtr == 0)
-			{
-				TextureReferenceCount.Remove(TextureId);
-
-				return true;
-			}
-		}
-		else
-		{
-			ensure(false); // Mutable texture reference count is incorrect
-			TextureReferenceCount.Remove(TextureId);
-		}
-
-		return false;
-	}
-
-
-	bool TextureHasReferences(const FMutableImageCacheKey& TextureId) const
-	{
-		const uint32* CountPtr = TextureReferenceCount.Find(TextureId);
-
-		if (CountPtr && *CountPtr > 0)
-		{
-			return true;
-		}
-
-		return false;
-	}
+	bool RemoveTextureReference(const FMutableImageCacheKey& TextureId);
 
+	bool TextureHasReferences(const FMutableImageCacheKey& TextureId) const;
 
 	// Init the async Skeletal Mesh creation/update
 	void InitUpdateSkeletalMesh(UCustomizableObjectInstance& Public, EQueuePriorityType Priority, bool bIsCloseDistTick, FInstanceUpdateDelegate* UpdateCallback = nullptr);
@@ -799,10 +730,10 @@ public:
 
 	void GetMipStreamingConfig(const UCustomizableObjectInstance& Instance, bool& bOutNeverStream, int32& OutMipsToSkip) const;
 	
-	bool IsReplaceDiscardedWithReferenceMeshEnabled() const { return bReplaceDiscardedWithReferenceMesh; }
-	void SetReplaceDiscardedWithReferenceMeshEnabled(bool bIsEnabled) { bReplaceDiscardedWithReferenceMesh = bIsEnabled; }
+	bool IsReplaceDiscardedWithReferenceMeshEnabled() const;
+	void SetReplaceDiscardedWithReferenceMeshEnabled(bool bIsEnabled);
 
-	int32 GetCountAllocatedSkeletalMesh() { return MutableStats.CountAllocatedSkeletalMesh; }
+	int32 GetCountAllocatedSkeletalMesh() const;
 
 	mutable FMutableStats MutableStats;
 
@@ -813,7 +744,7 @@ public:
 
 	static FCustomizableObjectCompilerBase* (*NewCompilerFunc)();
 
-	void CreatedTexture(UTexture2D* Texture);
+	void CreatedTexture(UTexture2D* Texture) const;
 
 	TMap<FMutableImageCacheKey, uint32> TextureReferenceCount; // Keeps a count of texture usage to decide if they have to be blocked from GC during an update
 
