[Mutable] Fix crash when skipping an update at Task_Game_StartUpdate.

#rnx

diff --git a/Engine/Plugins/Experimental/Mutable/Source/CustomizableObject/Private/MuCO/CustomizableObjectSystem.cpp b/Engine/Plugins/Experimental/Mutable/Source/CustomizableObject/Private/MuCO/CustomizableObjectSystem.cpp
--- a/Engine/Plugins/Experimental/Mutable/Source/CustomizableObject/Private/MuCO/CustomizableObjectSystem.cpp
+++ b/Engine/Plugins/Experimental/Mutable/Source/CustomizableObject/Private/MuCO/CustomizableObjectSystem.cpp
@@ -2331,7 +2331,7 @@ namespace impl
 
 
 	/** "Start Update" */
-	void Task_Game_StartUpdate(const TSharedPtr<FMutableOperation>& Operation)
+	void Task_Game_StartUpdate(TSharedPtr<FMutableOperation> Operation)
 	{
 		MUTABLE_CPUPROFILER_SCOPE(Task_Game_StartUpdate)
 
