[mutable] Several fixes to operations on block-compressed images, when the block size is not a divisor of the texture size, which may happen with the new astc formats.
#rnx

diff --git a/Engine/Plugins/Experimental/Mutable/Source/MutableRuntime/Private/MuR/Image.cpp b/Engine/Plugins/Experimental/Mutable/Source/MutableRuntime/Private/MuR/Image.cpp
--- a/Engine/Plugins/Experimental/Mutable/Source/MutableRuntime/Private/MuR/Image.cpp
+++ b/Engine/Plugins/Experimental/Mutable/Source/MutableRuntime/Private/MuR/Image.cpp
@@ -686,8 +686,8 @@ namespace mu
         int blockY = pixelY / fdata.PixelsPerBlockY;
         int blockPixelY = pixelY % fdata.PixelsPerBlockY;
 
-        int blocksPerRow = m_size[0] / fdata.PixelsPerBlockX;
-        int blockOffset = blockX + blockY * blocksPerRow;
+        int32 blocksPerRow = FMath::DivideAndRoundUp(int32(m_size[0]), int32(fdata.PixelsPerBlockX));
+        int32 blockOffset = blockX + blockY * blocksPerRow;
 
         // Non-generic part
         if ( m_format== EImageFormat::IF_RGB_UBYTE )
diff --git a/Engine/Plugins/Experimental/Mutable/Source/MutableRuntime/Private/MuR/OpImageCompose.cpp b/Engine/Plugins/Experimental/Mutable/Source/MutableRuntime/Private/MuR/OpImageCompose.cpp
--- a/Engine/Plugins/Experimental/Mutable/Source/MutableRuntime/Private/MuR/OpImageCompose.cpp
+++ b/Engine/Plugins/Experimental/Mutable/Source/MutableRuntime/Private/MuR/OpImageCompose.cpp
@@ -55,16 +55,16 @@ namespace mu
 		{
 			// Sizes in blocks of the current mips
 			UE::Math::TIntVector2<uint16> baseMipSize(
-				pBase->GetSizeX() / finfo.PixelsPerBlockX,
-				pBase->GetSizeY() / finfo.PixelsPerBlockY);
+				FMath::DivideAndRoundUp(pBase->GetSizeX(), uint16(finfo.PixelsPerBlockX)),
+				FMath::DivideAndRoundUp(pBase->GetSizeY(), uint16(finfo.PixelsPerBlockY)));
 
 			UE::Math::TIntVector2<uint16> blockMipPos(
 				rect.min[0] / finfo.PixelsPerBlockX,
 				rect.min[1] / finfo.PixelsPerBlockY);
 
 			UE::Math::TIntVector2<uint16> blockMipSize(
-				rect.size[0] / finfo.PixelsPerBlockX,
-				rect.size[1] / finfo.PixelsPerBlockY);
+				FMath::DivideAndRoundUp(rect.size[0], uint16(finfo.PixelsPerBlockX)),
+				FMath::DivideAndRoundUp(rect.size[1], uint16(finfo.PixelsPerBlockY)));
 
 			int doneMips = 0;
 			for (; doneMips < pBase->GetLODCount() && doneMips < pBlock->GetLODCount(); ++doneMips)
diff --git a/Engine/Plugins/Experimental/Mutable/Source/MutableRuntime/Private/MuR/OpImageCrop.cpp b/Engine/Plugins/Experimental/Mutable/Source/MutableRuntime/Private/MuR/OpImageCrop.cpp
--- a/Engine/Plugins/Experimental/Mutable/Source/MutableRuntime/Private/MuR/OpImageCrop.cpp
+++ b/Engine/Plugins/Experimental/Mutable/Source/MutableRuntime/Private/MuR/OpImageCrop.cpp
@@ -52,8 +52,8 @@ namespace mu
 		checkf( Rect.size[0] % finfo.PixelsPerBlockX == 0, TEXT("Rect must snap to blocks.") );
 		checkf( Rect.size[1] % finfo.PixelsPerBlockY == 0, TEXT("Rect must snap to blocks.") );
 
-		int baseRowSize = finfo.BytesPerBlock * Base->GetSizeX() / finfo.PixelsPerBlockX;
-		int cropRowSize = finfo.BytesPerBlock * Rect.size[0] / finfo.PixelsPerBlockX;
+		int baseRowSize = finfo.BytesPerBlock * FMath::DivideAndRoundUp(Base->GetSizeX(), uint16(finfo.PixelsPerBlockX));
+		int cropRowSize = finfo.BytesPerBlock * FMath::DivideAndRoundUp(Rect.size[0], int32(finfo.PixelsPerBlockX));
 
         const uint8_t* pBaseBuf = Base->GetData();
         uint8_t* pCropBuf = Cropped->GetData();
diff --git a/Engine/Plugins/Experimental/Mutable/Source/MutableTools/Private/MuT/CodeGenerator.cpp b/Engine/Plugins/Experimental/Mutable/Source/MutableTools/Private/MuT/CodeGenerator.cpp
--- a/Engine/Plugins/Experimental/Mutable/Source/MutableTools/Private/MuT/CodeGenerator.cpp
+++ b/Engine/Plugins/Experimental/Mutable/Source/MutableTools/Private/MuT/CodeGenerator.cpp
@@ -1580,9 +1580,9 @@ namespace mu
 								// we will devide the layout block by the format block
 								const FImageFormatData& finfo = GetImageFormatData(desc.m_format);
 
-								int mipsX = (int)ceilf(logf((float)blockSizeX / finfo.PixelsPerBlockX) / logf(2.0f));
-								int mipsY = (int)ceilf(logf((float)blockSizeY / finfo.PixelsPerBlockY) / logf(2.0f));
-								mop->BlockLevels = (uint8_t)FMath::Max(mipsX, mipsY);
+								int32 mipsX = FMath::CeilLogTwo(blockSizeX / finfo.PixelsPerBlockX);
+								int32 mipsY = FMath::CeilLogTwo(blockSizeY / finfo.PixelsPerBlockY);
+								mop->BlockLevels = (uint8)FMath::Max(mipsX, mipsY);
 
 								mop->AddressMode = mipmapNode->GetPrivate()->m_settings.m_addressMode;
 								mop->FilterType = mipmapNode->GetPrivate()->m_settings.m_filterType;
