[Mutable] Add default values to Texture Parameters.

- Texture Parameter Node now has a Default Value and Reference Value (prerviously known as Default Value).
- Created a new Editor Image Provider. Returns all textures available in the project.
- In editor, default Textures are correctly applied when creating a default Instance.
- Removed deprecated Instance/Descriptor Texture Parameter functions.

#rnx

diff --git a/Engine/Plugins/Experimental/Mutable/Source/CustomizableObject/Classes/MuCO/CustomizableObjectInstance.h b/Engine/Plugins/Experimental/Mutable/Source/CustomizableObject/Classes/MuCO/CustomizableObjectInstance.h
--- a/Engine/Plugins/Experimental/Mutable/Source/CustomizableObject/Classes/MuCO/CustomizableObjectInstance.h
+++ b/Engine/Plugins/Experimental/Mutable/Source/CustomizableObject/Classes/MuCO/CustomizableObjectInstance.h
@@ -180,7 +180,6 @@ public:
 
 	// UObject interface.
 #if WITH_EDITOR
-	virtual void PreEditChange(FProperty* PropertyAboutToChange) override;
 	virtual void PostEditChangeProperty(FPropertyChangedEvent& PropertyChangedEvent) override;
 	virtual bool CanEditChange( const FProperty* InProperty ) const override;
 #endif //WITH_EDITOR
@@ -397,10 +396,6 @@ public:
 	UFUNCTION(BlueprintCallable, Category = CustomizableObjectInstance)
 	void SetTextureParameterSelectedOption(const FString& TextureParamName, const FString& TextureValue, int32 RangeIndex = -1);
 
-	/** @deprecated Use SetTextureParameterSelectedOption instead. */
-	UFUNCTION(BlueprintCallable, Category = CustomizableObjectInstance)
-	void SetTextureParameterSelectedOptionT(const FString& TextureParamName, UTexture2D* TextureValue, int32 RangeIndex = -1);
-	
 	// Gets the value of a color parameter with name "ColorParamName"
 	UFUNCTION(BlueprintCallable, Category = CustomizableObjectInstance)
 	FLinearColor GetColorParameterSelectedOption(const FString& ColorParamName) const;
@@ -684,6 +679,10 @@ public:
 	/** Return the UCustomizableObjectInstance::Descriptor hash on the last successful update. */
 	FDescriptorRuntimeHash GetUpdateDescriptorRuntimeHash() const;
 
+#if WITH_EDITOR
+	const TArray<TObjectPtr<UTexture2D>>& GetTextureParameterDeclarations() const;
+#endif
+	
 	// --------------------------------------------------------------------
 
 	/** Flag to know if a property of this instance changed in the editor */
@@ -754,7 +753,7 @@ private:
 	int32 CurrentMaxLOD = -1;
 
 #if WITH_EDITORONLY_DATA
-	/** Textures used in the Texture Parameters. */
+	/** Textures which can used as values in Texture Parameters. */
 	UPROPERTY(EditAnywhere, Category = TextureParameter)
 	TArray<TObjectPtr<UTexture2D>> TextureParameterDeclarations;
 
diff --git a/Engine/Plugins/Experimental/Mutable/Source/CustomizableObject/Classes/MuCO/CustomizableObjectSystem.h b/Engine/Plugins/Experimental/Mutable/Source/CustomizableObject/Classes/MuCO/CustomizableObjectSystem.h
--- a/Engine/Plugins/Experimental/Mutable/Source/CustomizableObject/Classes/MuCO/CustomizableObjectSystem.h
+++ b/Engine/Plugins/Experimental/Mutable/Source/CustomizableObject/Classes/MuCO/CustomizableObjectSystem.h
@@ -290,15 +290,6 @@ public:
 	/** [Texture Parameters] Remove all images from the cache. */
 	void ClearImageCache();
 
-	/** Initialize (if was not already) and get the default image provider. */
-	UDefaultImageProvider& GetOrCreateDefaultImageProvider();
-
-private:
-	/** Mutable default image provider. Used by the COIEditor and Instance/Descriptor APIs. */
-	UPROPERTY()
-	TObjectPtr<UDefaultImageProvider> DefaultImageProvider = nullptr;
-
-public:
 	// Show a warning on-screen and via a notification (if in Editor) and log an error when a CustomizableObject is
 	// being used and it's not compiled.  Callers can add additional information to the error log.
 	void AddUncompiledCOWarning(const UCustomizableObject& InObject, FString const* OptionalLogInfo = nullptr);
diff --git a/Engine/Plugins/Experimental/Mutable/Source/CustomizableObject/Private/MuCO/CustomizableObjectInstance.cpp b/Engine/Plugins/Experimental/Mutable/Source/CustomizableObject/Private/MuCO/CustomizableObjectInstance.cpp
--- a/Engine/Plugins/Experimental/Mutable/Source/CustomizableObject/Private/MuCO/CustomizableObjectInstance.cpp
+++ b/Engine/Plugins/Experimental/Mutable/Source/CustomizableObject/Private/MuCO/CustomizableObjectInstance.cpp
@@ -244,23 +244,6 @@ void UCustomizableInstancePrivateData::PrepareForUpdate(const TSharedPtr<FMutabl
 
 #if WITH_EDITOR
 
-void UCustomizableObjectInstance::PreEditChange(FProperty* PropertyAboutToChange)
-{
-	UObject::PreEditChange(PropertyAboutToChange);
-	
-	const FName PropertyName = PropertyAboutToChange ? PropertyAboutToChange->GetFName() : NAME_None;
-
-	if (PropertyName == GET_MEMBER_NAME_CHECKED(UCustomizableObjectInstance, TextureParameterDeclarations))
-	{
-		UDefaultImageProvider& DefaultImageProvider = UCustomizableObjectSystem::GetInstance()->GetOrCreateDefaultImageProvider();
-		
-		for (const TObjectPtr<UTexture2D>& Texture : TextureParameterDeclarations)
-		{
-			DefaultImageProvider.Remove(Texture);
-		}
-	}
-}
-
 
 void UCustomizableObjectInstance::PostEditChangeProperty(FPropertyChangedEvent& PropertyChangedEvent)
 {
@@ -272,13 +255,6 @@ void UCustomizableObjectInstance::PostEditChangeProperty(FPropertyChangedEvent&
 
 	if (PropertyName == GET_MEMBER_NAME_CHECKED(UCustomizableObjectInstance, TextureParameterDeclarations))
 	{
-		UDefaultImageProvider& DefaultImageProvider = UCustomizableObjectSystem::GetInstance()->GetOrCreateDefaultImageProvider();
-		
-		for (const TObjectPtr<UTexture2D>& Texture : TextureParameterDeclarations)
-		{
-			DefaultImageProvider.Add(Texture);
-		}
-
 		UpdateSkeletalMeshAsync(true, true);
 	}
 }
@@ -2456,12 +2432,6 @@ void UCustomizableObjectInstance::SetTextureParameterSelectedOption(const FStrin
 }
 
 
-void UCustomizableObjectInstance::SetTextureParameterSelectedOptionT(const FString& TextureParamName, UTexture2D* TextureValue, const int32 RangeIndex)
-{
-	Descriptor.SetTextureParameterSelectedOptionT(TextureParamName, TextureValue, RangeIndex);
-}
-
-
 FLinearColor UCustomizableObjectInstance::GetColorParameterSelectedOption(const FString& ColorParamName) const
 {
 	return Descriptor.GetColorParameterSelectedOption(ColorParamName);
@@ -2953,6 +2923,13 @@ FDescriptorRuntimeHash UCustomizableObjectInstance::GetUpdateDescriptorRuntimeHa
 	return UpdateDescriptorRuntimeHash;
 }
 
+#if WITH_EDITOR
+const TArray<TObjectPtr<UTexture2D>>& UCustomizableObjectInstance::GetTextureParameterDeclarations() const
+{
+	return TextureParameterDeclarations;
+}
+#endif
+
 
 UCustomizableInstancePrivateData* UCustomizableObjectInstance::GetPrivate() const
 { 
diff --git a/Engine/Plugins/Experimental/Mutable/Source/CustomizableObject/Private/MuCO/CustomizableObjectInstanceDescriptor.cpp b/Engine/Plugins/Experimental/Mutable/Source/CustomizableObject/Private/MuCO/CustomizableObjectInstanceDescriptor.cpp
--- a/Engine/Plugins/Experimental/Mutable/Source/CustomizableObject/Private/MuCO/CustomizableObjectInstanceDescriptor.cpp
+++ b/Engine/Plugins/Experimental/Mutable/Source/CustomizableObject/Private/MuCO/CustomizableObjectInstanceDescriptor.cpp
@@ -1395,18 +1395,6 @@ void FCustomizableObjectInstanceDescriptor::SetTextureParameterSelectedOption(co
 }
 
 
-void FCustomizableObjectInstanceDescriptor::SetTextureParameterSelectedOptionT(const FString& TextureParamName,	UTexture2D* TextureValue, int32 RangeIndex)
-{
-	check(CustomizableObject);
-	RETURN_ON_UNCOMPILED_CO(CustomizableObject, TEXT("Error: Cannot set Int parameter "));
-
-	UDefaultImageProvider& UDefaultImageProvider = UCustomizableObjectSystem::GetInstance()->GetOrCreateDefaultImageProvider();
-	const FString Id = UDefaultImageProvider.Add(TextureValue);
-	
-	SetTextureParameterSelectedOption(TextureParamName, Id, RangeIndex);
-}
-
-
 FLinearColor FCustomizableObjectInstanceDescriptor::GetColorParameterSelectedOption(const FString& ColorParamName) const
 {
 	check(CustomizableObject);
diff --git a/Engine/Plugins/Experimental/Mutable/Source/CustomizableObject/Private/MuCO/CustomizableObjectSystem.cpp b/Engine/Plugins/Experimental/Mutable/Source/CustomizableObject/Private/MuCO/CustomizableObjectSystem.cpp
--- a/Engine/Plugins/Experimental/Mutable/Source/CustomizableObject/Private/MuCO/CustomizableObjectSystem.cpp
+++ b/Engine/Plugins/Experimental/Mutable/Source/CustomizableObject/Private/MuCO/CustomizableObjectSystem.cpp
@@ -28,6 +28,7 @@
 #include "UObject/UObjectIterator.h"
 #include "Widgets/Notifications/SNotificationList.h"
 #include "ContentStreaming.h"
+#include "MuCO/EditorImageProvider.h"
 
 #if WITH_EDITOR
 #include "Editor.h"
@@ -39,6 +40,7 @@
 #endif
 
 
+
 #include UE_INLINE_GENERATED_CPP_BY_NAME(CustomizableObjectSystem)
 
 class AActor;
@@ -385,6 +387,12 @@ void UCustomizableObjectSystem::InitSystem()
 	Private->ImageProvider = Provider;
 	Private->MutableSystem->SetImageParameterGenerator(Provider);
 
+#if WITH_EDITORONLY_DATA
+	Private->EditorImageProvider = NewObject<UEditorImageProvider>();
+	check(Private->EditorImageProvider);
+	RegisterImageProvider(Private->EditorImageProvider);
+#endif
+	
 #if WITH_EDITOR
 	if (!IsRunningGame())
 	{
@@ -523,6 +531,10 @@ void FCustomizableObjectSystemPrivate::AddReferencedObjects(FReferenceCollector&
 	{
 		Collector.AddReferencedObject(CurrentInstanceBeingUpdated);
 	}
+
+#if WITH_EDITORONLY_DATA
+	Collector.AddReferencedObject(EditorImageProvider);
+#endif
 }
 
 
@@ -2959,18 +2971,6 @@ void UCustomizableObjectSystem::UnregisterImageProvider(UCustomizableSystemImage
 }
 
 
-
-UDefaultImageProvider& UCustomizableObjectSystem::GetOrCreateDefaultImageProvider()
-{
-	if (!DefaultImageProvider)
-	{
-		DefaultImageProvider = NewObject<UDefaultImageProvider>();
-		RegisterImageProvider(DefaultImageProvider);
-	}
-
-	return *DefaultImageProvider;
-}
-
 static bool bRevertCacheTextureParameters = false;
 static FAutoConsoleVariableRef CVarRevertCacheTextureParameters(
 	TEXT("mutable.RevertCacheTextureParameters"), bRevertCacheTextureParameters,
diff --git a/Engine/Plugins/Experimental/Mutable/Source/CustomizableObject/Private/MuCO/CustomizableObjectSystemPrivate.h b/Engine/Plugins/Experimental/Mutable/Source/CustomizableObject/Private/MuCO/CustomizableObjectSystemPrivate.h
--- a/Engine/Plugins/Experimental/Mutable/Source/CustomizableObject/Private/MuCO/CustomizableObjectSystemPrivate.h
+++ b/Engine/Plugins/Experimental/Mutable/Source/CustomizableObject/Private/MuCO/CustomizableObjectSystemPrivate.h
@@ -24,6 +24,10 @@
 	#include "Async/TaskGraphInterfaces.h"
 #endif
 
+#if WITH_EDITORONLY_DATA
+class UEditorImageProvider;
+#endif
+
 class UCustomizableObjectSystem;
 namespace LowLevelTasks { enum class ETaskPriority : int8; }
 struct FTexturePlatformData;
@@ -766,5 +770,10 @@ public:
 	
 	/** Mutable TaskGraph system (Mutable Thread). */
 	FMutableTaskGraph MutableTaskGraph;
+	
+#if WITH_EDITORONLY_DATA
+	/** Mutable default image provider. Used by the COIEditor and Instance/Descriptor APIs. */
+	TObjectPtr<UEditorImageProvider> EditorImageProvider = nullptr;
+#endif
 };
 
diff --git a/Engine/Plugins/Experimental/Mutable/Source/CustomizableObject/Private/MuCO/DefaultImageProvider.cpp b/Engine/Plugins/Experimental/Mutable/Source/CustomizableObject/Private/MuCO/DefaultImageProvider.cpp
--- a/Engine/Plugins/Experimental/Mutable/Source/CustomizableObject/Private/MuCO/DefaultImageProvider.cpp
+++ b/Engine/Plugins/Experimental/Mutable/Source/CustomizableObject/Private/MuCO/DefaultImageProvider.cpp
@@ -69,7 +69,7 @@ FString UDefaultImageProvider::Add(UTexture2D* Texture)
 		return FString();
 	}
 
-	const FName Id(Texture->GetFullName());
+	const FName Id(Texture->GetPathName());
 	Textures.Add(Id, Texture);	
 
 	return Id.ToString();
diff --git a/Engine/Plugins/Experimental/Mutable/Source/CustomizableObject/Private/MuCO/EditorImageProvider.cpp b/Engine/Plugins/Experimental/Mutable/Source/CustomizableObject/Private/MuCO/EditorImageProvider.cpp
new file mode 100644
--- /dev/null
+++ b/Engine/Plugins/Experimental/Mutable/Source/CustomizableObject/Private/MuCO/EditorImageProvider.cpp
@@ -0,0 +1,22 @@
+// Copyright Epic Games, Inc. All Rights Reserved.
+
+
+#include "MuCO/EditorImageProvider.h" 
+
+#include "MuCO/CustomizableObjectSystem.h"
+#include "Engine/Texture2D.h"
+
+
+UCustomizableSystemImageProvider::ValueType UEditorImageProvider::HasTextureParameterValue(const FName& ID)
+{
+	const UTexture2D* Texture = Cast<UTexture2D>(FSoftObjectPath(ID.ToString()).TryLoad());
+	
+	return Texture ? ValueType::Unreal_Deferred : ValueType::None;
+}
+
+
+UTexture2D* UEditorImageProvider::GetTextureParameterValue(const FName& ID)
+{
+	return Cast<UTexture2D>(FSoftObjectPath(ID.ToString()).TryLoad());;
+}
+
diff --git a/Engine/Plugins/Experimental/Mutable/Source/CustomizableObject/Public/MuCO/CustomizableObjectCustomVersion.h b/Engine/Plugins/Experimental/Mutable/Source/CustomizableObject/Public/MuCO/CustomizableObjectCustomVersion.h
--- a/Engine/Plugins/Experimental/Mutable/Source/CustomizableObject/Public/MuCO/CustomizableObjectCustomVersion.h
+++ b/Engine/Plugins/Experimental/Mutable/Source/CustomizableObject/Public/MuCO/CustomizableObjectCustomVersion.h
@@ -96,6 +96,8 @@ struct CUSTOMIZABLEOBJECT_API FCustomizableObjectCustomVersion
 
 		TextureCompressionEnum,
 
+		NodeTextureParameterDefaultToReferenceValue,
+
 		// -----<new versions can be added above this line>-------------------------------------------------
 		VersionPlusOne,
 		LatestVersion = VersionPlusOne - 1
diff --git a/Engine/Plugins/Experimental/Mutable/Source/CustomizableObject/Public/MuCO/CustomizableObjectInstanceDescriptor.h b/Engine/Plugins/Experimental/Mutable/Source/CustomizableObject/Public/MuCO/CustomizableObjectInstanceDescriptor.h
--- a/Engine/Plugins/Experimental/Mutable/Source/CustomizableObject/Public/MuCO/CustomizableObjectInstanceDescriptor.h
+++ b/Engine/Plugins/Experimental/Mutable/Source/CustomizableObject/Public/MuCO/CustomizableObjectInstanceDescriptor.h
@@ -132,9 +132,6 @@ struct CUSTOMIZABLEOBJECT_API FCustomizableObjectInstanceDescriptor
 
 	/** Sets the texture value "TextureValue" of a texture parameter with index "TextureParamIndex". */
 	void SetTextureParameterSelectedOption(const FString& TextureParamName, const FString& TextureValue, int32 RangeIndex);
-
-	/** @deprecated Use SetTextureParameterSelectedOption instead. */
-	void SetTextureParameterSelectedOptionT(const FString& TextureParamName, UTexture2D* TextureValue, int32 RangeIndex);
 	
 	/** Gets the value of a color parameter with name "ColorParamName". */
 	FLinearColor GetColorParameterSelectedOption(const FString& ColorParamName) const;
diff --git a/Engine/Plugins/Experimental/Mutable/Source/CustomizableObject/Public/MuCO/EditorImageProvider.h b/Engine/Plugins/Experimental/Mutable/Source/CustomizableObject/Public/MuCO/EditorImageProvider.h
new file mode 100644
--- /dev/null
+++ b/Engine/Plugins/Experimental/Mutable/Source/CustomizableObject/Public/MuCO/EditorImageProvider.h
@@ -0,0 +1,23 @@
+// Copyright Epic Games, Inc. All Rights Reserved.
+
+#pragma once
+
+
+#include "MuCO/CustomizableObjectSystem.h"
+
+#include "EditorImageProvider.generated.h"
+
+class UTexture2D;
+class UCustomizableObjectInstance;
+
+
+UCLASS()
+class CUSTOMIZABLEOBJECT_API UEditorImageProvider : public UCustomizableSystemImageProvider
+{
+	GENERATED_BODY()
+
+public:
+	// UCustomizableSystemImageProvider interface
+	virtual ValueType HasTextureParameterValue(const FName& ID) override;
+	virtual UTexture2D* GetTextureParameterValue(const FName& ID) override;
+};
\ No newline at end of file
diff --git a/Engine/Plugins/Experimental/Mutable/Source/CustomizableObjectEditor/Classes/MuCOE/Nodes/CustomizableObjectNodeTextureParameter.h b/Engine/Plugins/Experimental/Mutable/Source/CustomizableObjectEditor/Classes/MuCOE/Nodes/CustomizableObjectNodeTextureParameter.h
--- a/Engine/Plugins/Experimental/Mutable/Source/CustomizableObjectEditor/Classes/MuCOE/Nodes/CustomizableObjectNodeTextureParameter.h
+++ b/Engine/Plugins/Experimental/Mutable/Source/CustomizableObjectEditor/Classes/MuCOE/Nodes/CustomizableObjectNodeTextureParameter.h
@@ -19,9 +19,14 @@ class CUSTOMIZABLEOBJECTEDITOR_API UCustomizableObjectNodeTextureParameter : pub
 public:
 	GENERATED_BODY()
 
+	/** Default value of the parameter. */
 	UPROPERTY(EditAnywhere, Category=CustomizableObject)
 	TObjectPtr<UTexture2D> DefaultValue;
 
+	/** Reference Texture where this parameter copies some properties from. */
+	UPROPERTY(EditAnywhere, Category=CustomizableObject)
+	TObjectPtr<UTexture2D> ReferenceValue;
+
 	UPROPERTY(EditAnywhere, Category=CustomizableObject)
 	FString ParameterName = "Default Name";
 
@@ -36,13 +41,14 @@ public:
 	UPROPERTY(EditAnywhere, Category = CustomizableObject)
 	int32 TextureSizeY = 0;
 
-	// Begin EdGraphNode interface
-	FText GetNodeTitle(ENodeTitleType::Type TitleType) const override;
-	FLinearColor GetNodeTitleColor() const override;
-	FText GetTooltipText() const override;
+	// EdGraphNode interface
+	virtual FText GetNodeTitle(ENodeTitleType::Type TitleType) const override;
+	virtual FLinearColor GetNodeTitleColor() const override;
+	virtual FText GetTooltipText() const override;
 
 	// UCustomizableObjectNode interface
-	void AllocateDefaultPins(UCustomizableObjectNodeRemapPins* RemapPins) override;
+	virtual void BackwardsCompatibleFixup() override;
+	virtual void AllocateDefaultPins(UCustomizableObjectNodeRemapPins* RemapPins) override;
 	virtual bool IsAffectedByLOD() const override { return false; }
 };
 
diff --git a/Engine/Plugins/Experimental/Mutable/Source/CustomizableObjectEditor/Private/MuCOE/GenerateMutableSource/GenerateMutableSource.cpp b/Engine/Plugins/Experimental/Mutable/Source/CustomizableObjectEditor/Private/MuCOE/GenerateMutableSource/GenerateMutableSource.cpp
--- a/Engine/Plugins/Experimental/Mutable/Source/CustomizableObjectEditor/Private/MuCOE/GenerateMutableSource/GenerateMutableSource.cpp
+++ b/Engine/Plugins/Experimental/Mutable/Source/CustomizableObjectEditor/Private/MuCOE/GenerateMutableSource/GenerateMutableSource.cpp
@@ -383,7 +383,7 @@ UTexture2D* FindReferenceImage(const UEdGraphPin* Pin, FMutableGraphGenerationCo
 
 	else if (const UCustomizableObjectNodeTextureParameter* ParamNodeTex = Cast<UCustomizableObjectNodeTextureParameter>(Node))
 	{
-		Result = ParamNodeTex->DefaultValue;
+		Result = ParamNodeTex->ReferenceValue;
 	}
 
 	else if (const UCustomizableObjectNodeMesh* TypedNodeMesh = Cast<UCustomizableObjectNodeMesh>(Node))
diff --git a/Engine/Plugins/Experimental/Mutable/Source/CustomizableObjectEditor/Private/MuCOE/GenerateMutableSource/GenerateMutableSourceImage.cpp b/Engine/Plugins/Experimental/Mutable/Source/CustomizableObjectEditor/Private/MuCOE/GenerateMutableSource/GenerateMutableSourceImage.cpp
--- a/Engine/Plugins/Experimental/Mutable/Source/CustomizableObjectEditor/Private/MuCOE/GenerateMutableSource/GenerateMutableSourceImage.cpp
+++ b/Engine/Plugins/Experimental/Mutable/Source/CustomizableObjectEditor/Private/MuCOE/GenerateMutableSource/GenerateMutableSourceImage.cpp
@@ -183,8 +183,10 @@ mu::NodeImagePtr GenerateMutableSourceImage(const UEdGraphPin* Pin, FMutableGrap
 		TextureNode->SetName(StringCast<ANSICHAR>(*TypedNodeParam->ParameterName).Get());
 		TextureNode->SetUid(StringCast<ANSICHAR>(*GenerationContext.GetNodeIdUnique(Node).ToString()).Get());
 
-		// TODO: Set a default value for the texture parameter?
-		//ColorNode->SetDefaultValue(TypedNodeParam->DefaultValue.R, TypedNodeParam->DefaultValue.G, TypedNodeParam->DefaultValue.B);
+		if (TypedNodeParam->DefaultValue)
+		{
+			TextureNode->SetDefaultValue(FName(TypedNodeParam->DefaultValue->GetPathName()));			
+		}
 
 		GenerationContext.ParameterUIDataMap.Add(TypedNodeParam->ParameterName, FParameterUIData(
 			TypedNodeParam->ParameterName,
@@ -209,7 +211,7 @@ mu::NodeImagePtr GenerateMutableSourceImage(const UEdGraphPin* Pin, FMutableGrap
 		ResizeNode->SetBase(FormatNode);
 		ResizeNode->SetRelative(false);
 
-		const UTexture2D* ReferenceTexture = TypedNodeParam->DefaultValue;
+		const UTexture2D* ReferenceTexture = TypedNodeParam->ReferenceValue;
 		if (ReferenceTexture)
 		{
 			ResizeNode->SetSize(FMath::Max(ReferenceTexture->GetImportedSize().X,1), FMath::Max(ReferenceTexture->GetImportedSize().X, 1));
diff --git a/Engine/Plugins/Experimental/Mutable/Source/CustomizableObjectEditor/Private/MuCOE/Nodes/CustomizableObjectNodeTextureParameter.cpp b/Engine/Plugins/Experimental/Mutable/Source/CustomizableObjectEditor/Private/MuCOE/Nodes/CustomizableObjectNodeTextureParameter.cpp
--- a/Engine/Plugins/Experimental/Mutable/Source/CustomizableObjectEditor/Private/MuCOE/Nodes/CustomizableObjectNodeTextureParameter.cpp
+++ b/Engine/Plugins/Experimental/Mutable/Source/CustomizableObjectEditor/Private/MuCOE/Nodes/CustomizableObjectNodeTextureParameter.cpp
@@ -37,5 +37,19 @@ FText UCustomizableObjectNodeTextureParameter::GetTooltipText() const
 	return LOCTEXT("Texture_Parameter_Tooltip", "Expose a runtime modifiable texture parameter from the Customizable Object.");
 }
 
+
+void UCustomizableObjectNodeTextureParameter::BackwardsCompatibleFixup()
+{
+	Super::BackwardsCompatibleFixup();
+
+	const int32 CustomizableObjectCustomVersion = GetLinkerCustomVersion(FCustomizableObjectCustomVersion::GUID);
+
+	if (CustomizableObjectCustomVersion < FCustomizableObjectCustomVersion::NodeTextureParameterDefaultToReferenceValue)
+	{
+		ReferenceValue = DefaultValue;
+		DefaultValue = {};
+	}	
+}
+
 #undef LOCTEXT_NAMESPACE
 
diff --git a/Engine/Plugins/Experimental/Mutable/Source/CustomizableObjectEditor/Private/MuCOE/SCustomizableInstanceProperties.cpp b/Engine/Plugins/Experimental/Mutable/Source/CustomizableObjectEditor/Private/MuCOE/SCustomizableInstanceProperties.cpp
--- a/Engine/Plugins/Experimental/Mutable/Source/CustomizableObjectEditor/Private/MuCOE/SCustomizableInstanceProperties.cpp
+++ b/Engine/Plugins/Experimental/Mutable/Source/CustomizableObjectEditor/Private/MuCOE/SCustomizableInstanceProperties.cpp
@@ -76,9 +76,20 @@ void SCustomizableInstanceProperties::Construct(const FArguments& InArgs)
 
 		StateNames.Sort(&CompareNames);
 
-		// Store the texture parameters data required for the ui.
-		TextureParameterValueNames.Add(MakeShareable(new FString("None")));
-		TextureParameterValues.Add(FName());
+		// Store which Texture Parameter values can be selected.
+		// Get default values.
+		for (const FCustomizableObjectTextureParameterValue& TextureParameter : CustomInstance->GetTextureParameters())
+		{
+			const FName DefaultValue = CustomizableObject->GetTextureParameterDefaultValue(TextureParameter.ParameterName);
+			
+			if (!TextureParameterValues.Contains(DefaultValue))
+			{
+				TextureParameterValueNames.Add(MakeShareable(new FString(DefaultValue.ToString())));
+				TextureParameterValues.Add(DefaultValue);	
+			}
+		}
+
+		// Get values from registered providers providers.
 		TArray<FCustomizableObjectExternalTexture> Textures = UCustomizableObjectSystem::GetInstance()->GetTextureParameterValues();
 		for (int i = 0; i < Textures.Num(); ++i)
 		{
@@ -86,6 +97,18 @@ void SCustomizableInstanceProperties::Construct(const FArguments& InArgs)
 			TextureParameterValues.Add(Textures[i].Value);
 		}
 
+		// Get values from TextureParameterDeclarations.
+		for (TObjectPtr<UTexture2D> Declaration : CustomInstance->GetTextureParameterDeclarations())
+		{
+			if (!Declaration)
+			{
+				continue;
+			}
+			
+			TextureParameterValueNames.Add(MakeShareable(new FString(Declaration.GetName())));
+			TextureParameterValues.Add(FName(Declaration.GetPathName()));
+		}
+
 		SetParameterProfileNamesOnEditor();
 		
 		int32 ProfileStringIndex = 0;
diff --git a/Engine/Plugins/Experimental/Mutable/Source/MutableRuntime/Private/MuR/Serialisation.cpp b/Engine/Plugins/Experimental/Mutable/Source/MutableRuntime/Private/MuR/Serialisation.cpp
--- a/Engine/Plugins/Experimental/Mutable/Source/MutableRuntime/Private/MuR/Serialisation.cpp
+++ b/Engine/Plugins/Experimental/Mutable/Source/MutableRuntime/Private/MuR/Serialisation.cpp
@@ -40,16 +40,15 @@ namespace mu
 	MUTABLE_IMPLEMENT_POD_SERIALISABLE(UE::Math::TIntVector2<uint16>);                
 	MUTABLE_IMPLEMENT_POD_SERIALISABLE(UE::Math::TIntVector2<int16>); 
 
-	template <>
-    void operator<<<FString>(OutputArchive& arch, const FString& t)
+	
+    void operator<<(OutputArchive& arch, const FString& t)
     {
 	    const TArray<TCHAR>& Data = t.GetCharArray();
     	arch << Data;
     }
 
 
-    template <>
-    void operator>><FString>(InputArchive& arch, FString& t)
+    void operator>>(InputArchive& arch, FString& t)
     {
 	    TArray<TCHAR> Data;
     	arch >> Data;
@@ -57,6 +56,21 @@ namespace mu
     	t = FString(Data.GetData()); // Construct from raw pointer to avoid double zero terminating character
     }
 
+
+    void operator<<(OutputArchive& arch, const FName& v)
+    {
+	    arch << v.ToString();
+    }
+
+
+    void operator>>(InputArchive& arch, FName& v)
+    {
+	    FString Temp;
+	    arch >> Temp;
+	    v = FName(Temp);
+    }
+
+
     //---------------------------------------------------------------------------------------------
     //---------------------------------------------------------------------------------------------
     //---------------------------------------------------------------------------------------------
diff --git a/Engine/Plugins/Experimental/Mutable/Source/MutableRuntime/Public/MuR/Serialisation.h b/Engine/Plugins/Experimental/Mutable/Source/MutableRuntime/Public/MuR/Serialisation.h
--- a/Engine/Plugins/Experimental/Mutable/Source/MutableRuntime/Public/MuR/Serialisation.h
+++ b/Engine/Plugins/Experimental/Mutable/Source/MutableRuntime/Public/MuR/Serialisation.h
@@ -373,11 +373,15 @@ namespace mu
 	MUTABLE_DEFINE_POD_VECTOR_SERIALISABLE(FVector4f);
 	
 	//---------------------------------------------------------------------------------------------
-	template<>
-	void operator<< <FString>(OutputArchive& arch, const FString& t);
+	void operator<< (OutputArchive& arch, const FString& t);
 	
-	template<>
-	void operator>> <FString>(InputArchive& arch, FString& t);
+	void operator>> (InputArchive& arch, FString& t);
+
+
+	//---------------------------------------------------------------------------------------------
+	MUTABLERUNTIME_API void operator<<(OutputArchive& arch, const FName& v);
+
+	MUTABLERUNTIME_API void operator>>(InputArchive& arch, FName& v);
 
 
 	//---------------------------------------------------------------------------------------------
@@ -405,21 +409,6 @@ namespace mu
 			Arch >> V[Index];
 		}
 	}
-
-	//---------------------------------------------------------------------------------------------
-	template<> 
-	inline void operator<<(OutputArchive& arch, const FName& v)
-	{
-		arch << v.ToString();
-	}
-
-	template<> 
-	inline void operator>>(InputArchive& arch, FName& v)
-	{
-		FString Temp;
-		arch >> Temp;
-		v = FName(Temp);
-	}
-
+	
 }
 
diff --git a/Engine/Plugins/Experimental/Mutable/Source/MutableTools/Private/MuT/CodeGenerator_Image.cpp b/Engine/Plugins/Experimental/Mutable/Source/MutableTools/Private/MuT/CodeGenerator_Image.cpp
--- a/Engine/Plugins/Experimental/Mutable/Source/MutableTools/Private/MuT/CodeGenerator_Image.cpp
+++ b/Engine/Plugins/Experimental/Mutable/Source/MutableTools/Private/MuT/CodeGenerator_Image.cpp
@@ -278,7 +278,7 @@ namespace mu
 			op->parameter.m_name = node.m_name;
 			op->parameter.m_uid = node.m_uid;
             op->parameter.m_type = PARAMETER_TYPE::T_IMAGE;
-        	op->parameter.m_defaultValue.Set<ParamImageType>(FName());
+        	op->parameter.m_defaultValue.Set<ParamImageType>(node.m_defaultValue);
 
 			// Generate the code for the ranges
 			for (int32 a = 0; a < node.m_ranges.Num(); ++a)
diff --git a/Engine/Plugins/Experimental/Mutable/Source/MutableTools/Private/MuT/NodeImageParameter.cpp b/Engine/Plugins/Experimental/Mutable/Source/MutableTools/Private/MuT/NodeImageParameter.cpp
--- a/Engine/Plugins/Experimental/Mutable/Source/MutableTools/Private/MuT/NodeImageParameter.cpp
+++ b/Engine/Plugins/Experimental/Mutable/Source/MutableTools/Private/MuT/NodeImageParameter.cpp
@@ -91,6 +91,20 @@ namespace mu
 		}
 	}
 
+	
+	//---------------------------------------------------------------------------------------------
+	FName NodeImageParameter::GetDefaultValue() const
+	{
+		return m_pD->m_defaultValue;
+	}
+
+	
+	//---------------------------------------------------------------------------------------------
+	void NodeImageParameter::SetDefaultValue(FName Value)
+	{
+    	m_pD->m_defaultValue = Value;
+	}
+
 
 	//---------------------------------------------------------------------------------------------
 	void NodeImageParameter::SetRangeCount(int i)
diff --git a/Engine/Plugins/Experimental/Mutable/Source/MutableTools/Private/MuT/NodeImageParameterPrivate.h b/Engine/Plugins/Experimental/Mutable/Source/MutableTools/Private/MuT/NodeImageParameterPrivate.h
--- a/Engine/Plugins/Experimental/Mutable/Source/MutableTools/Private/MuT/NodeImageParameterPrivate.h
+++ b/Engine/Plugins/Experimental/Mutable/Source/MutableTools/Private/MuT/NodeImageParameterPrivate.h
@@ -20,6 +20,7 @@ namespace mu
 
 		static NODE_TYPE s_type;
 
+    	FName m_defaultValue;
 		string m_name;
 		string m_uid;
 
@@ -28,9 +29,10 @@ namespace mu
 		//!
 		void Serialise( OutputArchive& arch ) const
 		{
-            uint32 ver = 2;
+            uint32 ver = 3;
 			arch << ver;
 
+			arch << m_defaultValue;
 			arch << m_name;
 			arch << m_uid;
 			arch << m_ranges;
@@ -41,8 +43,13 @@ namespace mu
 		{
             uint32 ver;
 			arch >> ver;
-            check(ver<=2);
+            check(ver <= 3);
 
+			if (ver >= 3)
+			{
+				arch >> m_defaultValue;
+			}
+			
 			arch >> m_name;
             arch >> m_uid;
 
diff --git a/Engine/Plugins/Experimental/Mutable/Source/MutableTools/Public/MuT/NodeImageParameter.h b/Engine/Plugins/Experimental/Mutable/Source/MutableTools/Public/MuT/NodeImageParameter.h
--- a/Engine/Plugins/Experimental/Mutable/Source/MutableTools/Public/MuT/NodeImageParameter.h
+++ b/Engine/Plugins/Experimental/Mutable/Source/MutableTools/Public/MuT/NodeImageParameter.h
@@ -58,6 +58,10 @@ namespace mu
 		//! Set the uid of the parameter.
 		void SetUid( const char* );
 
+    	//! Get the default value of the parameter.
+    	FName GetDefaultValue() const;
+    	void SetDefaultValue(FName Value);
+    	
 		//! Set the number of ranges (dimensions) for this parameter.
 		//! By default a parameter has 0 ranges, meaning it only has one value.
 		void SetRangeCount(int Index);
