[Mutable] Fix FImageGenerationOptions::RectSize not being zero initialized by default.

- This could cause large image memory allocations in some operations randomly, for instance IM_PLAINCOLOUR.

#rnx

diff --git a/Engine/Plugins/Experimental/Mutable/Source/MutableTools/Private/MuT/CodeGenerator.h b/Engine/Plugins/Experimental/Mutable/Source/MutableTools/Private/MuT/CodeGenerator.h
--- a/Engine/Plugins/Experimental/Mutable/Source/MutableTools/Private/MuT/CodeGenerator.h
+++ b/Engine/Plugins/Experimental/Mutable/Source/MutableTools/Private/MuT/CodeGenerator.h
@@ -340,7 +340,7 @@ namespace mu
 			CompilerOptions::TextureLayoutStrategy ImageLayoutStrategy = CompilerOptions::TextureLayoutStrategy::None;
 
 			/** This is used to introduce additional image generation safety. \TODO: Move this "safety" to optimization? */
-			UE::Math::TIntVector2<int32> RectSize;
+			UE::Math::TIntVector2<int32> RectSize = {0, 0};
 
 			/** */
 			int32 CurrentStateIndex = -1;
