[mutable] Fix possible crash on CO compilation because of mismatched instruction types.
#rnx

diff --git a/Engine/Plugins/Experimental/Mutable/Source/MutableTools/Private/MuT/ASTOpImageLayer.cpp b/Engine/Plugins/Experimental/Mutable/Source/MutableTools/Private/MuT/ASTOpImageLayer.cpp
--- a/Engine/Plugins/Experimental/Mutable/Source/MutableTools/Private/MuT/ASTOpImageLayer.cpp
+++ b/Engine/Plugins/Experimental/Mutable/Source/MutableTools/Private/MuT/ASTOpImageLayer.cpp
@@ -218,7 +218,7 @@ namespace mu
 				else if (FMath::IsNearlyEqual(colour[0], 1, UE_SMALL_NUMBER))
 				{
 					// If the mask is white, we can remove it
-					Ptr<ASTOpImageLayerColor> NewOp = mu::Clone<ASTOpImageLayerColor>(this);
+					Ptr<ASTOpImageLayer> NewOp = mu::Clone<ASTOpImageLayer>(this);
 					NewOp->mask = nullptr;
 					at = NewOp;
 				}
