#rnx

diff --git a/Engine/Plugins/Experimental/Mutable/Source/CustomizableObject/Private/MuCO/CustomizableObjectSystem.cpp b/Engine/Plugins/Experimental/Mutable/Source/CustomizableObject/Private/MuCO/CustomizableObjectSystem.cpp
--- a/Engine/Plugins/Experimental/Mutable/Source/CustomizableObject/Private/MuCO/CustomizableObjectSystem.cpp
+++ b/Engine/Plugins/Experimental/Mutable/Source/CustomizableObject/Private/MuCO/CustomizableObjectSystem.cpp
@@ -1701,9 +1701,17 @@ namespace impl
 					// Image references are just references to texture assets and require no work at all
 					if (!MutableImage->IsReference())
 					{
-						FTexturePlatformData* PlatformData = UCustomizableInstancePrivateData::MutableCreateImagePlatformData(MutableImage, -1, Image.FullImageSizeX, Image.FullImageSizeY);
-						OperationData->ImageToPlatformDataMap.Add(Image.ImageID, PlatformData);
-						OperationData->PendingTextureCoverageQueries.Add({ KeyName, Surface.MaterialIndex, PlatformData });
+						if (!OperationData->ImageToPlatformDataMap.Contains(Image.ImageID))
+						{
+							FTexturePlatformData* PlatformData = UCustomizableInstancePrivateData::MutableCreateImagePlatformData(MutableImage, -1, Image.FullImageSizeX, Image.FullImageSizeY);
+							OperationData->ImageToPlatformDataMap.Add(Image.ImageID, PlatformData);
+							OperationData->PendingTextureCoverageQueries.Add({ KeyName, Surface.MaterialIndex, PlatformData });
+						}
+						else
+						{
+							// The ImageID already exists in the ImageToPlatformDataMap, that means the equivalent surface in a lower
+							// LOD already created the PlatformData for that ImageID and added it to the ImageToPlatformDataMap.
+						}
 					}
 				}
 			}
