[mutable] Fix unloaded passthrough textures getting treated as mutable generated

#rnx

diff --git a/Engine/Plugins/Experimental/Mutable/Source/CustomizableObject/Private/MuCO/CustomizableObjectInstance.cpp b/Engine/Plugins/Experimental/Mutable/Source/CustomizableObject/Private/MuCO/CustomizableObjectInstance.cpp
--- a/Engine/Plugins/Experimental/Mutable/Source/CustomizableObject/Private/MuCO/CustomizableObjectInstance.cpp
+++ b/Engine/Plugins/Experimental/Mutable/Source/CustomizableObject/Private/MuCO/CustomizableObjectInstance.cpp
@@ -5771,7 +5771,7 @@ void UCustomizableInstancePrivateData::BuildMaterials(const TSharedPtr<FMutableO
 								TSoftObjectPtr<UTexture> Ref = CustomizableObject->ReferencedPassThroughTextures[ReferenceID];
 
 								// \TODO: This will force the load of the reference texture, potentially causing a hich. 
-								PassThroughTexture = Ref.Get();
+								PassThroughTexture = Ref.LoadSynchronous();
 							}
 							else
 							{
