[Mutable] Fix crash when generating the imported model of skeletal meshes with MinLOD > 0.

#rnx

diff --git a/Engine/Plugins/Experimental/Mutable/Source/CustomizableObject/Private/MuCO/UnrealBakeHelpers.cpp b/Engine/Plugins/Experimental/Mutable/Source/CustomizableObject/Private/MuCO/UnrealBakeHelpers.cpp
--- a/Engine/Plugins/Experimental/Mutable/Source/CustomizableObject/Private/MuCO/UnrealBakeHelpers.cpp
+++ b/Engine/Plugins/Experimental/Mutable/Source/CustomizableObject/Private/MuCO/UnrealBakeHelpers.cpp
@@ -693,11 +693,14 @@ void FUnrealBakeHelpers::BakeHelper_RegenerateImportedModel(USkeletalMesh* Skele
 		ImportedModel->LODModels[LODIndex].NumVertices = LODModel.GetNumVertices();
 
 		// Indices
-		int indexCount = LODModel.MultiSizeIndexContainer.GetIndexBuffer()->Num();
-		ImportedModel->LODModels[LODIndex].IndexBuffer.SetNum(indexCount);
-		for (int i = 0; i < indexCount; ++i)
+		if (LODModel.MultiSizeIndexContainer.IsIndexBufferValid())
 		{
-			ImportedModel->LODModels[LODIndex].IndexBuffer[i] = LODModel.MultiSizeIndexContainer.GetIndexBuffer()->Get(i);
+			const int32 NumIndices = LODModel.MultiSizeIndexContainer.GetIndexBuffer()->Num();
+			ImportedModel->LODModels[LODIndex].IndexBuffer.SetNum(NumIndices);
+			for (int32 Index = 0; Index < NumIndices; ++Index)
+			{
+				ImportedModel->LODModels[LODIndex].IndexBuffer[Index] = LODModel.MultiSizeIndexContainer.GetIndexBuffer()->Get(Index);
+			}
 		}
 
 		ImportedModel->LODModels[LODIndex].Sections.SetNum(LODModel.RenderSections.Num());
