[Mutable] Fix crash in UpdateSkeletalMesh_PostBeginUpdate0 when generating a mesh with invalid LODs.

#rnx

diff --git a/Engine/Plugins/Experimental/Mutable/Source/CustomizableObject/Private/MuCO/CustomizableObjectInstance.cpp b/Engine/Plugins/Experimental/Mutable/Source/CustomizableObject/Private/MuCO/CustomizableObjectInstance.cpp
--- a/Engine/Plugins/Experimental/Mutable/Source/CustomizableObject/Private/MuCO/CustomizableObjectInstance.cpp
+++ b/Engine/Plugins/Experimental/Mutable/Source/CustomizableObject/Private/MuCO/CustomizableObjectInstance.cpp
@@ -128,7 +128,7 @@ void UCustomizableInstancePrivateData::InvalidateGeneratedData()
 {
 	for (FCustomizableInstanceComponentData& ComponentData : ComponentsData)
 	{
-		ComponentData.LastMeshIdPerLOD.Init(MAX_uint64, NumLODsAvailable);
+		ComponentData.LastMeshIdPerLOD.Init(MAX_uint64, MAX_MESH_LOD_COUNT);
 	}
 
 	LastUpdateData.Clear();
@@ -218,7 +218,7 @@ void UCustomizableInstancePrivateData::PrepareForUpdate(const TSharedPtr<FMutabl
 
 			FCustomizableInstanceComponentData& NewComponentData = ComponentsData.AddDefaulted_GetRef();
 			NewComponentData.ComponentIndex = Component.Id;
-			NewComponentData.LastMeshIdPerLOD.Init(MAX_uint64, NumLODsAvailable);
+			NewComponentData.LastMeshIdPerLOD.Init(MAX_uint64, MAX_MESH_LOD_COUNT);
 		}
 	}
 
@@ -1752,109 +1752,201 @@ void SetMeshUVChannelDensity(FMeshUVChannelInfo& UVChannelInfo, float Density =
 	}
 }
 
+// TODO PRP: MTBL-1653 Remove CVar
+static bool bApplyFixDoComponentsNeedUpdate = true;
+FAutoConsoleVariableRef CVarApplyFixDoComponentsNeedUpdate(
+	TEXT("Mutable.ApplyFixDoComponentsNeedUpdate"),
+	bApplyFixDoComponentsNeedUpdate,
+	TEXT("If true, the new version of DoComponentsNeedUpdate will be used. This fix tries to detect cases where the output mesh would be invalid and cancel the update."));
 
-bool UCustomizableInstancePrivateData::DoComponentsNeedUpdate(UCustomizableObjectInstance* Public, const TSharedPtr<FMutableOperationData>& OperationData, TArray<bool>& OutComponentNeedsUpdate, bool& bOutEmptyMesh)
+
+bool UCustomizableInstancePrivateData::DoComponentsNeedUpdate(UCustomizableObjectInstance* Public, const TSharedPtr<FMutableOperationData>& OperationData, TArray<bool>& OutComponentNeedsUpdate, bool& bHasInvalidMesh)
 {
 	UCustomizableObject* CustomizableObject = Public->GetCustomizableObject();
 	check(CustomizableObject);
 
-	const int32 MaxNumLODs = OperationData->InstanceUpdateData.LODs.Num();
-	const int32 MaxNumComponents = CustomizableObject->GetComponentCount();
-
-	OutComponentNeedsUpdate.AddDefaulted(MaxNumComponents);
-
-	// Return true if at least one mesh needs to be updated
-	bool bUpdateMeshes = false;
+	const int32 NumLODs = OperationData->InstanceUpdateData.LODs.Num();
+	const int32 NumComponents = CustomizableObject->GetComponentCount();
 
-	for (int32 ComponentIndex = 0; ComponentIndex < MaxNumComponents; ++ComponentIndex)
+	if (bApplyFixDoComponentsNeedUpdate)
 	{
-		USkeletalMesh* SkeletalMesh = Public->SkeletalMeshes.IsValidIndex(ComponentIndex) ? Public->SkeletalMeshes[ComponentIndex] : nullptr;
-		FCustomizableInstanceComponentData* ComponentData = GetComponentData(ComponentIndex);
+		TArray<bool> ComponentWithMesh;
+		ComponentWithMesh.Init(false, NumComponents);
 
-		// Check if the mesh is/isn't generated and it should/shouldn't
-		bool bMeshNeedsUpdate = (SkeletalMesh && !ComponentData) || (!SkeletalMesh && ComponentData);
+		TArray<mu::FResourceID> MeshIDs;
+		MeshIDs.Init(MAX_uint64, NumComponents * MAX_MESH_LOD_COUNT);
 
-		if (SkeletalMesh)
+		// Gather the Mesh Ids of all components, and validate the integrity of the meshes to generate. 
+		for (int32 LODIndex = OperationData->CurrentMinLOD; LODIndex <= OperationData->CurrentMaxLOD && LODIndex < NumLODs; ++LODIndex)
 		{
-			bMeshNeedsUpdate = bMeshNeedsUpdate || (SkeletalMesh->GetLODNum() - 1) != OperationData->CurrentMaxLOD;
+			const FInstanceUpdateData::FLOD& LOD = OperationData->InstanceUpdateData.LODs[LODIndex];
+
+			for (int32 ComponentIndex = 0; ComponentIndex < LOD.ComponentCount; ++ComponentIndex)
+			{
+				const FInstanceUpdateData::FComponent& Component = OperationData->InstanceUpdateData.Components[LOD.FirstComponent + ComponentIndex];
+
+				if (Component.Id >= NumComponents)
+				{
+					ensureMsgf(false, TEXT("Mutable: Failed to generate SkeletalMesh. Invalid ComponentIndex. Index [%d]. Valid range [0 .. %d)."), Component.Id, NumComponents);
+					bHasInvalidMesh = true;
+					return false;
+				}
+
+				if (!Component.bGenerated || !Component.Mesh || Component.SurfaceCount == 0) // else
+				{
+					continue;
+				}
+
+				// Unreal does not support empty sections.
+				if (Component.Mesh->GetVertexCount() == 0) // else
+				{
+					UE_LOG(LogMutable, Error, TEXT("Failed to generate SkeletalMesh for CO Instance [%s]. CO [%s] has invalid geometry for LOD [%d] Component [%d]."),
+						*Public->GetName(), *CustomizableObject->GetName(),
+						OperationData->CurrentMaxLOD, ComponentIndex);
+					bHasInvalidMesh = true;
+					continue;
+				}
+
+				ComponentWithMesh[Component.Id] = true;
+				MeshIDs[Component.Id * MAX_MESH_LOD_COUNT + LODIndex] = Component.MeshID;
+			}
 		}
 
-		if (ComponentData)
+		// Find which components need an update
+		OutComponentNeedsUpdate.AddDefaulted(NumComponents);
+
+		for (int32 ComponentIndex = 0; ComponentIndex < NumComponents; ++ComponentIndex)
 		{
-			for (int32 LODIndex = 0; !bMeshNeedsUpdate && LODIndex < MaxNumLODs; ++LODIndex)
+			// Components with mesh must have valid geometry at CurrentMaxLOD
+			if (ComponentWithMesh[ComponentIndex] && MeshIDs[ComponentIndex * MAX_MESH_LOD_COUNT + OperationData->CurrentMaxLOD] == MAX_uint64)
+			{
+				UE_LOG(LogMutable, Error, TEXT("Failed to generate SkeletalMesh for CO Instance [%s]. CO [%s] is missing geometry for LOD [%d] Component [%d]."),
+					*Public->GetName(), *CustomizableObject->GetName(),
+					OperationData->CurrentMaxLOD, ComponentIndex);
+				bHasInvalidMesh = true;
+				continue;
+			}
+
+			// Update the component if there is a mesh and it shouldn't, or the other way around.
+			const bool bHasSkeletalMesh = Public->SkeletalMeshes.IsValidIndex(ComponentIndex) && Public->SkeletalMeshes[ComponentIndex];
+			OutComponentNeedsUpdate[ComponentIndex] = (ComponentWithMesh[ComponentIndex] != bHasSkeletalMesh);
+
+			const FCustomizableInstanceComponentData* ComponentData = GetComponentData(ComponentIndex);
+			if (!ComponentData) // Could be nullptr if the component has not been generated.
+			{
+				continue;
+			}
+
+			// Update if MeshIDs are different
+			const int32 ComponentOffset = ComponentIndex * MAX_MESH_LOD_COUNT;
+			for (int32 MeshIndex = 0; !OutComponentNeedsUpdate[ComponentIndex] && MeshIndex < MAX_MESH_LOD_COUNT; ++MeshIndex)
 			{
-				// Check if an LOD is generated and it shouldn't
-				const bool LODCanBeGenerated = LODIndex >= OperationData->CurrentMinLOD && LODIndex <= OperationData->CurrentMaxLOD;
-				bMeshNeedsUpdate = !LODCanBeGenerated && ComponentData->LastMeshIdPerLOD[LODIndex] != MAX_uint64;
+				OutComponentNeedsUpdate[ComponentIndex] = MeshIDs[ComponentOffset + MeshIndex] != ComponentData->LastMeshIdPerLOD[MeshIndex];
 			}
 		}
 
-		OutComponentNeedsUpdate[ComponentIndex] = bMeshNeedsUpdate;
+		// Mark as invalid if all components have empty meshes.
+		bHasInvalidMesh = bHasInvalidMesh || ComponentWithMesh.Find(true) == INDEX_NONE;
 
-		bUpdateMeshes = bUpdateMeshes || bMeshNeedsUpdate;
+		return !bHasInvalidMesh && OutComponentNeedsUpdate.Find(true) != INDEX_NONE;
 	}
+	else // TODO PRP: MTBL-1653 Remove old version
+	{
+		OutComponentNeedsUpdate.AddDefaulted(NumComponents);
 
-	bool bFoundEmptyMesh = false;
-	bool bFoundEmptyLOD = false;
+		// Return true if at least one mesh needs to be updated
+		bool bUpdateMeshes = false;
 
-	bool bFoundNonEmptyMesh = false;
+		for (int32 ComponentIndex = 0; ComponentIndex < NumComponents; ++ComponentIndex)
+		{
+			USkeletalMesh* SkeletalMesh = Public->SkeletalMeshes.IsValidIndex(ComponentIndex) ? Public->SkeletalMeshes[ComponentIndex] : nullptr;
+			FCustomizableInstanceComponentData* ComponentData = GetComponentData(ComponentIndex);
 
-	for (int32 LODIndex = OperationData->CurrentMinLOD; LODIndex <= OperationData->CurrentMaxLOD && LODIndex < MaxNumLODs; ++LODIndex)
-	{
+			// Check if the mesh is/isn't generated and it should/shouldn't
+			bool bMeshNeedsUpdate = (SkeletalMesh && !ComponentData) || (!SkeletalMesh && ComponentData);
 
-		const FInstanceUpdateData::FLOD& LOD = OperationData->InstanceUpdateData.LODs[LODIndex];
+			if (SkeletalMesh)
+			{
+				bMeshNeedsUpdate = bMeshNeedsUpdate || (SkeletalMesh->GetLODNum() - 1) != OperationData->CurrentMaxLOD;
+			}
 
-		for (int32 ComponentIndex = 0; ComponentIndex < LOD.ComponentCount; ++ComponentIndex)
+			if (ComponentData)
+			{
+				for (int32 LODIndex = 0; !bMeshNeedsUpdate && LODIndex < NumLODs; ++LODIndex)
+				{
+					// Check if an LOD is generated and it shouldn't
+					const bool LODCanBeGenerated = LODIndex >= OperationData->CurrentMinLOD && LODIndex <= OperationData->CurrentMaxLOD;
+					bMeshNeedsUpdate = !LODCanBeGenerated && ComponentData->LastMeshIdPerLOD[LODIndex] != MAX_uint64;
+				}
+			}
+
+			OutComponentNeedsUpdate[ComponentIndex] = bMeshNeedsUpdate;
+
+			bUpdateMeshes = bUpdateMeshes || bMeshNeedsUpdate;
+		}
+
+		bool bFoundEmptyMesh = false;
+		bool bFoundEmptyLOD = false;
+
+		bool bFoundNonEmptyMesh = false;
+
+		for (int32 LODIndex = OperationData->CurrentMinLOD; LODIndex <= OperationData->CurrentMaxLOD && LODIndex < NumLODs; ++LODIndex)
 		{
-			const FInstanceUpdateData::FComponent& Component = OperationData->InstanceUpdateData.Components[LOD.FirstComponent + ComponentIndex];
 
-			const mu::FResourceID LastMeshID = GetLastMeshId(Component.Id, LODIndex);
+			const FInstanceUpdateData::FLOD& LOD = OperationData->InstanceUpdateData.LODs[LODIndex];
 
-			if (Component.bGenerated)
+			for (int32 ComponentIndex = 0; ComponentIndex < LOD.ComponentCount; ++ComponentIndex)
 			{
-				if (Component.Mesh)
+				const FInstanceUpdateData::FComponent& Component = OperationData->InstanceUpdateData.Components[LOD.FirstComponent + ComponentIndex];
+
+				const mu::FResourceID LastMeshID = GetLastMeshId(Component.Id, LODIndex);
+
+				if (Component.bGenerated)
 				{
-					OutComponentNeedsUpdate[Component.Id] = OutComponentNeedsUpdate[Component.Id] || LastMeshID != Component.MeshID;
+					if (Component.Mesh)
+					{
+						OutComponentNeedsUpdate[Component.Id] = OutComponentNeedsUpdate[Component.Id] || LastMeshID != Component.MeshID;
 
-					// Don't build a degenerated mesh if something went wrong
-					if (Component.Mesh->GetVertexCount() > 0)
+						// Don't build a degenerated mesh if something went wrong
+						if (Component.Mesh->GetVertexCount() > 0)
+						{
+							bFoundNonEmptyMesh = true;
+						}
+						else if (LOD.ComponentCount == 1)
+						{
+							bFoundEmptyMesh = true;
+						}
+					}
+					else if (Component.bReuseMesh)
 					{
 						bFoundNonEmptyMesh = true;
+						check(LastMeshID != MAX_uint64);
 					}
-					else if (LOD.ComponentCount == 1)
+					else
 					{
-						bFoundEmptyMesh = true;
+						OutComponentNeedsUpdate[Component.Id] = true;
 					}
 				}
-				else if (Component.bReuseMesh)
-				{
-					bFoundNonEmptyMesh = true;
-					check(LastMeshID != MAX_uint64);
-				}
 				else
 				{
-					OutComponentNeedsUpdate[Component.Id] = true;
+					OutComponentNeedsUpdate[Component.Id] = OutComponentNeedsUpdate[Component.Id] || LastMeshID != MAX_uint64;
 				}
-			}
-			else
-			{
-				OutComponentNeedsUpdate[Component.Id] = OutComponentNeedsUpdate[Component.Id] || LastMeshID != MAX_uint64;
-			}
 
-			bUpdateMeshes = bUpdateMeshes || OutComponentNeedsUpdate[Component.Id];
+				bUpdateMeshes = bUpdateMeshes || OutComponentNeedsUpdate[Component.Id];
+			}
 		}
-	}
 
 
-	if (!bFoundNonEmptyMesh)
-	{
-		bFoundEmptyLOD = true;
-		bUpdateMeshes = true;
-		UE_LOG(LogMutable, Warning, TEXT("Building instance: An LOD has no mesh geometry. This cannot be handled by Unreal."));
-	}
+		if (!bFoundNonEmptyMesh)
+		{
+			bFoundEmptyLOD = true;
+			bUpdateMeshes = true;
+			UE_LOG(LogMutable, Warning, TEXT("Building instance: An LOD has no mesh geometry. This cannot be handled by Unreal."));
+		}
 
-	bOutEmptyMesh = bFoundEmptyLOD || bFoundEmptyMesh;
-	return bUpdateMeshes;
+		bHasInvalidMesh = bFoundEmptyLOD || bFoundEmptyMesh;
+		return bUpdateMeshes;
+	}
 }
 
 
@@ -1862,17 +1954,17 @@ bool UCustomizableInstancePrivateData::UpdateSkeletalMesh_PostBeginUpdate0(UCust
 {
 	MUTABLE_CPUPROFILER_SCOPE(UCustomizableInstancePrivateData::UpdateSkeletalMesh_PostBeginUpdate0)
 
-	bool bEmptyMesh = false;
+	bool bHasInvalidMesh = false;
 
 	TArray<bool> ComponentNeedsUpdate;
-	bool bUpdateMeshes = DoComponentsNeedUpdate(Public, OperationData, ComponentNeedsUpdate, bEmptyMesh);
+	bool bUpdateMeshes = DoComponentsNeedUpdate(Public, OperationData, ComponentNeedsUpdate, bHasInvalidMesh);
 
 	// We can not handle empty meshes, clear any generated mesh and return
-	if (bEmptyMesh)
+	if (bHasInvalidMesh)
 	{
-		Public->SkeletalMeshStatus = ESkeletalMeshState::UpdateError;
-		UE_LOG(LogMutable, Warning, TEXT("Updating skeletal mesh error for CO Instance %s"), *Public->GetName());
+		UE_LOG(LogMutable, Warning, TEXT("Failed to generate SkeletalMesh for CO Instance %s. CO [%s]"), *Public->GetName(), *GetNameSafe(Public->GetCustomizableObject()));
 
+		Public->SkeletalMeshStatus = ESkeletalMeshState::UpdateError;
 		Public->SkeletalMeshes.Reset(); // What about all the references
 
 		SetCOInstanceFlags(Generated);
@@ -1898,7 +1990,6 @@ bool UCustomizableInstancePrivateData::UpdateSkeletalMesh_PostBeginUpdate0(UCust
 	TArray<TObjectPtr<USkeletalMesh>> OldSkeletalMeshes = Public->SkeletalMeshes;
 
 	bool bSuccess = true;
-	bool bHasSkeletalMesh = false;
 
 	UCustomizableObject* CustomizableObject = Public->GetCustomizableObject();
 	check(CustomizableObject);
@@ -1911,10 +2002,14 @@ bool UCustomizableInstancePrivateData::UpdateSkeletalMesh_PostBeginUpdate0(UCust
 	{
 		const FInstanceUpdateData::FComponent& Component = OperationData->InstanceUpdateData.Components[ComponentIndex];
 
+		// If the component doesn't need an update copy the previously generated mesh.
 		if (!ComponentNeedsUpdate[Component.Id])
 		{
-			Public->SkeletalMeshes[Component.Id] = OldSkeletalMeshes[Component.Id];
-			bHasSkeletalMesh = true;
+			if (OldSkeletalMeshes.IsValidIndex(Component.Id))
+			{
+				Public->SkeletalMeshes[Component.Id] = OldSkeletalMeshes[Component.Id];
+			}
+
 			continue;
 		}
 
@@ -2133,12 +2228,10 @@ bool UCustomizableInstancePrivateData::UpdateSkeletalMesh_PostBeginUpdate0(UCust
 				Pair.Key->OnSkeletalMeshCreated(Pair.Value, ComponentIndex, SkeletalMesh);
 			}
 		}
-
-		bHasSkeletalMesh = true;
 	}
 
 
-	if (bSuccess && bHasSkeletalMesh)
+	if (bSuccess)
 	{
 		MUTABLE_CPUPROFILER_SCOPE(BuildSkeletalMeshData);
 
