[mutable] Fixed bug where refreshing a material node reset some texture options of the pin viewer.

#rnx

diff --git a/Engine/Plugins/Experimental/Mutable/Source/CustomizableObjectEditor/Classes/MuCOE/Nodes/CustomizableObjectNodeMaterial.h b/Engine/Plugins/Experimental/Mutable/Source/CustomizableObjectEditor/Classes/MuCOE/Nodes/CustomizableObjectNodeMaterial.h
--- a/Engine/Plugins/Experimental/Mutable/Source/CustomizableObjectEditor/Classes/MuCOE/Nodes/CustomizableObjectNodeMaterial.h
+++ b/Engine/Plugins/Experimental/Mutable/Source/CustomizableObjectEditor/Classes/MuCOE/Nodes/CustomizableObjectNodeMaterial.h
@@ -299,6 +299,10 @@ class CUSTOMIZABLEOBJECTEDITOR_API UCustomizableObjectNodeMaterialPinDataImage :
 public:
 	// UObject interface
 	virtual void PostEditChangeProperty(FPropertyChangedEvent& PropertyChangedEvent) override;
+
+	// NodePinDataParameter interface
+	/** Virtual function used to copy pin data when remapping pins. */
+	virtual void Copy(const UCustomizableObjectNodePinData& Other) override;
 	
 	// UCustomizableObjectNodeMaterialPinParameter interface
 	virtual bool IsDefault() const override;
diff --git a/Engine/Plugins/Experimental/Mutable/Source/CustomizableObjectEditor/Private/MuCOE/Nodes/CustomizableObjectNodeMaterial.cpp b/Engine/Plugins/Experimental/Mutable/Source/CustomizableObjectEditor/Private/MuCOE/Nodes/CustomizableObjectNodeMaterial.cpp
--- a/Engine/Plugins/Experimental/Mutable/Source/CustomizableObjectEditor/Private/MuCOE/Nodes/CustomizableObjectNodeMaterial.cpp
+++ b/Engine/Plugins/Experimental/Mutable/Source/CustomizableObjectEditor/Private/MuCOE/Nodes/CustomizableObjectNodeMaterial.cpp
@@ -113,6 +113,22 @@ void UCustomizableObjectNodeMaterialPinDataImage::Init(UCustomizableObjectNodeMa
 }
 
 
+void UCustomizableObjectNodeMaterialPinDataImage::Copy(const UCustomizableObjectNodePinData& Other)
+{
+	if (const UCustomizableObjectNodeMaterialPinDataImage* PinDataOldPin = Cast<UCustomizableObjectNodeMaterialPinDataImage>(&Other))
+	{
+		PinMode = PinDataOldPin->PinMode;
+		UVLayout = PinDataOldPin->UVLayout;
+		ReferenceTexture = PinDataOldPin->ReferenceTexture;
+
+		if (NodeMaterial)
+		{
+			NodeMaterial->UpdateImagePinMode(ParameterId);
+		}
+	}
+}
+
+
 FName UCustomizableObjectNodeMaterial::GetPinName(const EMaterialParameterType Type, const int32 ParameterIndex) const
 {
 	const FString ParameterName = GetParameterName(Type, ParameterIndex).ToString();
diff --git a/Engine/Plugins/Experimental/Mutable/Source/CustomizableObjectEditor/Private/MuCOE/Nodes/CustomizableObjectNodeTable.cpp b/Engine/Plugins/Experimental/Mutable/Source/CustomizableObjectEditor/Private/MuCOE/Nodes/CustomizableObjectNodeTable.cpp
--- a/Engine/Plugins/Experimental/Mutable/Source/CustomizableObjectEditor/Private/MuCOE/Nodes/CustomizableObjectNodeTable.cpp
+++ b/Engine/Plugins/Experimental/Mutable/Source/CustomizableObjectEditor/Private/MuCOE/Nodes/CustomizableObjectNodeTable.cpp
@@ -660,7 +660,7 @@ FString UCustomizableObjectNodeTable::GetRefreshMessage() const
 	return "Node data outdated. Please refresh node.";
 }
 
-
+// TODO(MTBL-1652): Move this to the Copy method of the NodePinData class which is more PinType specific (and the right plece to do this)
 void UCustomizableObjectNodeTable::RemapPinsData(const TMap<UEdGraphPin*, UEdGraphPin*>& PinsToRemap)
 {
 	const UEdGraphSchema_CustomizableObject* Schema = GetDefault<UEdGraphSchema_CustomizableObject>();
