[Mutable] Simplify UpdateSkeletalMesh flow.

- Removed unecessary flags.
- Removed Generated Instance flag from DoUpdateSkeletalMesh. Was incorrect. This flag still requires some improvements.
- Renamed InitUpdateSkeletalMesh to EnqueuepdateSkeletalMesh.
- DoUpdateSkeletalMesh is no longer mandatory to start an update. An update can start by simply adding the FMutableOperation to the MutableCurrentOperation.
- Removed no longer necessary bIsCloseDistTick. DoUpdateSkeletalMesh is no longer called from the tick. Now DoUpdateSkeletalMesh is only necessary to enqueue a task, which tick never does (executes a task directly).
- Removed unused BeginUpdate delegates.
- Replaced LODsStreamingEnabled with expression that was activating the flag (LODSettings).
- Replaced ForceGenerateAllLODs with the expression that was activating the flag (being on an EditorViewport).
- State is no longer obtained from the Instance, it is captured when generating mu::Parameters.
- Removed unecessary State copy when creating lambads.

#tests PIE, Village, Editor
#rnx

diff --git a/Engine/Plugins/Experimental/Mutable/Source/CustomizableObject/Classes/MuCO/CustomizableObjectInstance.h b/Engine/Plugins/Experimental/Mutable/Source/CustomizableObject/Classes/MuCO/CustomizableObjectInstance.h
--- a/Engine/Plugins/Experimental/Mutable/Source/CustomizableObject/Classes/MuCO/CustomizableObjectInstance.h
+++ b/Engine/Plugins/Experimental/Mutable/Source/CustomizableObject/Classes/MuCO/CustomizableObjectInstance.h
@@ -99,11 +99,6 @@ DECLARE_DYNAMIC_DELEGATE_OneParam(FInstanceUpdateDelegate, const FUpdateContext&
  * - Native delegates names should end with "NativeDelegate".
  * - Dynamic delegates broadcast before native delegates. */
 
-/** Broadcast at the end of an Instance update request (e.g., before returning from UpdateSkeletalMeshAsync).
- * Notice that Mutable internally can also request an Instance update. */
-DECLARE_DYNAMIC_MULTICAST_DELEGATE_OneParam(FBeginUpdateDelegate, UCustomizableObjectInstance*, Instance);
-DECLARE_MULTICAST_DELEGATE_OneParam(FBeginUpdateNativeDelegate, UCustomizableObjectInstance*);
-
 /** Broadcast when an Instance update has completed.
  * Notice that Mutable internally can also start an Instance update. */
 DECLARE_DYNAMIC_MULTICAST_DELEGATE_OneParam(FObjectInstanceUpdatedDelegate, UCustomizableObjectInstance*, Instance);
@@ -136,13 +131,6 @@ public:
 	
 	void SetDescriptor(const FCustomizableObjectInstanceDescriptor& InDescriptor);
 
-	/** Broadcast at the beginning of an Instance update. */
-	UPROPERTY(BlueprintAssignable, Category = CustomizableObjectInstance)
-	FBeginUpdateDelegate BeginUpdateDelegate;
-
-	/** Broadcast at the beginning of an Instance update. */
-	FBeginUpdateNativeDelegate BeginUpdateNativeDelegate;
-
 	/** Broadcast when the Customizable Object Instance is updated. */
 	UPROPERTY(Transient, BlueprintAssignable, Category = CustomizableObjectInstance)
 	FObjectInstanceUpdatedDelegate UpdatedDelegate;
@@ -292,18 +280,15 @@ public:
 
 private:
 	/** Perform all the checks required to see if the update should begin. */
-	EUpdateRequired IsUpdateRequired(bool bIsCloseDistTick, bool bOnlyUpdateIfNotGenerated, bool bIgnoreCloseDist) const;
+	EUpdateRequired IsUpdateRequired(bool bOnlyUpdateIfNotGenerated, bool bIgnoreCloseDist) const;
 
 public:
 	/** Private API.
 	 *
 	 * Update Skeletal Mesh asynchronously. Immersive function.
 	 * Once the update reaches this function, the update has been considered started and must complete all the update flow.
-	 * Starting at this function, all Update code paths must end up in FinishUpdateGlobal!
-	 *
-	 * @param bIsCloseDistTick true if and only if called from the tick.
-	 */
-	void DoUpdateSkeletalMesh(bool bIsCloseDistTick, bool bOnlyUpdateIfNotGenerated, bool bIgnoreCloseDist, bool bForceHighPriority, const EUpdateRequired* OptionalUpdateRequired, FInstanceUpdateDelegate* UpdateCallback);
+	 * Starting at this function, all Update code paths must end up in FinishUpdateGlobal! */
+	void EnqueueUpdateSkeletalMesh(bool bOnlyUpdateIfNotGenerated, bool bIgnoreCloseDist, bool bForceHighPriority, const EUpdateRequired* OptionalUpdateRequired, FInstanceUpdateDelegate* UpdateCallback);
 
 	// Clones the instance creating a new identical transient instance.
 	UFUNCTION(BlueprintCallable, Category = CustomizableObjectInstance)
@@ -322,8 +307,8 @@ public:
 	// Releases all the mutable resources this instance holds, should only be called when it is not going to be used any more.
 	void ReleaseMutableResources(bool bCalledFromBeginDestroy);
 
-	/** Returns the priority an update issued by DoUpdateSkeletalMesh would get in the current instance configuration and state */
-	EQueuePriorityType GetUpdatePriority(bool bIsCloseDistTick, bool bOnlyUpdateIfNotGenerated, bool bIgnoreCloseDist, bool bForceHighPriority) const;
+	/** Returns the priority an update issued by EnqueueUpdateSkeletalMesh would get in the current instance configuration and state */
+	EQueuePriorityType GetUpdatePriority(bool bForceHighPriority) const;
 
 	// Returns de description texture (ex: color bar) for this parameter and DescIndex
 	// This will only be valid if bBuildParameterDecorations was set to true before the last update.
@@ -837,7 +822,7 @@ public:
 	}
 
 private:
-	/** If true it means that DoUpdateSkeletalMesh has decided this update should be performed, if false it should be ignored. Just used for consistency checks */
+	/** If true it means that EnqueueUpdateSkeletalMesh has decided this update should be performed, if false it should be ignored. Just used for consistency checks */
 	bool bHasBeenIssued = false;
 };
 
diff --git a/Engine/Plugins/Experimental/Mutable/Source/CustomizableObject/Private/MuCO/CustomizableInstancePrivateData.h b/Engine/Plugins/Experimental/Mutable/Source/CustomizableObject/Private/MuCO/CustomizableInstancePrivateData.h
--- a/Engine/Plugins/Experimental/Mutable/Source/CustomizableObject/Private/MuCO/CustomizableInstancePrivateData.h
+++ b/Engine/Plugins/Experimental/Mutable/Source/CustomizableObject/Private/MuCO/CustomizableInstancePrivateData.h
@@ -33,7 +33,7 @@ enum ECOInstanceFlags
 	// Update process
 	Updating						= 1 << 0,	//
 	CreatingSkeletalMesh			= 1 << 1,	//
-	Generated						= 1 << 2,	// Generated or update in progress.
+	Generated						= 1 << 2,	//
 	ReuseTextures					= 1 << 3, 	// 
 	ReplacePhysicsAssets			= 1 << 4,	// Merge active PhysicsAssets and replace the base physics asset
 
@@ -46,12 +46,8 @@ enum ECOInstanceFlags
 	// Types of updates
 	PendingLODsUpdate				= 1 << 9,	// Used to queue an update due to a change in LODs required by the instance
 	PendingLODsDowngrade			= 1 << 10,	// Used to queue a downgrade update to reduce the number of LODs. LOD update goes from a high res level to a low res one, ex: 0 to 1 or 1 to 2
-
-	// Streaming
-	LODsStreamingEnabled			= 1 << 11,	// Stream LODs instead of generating all LODs at once. Enables LODs update(upgrade)/downgrade.
 	
 	// Generation
-	ForceGenerateAllLODs			= 1 << 12,	// If set, Requested LOD Levels will be ignored and all LODs in between the current min/max lod will be generated
 	ForceGenerateMipTail			= 1 << 13,	// If set, SkipGenerateResidentMips will be ignored and the mip tail will be generated
 };
 
diff --git a/Engine/Plugins/Experimental/Mutable/Source/CustomizableObject/Private/MuCO/CustomizableObjectInstance.cpp b/Engine/Plugins/Experimental/Mutable/Source/CustomizableObject/Private/MuCO/CustomizableObjectInstance.cpp
--- a/Engine/Plugins/Experimental/Mutable/Source/CustomizableObject/Private/MuCO/CustomizableObjectInstance.cpp
+++ b/Engine/Plugins/Experimental/Mutable/Source/CustomizableObject/Private/MuCO/CustomizableObjectInstance.cpp
@@ -527,16 +527,12 @@ void UCustomizableInstancePrivateData::InstanceUpdateFlags(const UCustomizableOb
 	FirstLODAvailable = CustomizableObject->LODSettings.FirstLODAvailable;
 	NumLODsAvailable = CustomizableObject->GetNumLODs();
 
-	if (CustomizableObject->LODSettings.bLODStreamingEnabled || !UCustomizableObjectSystem::GetInstance()->IsProgressiveMipStreamingEnabled())
+	if (CustomizableObject->LODSettings.bLODStreamingEnabled)
 	{
-		SetCOInstanceFlags(LODsStreamingEnabled);
-		ClearCOInstanceFlags(ForceGenerateAllLODs);
 		NumMaxLODsToStream = CustomizableObject->LODSettings.NumLODsToStream;
 	}
 	else
 	{
-		SetCOInstanceFlags(ForceGenerateAllLODs);
-		ClearCOInstanceFlags(LODsStreamingEnabled);
 		NumMaxLODsToStream = 0;
 	}
 }
@@ -739,17 +735,17 @@ bool UCustomizableObjectInstance::CanUpdateInstance() const
 
 void UCustomizableObjectInstance::UpdateSkeletalMeshAsync(bool bIgnoreCloseDist, bool bForceHighPriority)
 {
-	DoUpdateSkeletalMesh( false, false, bIgnoreCloseDist, bForceHighPriority, nullptr, nullptr);
+	EnqueueUpdateSkeletalMesh(false, bIgnoreCloseDist, bForceHighPriority, nullptr, nullptr);
 }
 
 
 void UCustomizableObjectInstance::UpdateSkeletalMeshAsyncResult(FInstanceUpdateDelegate Callback, bool bIgnoreCloseDist, bool bForceHighPriority)
 {
-	DoUpdateSkeletalMesh(false, false, bIgnoreCloseDist, bForceHighPriority, nullptr, &Callback);
+	EnqueueUpdateSkeletalMesh(false, bIgnoreCloseDist, bForceHighPriority, nullptr, &Callback);
 }
 
 
-EUpdateRequired UCustomizableObjectInstance::IsUpdateRequired(bool bIsCloseDistTick, bool bOnlyUpdateIfNotGenerated, bool bIgnoreCloseDist) const
+EUpdateRequired UCustomizableObjectInstance::IsUpdateRequired(bool bOnlyUpdateIfNotGenerated, bool bIgnoreCloseDist) const
 {
 	UCustomizableObjectSystem* System = UCustomizableObjectSystem::GetInstance();
 	const UCustomizableInstancePrivateData* const Private = GetPrivate();
@@ -794,7 +790,7 @@ EUpdateRequired UCustomizableObjectInstance::IsUpdateRequired(bool bIsCloseDistT
 
 	if (bIsGenerated &&
 		!bShouldUpdateLODs &&
-		(bOnlyUpdateIfNotGenerated || bIsCloseDistTick))
+		bOnlyUpdateIfNotGenerated)
 	{
 		return EUpdateRequired::NoUpdate;
 	}
@@ -803,7 +799,7 @@ EUpdateRequired UCustomizableObjectInstance::IsUpdateRequired(bool bIsCloseDistT
 }
 
 
-EQueuePriorityType UCustomizableObjectInstance::GetUpdatePriority(bool bIsCloseDistTick, bool bOnlyUpdateIfNotGenerated, bool bIgnoreCloseDist, bool bForceHighPriority) const
+EQueuePriorityType UCustomizableObjectInstance::GetUpdatePriority(bool bForceHighPriority) const
 {
 	const bool bIsGenerated = GetPrivate()->HasCOInstanceFlags(Generated);
 	const bool bShouldUpdateLODs = GetPrivate()->HasCOInstanceFlags(PendingLODsUpdate);
@@ -840,9 +836,9 @@ EQueuePriorityType UCustomizableObjectInstance::GetUpdatePriority(bool bIsCloseD
 }
 
 
-void UCustomizableObjectInstance::DoUpdateSkeletalMesh(bool bIsCloseDistTick, bool bOnlyUpdateIfNotGenerated, bool bIgnoreCloseDist, bool bForceHighPriority, const EUpdateRequired* OptionalUpdateRequired, FInstanceUpdateDelegate* UpdateCallback)
+void UCustomizableObjectInstance::EnqueueUpdateSkeletalMesh(bool bOnlyUpdateIfNotGenerated, bool bIgnoreCloseDist, bool bForceHighPriority, const EUpdateRequired* OptionalUpdateRequired, FInstanceUpdateDelegate* UpdateCallback)
 {
-	MUTABLE_CPUPROFILER_SCOPE(UCustomizableInstancePrivateData::DoUpdateSkeletalMesh);
+	MUTABLE_CPUPROFILER_SCOPE(UCustomizableInstancePrivateData::EnqueueUpdateSkeletalMesh);
 	check(IsInGameThread());
 	check(!OptionalUpdateRequired || *OptionalUpdateRequired != EUpdateRequired::NoUpdate); // If no update is required this functions must not be called.
 
@@ -854,7 +850,7 @@ void UCustomizableObjectInstance::DoUpdateSkeletalMesh(bool bIsCloseDistTick, bo
 		return;
 	}
 	
-	const EUpdateRequired UpdateRequired = OptionalUpdateRequired ? *OptionalUpdateRequired : IsUpdateRequired(bIsCloseDistTick, bOnlyUpdateIfNotGenerated, bIgnoreCloseDist);
+	const EUpdateRequired UpdateRequired = OptionalUpdateRequired ? *OptionalUpdateRequired : IsUpdateRequired(bOnlyUpdateIfNotGenerated, bIgnoreCloseDist);
 	switch (UpdateRequired)
 	{
 	case EUpdateRequired::NoUpdate:
@@ -864,7 +860,7 @@ void UCustomizableObjectInstance::DoUpdateSkeletalMesh(bool bIsCloseDistTick, bo
 	}		
 	case EUpdateRequired::Update:
 	{
-		EQueuePriorityType Priority = GetUpdatePriority(bIsCloseDistTick, bOnlyUpdateIfNotGenerated, bIgnoreCloseDist, bForceHighPriority);
+		EQueuePriorityType Priority = GetUpdatePriority(bForceHighPriority);
 
 		const uint32 InstanceId = GetUniqueID();
 		const float Distance = FMath::Sqrt(GetPrivate()->LastMinSquareDistFromComponentToPlayer);
@@ -885,13 +881,11 @@ void UCustomizableObjectInstance::DoUpdateSkeletalMesh(bool bIsCloseDistTick, bo
 
 		if (GetPrivate()->HasCOInstanceFlags(PendingLODsUpdate))
 		{
-			UE_LOG(LogMutable, Verbose, TEXT("LOD change: %d, %d -> %d, %d"), GetCurrentMinLOD(), GetCurrentMaxLOD(), GetMinLODToLoad(), GetMinLODToLoad());
+			UE_LOG(LogMutable, Verbose, TEXT("LOD change: %d, %d -> %d, %d"), GetCurrentMinLOD(), GetCurrentMaxLOD(), GetMinLODToLoad(), GetMaxLODToLoad());
 		}
 		
-		GetPrivate()->SetCOInstanceFlags(Generated); // Will be done in UpdateSkeletalMesh_PostBeginUpdate
-
-		// Do not do work after calling InitUpdateSkeletalMesh. This function can optimize an update and fully complete it before even exiting its scope.
-		System->GetPrivate()->InitUpdateSkeletalMesh(*this, Priority, bIsCloseDistTick, UpdateCallback);
+		// Do not do work after calling EnqueueUpdateSkeletalMesh. This function can optimize an update and fully complete it before even exiting its scope.
+		System->GetPrivate()->EnqueueUpdateSkeletalMesh(*this, Priority, UpdateCallback);
 		break;
 	}
 
@@ -916,17 +910,17 @@ void UCustomizableInstancePrivateData::TickUpdateCloseCustomizableObjects(UCusto
 		return;
 	}
 
-	const EUpdateRequired UpdateRequired = Public.IsUpdateRequired(true, true, false);
+	const EUpdateRequired UpdateRequired = Public.IsUpdateRequired(true, false);
 	if (UpdateRequired != EUpdateRequired::NoUpdate) // Since this is done in the tick, avoid starting an update that we know for sure that would not be performed. Once started it has some performance implications that we want to avoid.
 	{
 		if (UpdateRequired == EUpdateRequired::Discard)
 		{
-			Public.DoUpdateSkeletalMesh(true, true, false, false, &UpdateRequired, nullptr);
+			UCustomizableObjectSystem::GetInstance()->GetPrivate()->InitDiscardResourcesSkeletalMesh(&Public);
 			InOutRequestedUpdates.Remove(&Public);
 		}
 		else if (UpdateRequired == EUpdateRequired::Update)
 		{
-			EQueuePriorityType Priority = Public.GetUpdatePriority(true, true, false, false);
+			EQueuePriorityType Priority = Public.GetUpdatePriority(false);
 
 			FMutableUpdateCandidate* UpdateCandidate = InOutRequestedUpdates.Find(&Public);
 
@@ -971,9 +965,9 @@ void UCustomizableInstancePrivateData::UpdateInstanceIfNotGenerated(UCustomizabl
 		return;
 	}
 
-	Public.DoUpdateSkeletalMesh(false, true, false, false, nullptr, nullptr);
+	Public.EnqueueUpdateSkeletalMesh(true, false, false, nullptr, nullptr);
 
-	EQueuePriorityType Priority = Public.GetUpdatePriority(true, true, false, false);
+	EQueuePriorityType Priority = Public.GetUpdatePriority(false);
 	FMutableUpdateCandidate* UpdateCandidate = InOutRequestedUpdates.Find(&Public);
 
 	if (UpdateCandidate)
@@ -6447,16 +6441,17 @@ void UCustomizableObjectInstance::SetRequestedLODs(int32 InMinLOD, int32 InMaxLO
 	                                               FMutableInstanceUpdateMap& InOutRequestedUpdates)
 {
 	check(PrivateData);
-	if (!PrivateData->HasCOInstanceFlags(LODsStreamingEnabled))
+
+	if (!CanUpdateInstance())
 	{
 		return;
 	}
 
-	if (!CanUpdateInstance())
+	if (!GetCustomizableObject()->LODSettings.bLODStreamingEnabled)
 	{
 		return;
 	}
-
+	
 #if !UE_BUILD_SHIPPING
 	// Ignore Min/Max LOD limits. Mainly used for debug
 	if (!bIgnoreMinMaxLOD)
diff --git a/Engine/Plugins/Experimental/Mutable/Source/CustomizableObject/Private/MuCO/CustomizableObjectSystem.cpp b/Engine/Plugins/Experimental/Mutable/Source/CustomizableObject/Private/MuCO/CustomizableObjectSystem.cpp
--- a/Engine/Plugins/Experimental/Mutable/Source/CustomizableObject/Private/MuCO/CustomizableObjectSystem.cpp
+++ b/Engine/Plugins/Experimental/Mutable/Source/CustomizableObject/Private/MuCO/CustomizableObjectSystem.cpp
@@ -35,6 +35,7 @@
 #include "Logging/MessageLog.h"
 #include "Misc/ConfigCacheIni.h"
 #include "Misc/MessageDialog.h"
+#include "Engine/World.h"
 #else
 #include "Engine/Engine.h"
 #endif
@@ -858,16 +859,11 @@ bool FCustomizableObjectSystemPrivate::TextureHasReferences(const FMutableImageC
 }
 
 
-void FCustomizableObjectSystemPrivate::InitUpdateSkeletalMesh(UCustomizableObjectInstance& Instance, EQueuePriorityType Priority, bool bIsCloseDistTick, FInstanceUpdateDelegate* UpdateCallback)
+void FCustomizableObjectSystemPrivate::EnqueueUpdateSkeletalMesh(UCustomizableObjectInstance& Instance, const EQueuePriorityType Priority, FInstanceUpdateDelegate* UpdateCallback)
 {
 	MUTABLE_CPUPROFILER_SCOPE(FCustomizableObjectSystemPrivate::InitUpdateSkeletalMesh);
 
 	check(IsInGameThread());
-
-	bool bNeverStream = false;
-	int32 MipsToSkip = 0;
-
-	GetMipStreamingConfig(Instance, bNeverStream, MipsToSkip);
 	
 	const FDescriptorRuntimeHash UpdateDescriptorHash = Instance.GetUpdateDescriptorRuntimeHash();
 
@@ -888,13 +884,9 @@ void FCustomizableObjectSystemPrivate::InitUpdateSkeletalMesh(UCustomizableObjec
 		return; // The requested update is equal to the running update.
 	}
 	
-	// These delegates must be called at the end of the begin update.
-	Instance.BeginUpdateDelegate.Broadcast(&Instance);
-	Instance.BeginUpdateNativeDelegate.Broadcast(&Instance);
-
 	if (UpdateDescriptorHash.IsSubset(Instance.GetDescriptorRuntimeHash()) &&
 		!(CurrentMutableOperation &&
-			&Instance == CurrentMutableOperation->CustomizableObjectInstance)) // This condition is necessary because even if the descriptor is a subset, it will be replaced by the CurrentMutableOperation
+		&Instance == CurrentMutableOperation->CustomizableObjectInstance)) // This condition is necessary because even if the descriptor is a subset, it will be replaced by the CurrentMutableOperation
 	{
 		Instance.SkeletalMeshStatus = ESkeletalMeshState::Correct; // TODO FutureGMT MTBL-1033 should not be here. Move to UCustomizableObjectInstance::Updated
 		UpdateSkeletalMesh(Instance, Instance.GetDescriptorRuntimeHash(), EUpdateResult::Success, UpdateCallback);
@@ -938,11 +930,8 @@ void FCustomizableObjectSystemPrivate::InitUpdateSkeletalMesh(UCustomizableObjec
 
 		Instance.SkeletalMeshStatus = ESkeletalMeshState::AsyncUpdatePending;
 
-		if (!bIsCloseDistTick) // When called from bIsCloseDistTick, the update operation is directly processed without going to a queue first
-		{
-			const FMutablePendingInstanceUpdate InstanceUpdate(&Instance, Priority, UpdateCallback, bNeverStream, MipsToSkip);
-			MutablePendingInstanceWork.AddUpdate(InstanceUpdate);
-		}
+		const FMutablePendingInstanceUpdate InstanceUpdate(&Instance, Priority, UpdateCallback);
+		MutablePendingInstanceWork.AddUpdate(InstanceUpdate);
 	}
 }
 
@@ -1288,11 +1277,11 @@ namespace impl
 
 
 	// This runs in the mutable thread.
-	void Subtask_Mutable_BeginUpdate_GetMesh(const TSharedPtr<FMutableOperationData>& OperationData, TSharedPtr<mu::Model, ESPMode::ThreadSafe> Model, mu::Ptr<const mu::Parameters> MutableParameters, int32 State)
+	void Subtask_Mutable_BeginUpdate_GetMesh(const TSharedPtr<FMutableOperationData>& OperationData, TSharedPtr<mu::Model> Model)
 	{
 		MUTABLE_CPUPROFILER_SCOPE(Subtask_Mutable_BeginUpdate_GetMesh)
 
-		check(MutableParameters);
+		check(OperationData->MutableParameters);
 		check(OperationData);
 		OperationData->InstanceUpdateData.Clear();
 
@@ -1333,7 +1322,7 @@ namespace impl
 
 		// Main instance generation step
 		// LOD mask, set to all ones to build  all LODs
-		const mu::Instance* Instance = System->BeginUpdate(OperationData->InstanceID, MutableParameters, State, mu::System::AllLODs);
+		const mu::Instance* Instance = System->BeginUpdate(OperationData->InstanceID, OperationData->MutableParameters, OperationData->State, mu::System::AllLODs);
 		if (!Instance)
 		{
 			UE_LOG(LogMutable, Warning, TEXT("An Instace update has failed."));
@@ -1869,7 +1858,7 @@ namespace impl
 		const uint32 StartCycles = FPlatformTime::Cycles();
 #endif
 
-		Subtask_Mutable_BeginUpdate_GetMesh(OperationData, Model, OperationData->MutableParameters, OperationData->State);
+		Subtask_Mutable_BeginUpdate_GetMesh(OperationData, Model);
 
 		// TODO: Not strictly mutable: move to another worker thread task to free mutable access?
 		Subtask_Mutable_PrepareSkeletonData(OperationData);
@@ -2378,6 +2367,8 @@ namespace impl
 			FinishUpdateGlobal(nullptr, EUpdateResult::Error, &Operation->UpdateCallback);
 			return;
 		}
+
+		CandidateInstancePrivateData->InstanceUpdateFlags(*CandidateInstance);
 		
 		if (CandidateInstancePrivateData->HasCOInstanceFlags(PendingLODsUpdate))
 		{
@@ -2465,7 +2456,12 @@ namespace impl
 			bLiveUpdateMode = StateData ? StateData->bLiveUpdateMode : false;
 		}
 
-		if (bLiveUpdateMode && (!Operation->bNeverStream || Operation->MipsToSkip > 0))
+		bool bNeverStream = false;
+		int32 MipsToSkip = 0;
+
+		SystemPrivateData->GetMipStreamingConfig(*CandidateInstance, bNeverStream, MipsToSkip);
+		
+		if (bLiveUpdateMode && (!bNeverStream || MipsToSkip > 0))
 		{
 			UE_LOG(LogMutable, Warning, TEXT("Instance LiveUpdateMode does not yet support progressive streaming of Mutable textures. Disabling LiveUpdateMode for this update."));
 			bLiveUpdateMode = false;
@@ -2478,7 +2474,7 @@ namespace impl
 			bReuseInstanceTextures = StateData ? StateData->bReuseInstanceTextures : false;
 			bReuseInstanceTextures |= CandidateInstancePrivateData->HasCOInstanceFlags(ReuseTextures);
 			
-			if (bReuseInstanceTextures && !Operation->bNeverStream)
+			if (bReuseInstanceTextures && !bNeverStream)
 			{
 				UE_LOG(LogMutable, Warning, TEXT("Instance texture reuse requires that the current Mutable state is in non-streaming mode. Change it in the Mutable graph base node in the state definition."));
 				bReuseInstanceTextures = false;
@@ -2492,6 +2488,7 @@ namespace impl
 			Task_Game_ReleaseInstanceID(CandidateInstancePrivateData->LiveUpdateModeInstanceID);
 			CandidateInstancePrivateData->LiveUpdateModeInstanceID = 0;
 		}
+
 		
 		// Task: Mutable Update and GetMesh
 		//-------------------------------------------------------------
@@ -2504,13 +2501,13 @@ namespace impl
 		CurrentOperationData->LastUpdateData = SystemPrivateData->bEnableMutableReusePreviousUpdateData ? CandidateInstancePrivateData->LastUpdateData : FInstanceGeneratedData();
 		CurrentOperationData->CurrentMinLOD = Operation->InstanceDescriptorRuntimeHash.GetMinLOD();
 		CurrentOperationData->CurrentMaxLOD = Operation->InstanceDescriptorRuntimeHash.GetMaxLOD();
-		CurrentOperationData->bNeverStream = Operation->bNeverStream;
+		CurrentOperationData->bNeverStream = bNeverStream;
 		CurrentOperationData->bLiveUpdateMode = bLiveUpdateMode;
 		CurrentOperationData->bReuseInstanceTextures = bReuseInstanceTextures;
 		CurrentOperationData->InstanceID = bLiveUpdateMode ? CandidateInstancePrivateData->LiveUpdateModeInstanceID : 0;
-		CurrentOperationData->MipsToSkip = Operation->MipsToSkip;
+		CurrentOperationData->MipsToSkip = MipsToSkip;
 		CurrentOperationData->MutableParameters = Parameters;
-		CurrentOperationData->State = CandidateInstance->GetState();
+		CurrentOperationData->State = Operation->GetState();
 		CurrentOperationData->UpdateResult = EUpdateResult::Success;
 		CurrentOperationData->UpdateCallback = Operation->UpdateCallback;
 		CurrentOperationData->bBuildParameterRelevancy = Operation->IsBuildParameterRelevancy();
@@ -2523,8 +2520,35 @@ namespace impl
 			CustomizableObject->GetLowPriorityTextureNames(CurrentOperationData->LowPriorityTextures);
 		}
 
+		bool bIsInEditorViewport = false;
+
+#if WITH_EDITOR
+		for (TObjectIterator<UCustomizableSkeletalComponent> CustomizableSkeletalComponent; CustomizableSkeletalComponent && !bIsInEditorViewport; ++CustomizableSkeletalComponent)
+		{
+			if (CustomizableSkeletalComponent &&
+				CustomizableSkeletalComponent->CustomizableObjectInstance == CandidateInstance)
+			{
+				EWorldType::Type WorldType = EWorldType::Type::None;
+
+				if (CustomizableSkeletalComponent->GetWorld())
+				{
+					WorldType = CustomizableSkeletalComponent->GetWorld()->WorldType;
+				}
+
+				switch (WorldType)
+				{
+					// Editor preview instances
+					case EWorldType::EditorPreview:
+					case EWorldType::None:
+						bIsInEditorViewport = true;
+					default: ;
+				}
+			}
+		}
+#endif // WITH_EDITOR
+		
 		if (System->IsOnlyGenerateRequestedLODsEnabled() && System->CurrentInstanceLODManagement->IsOnlyGenerateRequestedLODLevelsEnabled() && 
-			!Operation->bForceGenerateAllLODs)
+			!bIsInEditorViewport)
 		{
 			CurrentOperationData->RequestedLODs = Operation->InstanceDescriptorRuntimeHash.GetRequestedLODs();
 		}
@@ -2687,8 +2711,6 @@ bool UCustomizableObjectSystem::Tick(float DeltaTime)
 			}
 		}
 
-		TSharedPtr<FMutableOperation> FoundOperation;
-
 		{
 			// Look for the highest priority update between the pending updates and the LOD Requested Updates
 			EQueuePriorityType MaxPriorityFound = EQueuePriorityType::Low;
@@ -2704,7 +2726,7 @@ bool UCustomizableObjectSystem::Tick(float DeltaTime)
 
 				if (PendingUpdate.CustomizableObjectInstance.IsValid())
 				{
-					EQueuePriorityType PriorityType = PendingUpdate.CustomizableObjectInstance->GetUpdatePriority(true, true, false, false);
+					const EQueuePriorityType PriorityType = PendingUpdate.CustomizableObjectInstance->GetUpdatePriority(false);
 					
 					if (PendingUpdate.PriorityType <= MaxPriorityFound)
 					{
@@ -2770,7 +2792,7 @@ bool UCustomizableObjectSystem::Tick(float DeltaTime)
 				check(!LODUpdateCandidateFound);
 
 				UCustomizableObjectInstance* PendingInstance = PendingInstanceUpdateFound->CustomizableObjectInstance.Get();
-				ensure(PendingInstance);
+				check(PendingInstance);
 
 				// Maybe there's a LODUpdate that has the same instance, merge both updates as an optimization
 				FMutableUpdateCandidate* LODUpdateWithSameInstance = RequestedLODUpdates.Find(PendingInstance);
@@ -2780,33 +2802,16 @@ bool UCustomizableObjectSystem::Tick(float DeltaTime)
 					LODUpdateWithSameInstance->ApplyLODUpdateParamsToInstance();
 				}
 
-				// No need to do a DoUpdateSkeletalMesh, as it was already done when adding the pending update
-
-				FoundOperation = MakeShared<FMutableOperation>(FMutableOperation::CreateInstanceUpdate(PendingInstance,
-					PendingInstanceUpdateFound->bNeverStream, PendingInstanceUpdateFound->MipsToSkip, PendingInstanceUpdateFound->Callback));
+				Private->CurrentMutableOperation = MakeShared<FMutableOperation>(FMutableOperation::CreateInstanceUpdate(*PendingInstance, PendingInstanceUpdateFound->Callback));
 
 				Private->MutablePendingInstanceWork.RemoveUpdate(PendingInstanceUpdateFound->CustomizableObjectInstance);
 			}
-			else if(LODUpdateCandidateFound)
+			else if (LODUpdateCandidateFound)
 			{
 				// Commit the LOD changes
 				LODUpdateCandidateFound->ApplyLODUpdateParamsToInstance();
 
-				// No need to check if a PendingUpdate has the same instance, it would already have been caught in the 
-				// previous if statements
-
-				// LOD updates are detected automatically and have never had a DoUpdateSkeletalMesh, so call it without enqueing a new update
-				const EUpdateRequired Required = EUpdateRequired::Update;
-				LODUpdateCandidateFound->CustomizableObjectInstance->DoUpdateSkeletalMesh(true, false, false, false, &Required, nullptr);
-				
-				bool bNeverStream = false;
-				int32 MipsToSkip = 0;
-
-				ensure(LODUpdateCandidateFound->CustomizableObjectInstance);
-				GetPrivate()->GetMipStreamingConfig(*LODUpdateCandidateFound->CustomizableObjectInstance, bNeverStream, MipsToSkip);
-
-				FoundOperation = MakeShared<FMutableOperation>(FMutableOperation::CreateInstanceUpdate(LODUpdateCandidateFound->CustomizableObjectInstance,
-					bNeverStream, MipsToSkip, nullptr));
+				Private->CurrentMutableOperation = MakeShared<FMutableOperation>(FMutableOperation::CreateInstanceUpdate(*LODUpdateCandidateFound->CustomizableObjectInstance, nullptr));
 			}
 		}
 
@@ -2827,26 +2832,8 @@ bool UCustomizableObjectSystem::Tick(float DeltaTime)
 		// Free memory before starting the new update
 		DiscardInstances();
 		ReleaseInstanceIDs();
-
-		// Decide the next mutable operation to perform.
-		if (FoundOperation.IsValid())
-		{
-			Private->CurrentMutableOperation = FoundOperation;
-
-			UCustomizableObjectInstance* FoundInstance = Private->CurrentMutableOperation->CustomizableObjectInstance.Get();
-			
-			if (FoundInstance->GetPrivate()->HasCOInstanceFlags(ForceGenerateAllLODs))
-			{
-				Private->CurrentMutableOperation->bForceGenerateAllLODs = true;
-			}
-
-			Private->CurrentMutableOperation->CustomizableObjectInstance->GetPrivate()->InstanceUpdateFlags(*FoundInstance);
-		}
-		else
-		{
-			Private->CurrentMutableOperation = nullptr;
-		}
 	}
+	
 	// Advance the current operation
 	if (Private->CurrentMutableOperation)
 	{
@@ -3026,17 +3013,13 @@ void UnCacheTexturesParameters(const TArray<FName>& TextureParameters)
 
 FMutableOperation::FMutableOperation(const FMutableOperation& Other)
 {
-	bNeverStream = Other.bNeverStream;
-	MipsToSkip = Other.MipsToSkip;
 	CustomizableObjectInstance = Other.CustomizableObjectInstance;
 	InstanceDescriptorRuntimeHash = Other.InstanceDescriptorRuntimeHash;
-	bStarted = Other.bStarted;
 	bBuildParameterRelevancy = Other.bBuildParameterRelevancy;
 	Parameters = Other.Parameters;
 	TextureParameters = Other.TextureParameters;
 	UpdateCallback = Other.UpdateCallback;
 	State = Other.State;
-	IDToRelease = Other.IDToRelease;
 
 	CacheTexturesParameters(TextureParameters);
 }
@@ -3044,11 +3027,8 @@ FMutableOperation::FMutableOperation(const FMutableOperation& Other)
 
 FMutableOperation& FMutableOperation::operator=(const FMutableOperation& Other)
 {
-	bNeverStream = Other.bNeverStream;
-	MipsToSkip = Other.MipsToSkip;
 	CustomizableObjectInstance = Other.CustomizableObjectInstance;
 	InstanceDescriptorRuntimeHash = Other.InstanceDescriptorRuntimeHash;
-	bStarted = Other.bStarted;
 	bBuildParameterRelevancy = Other.bBuildParameterRelevancy;
 	Parameters = Other.Parameters;
 	TextureParameters = Other.TextureParameters;
@@ -3067,28 +3047,25 @@ FMutableOperation::~FMutableOperation()
 }
 
 
-FMutableOperation FMutableOperation::CreateInstanceUpdate(UCustomizableObjectInstance* InCustomizableObjectInstance, bool bInNeverStream, int32 InMipsToSkip, const FInstanceUpdateDelegate* UpdateCallback)
+FMutableOperation FMutableOperation::CreateInstanceUpdate(UCustomizableObjectInstance& InCustomizableObjectInstance, const FInstanceUpdateDelegate* UpdateCallback)
 {
-	check(InCustomizableObjectInstance != nullptr);
-	check(InCustomizableObjectInstance->GetPrivate() != nullptr);
-	check(InCustomizableObjectInstance->GetCustomizableObject() != nullptr);
+	check(InCustomizableObjectInstance.GetPrivate() != nullptr);
+	check(InCustomizableObjectInstance.GetCustomizableObject() != nullptr);
 
 	FMutableOperation Op;
-	Op.bNeverStream = bInNeverStream;
-	Op.MipsToSkip = InMipsToSkip;
-	Op.CustomizableObjectInstance = InCustomizableObjectInstance;
-	Op.InstanceDescriptorRuntimeHash = InCustomizableObjectInstance->GetUpdateDescriptorRuntimeHash();
-	Op.bStarted = false;
-	Op.bBuildParameterRelevancy = InCustomizableObjectInstance->GetBuildParameterRelevancy();
-	Op.Parameters = InCustomizableObjectInstance->GetDescriptor().GetParameters();
-	Op.TextureParameters = InCustomizableObjectInstance->GetPrivate()->UpdateTextureParameters;
+	Op.CustomizableObjectInstance = &InCustomizableObjectInstance;
+	Op.InstanceDescriptorRuntimeHash = InCustomizableObjectInstance.GetUpdateDescriptorRuntimeHash();
+	Op.State = InCustomizableObjectInstance.GetState();
+	Op.bBuildParameterRelevancy = InCustomizableObjectInstance.GetBuildParameterRelevancy();
+	Op.Parameters = InCustomizableObjectInstance.GetDescriptor().GetParameters();
+	Op.TextureParameters = InCustomizableObjectInstance.GetPrivate()->UpdateTextureParameters;
 
 	if (UpdateCallback)
 	{
 		Op.UpdateCallback = *UpdateCallback;		
 	}
 	
-	InCustomizableObjectInstance->GetCustomizableObject()->ApplyStateForcedValuesToParameters(InCustomizableObjectInstance->GetState(), Op.Parameters.get());
+	InCustomizableObjectInstance.GetCustomizableObject()->ApplyStateForcedValuesToParameters(Op.State, Op.Parameters.get());
 
 	if (!Op.Parameters)
 	{
diff --git a/Engine/Plugins/Experimental/Mutable/Source/CustomizableObject/Private/MuCO/CustomizableObjectSystemPrivate.h b/Engine/Plugins/Experimental/Mutable/Source/CustomizableObject/Private/MuCO/CustomizableObjectSystemPrivate.h
--- a/Engine/Plugins/Experimental/Mutable/Source/CustomizableObject/Private/MuCO/CustomizableObjectSystemPrivate.h
+++ b/Engine/Plugins/Experimental/Mutable/Source/CustomizableObject/Private/MuCO/CustomizableObjectSystemPrivate.h
@@ -42,13 +42,13 @@ class FMutableOperation
 	TArray<FName> TextureParameters;
 
 	bool bBuildParameterRelevancy = false;
-	bool bMeshNeedsUpdate = false;
-
-	/** Protected constructor. */
-	FMutableOperation() {}
 
+	/** Instance optimization state. */
+	int32 State = 0;
+	
+	FMutableOperation() = default;
+	
 public:
-
 	FMutableOperation(const FMutableOperation&);
 	FMutableOperation(FMutableOperation&&) = default;
 	FMutableOperation& operator=(const FMutableOperation&);
@@ -56,18 +56,7 @@ public:
 
 	~FMutableOperation();
 
-	static FMutableOperation CreateInstanceUpdate(UCustomizableObjectInstance* COInstance, bool bInNeverStream, int32 MipsToSkip, const FInstanceUpdateDelegate* UpdateCallback);
-
-	bool bStarted = false;
-
-	bool bForceGenerateAllLODs = false;
-
-	bool bNeverStream = false;
-
-	/** If this is non-zero, the instance images will be generated with a certain amount of mips starting from the smallest. Otherwise, the full images will be generated. 
-	 * This is used for both Update and UpdateImage operations. 
-	 */
-	int32 MipsToSkip = 0;
+	static FMutableOperation CreateInstanceUpdate(UCustomizableObjectInstance& COInstance, const FInstanceUpdateDelegate* UpdateCallback);
 
 	// Weak reference to the instance we are operating on.
 	// It is weak because we don't want to lock it in case it becomes irrelevant in the game while operations are pending and it needs to be destroyed.
@@ -79,12 +68,6 @@ public:
 	//! This is used to calculate stats.
 	double StartUpdateTime = 0.0;
 
-	/** Instance optimization state. */
-	int32 State = -1;
-
-	/** Only used in the IDRelease operation type */
-	mu::Instance::ID IDToRelease;
-
 	FInstanceUpdateDelegate UpdateCallback;
 	
 	bool IsBuildParameterRelevancy() const
@@ -97,6 +80,11 @@ public:
 	{
 		return Parameters;
 	}
+
+	int32 GetState() const
+	{
+		return State;
+	}
 };
 
 
@@ -109,12 +97,6 @@ struct FMutablePendingInstanceUpdate
 
 	double SecondsAtUpdate = 0;
 	FInstanceUpdateDelegate* Callback = nullptr;
-	bool bNeverStream = false;
-
-	/** If this is non-zero, the instance images will be generated with a certain amount of mips starting from the smallest. Otherwise, the full images will be generated.
-	 * This is used for both Update and UpdateImage operations.
-	 */
-	int32 MipsToSkip = 0;
 
 	FMutablePendingInstanceUpdate(UCustomizableObjectInstance* InCustomizableObjectInstance)
 	{
@@ -123,7 +105,7 @@ struct FMutablePendingInstanceUpdate
 	}
 
 	FMutablePendingInstanceUpdate(UCustomizableObjectInstance* InCustomizableObjectInstance, EQueuePriorityType NewPriorityType, 
-		                          FInstanceUpdateDelegate* InCallback, bool bInNeverStream, int32 InMipsToSkip)
+		                          FInstanceUpdateDelegate* InCallback)
 	{
 		PriorityType = NewPriorityType;
 
@@ -133,8 +115,6 @@ struct FMutablePendingInstanceUpdate
 
 		SecondsAtUpdate = FPlatformTime::Seconds();
 		Callback = InCallback;
-		bNeverStream = bInNeverStream;
-		MipsToSkip = InMipsToSkip;
 	}
 
 	friend bool operator ==(const FMutablePendingInstanceUpdate& A, const FMutablePendingInstanceUpdate& B)
@@ -723,8 +703,7 @@ public:
 
 	bool TextureHasReferences(const FMutableImageCacheKey& TextureId) const;
 
-	// Init the async Skeletal Mesh creation/update
-	void InitUpdateSkeletalMesh(UCustomizableObjectInstance& Public, EQueuePriorityType Priority, bool bIsCloseDistTick, FInstanceUpdateDelegate* UpdateCallback = nullptr);
+	void EnqueueUpdateSkeletalMesh(UCustomizableObjectInstance& Instance, EQueuePriorityType Priority, FInstanceUpdateDelegate* UpdateCallback);
 		
 	// Init an async and safe release of the UE and Mutable resources used by the instance without actually destroying the instance, for example if it's very far away
 	void InitDiscardResourcesSkeletalMesh(UCustomizableObjectInstance* InCustomizableObjectInstance);
diff --git a/Engine/Plugins/Experimental/Mutable/Source/CustomizableObject/Private/MuCO/CustomizableSkeletalComponent.cpp b/Engine/Plugins/Experimental/Mutable/Source/CustomizableObject/Private/MuCO/CustomizableSkeletalComponent.cpp
--- a/Engine/Plugins/Experimental/Mutable/Source/CustomizableObject/Private/MuCO/CustomizableSkeletalComponent.cpp
+++ b/Engine/Plugins/Experimental/Mutable/Source/CustomizableObject/Private/MuCO/CustomizableSkeletalComponent.cpp
@@ -185,7 +185,7 @@ void UCustomizableSkeletalComponent::EditorUpdateComponent()
 {
 	if (CustomizableObjectInstance)
 	{
-		CustomizableObjectInstance->GetPrivate()->SetCOInstanceFlags((ECOInstanceFlags)(UsedByComponent | ForceGenerateAllLODs));
+		CustomizableObjectInstance->GetPrivate()->SetCOInstanceFlags(UsedByComponent);
 
 		AActor* ParentActor = GetAttachmentRootActor();
 		if (ParentActor)
