#rnx

diff --git a/Engine/Plugins/Experimental/Mutable/Source/CustomizableObject/Private/MuCO/CustomizableObjectInstance.cpp b/Engine/Plugins/Experimental/Mutable/Source/CustomizableObject/Private/MuCO/CustomizableObjectInstance.cpp
--- a/Engine/Plugins/Experimental/Mutable/Source/CustomizableObject/Private/MuCO/CustomizableObjectInstance.cpp
+++ b/Engine/Plugins/Experimental/Mutable/Source/CustomizableObject/Private/MuCO/CustomizableObjectInstance.cpp
@@ -1748,7 +1748,7 @@ void SetMeshUVChannelDensity(FMeshUVChannelInfo& UVChannelInfo, float Density =
 }
 
 // TODO PRP: MTBL-1653 Remove CVar
-static bool bApplyFixDoComponentsNeedUpdate = true;
+static bool bApplyFixDoComponentsNeedUpdate = false;
 FAutoConsoleVariableRef CVarApplyFixDoComponentsNeedUpdate(
 	TEXT("Mutable.ApplyFixDoComponentsNeedUpdate"),
 	bApplyFixDoComponentsNeedUpdate,
