#rnx

diff --git a/Engine/Plugins/Experimental/Mutable/Source/CustomizableObjectEditor/Private/MuCOE/CustomizableObjectEditorViewportClient.cpp b/Engine/Plugins/Experimental/Mutable/Source/CustomizableObjectEditor/Private/MuCOE/CustomizableObjectEditorViewportClient.cpp
--- a/Engine/Plugins/Experimental/Mutable/Source/CustomizableObjectEditor/Private/MuCOE/CustomizableObjectEditorViewportClient.cpp
+++ b/Engine/Plugins/Experimental/Mutable/Source/CustomizableObjectEditor/Private/MuCOE/CustomizableObjectEditorViewportClient.cpp
@@ -2233,6 +2233,7 @@ void FCustomizableObjectEditorViewportClient::BakeInstance()
 	BakeInstance(nullptr);
 }
 
+
 //-------------------------------------------------------------------------------------------------
 void FCustomizableObjectEditorViewportClient::BakeInstance(UCustomizableObjectInstance* InInstance)
 {
@@ -2651,32 +2652,49 @@ void FCustomizableObjectEditorViewportClient::BakeInstance(UCustomizableObjectIn
 								{
 									if (Inst->TextureParameterValues[t].ParameterValue)
 									{
-										UTexture2D* SrcTex = Cast<UTexture2D>(Inst->TextureParameterValues[t].ParameterValue);
-										FString ParameterSanitized = Inst->TextureParameterValues[t].ParameterInfo.Name.ToString();
-										RemoveRestrictedChars(ParameterSanitized);
-
-										FString TexObjName = ObjectName + "_" + MaterialName + "_" + ParameterSanitized;
-
-										if (!GetUniqueResourceName(SrcTex, TexObjName, ArrayCachedObject, ArrayCachedElement))
+										if (Inst->TextureParameterValues[t].ParameterValue->HasAnyFlags(RF_Transient))
 										{
-											UTexture* PrevTexture = Cast<UTexture>(ArrayCachedObject[ArrayCachedElement.Find(TexObjName)]);
-											InstDynamic->SetTextureParameterValue(Inst->TextureParameterValues[t].ParameterInfo.Name, PrevTexture);
-											continue;
+											UTexture2D* SrcTex = Cast<UTexture2D>(Inst->TextureParameterValues[t].ParameterValue);
+
+											if (SrcTex)
+											{
+												FString ParameterSanitized = Inst->TextureParameterValues[t].ParameterInfo.Name.ToString();
+												RemoveRestrictedChars(ParameterSanitized);
+
+												FString TexObjName = ObjectName + "_" + MaterialName + "_" + ParameterSanitized;
+
+												if (!GetUniqueResourceName(SrcTex, TexObjName, ArrayCachedObject, ArrayCachedElement))
+												{
+													UTexture* PrevTexture = Cast<UTexture>(ArrayCachedObject[ArrayCachedElement.Find(TexObjName)]);
+													InstDynamic->SetTextureParameterValue(Inst->TextureParameterValues[t].ParameterInfo.Name, PrevTexture);
+													continue;
+												}
+
+												if (!ManageBakingAction(AssetPath, TexObjName))
+												{
+													return;
+												}
+
+												FString TexPkgName = FolderDlg->GetAssetPath() + FString("/") + TexObjName;
+												TMap<UObject*, UObject*> FakeReplacementMap;
+												UTexture2D* DupTex = FUnrealBakeHelpers::BakeHelper_CreateAssetTexture(SrcTex, TexObjName, TexPkgName, nullptr, false, FakeReplacementMap, BakingOverwritePermission);
+												ArrayCachedObject.Add(DupTex);
+												ArrayCachedElement.Add(TexObjName);
+												PackagesToSave.Add(DupTex->GetPackage());
+
+												InstDynamic->SetTextureParameterValue(Inst->TextureParameterValues[t].ParameterInfo.Name, DupTex);
+											}
+											else
+											{
+												UE_LOG(LogMutable, Error, TEXT("A Mutable texture that is not a Texture2D has been found while baking a CustomizableObjectInstance."));
+											}
 										}
-
-										if (!ManageBakingAction(AssetPath, TexObjName))
+										else
 										{
-											return;
+											// If it's not transient it's not a mutable texture, it's a pass-through texture
+											// Just set the original texture
+											InstDynamic->SetTextureParameterValue(Inst->TextureParameterValues[t].ParameterInfo.Name, Inst->TextureParameterValues[t].ParameterValue);
 										}
-
-										FString TexPkgName = FolderDlg->GetAssetPath() + FString("/") + TexObjName;
-										TMap<UObject*, UObject*> FakeReplacementMap;
-										UTexture2D* DupTex = FUnrealBakeHelpers::BakeHelper_CreateAssetTexture(SrcTex, TexObjName, TexPkgName, nullptr, false, FakeReplacementMap, BakingOverwritePermission);
-										ArrayCachedObject.Add(DupTex);
-										ArrayCachedElement.Add(TexObjName);
-										PackagesToSave.Add(DupTex->GetPackage());
-
-										InstDynamic->SetTextureParameterValue(Inst->TextureParameterValues[t].ParameterInfo.Name, DupTex);
 									}
 								}
 							}
