[Mutable] Switch to the New TaskGraph API.

#rnx

diff --git a/Engine/Plugins/Experimental/Mutable/Source/CustomizableObject/Private/MuCO/UnrealMutableImageProvider.cpp b/Engine/Plugins/Experimental/Mutable/Source/CustomizableObject/Private/MuCO/UnrealMutableImageProvider.cpp
--- a/Engine/Plugins/Experimental/Mutable/Source/CustomizableObject/Private/MuCO/UnrealMutableImageProvider.cpp
+++ b/Engine/Plugins/Experimental/Mutable/Source/CustomizableObject/Private/MuCO/UnrealMutableImageProvider.cpp
@@ -352,7 +352,14 @@ mu::EImageFormat GetMutablePixelFormat(EPixelFormat InTextureFormat)
 				{
 					UE_LOG(LogMutable, Warning, TEXT("Failed to create an IORequest for a UTexture2D BulkData for an application-specific image parameter."));
 
-					IORequestCompletionEvent->DispatchSubsequents();
+#ifdef MUTABLE_USE_NEW_TASKGRAPH
+					IORequestCompletionEvent.Trigger();
+#else
+					if (IORequestCompletionEvent.IsValid())
+					{
+						IORequestCompletionEvent->DispatchSubsequents();
+					}
+#endif
 					
 					ResultCallback(CreateDummy());
 					return Invoke(TrivialReturn);
diff --git a/Engine/Plugins/Experimental/Mutable/Source/MutableRuntime/Public/MuR/System.h b/Engine/Plugins/Experimental/Mutable/Source/MutableRuntime/Public/MuR/System.h
--- a/Engine/Plugins/Experimental/Mutable/Source/MutableRuntime/Public/MuR/System.h
+++ b/Engine/Plugins/Experimental/Mutable/Source/MutableRuntime/Public/MuR/System.h
@@ -14,7 +14,7 @@
 #include "Templates/Tuple.h"
 
 /** This define will use the newer task graph interface to manage mutable concurrency. It has not been fully tested. */
-//#define MUTABLE_USE_NEW_TASKGRAPH
+#define MUTABLE_USE_NEW_TASKGRAPH
 
 /** If set to 1, this enables some expensive Unreal Insights traces, but can lead to 5x slower mutable operation. 
 * Other cheaper traces are enabled at all times.
