[Mutable] Fix COI Editor Texture Parameter Selector when adding a Default Value for the first time.

#rnx

diff --git a/Engine/Plugins/Experimental/Mutable/Source/CustomizableObjectEditor/Private/MuCOE/SCustomizableInstanceProperties.cpp b/Engine/Plugins/Experimental/Mutable/Source/CustomizableObjectEditor/Private/MuCOE/SCustomizableInstanceProperties.cpp
--- a/Engine/Plugins/Experimental/Mutable/Source/CustomizableObjectEditor/Private/MuCOE/SCustomizableInstanceProperties.cpp
+++ b/Engine/Plugins/Experimental/Mutable/Source/CustomizableObjectEditor/Private/MuCOE/SCustomizableInstanceProperties.cpp
@@ -80,12 +80,24 @@ void SCustomizableInstanceProperties::Construct(const FArguments& InArgs)
 		// Get default values.
 		for (const FCustomizableObjectTextureParameterValue& TextureParameter : CustomInstance->GetTextureParameters())
 		{
-			const FName DefaultValue = CustomizableObject->GetTextureParameterDefaultValue(TextureParameter.ParameterName);
+			const FName Value = CustomizableObject->GetTextureParameterDefaultValue(TextureParameter.ParameterName);
 			
-			if (!TextureParameterValues.Contains(DefaultValue))
+			if (!TextureParameterValues.Contains(Value))
 			{
-				TextureParameterValueNames.Add(MakeShareable(new FString(DefaultValue.ToString())));
-				TextureParameterValues.Add(DefaultValue);	
+				TextureParameterValueNames.Add(MakeShareable(new FString(Value.ToString())));
+				TextureParameterValues.Add(Value);	
+			}
+		}
+
+		// Selected parameter
+		for (const FCustomizableObjectTextureParameterValue& TextureParameter : CustomInstance->GetTextureParameters())
+		{
+			const FName& Value = TextureParameter.ParameterValue;
+			
+			if (!TextureParameterValues.Contains(Value))
+			{
+				TextureParameterValueNames.Add(MakeShareable(new FString(Value.ToString())));
+				TextureParameterValues.Add(Value);	
 			}
 		}
 
@@ -93,8 +105,13 @@ void SCustomizableInstanceProperties::Construct(const FArguments& InArgs)
 		TArray<FCustomizableObjectExternalTexture> Textures = UCustomizableObjectSystem::GetInstance()->GetTextureParameterValues();
 		for (int i = 0; i < Textures.Num(); ++i)
 		{
-			TextureParameterValueNames.Add(MakeShareable(new FString(Textures[i].Name)));
-			TextureParameterValues.Add(Textures[i].Value);
+			const FName& Value = FName(Textures[i].Value);
+
+			if (!TextureParameterValues.Contains(Value))
+			{
+				TextureParameterValueNames.Add(MakeShareable(new FString(Textures[i].Name)));
+				TextureParameterValues.Add(Value);
+			}
 		}
 
 		// Get values from TextureParameterDeclarations.
@@ -104,10 +121,15 @@ void SCustomizableInstanceProperties::Construct(const FArguments& InArgs)
 			{
 				continue;
 			}
+
+			const FName Value = FName(Declaration.GetPathName());
 			
-			TextureParameterValueNames.Add(MakeShareable(new FString(Declaration.GetName())));
-			TextureParameterValues.Add(FName(Declaration.GetPathName()));
-		}
+			if (!TextureParameterValues.Contains(Value))
+			{
+				TextureParameterValueNames.Add(MakeShareable(new FString(Declaration.GetName())));
+				TextureParameterValues.Add(Value);
+			}
+		}		
 
 		SetParameterProfileNamesOnEditor();
 		
@@ -836,11 +858,9 @@ void SCustomizableInstanceProperties::AddParameter(int32 ParamIndexInObject)
 
 	case EMutableParameterType::Texture:
 	{
-		const TArray<FCustomizableObjectTextureParameterValue>& TextureParameters = CustomInstance->GetTextureParameters();
-		TSharedPtr<FString> InitiallySelected = TextureParameterValueNames[0]; // First index is always the None option.
-
 		const FName ParameterValue = CustomInstance->GetTextureParameterSelectedOption(ParamName);
-		
+		TSharedPtr<FString> InitiallySelected;
+
 		// Look for the value index
 		for (int32 ValueIndex = 0; ValueIndex < TextureParameterValueNames.Num(); ++ValueIndex)
 		{
