#rnx

diff --git a/Engine/Plugins/Experimental/Mutable/Source/CustomizableObject/Private/MuCO/CustomizableInstancePrivateData.h b/Engine/Plugins/Experimental/Mutable/Source/CustomizableObject/Private/MuCO/CustomizableInstancePrivateData.h
--- a/Engine/Plugins/Experimental/Mutable/Source/CustomizableObject/Private/MuCO/CustomizableInstancePrivateData.h
+++ b/Engine/Plugins/Experimental/Mutable/Source/CustomizableObject/Private/MuCO/CustomizableInstancePrivateData.h
@@ -234,7 +234,7 @@ public:
 
 	void BuildMaterials(const TSharedPtr<FMutableOperationData>& OperationData, UCustomizableObjectInstance* Public);
 
-	void ReuseTexture(UTexture2D* Texture);
+	void ReuseTexture(UTexture2D* Texture, TSharedRef<FTexturePlatformData, ESPMode::ThreadSafe>& PlatformData);
 
 	// Return an event that will be fired when the assets  have been loaded. It returns null if no asset needs loading.
 	FGraphEventRef LoadAdditionalAssetsAsync(const TSharedPtr<FMutableOperationData>& OperationData, UCustomizableObjectInstance* Public, struct FStreamableManager &StreamableManager);
diff --git a/Engine/Plugins/Experimental/Mutable/Source/CustomizableObject/Private/MuCO/CustomizableObjectInstance.cpp b/Engine/Plugins/Experimental/Mutable/Source/CustomizableObject/Private/MuCO/CustomizableObjectInstance.cpp
--- a/Engine/Plugins/Experimental/Mutable/Source/CustomizableObject/Private/MuCO/CustomizableObjectInstance.cpp
+++ b/Engine/Plugins/Experimental/Mutable/Source/CustomizableObject/Private/MuCO/CustomizableObjectInstance.cpp
@@ -5526,7 +5526,8 @@ void UCustomizableInstancePrivateData::AdditionalAssetsAsyncLoaded(UCustomizable
 }
 
 
-void UpdateTextureRegionsMutable(UTexture2D* Texture, int32 MipIndex, uint32 NumMips, const FUpdateTextureRegion2D& Region, uint32 SrcPitch, const FByteBulkData* BulkData)
+void UpdateTextureRegionsMutable(UTexture2D* Texture, int32 MipIndex, uint32 NumMips, const FUpdateTextureRegion2D& Region, uint32 SrcPitch, 
+								 const FByteBulkData* BulkData, TSharedRef<FTexturePlatformData, ESPMode::ThreadSafe>& PlatformData)
 {
 	if (Texture->GetResource())
 	{
@@ -5537,9 +5538,15 @@ void UpdateTextureRegionsMutable(UTexture2D* Texture, int32 MipIndex, uint32 Num
 			FUpdateTextureRegion2D Region;
 			uint32 SrcPitch;
 			uint32 NumMips;
+			
+			// The Platform Data mips will be automatically deleted when all FUpdateTextureRegionsData that reference it are deleted
+			// in the render thread after being used to update the texture
+			TSharedRef<FTexturePlatformData, ESPMode::ThreadSafe> PlatformData;
+
+			FUpdateTextureRegionsData(TSharedRef<FTexturePlatformData, ESPMode::ThreadSafe>& InPlatformData) : PlatformData(InPlatformData) {}
 		};
 
-		FUpdateTextureRegionsData* RegionData = new FUpdateTextureRegionsData;
+		FUpdateTextureRegionsData* RegionData = new FUpdateTextureRegionsData(PlatformData);
 
 		RegionData->Texture2DResource = (FTexture2DResource*)Texture->GetResource();
 		RegionData->MipIndex = MipIndex;
@@ -5575,38 +5582,36 @@ void UpdateTextureRegionsMutable(UTexture2D* Texture, int32 MipIndex, uint32 Num
 				}
 
 				BulkData->Unlock();
-				delete RegionData;
+				delete RegionData; // This will implicitly delete the Platform Data if this is the last RegionData referencing it
 			});
 	}
 }
 
 
-void UCustomizableInstancePrivateData::ReuseTexture(UTexture2D* Texture)
+void UCustomizableInstancePrivateData::ReuseTexture(UTexture2D* Texture, TSharedRef<FTexturePlatformData, ESPMode::ThreadSafe>& PlatformData)
 {
-	if (Texture->GetPlatformData())
+	uint32 NumMips = PlatformData->Mips.Num();
+
+	for (uint32 i = 0; i < NumMips; i++)
 	{
-		uint32 NumMips = Texture->GetPlatformData()->Mips.Num();
+		FTexture2DMipMap& Mip = PlatformData->Mips[i];
 
-		for (uint32 i = 0; i < NumMips; i++)
+		if (Mip.BulkData.GetElementCount() > 0)
 		{
-			FTexture2DMipMap& Mip = Texture->GetPlatformData()->Mips[i];
-
-			if (Mip.BulkData.GetElementCount() > 0)
-			{
-				FUpdateTextureRegion2D Region;
+			FUpdateTextureRegion2D Region;
 
-				Region.DestX = 0;
-				Region.DestY = 0;
-				Region.SrcX = 0;
-				Region.SrcY = 0;
-				Region.Width = Mip.SizeX;
-				Region.Height = Mip.SizeY;
+			Region.DestX = 0;
+			Region.DestY = 0;
+			Region.SrcX = 0;
+			Region.SrcY = 0;
+			Region.Width = Mip.SizeX;
+			Region.Height = Mip.SizeY;
 
-				check(int32(Region.Width) <= Texture->GetSizeX());
-				check(int32(Region.Height) <= Texture->GetSizeY());
+			check(int32(Region.Width) <= Texture->GetSizeX());
+			check(int32(Region.Height) <= Texture->GetSizeY());
 
-				UpdateTextureRegionsMutable(Texture, i, NumMips, Region, Mip.SizeX * sizeof(uint8) * 4, &Mip.BulkData);
-			}
+			UpdateTextureRegionsMutable(Texture, i, NumMips, Region, 
+					                    Mip.SizeX * sizeof(uint8) * 4, &Mip.BulkData, PlatformData);
 		}
 	}
 }
@@ -5868,6 +5873,10 @@ void UCustomizableInstancePrivateData::BuildMaterials(const TSharedPtr<FMutableO
 							if (!MutableTexture && !PassThroughTexture && MutableImage)
 							{
 								TWeakObjectPtr<UTexture2D>* ReusedTexture = bReuseTextures ? TextureReuseCache.Find(TextureReuseCacheRef) : nullptr;
+								
+								// This shared ptr will hold the reused texture platform data (mips) until the reused texture is updated 
+								// and delete it automatically
+								TSharedPtr<FTexturePlatformData, ESPMode::ThreadSafe> ReusedTexturePlatformData;
 
 								if (ReusedTexture && (*ReusedTexture).IsValid() && !(*ReusedTexture)->HasAnyFlags(RF_BeginDestroyed))
 								{
@@ -5912,9 +5921,16 @@ void UCustomizableInstancePrivateData::BuildMaterials(const TSharedPtr<FMutableO
 											check(PlatformData->Mips.Num() == MutableTexture->GetPlatformData()->Mips.Num());
 											check(PlatformData->Mips[0].SizeX == MutableTexture->GetPlatformData()->Mips[0].SizeX);
 											check(PlatformData->Mips[0].SizeY == MutableTexture->GetPlatformData()->Mips[0].SizeY);
+
+											// Now the ReusedTexturePlatformData shared ptr owns the platform data
+											ReusedTexturePlatformData = TSharedPtr<FTexturePlatformData, ESPMode::ThreadSafe>(PlatformData);
+										}
+										else
+										{
+											// Now the MutableTexture owns the platform data
+											MutableTexture->SetPlatformData(PlatformData);
 										}
 
-										MutableTexture->SetPlatformData(PlatformData);
 										OperationData->ImageToPlatformDataMap.Remove(Image.ImageID);
 									}
 									else
@@ -5957,7 +5973,13 @@ void UCustomizableInstancePrivateData::BuildMaterials(const TSharedPtr<FMutableO
 												}
 											}
 
-											ReuseTexture(MutableTexture);
+											check(ReusedTexturePlatformData.IsValid());
+
+											if (ReusedTexturePlatformData.IsValid())
+											{
+												TSharedRef<FTexturePlatformData, ESPMode::ThreadSafe> PlatformDataRef = ReusedTexturePlatformData.ToSharedRef();
+												ReuseTexture(MutableTexture, PlatformDataRef);
+											}
 										}
 										else if (MutableTexture)
 										{	
