[mutable] Protect images that are being loaded from being garbage collected during compilation.
Additional profiling for object compilation.
Minor changes to ease backports.
#rnx

diff --git a/Engine/Plugins/Experimental/Mutable/Source/CustomizableObject/Public/MuCO/UnrealPortabilityHelpers.h b/Engine/Plugins/Experimental/Mutable/Source/CustomizableObject/Public/MuCO/UnrealPortabilityHelpers.h
--- a/Engine/Plugins/Experimental/Mutable/Source/CustomizableObject/Public/MuCO/UnrealPortabilityHelpers.h
+++ b/Engine/Plugins/Experimental/Mutable/Source/CustomizableObject/Public/MuCO/UnrealPortabilityHelpers.h
@@ -7,6 +7,14 @@
 
 /** Helpers to ease portability across unreal engine versions */
 
+#if ENGINE_MAJOR_VERSION==5 && ENGINE_MINOR_VERSION<3
+
+#define ObjectPtrWrap(X)	X
+#define ObjectPtrDecay(X)	X
+
+#endif
+
+
 #if ENGINE_MAJOR_VERSION==5 && ENGINE_MINOR_VERSION>=1
 
 #include "Engine/SkinnedAssetCommon.h"
diff --git a/Engine/Plugins/Experimental/Mutable/Source/CustomizableObjectEditor/Private/MuCOE/CustomizableObjectCompileRunnable.cpp b/Engine/Plugins/Experimental/Mutable/Source/CustomizableObjectEditor/Private/MuCOE/CustomizableObjectCompileRunnable.cpp
--- a/Engine/Plugins/Experimental/Mutable/Source/CustomizableObjectEditor/Private/MuCOE/CustomizableObjectCompileRunnable.cpp
+++ b/Engine/Plugins/Experimental/Mutable/Source/CustomizableObjectEditor/Private/MuCOE/CustomizableObjectCompileRunnable.cpp
@@ -15,6 +15,8 @@ class ITargetPlatform;
 
 #define LOCTEXT_NAMESPACE "CustomizableObjectEditor"
 
+#define UE_MUTABLE_CORE_REGION	TEXT("Mutable Core")
+
 
 FCustomizableObjectCompileRunnable::FCustomizableObjectCompileRunnable(mu::Ptr<mu::Node> Root)
 	: MutableRoot(Root)
@@ -27,6 +29,8 @@ FCustomizableObjectCompileRunnable::FCustomizableObjectCompileRunnable(mu::Ptr<m
 
 uint32 FCustomizableObjectCompileRunnable::Run()
 {
+	TRACE_BEGIN_REGION(UE_MUTABLE_CORE_REGION);
+
 	UE_LOG(LogMutable, Verbose, TEXT("PROFILE: [ %16.8f ] FCustomizableObjectCompileRunnable::Run start."), FPlatformTime::Seconds());
 
 	uint32 Result = 1;
@@ -148,6 +152,8 @@ uint32 FCustomizableObjectCompileRunnable::Run()
 
 	UE_LOG(LogMutable, Verbose, TEXT("PROFILE: [ %16.8f ] FCustomizableObjectCompileRunnable::Run end."), FPlatformTime::Seconds());
 
+	TRACE_END_REGION(UE_MUTABLE_CORE_REGION);
+
 	return Result;
 }
 
diff --git a/Engine/Plugins/Experimental/Mutable/Source/CustomizableObjectEditor/Private/MuCOE/CustomizableObjectCompiler.cpp b/Engine/Plugins/Experimental/Mutable/Source/CustomizableObjectEditor/Private/MuCOE/CustomizableObjectCompiler.cpp
--- a/Engine/Plugins/Experimental/Mutable/Source/CustomizableObjectEditor/Private/MuCOE/CustomizableObjectCompiler.cpp
+++ b/Engine/Plugins/Experimental/Mutable/Source/CustomizableObjectEditor/Private/MuCOE/CustomizableObjectCompiler.cpp
@@ -27,6 +27,10 @@ class UTexture2D;
 
 #define LOCTEXT_NAMESPACE "CustomizableObjectEditor"
 
+#define UE_MUTABLE_COMPILE_REGION		TEXT("Mutable Compile")
+#define UE_MUTABLE_PRELOAD_REGION		TEXT("Mutable Preload")
+#define UE_MUTABLE_SAVEDD_REGION		TEXT("Mutable SaveDD")
+#define UE_MUTABLE_CONVERTING_REGION	TEXT("Mutable Converting")
 
 UCustomizableObjectNodeObject* GetRootNode(UCustomizableObject* Object, bool &bOutMultipleBaseObjectsFound);
 
@@ -77,6 +81,7 @@ bool FCustomizableObjectCompiler::Tick()
 
 		CompilationLogsContainer.ClearMessagesArray();
 
+		TRACE_END_REGION(UE_MUTABLE_COMPILE_REGION);
 
 	}
 
@@ -92,6 +97,8 @@ bool FCustomizableObjectCompiler::Tick()
 
 		UE_LOG(LogMutable, Verbose, TEXT("PROFILE: [ %16.8f ] Finished Saving Derived Data task."), FPlatformTime::Seconds());
 		UE_LOG(LogMutable, Verbose, TEXT("PROFILE: -----------------------------------------------------------"));
+
+		TRACE_END_REGION(UE_MUTABLE_SAVEDD_REGION);
 	}
 
 	// In editor, when compiling a CO, referencer assets and Unreal to Mutable texture conversion are performed asynchronously
@@ -175,12 +182,16 @@ void FCustomizableObjectCompiler::PreloadingReferencerAssetsCallback(UCustomizab
 	CustomizableObjectCompiler->SetPreloadingReferencerAssets(false);
 	UCustomizableObjectSystem::GetInstance()->UnlockObject(Object);
 
+	TRACE_END_REGION(UE_MUTABLE_PRELOAD_REGION);
+
 	CustomizableObjectCompiler->CompileInternal(Object, Options, bAsync);
 }
 
 
 void FCustomizableObjectCompiler::Compile(UCustomizableObject& Object, const FCompilationOptions& InOptions, bool bAsync)
 {
+	TRACE_BEGIN_REGION(UE_MUTABLE_COMPILE_REGION);
+
 	UCustomizableObjectSystem* System = UCustomizableObjectSystem::GetInstance();
 	check(System);
 
@@ -233,6 +244,7 @@ void FCustomizableObjectCompiler::Compile(UCustomizableObject& Object, const FCo
 	CurrentObject = &Object;
 
 	PreloadingReferencerAssets = true;
+	TRACE_BEGIN_REGION(UE_MUTABLE_PRELOAD_REGION);
 
 	CleanCachedReferencers();
 	UpdateArrayGCProtect();
@@ -295,6 +307,12 @@ void FCustomizableObjectCompiler::AddReferencedObjects(FReferenceCollector& Coll
 		Collector.AddReferencedObject(ArrayGCProtect[i]);
 	}
 
+	// Protect images that are pending to be loaded as well
+	for (FTextureUnrealToMutableTask& Image: ArrayTextureUnrealToMutableTask )
+	{
+		Collector.AddReferencedObject(Image.Texture);
+	}
+
 	if (CurrentObject)
 	{
 		Collector.AddReferencedObject(CurrentObject);
@@ -864,6 +882,8 @@ void FCustomizableObjectCompiler::UpdatePendingTextureConversion(bool UseTimeLim
 			PendingTexturesToLoad = false;
 			ArrayTextureUnrealToMutableTask.Empty();
 			CompletedUnrealToMutableTask = 0;
+
+			TRACE_END_REGION(UE_MUTABLE_CONVERTING_REGION);
 		}
 	}
 }
@@ -1288,7 +1308,7 @@ void FCustomizableObjectCompiler::CompileInternal(UCustomizableObject* Object, c
 		{
 			if (!Object->MaskOutCache.ToString().IsEmpty() || Object->MaskOutCache.Get())
 			{
-				if (!ParamNamesToSelectedOptions.Num()) // Don't marked the object as modified because of a partial compilation
+				if (!ParamNamesToSelectedOptions.Num()) // Don't mark the object as modified because of a partial compilation
 				{
 					Object->MaskOutCache = nullptr;
 					Object->MarkPackageDirty();
@@ -1318,6 +1338,7 @@ void FCustomizableObjectCompiler::CompileInternal(UCustomizableObject* Object, c
 		{
 			ArrayTextureUnrealToMutableTask.Insert(GenerationContext.ArrayTextureUnrealToMutableTask, ArrayTextureUnrealToMutableTask.Num());
 			PendingTexturesToLoad = true;
+			TRACE_BEGIN_REGION(UE_MUTABLE_CONVERTING_REGION);
 		}
 
 		// If synchronous compilation is requested, proceed the same way if in editor and if packaging
@@ -1460,6 +1481,8 @@ void FCustomizableObjectCompiler::FinishCompilation()
 
 	if (!Options.bDontUpdateStreamedDataAndCache)
 	{
+		TRACE_BEGIN_REGION(UE_MUTABLE_SAVEDD_REGION);
+
 		SaveDDTask = MakeShareable(new FCustomizableObjectSaveDDRunnable(CurrentObject, Options));
 	}
 	else
diff --git a/Engine/Plugins/Experimental/Mutable/Source/CustomizableObjectEditor/Private/MuCOE/GenerateMutableSource/GenerateMutableSource.h b/Engine/Plugins/Experimental/Mutable/Source/CustomizableObjectEditor/Private/MuCOE/GenerateMutableSource/GenerateMutableSource.h
--- a/Engine/Plugins/Experimental/Mutable/Source/CustomizableObjectEditor/Private/MuCOE/GenerateMutableSource/GenerateMutableSource.h
+++ b/Engine/Plugins/Experimental/Mutable/Source/CustomizableObjectEditor/Private/MuCOE/GenerateMutableSource/GenerateMutableSource.h
@@ -257,7 +257,7 @@ struct FTextureUnrealToMutableTask
 
 	mu::NodeImageConstantPtr ImageNode;
 	mu::TablePtr TableNode;
-	UTexture2D* Texture;
+	TObjectPtr<UTexture2D> Texture;
 	const UCustomizableObjectNode* Node;
 	int32 TableColumn;
 	int32 TableRow;
diff --git a/Engine/Plugins/Experimental/Mutable/Source/CustomizableObjectEditor/Private/MuCOE/GenerateMutableSource/GenerateMutableSourceImage.cpp b/Engine/Plugins/Experimental/Mutable/Source/CustomizableObjectEditor/Private/MuCOE/GenerateMutableSource/GenerateMutableSourceImage.cpp
--- a/Engine/Plugins/Experimental/Mutable/Source/CustomizableObjectEditor/Private/MuCOE/GenerateMutableSource/GenerateMutableSourceImage.cpp
+++ b/Engine/Plugins/Experimental/Mutable/Source/CustomizableObjectEditor/Private/MuCOE/GenerateMutableSource/GenerateMutableSourceImage.cpp
@@ -72,11 +72,11 @@ mu::ImagePtr ConvertTextureUnrealToMutable(UTexture2D* Texture, const UCustomiza
 					Node);
 			break;
 		}
-	case EUnrealToMutableConversionError::CompositeImageDimensionMissmatch:
+	case EUnrealToMutableConversionError::CompositeImageDimensionMismatch:
 		{
 			Compiler->CompilerLog(
-					LOCTEXT("CompositeImageDimensionMissmatch", 
-							"Composite image dimension missmatch."), 
+					LOCTEXT("CompositeImageDimensionMismatch", 
+							"Composite image dimension mismatch."), 
 					Node);
 			break;
 		}
diff --git a/Engine/Plugins/Experimental/Mutable/Source/CustomizableObjectEditor/Private/MuCOE/UnrealToMutableTextureConversionUtils.cpp b/Engine/Plugins/Experimental/Mutable/Source/CustomizableObjectEditor/Private/MuCOE/UnrealToMutableTextureConversionUtils.cpp
--- a/Engine/Plugins/Experimental/Mutable/Source/CustomizableObjectEditor/Private/MuCOE/UnrealToMutableTextureConversionUtils.cpp
+++ b/Engine/Plugins/Experimental/Mutable/Source/CustomizableObjectEditor/Private/MuCOE/UnrealToMutableTextureConversionUtils.cpp
@@ -48,7 +48,7 @@ EUnrealToMutableConversionError ApplyCompositeTexture(
    
     if (SizeX < Image.SizeX || SizeY < Image.SizeY || !bHaveSimilarAspect)
     {
-        return EUnrealToMutableConversionError::CompositeImageDimensionMissmatch;
+        return EUnrealToMutableConversionError::CompositeImageDimensionMismatch;
     }
 
     {
diff --git a/Engine/Plugins/Experimental/Mutable/Source/CustomizableObjectEditor/Private/MuCOE/UnrealToMutableTextureConversionUtils.h b/Engine/Plugins/Experimental/Mutable/Source/CustomizableObjectEditor/Private/MuCOE/UnrealToMutableTextureConversionUtils.h
--- a/Engine/Plugins/Experimental/Mutable/Source/CustomizableObjectEditor/Private/MuCOE/UnrealToMutableTextureConversionUtils.h
+++ b/Engine/Plugins/Experimental/Mutable/Source/CustomizableObjectEditor/Private/MuCOE/UnrealToMutableTextureConversionUtils.h
@@ -12,7 +12,7 @@ enum class EUnrealToMutableConversionError
 {
     Success,
     UnsupportedFormat,
-    CompositeImageDimensionMissmatch,
+    CompositeImageDimensionMismatch,
     CompositeUnsupportedFormat,
     Unknown
 };
diff --git a/Engine/Plugins/Experimental/Mutable/Source/MutableRuntime/Private/MuR/CodeRunner.cpp b/Engine/Plugins/Experimental/Mutable/Source/MutableRuntime/Private/MuR/CodeRunner.cpp
--- a/Engine/Plugins/Experimental/Mutable/Source/MutableRuntime/Private/MuR/CodeRunner.cpp
+++ b/Engine/Plugins/Experimental/Mutable/Source/MutableRuntime/Private/MuR/CodeRunner.cpp
@@ -4064,7 +4064,10 @@ namespace mu
 				int32 MipsToSkip = item.ExecutionOptions;
 				while ( MipsToSkip>0 && rect.size[0]>0 && rect.size[1]>0 )
 				{
-					rect.ShrinkToHalf();
+					rect.min[0] /= 2;
+					rect.min[1] /= 2;
+					rect.size[0] /= 2;
+					rect.size[1] /= 2;
 					MipsToSkip--;
 				}
 
@@ -4524,6 +4527,7 @@ namespace mu
 
                 ImageMakeGrowMap(Result.get(), Mask.get(), args.border );
 				Result->m_flags |= Image::IF_CANNOT_BE_SCALED;
+				check(Result->GetData());
 
 				Release(Mask);
                 StoreImage( item, Result);
diff --git a/Engine/Plugins/Experimental/Mutable/Source/MutableRuntime/Private/MuR/System.cpp b/Engine/Plugins/Experimental/Mutable/Source/MutableRuntime/Private/MuR/System.cpp
--- a/Engine/Plugins/Experimental/Mutable/Source/MutableRuntime/Private/MuR/System.cpp
+++ b/Engine/Plugins/Experimental/Mutable/Source/MutableRuntime/Private/MuR/System.cpp
@@ -1471,10 +1471,11 @@ namespace mu
 				bool bInclude = Params->GetPrivate()->HasMultipleValues(ParamIndex);
 				if (!bInclude)
 				{
-					bInclude =
+					bInclude = !(
 						Params->GetPrivate()->m_values[ParamIndex]
-						!=
-						ParamDescs[ParamIndex].m_defaultValue;
+						==
+						ParamDescs[ParamIndex].m_defaultValue
+						);
 				}
 
 				IncludedParameters[IndexIndex] = bInclude;
diff --git a/Engine/Plugins/Experimental/Mutable/Source/MutableRuntime/Public/MuR/MutableMath.h b/Engine/Plugins/Experimental/Mutable/Source/MutableRuntime/Public/MuR/MutableMath.h
--- a/Engine/Plugins/Experimental/Mutable/Source/MutableRuntime/Public/MuR/MutableMath.h
+++ b/Engine/Plugins/Experimental/Mutable/Source/MutableRuntime/Public/MuR/MutableMath.h
@@ -1025,13 +1025,6 @@ namespace mu
 			Bound(other.min + other.size);
 		}
 
-		//-----------------------------------------------------------------------------------------
-		inline void ShrinkToHalf()
-		{
-			min /= 2;
-			size /= 2;
-		}
-
 		//-----------------------------------------------------------------------------------------
 		inline bool IsEmpty() const
 		{
diff --git a/Engine/Plugins/Experimental/Mutable/Source/MutableRuntime/Public/MuR/Serialisation.h b/Engine/Plugins/Experimental/Mutable/Source/MutableRuntime/Public/MuR/Serialisation.h
--- a/Engine/Plugins/Experimental/Mutable/Source/MutableRuntime/Public/MuR/Serialisation.h
+++ b/Engine/Plugins/Experimental/Mutable/Source/MutableRuntime/Public/MuR/Serialisation.h
@@ -39,8 +39,7 @@ namespace mu
     void DLLEXPORT operator>> <T>(InputArchive& arch, T& t);		\
 	
 	
-    //! \brief
-    //! \ingroup model
+    /** This repesents a handle to Mutable resource. */
     template<class R>
     class ResourceProxy : public RefCounted
     {
@@ -51,6 +50,7 @@ namespace mu
     };
 
 
+	/** Proxy implementation that always has the resource loaded in memory. */
 	template<class R>
 	class ResourceProxyMemory : public ResourceProxy<R>
 	{
@@ -65,7 +65,7 @@ namespace mu
 			m_resource = i;
 		}
 
-		// ImageProxy interface
+		// ResourceProxy interface
 		Ptr<const R> Get() override
 		{
 			return m_resource;
