[Mutable] Add laplacian smoothing to mesh reshape oprations.

#rnx

diff --git a/Engine/Plugins/Experimental/Mutable/Source/CustomizableObject/Classes/MuCO/CustomizableObject.h b/Engine/Plugins/Experimental/Mutable/Source/CustomizableObject/Classes/MuCO/CustomizableObject.h
--- a/Engine/Plugins/Experimental/Mutable/Source/CustomizableObject/Classes/MuCO/CustomizableObject.h
+++ b/Engine/Plugins/Experimental/Mutable/Source/CustomizableObject/Classes/MuCO/CustomizableObject.h
@@ -1459,7 +1459,7 @@ private:
 	// This is a manual version number for the binary blobs in this asset.
 	// Increasing it invalidates all the previously compiled models.
 	// Warning: If while merging code both versions have changed, take the highest+1.
-	static const int32 CurrentSupportedVersion = 402;
+	static const int32 CurrentSupportedVersion = 403;
 
 public:
 
diff --git a/Engine/Plugins/Experimental/Mutable/Source/CustomizableObjectEditor/Classes/MuCOE/Nodes/CustomizableObjectNodeMeshReshape.h b/Engine/Plugins/Experimental/Mutable/Source/CustomizableObjectEditor/Classes/MuCOE/Nodes/CustomizableObjectNodeMeshReshape.h
--- a/Engine/Plugins/Experimental/Mutable/Source/CustomizableObjectEditor/Classes/MuCOE/Nodes/CustomizableObjectNodeMeshReshape.h
+++ b/Engine/Plugins/Experimental/Mutable/Source/CustomizableObjectEditor/Classes/MuCOE/Nodes/CustomizableObjectNodeMeshReshape.h
@@ -73,10 +73,14 @@ public:
 
 	void Serialize(FArchive& Ar) override;
 
-	/** Enable the deformation of the skeleton of the base mesh. */
+	/** Enable the deformation of the vertices of the base mesh. */
 	UPROPERTY(EditAnywhere, Category = CustomizableObject)
 	bool bReshapeVertices = true;	
 
+	/** Enable laplacian smoothing to the result of the base mesh reshape. */
+	UPROPERTY(EditAnywhere, Category = CustomizableObject, meta = (EditCondition = "bReshapeVertices"))
+	bool bApplyLaplacianSmoothing = false;
+
 	/** Enable the deformation of the skeleton of the base mesh. */
 	UPROPERTY(EditAnywhere, Category = CustomizableObject)
 	bool bReshapeSkeleton = false;	
diff --git a/Engine/Plugins/Experimental/Mutable/Source/CustomizableObjectEditor/Private/MuCOE/GenerateMutableSource/GenerateMutableSourceMesh.cpp b/Engine/Plugins/Experimental/Mutable/Source/CustomizableObjectEditor/Private/MuCOE/GenerateMutableSource/GenerateMutableSourceMesh.cpp
--- a/Engine/Plugins/Experimental/Mutable/Source/CustomizableObjectEditor/Private/MuCOE/GenerateMutableSource/GenerateMutableSourceMesh.cpp
+++ b/Engine/Plugins/Experimental/Mutable/Source/CustomizableObjectEditor/Private/MuCOE/GenerateMutableSource/GenerateMutableSourceMesh.cpp
@@ -3400,6 +3400,7 @@ mu::NodeMeshPtr GenerateMutableSourceMesh(const UEdGraphPin* Pin,
 		{
 
 			MeshNode->SetReshapeVertices(TypedNodeReshape->bReshapeVertices);
+			MeshNode->SetApplyLaplacian(TypedNodeReshape->bApplyLaplacianSmoothing);
 			MeshNode->SetReshapeSkeleton(TypedNodeReshape->bReshapeSkeleton);
 			MeshNode->SetReshapePhysicsVolumes(TypedNodeReshape->bReshapePhysicsVolumes);
 
diff --git a/Engine/Plugins/Experimental/Mutable/Source/MutableRuntime/Private/MuR/CodeRunner.cpp b/Engine/Plugins/Experimental/Mutable/Source/MutableRuntime/Private/MuR/CodeRunner.cpp
--- a/Engine/Plugins/Experimental/Mutable/Source/MutableRuntime/Private/MuR/CodeRunner.cpp
+++ b/Engine/Plugins/Experimental/Mutable/Source/MutableRuntime/Private/MuR/CodeRunner.cpp
@@ -2334,7 +2334,7 @@ namespace mu
 						if (!EnumHasAnyFlags(BindFlags, EMeshBindShapeFlags::ReshapeVertices))
 						{
 							Ptr<Mesh> BindMeshNoVertsResult = CloneOrTakeOver(BaseMesh);
-							BindMeshNoVertsResult->m_AdditionalBuffers = MoveTemp(BindMeshResult->m_AdditionalBuffers);
+							BindMeshNoVertsResult->AdditionalBuffers = MoveTemp(BindMeshResult->AdditionalBuffers);
 
 							Release(BaseMesh);
 							Release(BindMeshResult);
diff --git a/Engine/Plugins/Experimental/Mutable/Source/MutableRuntime/Private/MuR/Mesh.cpp b/Engine/Plugins/Experimental/Mutable/Source/MutableRuntime/Private/MuR/Mesh.cpp
--- a/Engine/Plugins/Experimental/Mutable/Source/MutableRuntime/Private/MuR/Mesh.cpp
+++ b/Engine/Plugins/Experimental/Mutable/Source/MutableRuntime/Private/MuR/Mesh.cpp
@@ -59,7 +59,7 @@ MeshPtr Mesh::Clone() const
     pResult->m_FaceBuffers = m_FaceBuffers;
 
 	// Clone additional buffers
-	pResult->m_AdditionalBuffers = m_AdditionalBuffers;
+	pResult->AdditionalBuffers = AdditionalBuffers;
 
     // Clone the layouts
 	pResult->m_layouts = m_layouts;
@@ -134,7 +134,7 @@ MeshPtr Mesh::Clone(EMeshCopyFlags Flags) const
 	// Clone additional buffers
 	if (EnumHasAnyFlags(Flags, EMeshCopyFlags::WithAdditionalBuffers))
 	{
-		pResult->m_AdditionalBuffers = m_AdditionalBuffers;
+		pResult->AdditionalBuffers = AdditionalBuffers;
 	}
 
     // Clone the layout	
@@ -226,7 +226,7 @@ void Mesh::CopyFrom(const Mesh& From, EMeshCopyFlags Flags)
 	// Copy additional buffers
 	if (EnumHasAnyFlags(Flags, EMeshCopyFlags::WithAdditionalBuffers))
 	{
-		m_AdditionalBuffers = From.m_AdditionalBuffers;
+		AdditionalBuffers = From.AdditionalBuffers;
 	}
 
     // Copy the layout	
@@ -698,7 +698,7 @@ int32 Mesh::GetDataSize() const
 
 	// Should be allocation sizes used for this?
 	int32 AdditionalBuffersSize = 0;
-	for (const TPair<EMeshBufferType, FMeshBufferSet>&  AdditionalBuffer : m_AdditionalBuffers)
+	for (const TPair<EMeshBufferType, FMeshBufferSet>&  AdditionalBuffer : AdditionalBuffers)
 	{
 		AdditionalBuffersSize += AdditionalBuffer.Value.GetDataSize();
 	}
@@ -1275,7 +1275,7 @@ void Mesh::Serialise(OutputArchive& arch) const
 	arch << m_IndexBuffers;
 	arch << m_VertexBuffers;
 	arch << m_FaceBuffers;
-	arch << m_AdditionalBuffers;
+	arch << AdditionalBuffers;
 	arch << m_layouts;
 
 	arch << SkeletonIDs;
@@ -1306,7 +1306,7 @@ void Mesh::Unserialise(InputArchive& arch)
 	arch >> m_IndexBuffers;
 	arch >> m_VertexBuffers;
 	arch >> m_FaceBuffers;
-	arch >> m_AdditionalBuffers;
+	arch >> AdditionalBuffers;
 	arch >> m_layouts;
 
 	if (ver >= 14)
diff --git a/Engine/Plugins/Experimental/Mutable/Source/MutableRuntime/Private/MuR/OpMeshBind.h b/Engine/Plugins/Experimental/Mutable/Source/MutableRuntime/Private/MuR/OpMeshBind.h
--- a/Engine/Plugins/Experimental/Mutable/Source/MutableRuntime/Private/MuR/OpMeshBind.h
+++ b/Engine/Plugins/Experimental/Mutable/Source/MutableRuntime/Private/MuR/OpMeshBind.h
@@ -14,6 +14,7 @@
 #include "MuR/Operations.h"
 #include "Spatial/MeshAABBTree3.h"
 
+#include "MuR/OpMeshSmoothing.h"
 
 
 // TODO: Make the handling of rotations an option. It is more expensive on CPU and memory, and for some
@@ -781,6 +782,123 @@ namespace mu
 		return BindData;
 	}
 
+	inline void GenerateAndAddLaplacianData(Mesh& InOutMesh)
+	{
+		MUTABLE_CPUPROFILER_SCOPE(GenerateLaplacianData);
+
+
+		// Storage for buffers in a format different than the supported one.
+		// Not used if the buffer data is compatible. The data will always be accessed using a 
+		// view regardless of compatibility.
+		TArray<FVector3f> ConvertedVerticesStorage;
+		TArray<uint32>    ConvertedIndicesStorage;
+
+		TArrayView<const FVector3f> VerticesView;
+		{
+			const UntypedMeshBufferIteratorConst PositionBegin(InOutMesh.GetVertexBuffers(), MBS_POSITION);
+			const int32 NumVertices = InOutMesh.GetVertexBuffers().GetElementCount();
+			
+			const bool bIsCompatibleBuffer =
+				PositionBegin.GetElementSize() == sizeof(FVector3f) &&
+				PositionBegin.GetFormat() == MBF_FLOAT32 &&
+				PositionBegin.GetComponents() == 3;
+			
+			const bool bIsAlignmentGood = reinterpret_cast<UPTRINT>(PositionBegin.ptr()) % alignof(FVector3f) == 0;
+
+			if (!bIsCompatibleBuffer || !bIsAlignmentGood)
+			{
+				ConvertedVerticesStorage.SetNumUninitialized(NumVertices);
+				
+				for (int32 VertexIndex = 0; VertexIndex < NumVertices; ++VertexIndex)
+				{
+					ConvertedVerticesStorage[VertexIndex] = (PositionBegin + VertexIndex).GetAsVec3f();
+				}
+
+				VerticesView = TArrayView<const FVector3f>(ConvertedVerticesStorage.GetData(), NumVertices);
+			}
+			else
+			{
+				VerticesView = TArrayView<const FVector3f>(reinterpret_cast<const FVector3f*>(PositionBegin.ptr()), NumVertices);
+			}
+		}
+
+		TArrayView<const uint32> IndicesView;
+		{
+			const UntypedMeshBufferIteratorConst IndicesBegin(InOutMesh.GetIndexBuffers(), MBS_VERTEXINDEX);
+			const int32 NumIndices = InOutMesh.GetIndexBuffers().GetElementCount();
+			
+			const bool bIsCompatibleBuffer =
+				IndicesBegin.GetElementSize() == sizeof(uint32) &&
+				IndicesBegin.GetFormat() == MBF_UINT32 &&
+				IndicesBegin.GetComponents() == 1;
+
+			const bool bIsAlignmentGood = reinterpret_cast<UPTRINT>(IndicesBegin.ptr()) % alignof(uint32) == 0;
+			
+			if (!bIsCompatibleBuffer || !bIsAlignmentGood)
+			{
+				ConvertedIndicesStorage.SetNumUninitialized(NumIndices);
+				for (int32 I = 0; I < NumIndices; ++I)
+				{
+					ConvertedIndicesStorage[I] = (IndicesBegin + I).GetAsUINT32();
+				}
+
+				IndicesView = TArrayView<uint32>(ConvertedIndicesStorage.GetData(), ConvertedIndicesStorage.Num());
+			}
+			else
+			{
+				IndicesView = TArrayView<const uint32>(reinterpret_cast<const uint32*>(IndicesBegin.ptr()), NumIndices);
+			}
+		}
+
+		TArray<int32> UniqueVertexMap = MakeUniqueVertexMap(VerticesView);
+		TArray<TArray<int32, TInlineAllocator<8>>> VertexFaces = BuildVertexFaces(IndicesView, UniqueVertexMap);
+		TMap<uint64, UE::Geometry::FIndex2i> EdgesFaces = BuildEdgesFaces(IndicesView, UniqueVertexMap);
+		
+		TArray<int32> VertexRingsOffsets;
+		TArray<int32> VertexRingsData;
+		Tie(VertexRingsOffsets, VertexRingsData) = BuildVertexRings(IndicesView, UniqueVertexMap, VertexFaces, EdgesFaces);
+
+		FMeshBufferSet MeshLaplacianOffsetsBuffer;
+		MeshLaplacianOffsetsBuffer.SetBufferCount(1);
+		MeshLaplacianOffsetsBuffer.SetElementCount(VertexRingsOffsets.Num());
+
+		FMeshBufferSet MeshLaplacianDataBuffer;
+		MeshLaplacianDataBuffer.SetBufferCount(1);
+		MeshLaplacianDataBuffer.SetElementCount(VertexRingsData.Num());
+
+		// Don't add this to the vertex buffer set for now, it is currently only used for Laplacian smoothing a and it is removed
+		// right away after use.
+		FMeshBufferSet UniqueVertexMapBuffer;
+		UniqueVertexMapBuffer.SetBufferCount(1);
+		UniqueVertexMapBuffer.SetElementCount(UniqueVertexMap.Num());
+
+		const FIntBufferDescriptor BufDesc;
+		MeshLaplacianOffsetsBuffer.SetBuffer(0, BufDesc.ElementSize, BufDesc.Channels, BufDesc.Semantics, BufDesc.SemanticIndices, BufDesc.Formats, BufDesc.Components);
+		MeshLaplacianDataBuffer.SetBuffer(0, BufDesc.ElementSize, BufDesc.Channels, BufDesc.Semantics, BufDesc.SemanticIndices, BufDesc.Formats, BufDesc.Components);
+		UniqueVertexMapBuffer.SetBuffer(0, BufDesc.ElementSize, BufDesc.Channels, BufDesc.Semantics, BufDesc.SemanticIndices, BufDesc.Formats, BufDesc.Components);
+
+		// TODO: Add a way for buffers to steal memory from temporaries, so we can avoid a copy here.
+		FMemory::Memcpy(
+				UniqueVertexMapBuffer.GetBufferData(0), 
+				UniqueVertexMap.GetData(), 
+				UniqueVertexMapBuffer.GetDataSize());
+
+		FMemory::Memcpy(
+				MeshLaplacianOffsetsBuffer.GetBufferData(0), 
+				VertexRingsOffsets.GetData(), 
+				MeshLaplacianOffsetsBuffer.GetDataSize());
+
+		FMemory::Memcpy(
+				MeshLaplacianDataBuffer.GetBufferData(0), 
+				VertexRingsData.GetData(), 
+				MeshLaplacianDataBuffer.GetDataSize());
+
+		InOutMesh.AdditionalBuffers.Emplace(EMeshBufferType::UniqueVertexMap, MoveTemp(UniqueVertexMapBuffer));
+		InOutMesh.AdditionalBuffers.Emplace(EMeshBufferType::MeshLaplacianOffsets, MoveTemp(MeshLaplacianOffsetsBuffer));
+		InOutMesh.AdditionalBuffers.Emplace(EMeshBufferType::MeshLaplacianData, MoveTemp(MeshLaplacianDataBuffer));
+	}
+
+
 	inline TTuple<TArray<FReshapePointBindingData>, TArray<int32>> BindPose(
 			const Mesh* Mesh, FShapeMeshTree& ShapeMeshTree, const TArray<uint16>& BonesToDeform )
 	{
@@ -877,6 +995,7 @@ namespace mu
 		}
 
 		const bool bReshapeVertices = EnumHasAnyFlags(BindFlags, EMeshBindShapeFlags::ReshapeVertices);
+		const bool bApplyLaplacian  = EnumHasAnyFlags(BindFlags, EMeshBindShapeFlags::ApplyLaplacian);
 		const bool bReshapeSkeleton = EnumHasAnyFlags(BindFlags, EMeshBindShapeFlags::ReshapeSkeleton);
 		const bool bReshapePhysics  = EnumHasAnyFlags(BindFlags, EMeshBindShapeFlags::ReshapePhysicsVolumes);
 
@@ -985,12 +1104,17 @@ namespace mu
 			// \TODO: Check that there is no other binding data.
 			// \TODO: Support specifying the binding data channel for multiple binding support.
 			FMeshBufferSet& VB = Result->GetVertexBuffers();
-			int NewBufferIndex = VB.GetBufferCount();
+			int32 NewBufferIndex = VB.GetBufferCount();
 			VB.SetBufferCount(NewBufferIndex + 1);
 
 			FReshapeVertexBindingDataBufferDescriptor BufDesc(BindingDataIndex);
 			VB.SetBuffer(NewBufferIndex, sizeof(FReshapeVertexBindingData), BufDesc.Channels, BufDesc.Semantics, BufDesc.SemanticIndices, BufDesc.Formats, BufDesc.Components);
 			FMemory::Memcpy(VB.GetBufferData(NewBufferIndex), VerticesBindData.GetData(), VerticesBindData.Num() * sizeof(FReshapeVertexBindingData));
+	
+			if (bApplyLaplacian)
+			{
+				GenerateAndAddLaplacianData(*Result);
+			}
 		}
 
 		// Bind the skeleton bones
@@ -1027,7 +1151,7 @@ namespace mu
 			FMemory::Memcpy(SkeletonBuffer.GetBufferData(0), SkeletonBindDataArray.GetData(), NumBonesToDeform * sizeof(FReshapePointBindingData));
 			FMemory::Memcpy(SkeletonBuffer.GetBufferData(1), BoneIndices.GetData(), NumBonesToDeform * sizeof(int32));
 
-			Result->m_AdditionalBuffers.Emplace(EMeshBufferType::SkeletonDeformBinding, MoveTemp(SkeletonBuffer));
+			Result->AdditionalBuffers.Emplace(EMeshBufferType::SkeletonDeformBinding, MoveTemp(SkeletonBuffer));
 		}
 
 		const PhysicsBody* ResultPhysicsBody = Result->m_pPhysicsBody.get();
@@ -1078,9 +1202,9 @@ namespace mu
 			PhysicsBodySelectionOffsetsBuffer.SetBuffer(0, sizeof(int32), IntBufDesc.Channels, IntBufDesc.Semantics, IntBufDesc.SemanticIndices, IntBufDesc.Formats, IntBufDesc.Components, IntBufDesc.Offsets);	
 			FMemory::Memcpy(PhysicsBodySelectionOffsetsBuffer.GetBufferData(0), DeformedBodyIndicesOffsets.GetData(), DeformedBodyIndicesOffsets.Num() * sizeof(int32));
 
-			Result->m_AdditionalBuffers.Emplace(EMeshBufferType::PhysicsBodyDeformBinding, MoveTemp(PhysicsBodyBuffer));
-			Result->m_AdditionalBuffers.Emplace(EMeshBufferType::PhysicsBodyDeformSelection, MoveTemp(PhysicsBodySelectionBuffer));
-			Result->m_AdditionalBuffers.Emplace(EMeshBufferType::PhysicsBodyDeformOffsets, MoveTemp(PhysicsBodySelectionOffsetsBuffer));
+			Result->AdditionalBuffers.Emplace(EMeshBufferType::PhysicsBodyDeformBinding, MoveTemp(PhysicsBodyBuffer));
+			Result->AdditionalBuffers.Emplace(EMeshBufferType::PhysicsBodyDeformSelection, MoveTemp(PhysicsBodySelectionBuffer));
+			Result->AdditionalBuffers.Emplace(EMeshBufferType::PhysicsBodyDeformOffsets, MoveTemp(PhysicsBodySelectionOffsetsBuffer));
 		}	
     }
 
diff --git a/Engine/Plugins/Experimental/Mutable/Source/MutableRuntime/Private/MuR/OpMeshFormat.cpp b/Engine/Plugins/Experimental/Mutable/Source/MutableRuntime/Private/MuR/OpMeshFormat.cpp
--- a/Engine/Plugins/Experimental/Mutable/Source/MutableRuntime/Private/MuR/OpMeshFormat.cpp
+++ b/Engine/Plugins/Experimental/Mutable/Source/MutableRuntime/Private/MuR/OpMeshFormat.cpp
@@ -683,7 +683,7 @@ namespace mu
 
 		Result->m_tags = pSource->m_tags;
 
-		Result->m_AdditionalBuffers = pSource->m_AdditionalBuffers;
+		Result->AdditionalBuffers = pSource->AdditionalBuffers;
 
 		Result->BonePoses = pSource->BonePoses;
 		Result->BoneMap = pSource->BoneMap;
diff --git a/Engine/Plugins/Experimental/Mutable/Source/MutableRuntime/Private/MuR/OpMeshReshape.h b/Engine/Plugins/Experimental/Mutable/Source/MutableRuntime/Private/MuR/OpMeshReshape.h
--- a/Engine/Plugins/Experimental/Mutable/Source/MutableRuntime/Private/MuR/OpMeshReshape.h
+++ b/Engine/Plugins/Experimental/Mutable/Source/MutableRuntime/Private/MuR/OpMeshReshape.h
@@ -15,6 +15,7 @@
 #include "Spatial/MeshAABBTree3.h"
 #include "Distance/DistPoint3Triangle3.h"
 #include "LineTypes.h"
+#include "OpMeshSmoothing.h"
 
 #include "Algo/AnyOf.h" 
 
@@ -63,7 +64,7 @@ namespace mu
 		FVector3f CurrentShapeNormal = ((Shape.Positions[Triangle.B] - Shape.Positions[Triangle.A]) ^ 
 										(Shape.Positions[Triangle.C] - Shape.Positions[Triangle.A])).GetSafeNormal();
 
-		Rotation = FQuat4f::FindBetween(Binding.ShapeNormal, CurrentShapeNormal);
+		Rotation = FQuat4f::FindBetween(Binding.ShapeNormal.GetSafeNormal(), CurrentShapeNormal);
 	}
 
 	inline void GetDeform( const FShapeMeshDescriptorApply& Shape, const FReshapePointBindingData& Binding, FVector3f& NewPosition)
@@ -388,6 +389,7 @@ namespace mu
 		UntypedMeshBufferIterator ItPosition(Mesh->GetVertexBuffers(), MBS_POSITION);
 		UntypedMeshBufferIterator ItNormal(Mesh->GetVertexBuffers(), MBS_NORMAL);
 		UntypedMeshBufferIterator ItTangent(Mesh->GetVertexBuffers(), MBS_TANGENT);
+		UntypedMeshBufferIterator ItBinormal(Mesh->GetVertexBuffers(), MBS_BINORMAL);
 
 #if DO_CHECK
 		// checking if the Base shape has more triangles than the target shape
@@ -417,18 +419,17 @@ namespace mu
 			{
 				GetDeform(Shape, Binding, NewPosition, TangentSpaceCorrection);
 
-				FVector3f OldPosition = ItPosition.GetAsVec3f();
-				
+				const FVector3f OldPosition = ItPosition.GetAsVec3f();
+				FVector3f Displacement = NewPosition - OldPosition;
+
 				if (!FMath::IsNearlyEqual(Binding.Weight, 1.0f))
 				{
-					// Zero weighted vertices are already discarded at the binding phase so there is no gain
-					// checking here.
-					NewPosition = FMath::Lerp(OldPosition, NewPosition, Binding.Weight);
-					TangentSpaceCorrection = FQuat4f::Slerp(FQuat4f::Identity, TangentSpaceCorrection, Binding.Weight).GetNormalized();
+					Displacement *= Binding.Weight;
+					TangentSpaceCorrection = FQuat4f::Slerp(FQuat4f::Identity, TangentSpaceCorrection, Binding.Weight);
 				}
 				
 				// Non rigid vertices will not rotate since are attached to themselves 
-				NewPosition = TangentSpaceCorrection.RotateVector(OldPosition - Binding.AttachmentPoint) + NewPosition;
+				NewPosition = TangentSpaceCorrection.RotateVector(OldPosition - Binding.AttachmentPoint) + (OldPosition + Displacement);
 				
 				ItPosition.SetFromVec3f(NewPosition);
 					
@@ -445,6 +446,13 @@ namespace mu
 					FVector3f NewTangent = TangentSpaceCorrection.RotateVector(OldTangent);
 					ItTangent.SetFromVec3f(NewTangent);
 				}
+
+				if (ItBinormal.ptr())
+				{
+					FVector3f OldBinormal = ItBinormal.GetAsVec3f();
+					FVector3f NewBinormal = TangentSpaceCorrection.RotateVector(OldBinormal);
+					ItBinormal.SetFromVec3f(NewBinormal);
+				}
 			}
 			
 			++ItPosition;
@@ -458,8 +466,12 @@ namespace mu
 			{
 				++ItTangent;
 			}
+
+			if (ItBinormal.ptr())
+			{
+				++ItBinormal;
+			}
 		}
-		
 	}
 
 	inline void ApplyToPose(Mesh* Result, 
@@ -719,6 +731,7 @@ namespace mu
 		const bool bReshapeVertices = EnumHasAnyFlags(BindFlags, EMeshBindShapeFlags::ReshapeVertices);
 		const bool bReshapeSkeleton = EnumHasAnyFlags(BindFlags, EMeshBindShapeFlags::ReshapeSkeleton);
 		const bool bReshapePhysicsVolumes = EnumHasAnyFlags(BindFlags, EMeshBindShapeFlags::ReshapePhysicsVolumes);
+		const bool bApplyLaplacian = EnumHasAnyFlags(BindFlags, EMeshBindShapeFlags::ApplyLaplacian);
 
 		// Early out if nothing will be modified and the vertices discarted.
 		const bool bSkeletonModification = BaseMesh->GetSkeleton() && bReshapeSkeleton;
@@ -731,26 +744,35 @@ namespace mu
 		}
 	
 		// \TODO: Multiple binding data support
-		int BindingDataIndex = 0;
+		int32 BindingDataIndex = 0;
 
 		// If the base mesh has no binding data, just clone it.
-		int BarycentricDataBuffer = 0;
-		int BarycentricDataChannel = 0;
+		int32 BarycentricDataBuffer = 0;
+		int32 BarycentricDataChannel = 0;
 		const FMeshBufferSet& VB = BaseMesh->GetVertexBuffers();
 		VB.FindChannel(MBS_BARYCENTRICCOORDS, BindingDataIndex, &BarycentricDataBuffer, &BarycentricDataChannel);
 		
+		Ptr<Mesh> TempMesh = nullptr;
+		Mesh* VerticesReshapeMesh = Result;
+		
+		if (bApplyLaplacian)
+		{
+			TempMesh = new Mesh();
+			VerticesReshapeMesh = TempMesh.get();
+		}
+
 		// Copy Without VertexBuffers or AdditionalBuffers
 		constexpr EMeshCopyFlags CopyFlags = ~(EMeshCopyFlags::WithVertexBuffers | EMeshCopyFlags::WithAdditionalBuffers);
-		Result->CopyFrom(*BaseMesh, CopyFlags);
+		VerticesReshapeMesh->CopyFrom(*BaseMesh, CopyFlags);
 	
-		FMeshBufferSet& ResultBuffers = Result->GetVertexBuffers();
+		FMeshBufferSet& ResultBuffers = VerticesReshapeMesh->GetVertexBuffers();
 
 		check(ResultBuffers.m_buffers.Num() == 0);
 
 		// Copy buffers skipping binding data. 
 		ResultBuffers.m_elementCount = VB.m_elementCount;
 		// Remove one element to the number of buffers if BarycentricDataBuffer found. 
-		ResultBuffers.m_buffers.SetNum( FMath::Max( 0, VB.m_buffers.Num() - int32(BarycentricDataBuffer >= 0) ) );
+		ResultBuffers.m_buffers.SetNum(FMath::Max(0, VB.m_buffers.Num() - int32(BarycentricDataBuffer >= 0)));
 		
 		for (int32 B = 0, R = 0; B < VB.m_buffers.Num(); ++B)
 		{
@@ -761,8 +783,8 @@ namespace mu
 		}
 
 		// Copy the additional buffers skipping binding data.
-		Result->m_AdditionalBuffers.Reserve( BaseMesh->m_AdditionalBuffers.Num() );
-		for (const TPair<EMeshBufferType, FMeshBufferSet>& A : BaseMesh->m_AdditionalBuffers)
+		VerticesReshapeMesh->AdditionalBuffers.Reserve(BaseMesh->AdditionalBuffers.Num());
+		for (const TPair<EMeshBufferType, FMeshBufferSet>& A : BaseMesh->AdditionalBuffers)
 		{	
 			const bool bIsBindOpBuffer = 
 					A.Key == EMeshBufferType::SkeletonDeformBinding || 
@@ -772,7 +794,7 @@ namespace mu
 	
 			if (!bIsBindOpBuffer)
 			{
-				Result->m_AdditionalBuffers.Add(A);
+				VerticesReshapeMesh->AdditionalBuffers.Add(A);
 			}
 		}
 		if (!ShapeMesh)
@@ -837,16 +859,26 @@ namespace mu
 
 		if (bReshapeVertices && BarycentricDataBuffer >= 0)
 		{
-			MUTABLE_CPUPROFILER_SCOPE(ReshapeVertices);
-			// \TODO: More checks
-			check(BarycentricDataChannel == 0);
-			check(VB.GetElementSize(BarycentricDataBuffer) == (int)sizeof(FReshapeVertexBindingData));
+			{
+				MUTABLE_CPUPROFILER_SCOPE(ReshapeVertices);
+				// \TODO: More checks
+				check(BarycentricDataChannel == 0);
+				check(VB.GetElementSize(BarycentricDataBuffer) == (int)sizeof(FReshapeVertexBindingData));
+
+				TArrayView<const FReshapeVertexBindingData> VerticesBindingData(
+						(const FReshapeVertexBindingData*)VB.GetBufferData(BarycentricDataBuffer),
+						VB.GetElementCount());
 
-			TArrayView<const FReshapeVertexBindingData> VerticesBindingData(
-					(const FReshapeVertexBindingData*)VB.GetBufferData(BarycentricDataBuffer),
-					VB.GetElementCount());
+				ApplyToVertices(VerticesReshapeMesh, VerticesBindingData, ShapeDescriptor);
+			}
 
-			ApplyToVertices(Result, VerticesBindingData, ShapeDescriptor);
+			if (bApplyLaplacian)
+			{
+				check(Result && VerticesReshapeMesh);
+				// check result is empty at this point.
+				check(Result->GetVertexCount() == 0 && Result->GetIndexCount() == 0); 
+				SmoothMeshLaplacian(*Result, *VerticesReshapeMesh);
+			}
 		}
 	
 		if (bReshapeSkeleton)
@@ -855,7 +887,7 @@ namespace mu
 
 			// If the base mesh has no binding data for the skeleton don't do anything.
 			const FMeshBufferSet* SkeletonBindBuffer = nullptr;
-			for ( const TPair<EMeshBufferType, FMeshBufferSet>& A : BaseMesh->m_AdditionalBuffers )
+			for ( const TPair<EMeshBufferType, FMeshBufferSet>& A : BaseMesh->AdditionalBuffers )
 			{
 				if ( A.Key == EMeshBufferType::SkeletonDeformBinding )
 				{
@@ -900,13 +932,13 @@ namespace mu
 			MUTABLE_CPUPROFILER_SCOPE(ReshapePhysicsBodies);
 			
 			using BufferEntryType = TPair<EMeshBufferType, FMeshBufferSet>;
-			const BufferEntryType* FoundPhysicsBindBuffer = BaseMesh->m_AdditionalBuffers.FindByPredicate(
+			const BufferEntryType* FoundPhysicsBindBuffer = BaseMesh->AdditionalBuffers.FindByPredicate(
 					[](BufferEntryType& E){ return E.Key == EMeshBufferType::PhysicsBodyDeformBinding; });
 
-			const BufferEntryType* FoundPhysicsBindSelectionBuffer = BaseMesh->m_AdditionalBuffers.FindByPredicate(
+			const BufferEntryType* FoundPhysicsBindSelectionBuffer = BaseMesh->AdditionalBuffers.FindByPredicate(
 					[](BufferEntryType& E){ return E.Key == EMeshBufferType::PhysicsBodyDeformSelection; });
 	
-			const BufferEntryType* FoundPhysicsBindOffsetsBuffer = BaseMesh->m_AdditionalBuffers.FindByPredicate(
+			const BufferEntryType* FoundPhysicsBindOffsetsBuffer = BaseMesh->AdditionalBuffers.FindByPredicate(
 					[](BufferEntryType& E){ return E.Key == EMeshBufferType::PhysicsBodyDeformOffsets; });
 
 			BarycentricDataBuffer = -1;
diff --git a/Engine/Plugins/Experimental/Mutable/Source/MutableRuntime/Private/MuR/OpMeshSmoothing.h b/Engine/Plugins/Experimental/Mutable/Source/MutableRuntime/Private/MuR/OpMeshSmoothing.h
new file mode 100644
--- /dev/null
+++ b/Engine/Plugins/Experimental/Mutable/Source/MutableRuntime/Private/MuR/OpMeshSmoothing.h
@@ -0,0 +1,578 @@
+// Copyright Epic Games, Inc. All Rights Reserved.
+
+#pragma once
+
+
+//#include "CoreMinimal.h"
+#include "Algo/Reverse.h"
+#include "Math/Vector.h"
+#include "Math/UnrealMathUtility.h"
+#include "Spatial/PointHashGrid3.h"
+#include "Containers/Array.h"
+#include "Containers/ArrayView.h"
+#include "Containers/Map.h"
+
+#include "MuR/Mesh.h"
+#include "MuR/MeshPrivate.h"
+#include "MuR/MutableTrace.h"
+
+
+namespace mu
+{
+	TArray<int32> MakeUniqueVertexMap(TArrayView<const FVector3f> Vertices)
+	{
+		MUTABLE_CPUPROFILER_SCOPE(MeshSmoothing_MakeUniqueVertexMap);
+
+		// Build unique vertex map.
+		const int32 NumVertices = Vertices.Num();
+
+		const float CellSize = 0.01f; // TODO: this should be proportional to the mesh bounding box.
+		UE::Geometry::TPointHashGrid3f<int32> VertHash(CellSize, INDEX_NONE);
+		VertHash.Reserve(NumVertices);
+
+		for (int32 VertexIndex = 0; VertexIndex < NumVertices; ++VertexIndex)
+		{
+			VertHash.InsertPointUnsafe(VertexIndex, Vertices[VertexIndex]);
+		}
+
+		TArray<int32> UniqueVertexMap;
+		UniqueVertexMap.Init(INDEX_NONE, NumVertices);
+	
+		TArray<int32> CollapsedVertices;
+		for (int32 VertexIndex = 0; VertexIndex < NumVertices; ++VertexIndex)
+		{
+			if (UniqueVertexMap[VertexIndex] != INDEX_NONE)
+			{
+				continue;
+			}
+	
+			FVector3f CurrentVertex = Vertices[VertexIndex];
+
+			constexpr bool bAllowSrink = false;
+			CollapsedVertices.Empty(32);
+
+			VertHash.FindPointsInBall(CurrentVertex, TMathUtilConstants<float>::ZeroTolerance,
+				[&CurrentVertex, &Vertices](const int32& Other) -> float 
+				{ 
+					return FVector3f::DistSquared(Vertices[Other], CurrentVertex); 
+				},
+				CollapsedVertices);
+	
+			for (int32 CollapsedVertexIndex : CollapsedVertices)
+			{
+				UniqueVertexMap[CollapsedVertexIndex] = VertexIndex;
+			}
+		}
+
+		return UniqueVertexMap;
+	}
+
+	TArray<TArray<int32, TInlineAllocator<8>>> BuildVertexFaces(TArrayView<const uint32> Indices, const TArray<int32>& UniqueVertexMap)
+	{
+		MUTABLE_CPUPROFILER_SCOPE(MeshSmoothing_BuildVertexFaces);
+		
+		const int32 NumIndices  = Indices.Num();
+		const int32 NumVertices = UniqueVertexMap.Num();
+		
+		TArray<TArray<int32, TInlineAllocator<8>>> VertexFaces;
+		VertexFaces.SetNum(UniqueVertexMap.Num());
+
+		check(NumIndices % 3 == 0);
+		for (int32 Face = 0; Face < NumIndices; Face += 3)
+		{
+			VertexFaces[UniqueVertexMap[Indices[Face + 0]]].Add(Face);
+			VertexFaces[UniqueVertexMap[Indices[Face + 1]]].Add(Face);
+			VertexFaces[UniqueVertexMap[Indices[Face + 2]]].Add(Face);
+		}
+
+		return VertexFaces;
+	}
+
+	FORCEINLINE uint32 GetEdgeLowVertex(uint64 Edge)
+	{
+		return static_cast<uint32>(Edge);
+	}
+
+	FORCEINLINE uint32 GetEdgeHighVertex(uint64 Edge)
+	{
+		return static_cast<uint32>(Edge >> 32);
+	}
+
+	FORCEINLINE uint64 MakeEdge(uint32 E0, uint32 E1)
+	{
+		const uint64 TestMask = E0 < E1 ? 0 : ~0;
+		return (static_cast<uint64>(E0) << (32 & TestMask)) | (static_cast<uint64>(E1) << (32 & (~TestMask)));
+	}
+
+	TMap<uint64, UE::Geometry::FIndex2i> BuildEdgesFaces(TArrayView<const uint32> Indices, const TArray<int32>& UniqueVertexMap)
+	{
+		MUTABLE_CPUPROFILER_SCOPE(MeshSmoothing_BuildEdgesFaces);
+
+		const int32 NumIndices = Indices.Num();
+		
+		auto RemapVertexFaces = [&UniqueVertexMap](const UE::Geometry::FIndex3i& F) -> UE::Geometry::FIndex3i
+		{
+			return UE::Geometry::FIndex3i { UniqueVertexMap[F[0]], UniqueVertexMap[F[1]], UniqueVertexMap[F[2]] };
+		};
+
+		// 2-manifold meshes are assumed.	
+		TMap<uint64, UE::Geometry::FIndex2i> EdgeFaces;
+		EdgeFaces.Reserve(NumIndices);
+
+		check(NumIndices % 3 == 0);
+		for (int32 Face = 0; Face < NumIndices; Face += 3)
+		{
+			const UE::Geometry::FIndex3i FaceIndices = RemapVertexFaces(
+			{
+				static_cast<int32>(Indices[Face + 0]),
+				static_cast<int32>(Indices[Face + 1]),
+				static_cast<int32>(Indices[Face + 2]),
+			});
+
+			uint64 Edge0 = MakeEdge(FaceIndices[0], FaceIndices[1]);
+			uint64 Edge1 = MakeEdge(FaceIndices[1], FaceIndices[2]);
+			uint64 Edge2 = MakeEdge(FaceIndices[2], FaceIndices[0]);
+
+			UE::Geometry::FIndex2i& Edge0Faces = EdgeFaces.FindOrAdd(Edge0, {-1, -1});
+			Edge0Faces = Edge0Faces[0] < 0 ? UE::Geometry::FIndex2i{Face, -1} : UE::Geometry::FIndex2i{Edge0Faces[0], Face};
+
+			UE::Geometry::FIndex2i& Edge1Faces = EdgeFaces.FindOrAdd(Edge1, {-1, -1});
+			Edge1Faces = Edge1Faces[0] < 0 ? UE::Geometry::FIndex2i{Face, -1} : UE::Geometry::FIndex2i{Edge1Faces[0], Face};
+			
+			UE::Geometry::FIndex2i& Edge2Faces = EdgeFaces.FindOrAdd(Edge2, {-1, -1});
+			Edge2Faces = Edge2Faces[0] < 0 ? UE::Geometry::FIndex2i{Face, -1} : UE::Geometry::FIndex2i{Edge2Faces[0], Face};
+		}
+
+		return EdgeFaces;
+	}
+
+	FORCEINLINE int32 FindEdgeOtherVertex(uint64 Edge, uint32 V)
+	{
+		check(GetEdgeLowVertex(Edge) == V || GetEdgeHighVertex(Edge) == V);
+		return GetEdgeLowVertex(Edge) != V ? GetEdgeLowVertex(Edge) : GetEdgeHighVertex(Edge);
+	}
+
+	FORCEINLINE int32 FindOppositeFace(UE::Geometry::FIndex2i EdgeFaces, int32 Face)
+	{
+		check(EdgeFaces[0] == Face || EdgeFaces[1] == Face);
+		return Face == EdgeFaces[0] ? EdgeFaces[1] : EdgeFaces[0];
+	}
+
+	FORCEINLINE uint32 FindFaceVertexNotInEdge(uint64 Edge, const UE::Geometry::FIndex3i& F)
+	{
+		if (Edge == MakeEdge(F[0], F[1])) return F[2];
+		if (Edge == MakeEdge(F[0], F[2])) return F[1];
+		if (Edge == MakeEdge(F[1], F[2])) return F[0];
+	
+		check(false);
+		return -1;
+	}
+
+	FORCEINLINE int32 FindNextFaceVertex(int32 V, const UE::Geometry::FIndex3i& F)
+	{
+		if (V == F[0]) return F[1];
+		if (V == F[1]) return F[2];
+		if (V == F[2]) return F[0];
+
+		check(false);
+		return -1;
+	}
+
+	TTuple<TArray<int32>, TArray<int32>> BuildVertexRings(
+		const TArrayView<const uint32> Indices,
+		const TArray<int32>& UniqueVertexMap,
+		const TArray<TArray<int32, TInlineAllocator<8>>>& VertexFaces,
+		const TMap<uint64, UE::Geometry::FIndex2i>& EdgesFaces)
+	{
+		MUTABLE_CPUPROFILER_SCOPE(MeshSmoothing_BuildVertexRings);
+
+#if DO_CHECK
+		int32 BowtiedRingsFound = 0;
+#endif
+		auto RemapVertexFaces = [&UniqueVertexMap](const UE::Geometry::FIndex3i& F) -> UE::Geometry::FIndex3i
+		{
+			return UE::Geometry::FIndex3i { UniqueVertexMap[F[0]], UniqueVertexMap[F[1]], UniqueVertexMap[F[2]] };
+		};
+
+		const int32 NumVertices = UniqueVertexMap.Num();
+	
+		TArray<int32> VertexRingsOffsets;
+		VertexRingsOffsets.Reserve(NumVertices*2);
+		
+		TArray<int32> VertexRings;
+		VertexRings.Reserve(NumVertices*6);
+
+		VertexRingsOffsets.Add(0);
+
+		for (int32 V = 0; V < NumVertices; ++V)
+		{
+			if (UniqueVertexMap[V] != V)
+			{
+				continue;
+			}
+
+			const int32 NumVertexFaces = VertexFaces[V].Num();
+
+			const int32 InitialFace = VertexFaces[V][0];
+			UE::Geometry::FIndex3i FaceVertices = RemapVertexFaces(
+			{
+				static_cast<int32>(Indices[InitialFace + 0]), 
+				static_cast<int32>(Indices[InitialFace + 1]), 
+				static_cast<int32>(Indices[InitialFace + 2])
+			});
+
+			int32 V0 = FindNextFaceVertex(V,  FaceVertices);
+			int32 V1 = FindNextFaceVertex(V0, FaceVertices);
+			uint64 Edge0 = MakeEdge(V, V0);
+			uint64 Edge1 = MakeEdge(V, V1);
+			
+			const int32 RingBeginOffset = VertexRings.Num();
+		
+			uint64 Edge = Edge1;
+			int32 Face = FindOppositeFace(EdgesFaces[Edge1], InitialFace);
+
+			VertexRings.Add(V0);
+			VertexRings.Add(V1);
+
+			if (Face >= 0)
+			{
+				for (;;)
+				{
+					FaceVertices = RemapVertexFaces(
+					{
+						static_cast<int32>(Indices[Face + 0]), 
+						static_cast<int32>(Indices[Face + 1]), 
+						static_cast<int32>(Indices[Face + 2])
+					});
+
+					const int32 NextVertex = FindFaceVertexNotInEdge(Edge, FaceVertices);
+					
+					Edge = MakeEdge(V, NextVertex);
+					Face = FindOppositeFace(EdgesFaces[Edge], Face);
+
+					if (Face == InitialFace)
+					{
+						break;
+					}
+
+					VertexRings.Add(NextVertex);
+					
+					if (Face < 0)
+					{
+						break;
+					}
+				}
+			}
+
+			// We have a closed ring.
+			if (Face >= 0)
+			{
+				check(Edge == Edge0);
+				check(Face == InitialFace);
+
+				VertexRingsOffsets.Add(VertexRings.Num());
+				continue;
+			}
+
+			const int32 RingHalfOffset = VertexRings.Num();
+
+			Face = FindOppositeFace(EdgesFaces[Edge0], InitialFace);
+			Edge = Edge0;
+
+			if (Face >= 0)
+			{
+				for (;;)
+				{
+					FaceVertices = RemapVertexFaces(
+					{
+						static_cast<int32>(Indices[Face + 0]), 
+						static_cast<int32>(Indices[Face + 1]), 
+						static_cast<int32>(Indices[Face + 2])
+					});
+
+					const int32 NextVertex = FindFaceVertexNotInEdge(Edge, FaceVertices);	
+					VertexRings.Add(NextVertex);
+
+					Edge = MakeEdge(V, NextVertex);
+					Face = FindOppositeFace(EdgesFaces[Edge], Face);
+					
+					if (Face < 0)
+					{
+						break;
+					}	
+				}
+				
+				const int32 RingEndOffset = VertexRings.Num();
+
+				Algo::Reverse(VertexRings.GetData() + RingBeginOffset, RingEndOffset - RingBeginOffset);
+				Algo::Reverse(VertexRings.GetData() + RingBeginOffset + RingEndOffset - RingHalfOffset, RingHalfOffset - RingBeginOffset);
+			}
+
+			// Mark ring as open.
+			VertexRings.Add(-1);	
+
+			const int32 OpenRingNumFaces = VertexRings.Num() - 1 - RingBeginOffset - 1;
+			if (NumVertexFaces <= OpenRingNumFaces)
+			{
+				check(NumVertexFaces == OpenRingNumFaces);
+					
+				VertexRingsOffsets.Add(VertexRings.Num());
+				continue;
+			}
+
+			// This case should happen rarely on well behaved meshes. 
+			// TODO: For now the case is not supported. Review if actually needed.
+#if DO_CHECK
+			++BowtiedRingsFound;
+#endif
+			//continue;
+		
+			// The following code for computing bowtied rings is incomplete, but it gives an idea of how to proceed.
+			//TBitArray VisitedVertexFaces(false, VertexFaces.Num());
+
+			//for (int32 RingVertex = RingBeginOffset; RingVertex < RingEndOffset; ++RingVertex)
+			//{
+			//	UE::Geometry::FIndex2i EdgeFaces = EdgesFaces[MakeEdge(V, VetexRing[RingVetex])];
+
+			//	if (EdgeFaces[0] > 0)
+			//	{
+			//		const int32 VertexFaceIndex = VertexFaces[V].IndexOf(EdgeFaces[0]); 
+			//		check(VertexFaceIndex == INDEX_NONE);
+			//		VisitedVertexFaces[VertexFaceIndex] = true;
+			//	}
+
+			//	if (EdgeFaces[1] > 0)
+			//	{
+			//		const int32 VertexFaceIndex = VertexFaces[V].IndexOf(EdgeFaces[1]); 
+			//		check(VertexFaceIndex == INDEX_NONE);
+			//		VisitedVertexFaces[VertexFaceIndex] = true;
+			//	}
+			//}
+
+
+			//for (;;)
+			//{
+
+			//	int32 VertexFace = VisitedVeretexFaces.Find(false);
+			//	if (VertexFace == INDEX_NONE)
+			//	{
+			//		break;
+			//	}
+
+			//	Face = VertexFaces[V][VertexFace];
+			//	
+			//	FaceVertices = RemapVertexFaces({ Indices[Face + 0], Indices[Face + 1], Indices[Face + 2] });
+			//	
+			//	V0 = FindNextVertexFace(V, FaceVertices);
+			//	V1 = FindNextVertexFace(V0, FaceVertices);
+
+			//	Edge0 = MakeEdge(V, V0); 
+			//	Edge1 = MakeEdge(V, V1); 
+
+			//	const int32 CurrentRingBegin = VetexRings.Num();
+
+			//	VertexRings.Add(V0);
+			//	VertexRings.Add(V1);
+			//	
+			//	const int32 
+			//	Edge = Edge1;
+			//	for (;;)
+			//	{
+			//		if (Face < 0)
+			//		{
+			//			break;
+			//		}
+
+
+			//		FaceVertices = RemapVertexFaces({ Indices[Face + 0], Indices[Face + 1], Indices[Face + 2] });
+
+			//		const int32 NextVertex = FindFaceVertexNotInEdge(Edge, FaceVertices);
+			//		VertexRings.Add(NextVertex);
+
+			//		Edge = MakeEdge(V, NextVertex);
+			//		Face = FindOtherFace(EdgesFaces[Edge], Face);
+
+			//		const int32 VertexFaceIndex = VertexFaces[V].IndexOf(Face);
+			//		check(VertexFaceIndex != INDEX_NONE);
+			//		VisitedVertexFaces[VertexFaceIndex] = true;
+			//	}
+
+			//	const int32 CurrentRingHalf = VertexRings.Num();
+
+			//	Edge = Edge0;
+			//	for (;;)
+			//	{
+			//		if (Face < 0)
+			//		{
+			//			break;
+			//		}
+
+			//		FaceVertices = { Indices[Face + 0], Indices[Face + 1], Indices[Face + 2] };
+			//		FaceVertices = { UniqueVertexMap[FaceVertices[0]], UniqueVertexMap[FaceVertices[1]], UniqueVertexMap[FaceVertices[2]] };
+
+			//		const int32 NextVertex = FindFaceVertexNotInEdge(Edge, FaceVertices);
+			//		VertexRings.Add(NextVertex);
+			//		
+			//		Edge = MakeEdge(V, FindFaceVertexNotInEdge(Edge, FaceVertices));
+			//		Face = FindOtherFace(Face, EdgeFaces[Edge]);
+
+			//		const int32 VertexFaceIndex = VertexFaces[V].IndexOf(Face);
+			//		check(VertexFaceIndex != INDEX_NONE);
+			//		VisitedVertexFaces[VertexFaceIndex] = true;
+
+			//	}
+
+			//	const int32 CurrentRingEnd = VertexRings.Num();
+
+
+			//	Algo::Reverse(VertexRings.GetData() + CurrentRingBegin, CurrentRingEnd - CurrentRingBegin);
+			//	Algo::Reverse(VertexRings.GetData() + CurrentRingEnd - CurrentRingHalf, CurrentRingHalf - CurrentRingBegin);
+			//}			
+		}
+		
+#if DO_CHECK
+		if (BowtiedRingsFound > 0)
+		{
+			UE_LOG(LogMutableCore, Warning, 
+				TEXT("%d Bowtied vertex rings found in a mesh while computing smoothing operation data, this is currently not supported. " 
+					 "Smoothing for the vertices involved may not work properly."), 
+				BowtiedRingsFound);
+		}
+#endif
+
+		return MakeTuple(MoveTemp(VertexRingsOffsets), MoveTemp(VertexRings));
+	}
+
+	inline void SmoothMeshLaplacian(Mesh& DstMesh, const Mesh& SrcMesh, int32 Iterations = 1)
+	{
+		MUTABLE_CPUPROFILER_SCOPE(MeshSmoothing_SmoothMeshLaplacian);
+
+		constexpr EMeshCopyFlags CopyFlags = ~EMeshCopyFlags::WithAdditionalBuffers;
+		DstMesh.CopyFrom(SrcMesh, CopyFlags);
+
+		using BufferEntryType = TPair<EMeshBufferType, FMeshBufferSet>;
+		
+		// Copy AdditionalBuffers skipping LaplacianData. 
+		for (const BufferEntryType& AdditionalBuffer : SrcMesh.AdditionalBuffers)
+		{
+			const bool bIsLaplacianBuffer =
+				AdditionalBuffer.Key == EMeshBufferType::UniqueVertexMap   ||
+				AdditionalBuffer.Key == EMeshBufferType::MeshLaplacianData ||
+				AdditionalBuffer.Key == EMeshBufferType::MeshLaplacianOffsets;
+
+			if (!bIsLaplacianBuffer)
+			{
+				DstMesh.AdditionalBuffers.Add(AdditionalBuffer);
+			}
+		}
+
+		const BufferEntryType* MeshLaplacianDataBuffer = SrcMesh.AdditionalBuffers.FindByPredicate(
+				[](const BufferEntryType& E) { return E.Key == EMeshBufferType::MeshLaplacianData; });
+
+		const BufferEntryType* MeshLaplacianOffsetsBuffer = SrcMesh.AdditionalBuffers.FindByPredicate(
+				[](const BufferEntryType& E) { return E.Key == EMeshBufferType::MeshLaplacianOffsets; });
+
+		const BufferEntryType* MeshUniqueVertexMapBuffer = SrcMesh.AdditionalBuffers.FindByPredicate(
+				[](const BufferEntryType& E) { return E.Key == EMeshBufferType::UniqueVertexMap; });
+
+		const bool bHasNecessaryBuffers = MeshLaplacianDataBuffer && MeshLaplacianOffsetsBuffer && MeshUniqueVertexMapBuffer;
+		if (!bHasNecessaryBuffers)
+		{
+			UE_LOG(LogMutableCore, Warning, TEXT("Smooth laplacian with source mesh missing some necessary buffers. Ignoring operation."));
+			return;
+		}
+
+		TArrayView<const int32> LaplacianDataVertexOffsets = TArrayView<const int32>(
+				reinterpret_cast<const int32*>(MeshLaplacianOffsetsBuffer->Value.GetBufferData(0)),
+				MeshLaplacianOffsetsBuffer->Value.GetElementCount());
+
+		TArrayView<const int32> LaplacianVertexRingData = TArrayView<const int32>(
+				reinterpret_cast<const int32*>(MeshLaplacianDataBuffer->Value.GetBufferData(0)),
+				MeshLaplacianDataBuffer->Value.GetElementCount());
+		
+		TArrayView<const int32> UniqueVertexMap = TArrayView<const int32>(
+				reinterpret_cast<const int32*>(MeshUniqueVertexMapBuffer->Value.GetBufferData(0)),
+				MeshUniqueVertexMapBuffer->Value.GetElementCount());
+
+		const UntypedMeshBufferIteratorConst SrcPositionBegin(SrcMesh.GetVertexBuffers(), MBS_POSITION);
+		const UntypedMeshBufferIterator      DstPositionBegin(DstMesh.GetVertexBuffers(), MBS_POSITION);
+
+		const int32 NumVertices = DstMesh.GetVertexCount();
+ 
+		TArray<FVector3f, TInlineAllocator<8>> RingVerticesStorage;
+		
+		for (int32 VertexIndex = 0, UniqueVertexIndex = 0; VertexIndex < NumVertices; ++VertexIndex)
+		{
+			if (VertexIndex != UniqueVertexMap[VertexIndex])
+			{
+				continue;
+			}
+
+			const int32 LaplacianVertexDataBegin = LaplacianDataVertexOffsets[UniqueVertexIndex];
+			const int32 LaplacianVertexDataEnd   = LaplacianDataVertexOffsets[UniqueVertexIndex + 1];
+
+			const bool bIsOpenRing   = LaplacianVertexRingData[LaplacianVertexDataEnd - 1] < 0;
+			const int32 RingNumElems = LaplacianVertexDataEnd - LaplacianVertexDataBegin - static_cast<int32>(bIsOpenRing);
+
+			TArray<FVector3f, TInlineAllocator<8>>& Ring = RingVerticesStorage;
+			Ring.SetNum(RingNumElems);
+			for (int32 R = 0; R < RingNumElems; ++R)
+			{
+				Ring[R] = (SrcPositionBegin + LaplacianVertexRingData[LaplacianVertexDataBegin + R]).GetAsVec3f();
+			}
+
+			const FVector3f VertexPosition = (SrcPositionBegin + VertexIndex).GetAsVec3f();
+			
+			FVector3f WeightedPositionLaplacian = FVector3f::ZeroVector;
+			float TotalWeight = 0.0f;
+			
+			int32 PrevIndex = bIsOpenRing ? 0 : RingNumElems - 1;
+			int32 NextIndex = 1;
+
+			for (int32 R = 0; R < RingNumElems; ++R)
+			{
+				const FVector3f CoAU = VertexPosition - Ring[PrevIndex]; 
+				const FVector3f CoAV = Ring[R] - Ring[PrevIndex];
+
+				const FVector3f CoBU = VertexPosition - Ring[NextIndex]; 
+				const FVector3f CoBV = Ring[R] - Ring[NextIndex]; 
+
+				const float CoA = FVector3f::DotProduct(CoAU, CoAV) * 
+								  FMath::InvSqrt(
+										  FMath::Max(UE_KINDA_SMALL_NUMBER, FVector3f::CrossProduct(CoAU, CoAV).SquaredLength())); 
+
+				const float CoB = FVector3f::DotProduct(CoBU, CoBV) * 
+								  FMath::InvSqrt(
+										  FMath::Max(UE_KINDA_SMALL_NUMBER, FVector3f::CrossProduct(CoBU, CoBV).SquaredLength()));
+
+				const float Weight = (CoA + CoB);// * 0.5f;	
+
+				WeightedPositionLaplacian += Ring[R] * Weight;
+				TotalWeight += Weight;
+
+				PrevIndex = R;
+				NextIndex = R + 2 < RingNumElems ? R + 2 : (R + 1)*static_cast<int32>(bIsOpenRing);
+			}
+
+			const FVector3f VertexPositionLaplacian = 
+				(WeightedPositionLaplacian / FMath::Max(UE_KINDA_SMALL_NUMBER, TotalWeight));
+
+			(DstPositionBegin + VertexIndex).SetFromVec3f(VertexPosition + (VertexPositionLaplacian-VertexPosition)*0.5f);
+			++UniqueVertexIndex;
+		}
+
+		// Copy positions to the collapsed vertices.
+		const SIZE_T PositionElemSize = GetMeshFormatData(DstPositionBegin.GetFormat()).m_size * DstPositionBegin.GetComponents(); 
+		for (int32 VertexIndex = 0; VertexIndex < NumVertices; ++VertexIndex)
+		{
+			const int32 MappedIndex = UniqueVertexMap[VertexIndex];  
+			if (VertexIndex == MappedIndex)
+			{
+				continue;
+			}
+
+			FMemory::Memcpy((DstPositionBegin + VertexIndex).ptr(), (DstPositionBegin + MappedIndex).ptr(), PositionElemSize);
+		}
+	}
+}
+
diff --git a/Engine/Plugins/Experimental/Mutable/Source/MutableRuntime/Private/MuR/Operations.h b/Engine/Plugins/Experimental/Mutable/Source/MutableRuntime/Private/MuR/Operations.h
--- a/Engine/Plugins/Experimental/Mutable/Source/MutableRuntime/Private/MuR/Operations.h
+++ b/Engine/Plugins/Experimental/Mutable/Source/MutableRuntime/Private/MuR/Operations.h
@@ -343,6 +343,7 @@ namespace mu
         COUNT
 
     };
+
 	enum class EMeshBindShapeFlags : uint32
 	{
 		None				   = 0,
@@ -350,6 +351,7 @@ namespace mu
 		EnableRigidParts       = 1 << 2,
 		ReshapePhysicsVolumes  = 1 << 4,
 		ReshapeVertices		   = 1 << 5,
+		ApplyLaplacian		   = 1 << 6,
 	};
 	ENUM_CLASS_FLAGS(EMeshBindShapeFlags);
 
diff --git a/Engine/Plugins/Experimental/Mutable/Source/MutableRuntime/Public/MuR/Mesh.h b/Engine/Plugins/Experimental/Mutable/Source/MutableRuntime/Public/MuR/Mesh.h
--- a/Engine/Plugins/Experimental/Mutable/Source/MutableRuntime/Public/MuR/Mesh.h
+++ b/Engine/Plugins/Experimental/Mutable/Source/MutableRuntime/Public/MuR/Mesh.h
@@ -105,7 +105,10 @@ namespace mu
 		SkeletonDeformBinding,
 		PhysicsBodyDeformBinding,
 		PhysicsBodyDeformSelection,
-		PhysicsBodyDeformOffsets
+		PhysicsBodyDeformOffsets,
+		MeshLaplacianData,
+		MeshLaplacianOffsets,
+		UniqueVertexMap
 	};
 
 	//!
@@ -370,7 +373,7 @@ namespace mu
 		FMeshBufferSet m_FaceBuffers;
 
 		//! Additional buffers used for temporary or custom data in different algorithms.
-		TArray< TPair<EMeshBufferType, FMeshBufferSet> > m_AdditionalBuffers;
+		TArray<TPair<EMeshBufferType, FMeshBufferSet>> AdditionalBuffers;
 
 		//! This is bit-mask on the STATIC_MESH_FORMATS enumeration, marking what static formats
 		//! are compatible with this one. Usually precalculated at model compilation time.
@@ -482,10 +485,10 @@ namespace mu
 				equal &= (*m_layouts[i]) == (*o.m_layouts[i]);
 			}
 
-			equal &= m_AdditionalBuffers.Num() == o.m_AdditionalBuffers.Num();
-			for (int32 i = 0; equal && i < m_AdditionalBuffers.Num(); ++i)
+			equal &= AdditionalBuffers.Num() == o.AdditionalBuffers.Num();
+			for (int32 i = 0; equal && i < AdditionalBuffers.Num(); ++i)
 			{
-				equal &= m_AdditionalBuffers[i] == o.m_AdditionalBuffers[i];
+				equal &= AdditionalBuffers[i] == o.AdditionalBuffers[i];
 			}
 
 			equal &= BonePoses.Num() == o.BonePoses.Num();
diff --git a/Engine/Plugins/Experimental/Mutable/Source/MutableRuntime/Public/MuR/Model.h b/Engine/Plugins/Experimental/Mutable/Source/MutableRuntime/Public/MuR/Model.h
--- a/Engine/Plugins/Experimental/Mutable/Source/MutableRuntime/Public/MuR/Model.h
+++ b/Engine/Plugins/Experimental/Mutable/Source/MutableRuntime/Public/MuR/Model.h
@@ -13,7 +13,7 @@
 //! data structures. Compiled models are not necessarily compatible when the runtime is updated,
 //! so this version number can be used externally to verify this. It is not used internally, and
 //! serializing models from different versions than this runtime will probably result in a crash.
-#define MUTABLE_COMPILED_MODEL_CODE_VERSION		uint32( 76 )
+#define MUTABLE_COMPILED_MODEL_CODE_VERSION		uint32( 77 )
 #define MUTABLE_PARAMETERS_VERSION              uint32( 2 )
 
 
diff --git a/Engine/Plugins/Experimental/Mutable/Source/MutableTools/Private/MuT/ASTOpMeshApplyShape.cpp b/Engine/Plugins/Experimental/Mutable/Source/MutableTools/Private/MuT/ASTOpMeshApplyShape.cpp
--- a/Engine/Plugins/Experimental/Mutable/Source/MutableTools/Private/MuT/ASTOpMeshApplyShape.cpp
+++ b/Engine/Plugins/Experimental/Mutable/Source/MutableTools/Private/MuT/ASTOpMeshApplyShape.cpp
@@ -40,7 +40,8 @@ namespace mu
 			const bool bSameFlags =
 				bReshapePhysicsVolumes == Other->bReshapePhysicsVolumes &&
 				bReshapeSkeleton == Other->bReshapeSkeleton &&
-				bReshapeVertices == Other->bReshapeVertices;
+				bReshapeVertices == Other->bReshapeVertices &&
+				bApplyLaplacian == Other->bApplyLaplacian;
 
 			return Mesh == Other->Mesh && Shape == Other->Shape && bSameFlags;
 		}
@@ -50,12 +51,14 @@ namespace mu
 
 	uint64 ASTOpMeshApplyShape::Hash() const
 	{
-		uint64 res = std::hash<void*>()(Mesh.child().get());
-		hash_combine(res, Shape.child().get());
-		hash_combine(res, bool(bReshapeSkeleton));
-		hash_combine(res, bool(bReshapePhysicsVolumes));
-		hash_combine(res, bool(bReshapeVertices));
-		return res;
+		uint64 Result = std::hash<void*>()(Mesh.child().get());
+		hash_combine(Result, Shape.child().get());
+		hash_combine(Result, bool(bReshapeSkeleton));
+		hash_combine(Result, bool(bReshapePhysicsVolumes));
+		hash_combine(Result, bool(bReshapeVertices));
+		hash_combine(Result, bool(bApplyLaplacian));
+
+		return Result;
 	}
 
 
@@ -71,10 +74,10 @@ namespace mu
 	}
 
 
-	void ASTOpMeshApplyShape::ForEachChild(const TFunctionRef<void(ASTChild&)> f)
+	void ASTOpMeshApplyShape::ForEachChild(const TFunctionRef<void(ASTChild&)> Func)
 	{
-		f(Mesh);
-		f(Shape);
+		Func(Mesh);
+		Func(Shape);
 	}
 
 
@@ -91,6 +94,7 @@ namespace mu
 			EnumAddFlags(BindFlags, bReshapeSkeleton ? EMeshBindShapeFlags::ReshapeSkeleton : NoFlags);
 			EnumAddFlags(BindFlags, bReshapePhysicsVolumes ? EMeshBindShapeFlags::ReshapePhysicsVolumes : NoFlags);
 			EnumAddFlags(BindFlags, bReshapeVertices ? EMeshBindShapeFlags::ReshapeVertices : NoFlags);
+			EnumAddFlags(BindFlags, bApplyLaplacian ? EMeshBindShapeFlags::ApplyLaplacian : NoFlags);
 
 			Args.flags = static_cast<uint32>(BindFlags);
 
diff --git a/Engine/Plugins/Experimental/Mutable/Source/MutableTools/Private/MuT/ASTOpMeshApplyShape.h b/Engine/Plugins/Experimental/Mutable/Source/MutableTools/Private/MuT/ASTOpMeshApplyShape.h
--- a/Engine/Plugins/Experimental/Mutable/Source/MutableTools/Private/MuT/ASTOpMeshApplyShape.h
+++ b/Engine/Plugins/Experimental/Mutable/Source/MutableTools/Private/MuT/ASTOpMeshApplyShape.h
@@ -22,10 +22,10 @@ struct FProgram;
 		ASTChild Mesh;
 		ASTChild Shape;
 
-		uint32 bReshapeSkeleton : 1;
+		uint32 bReshapeSkeleton       : 1;
 		uint32 bReshapePhysicsVolumes : 1;
-		uint32 bReshapeVertices : 1;
-
+		uint32 bReshapeVertices       : 1;
+		uint32 bApplyLaplacian        : 1;
 	public:
 
 		ASTOpMeshApplyShape();
diff --git a/Engine/Plugins/Experimental/Mutable/Source/MutableTools/Private/MuT/ASTOpMeshBindShape.cpp b/Engine/Plugins/Experimental/Mutable/Source/MutableTools/Private/MuT/ASTOpMeshBindShape.cpp
--- a/Engine/Plugins/Experimental/Mutable/Source/MutableTools/Private/MuT/ASTOpMeshBindShape.cpp
+++ b/Engine/Plugins/Experimental/Mutable/Source/MutableTools/Private/MuT/ASTOpMeshBindShape.cpp
@@ -25,6 +25,8 @@ namespace mu
 		, Shape(this)
 		, bReshapeSkeleton(false)
 		, bReshapePhysicsVolumes(false)
+		, bReshapeVertices(false)
+		, bApplyLaplacian(false)
 	{
 	}
 
@@ -43,7 +45,8 @@ namespace mu
 			const bool bSameFlags =
 				bReshapeSkeleton == Other->bReshapeSkeleton	&&
 				bReshapePhysicsVolumes == Other->bReshapePhysicsVolumes &&
-				bReshapeVertices == Other->bReshapeVertices;
+				bReshapeVertices == Other->bReshapeVertices &&
+				bApplyLaplacian == Other->bApplyLaplacian;
 
 			return bSameFlags &&
 				Mesh == Other->Mesh &&
@@ -63,29 +66,30 @@ namespace mu
 
 	uint64 ASTOpMeshBindShape::Hash() const
 	{
-		uint64 res = std::hash<void*>()(Mesh.child().get());
-		hash_combine(res, Shape.child().get());
-		hash_combine(res, bool(bReshapeSkeleton));
-		hash_combine(res, bool(bReshapePhysicsVolumes));
-		hash_combine(res, bool(bReshapeVertices));
-		hash_combine(res, bool(BindingMethod));
-
-		hash_combine(res, static_cast<uint32>(RChannelUsage));
-		hash_combine(res, static_cast<uint32>(GChannelUsage));
-		hash_combine(res, static_cast<uint32>(BChannelUsage));
-		hash_combine(res, static_cast<uint32>(AChannelUsage));
+		uint64 Result = std::hash<void*>()(Mesh.child().get());
+		hash_combine(Result, Shape.child().get());
+		hash_combine(Result, bool(bReshapeSkeleton));
+		hash_combine(Result, bool(bReshapePhysicsVolumes));
+		hash_combine(Result, bool(bReshapeVertices));
+		hash_combine(Result, bool(bApplyLaplacian));
+		hash_combine(Result, bool(BindingMethod));
+
+		hash_combine(Result, static_cast<uint32>(RChannelUsage));
+		hash_combine(Result, static_cast<uint32>(GChannelUsage));
+		hash_combine(Result, static_cast<uint32>(BChannelUsage));
+		hash_combine(Result, static_cast<uint32>(AChannelUsage));
 
 		for (const uint16 S : BonesToDeform)
 		{
-			hash_combine(res, S);
+			hash_combine(Result, S);
 		}
 
 		for (const uint16 S : PhysicsToDeform)
 		{
-			hash_combine(res, S);
+			hash_combine(Result, S);
 		}
 
-		return res;
+		return Result;
 	}
 
 
@@ -97,6 +101,7 @@ namespace mu
 		NewOp->bReshapeSkeleton	= bReshapeSkeleton;
 		NewOp->bReshapePhysicsVolumes = bReshapePhysicsVolumes;
 		NewOp->bReshapeVertices = bReshapeVertices;
+		NewOp->bApplyLaplacian = bApplyLaplacian;
 		NewOp->BonesToDeform = BonesToDeform;
 		NewOp->PhysicsToDeform = PhysicsToDeform;
 		NewOp->BindingMethod = BindingMethod;
@@ -110,10 +115,10 @@ namespace mu
 	}
 
 
-	void ASTOpMeshBindShape::ForEachChild(const TFunctionRef<void(ASTChild&)> f)
+	void ASTOpMeshBindShape::ForEachChild(const TFunctionRef<void(ASTChild&)> Func)
 	{
-		f(Mesh);
-		f(Shape);
+		Func(Mesh);
+		Func(Shape);
 	}
 
 
@@ -130,6 +135,7 @@ namespace mu
 			EnumAddFlags(BindFlags, bReshapeSkeleton ? EMeshBindShapeFlags::ReshapeSkeleton : NoFlags);
 			EnumAddFlags(BindFlags, bReshapePhysicsVolumes ? EMeshBindShapeFlags::ReshapePhysicsVolumes : NoFlags);
 			EnumAddFlags(BindFlags, bReshapeVertices ? EMeshBindShapeFlags::ReshapeVertices : NoFlags);
+			EnumAddFlags(BindFlags, bApplyLaplacian ? EMeshBindShapeFlags::ApplyLaplacian : NoFlags);
 
 			{
 				auto ConvertColorUsage = [](EVertexColorUsage Usage)
diff --git a/Engine/Plugins/Experimental/Mutable/Source/MutableTools/Private/MuT/ASTOpMeshBindShape.h b/Engine/Plugins/Experimental/Mutable/Source/MutableTools/Private/MuT/ASTOpMeshBindShape.h
--- a/Engine/Plugins/Experimental/Mutable/Source/MutableTools/Private/MuT/ASTOpMeshBindShape.h
+++ b/Engine/Plugins/Experimental/Mutable/Source/MutableTools/Private/MuT/ASTOpMeshBindShape.h
@@ -32,6 +32,7 @@ struct FProgram;
 		uint32 bReshapeSkeleton	      : 1;
 		uint32 bReshapePhysicsVolumes : 1;
 		uint32 bReshapeVertices       : 1;
+		uint32 bApplyLaplacian        : 1;
 
 		EVertexColorUsage RChannelUsage = EVertexColorUsage::None;
 		EVertexColorUsage GChannelUsage = EVertexColorUsage::None;
diff --git a/Engine/Plugins/Experimental/Mutable/Source/MutableTools/Private/MuT/CodeGenerator_Mesh.cpp b/Engine/Plugins/Experimental/Mutable/Source/MutableTools/Private/MuT/CodeGenerator_Mesh.cpp
--- a/Engine/Plugins/Experimental/Mutable/Source/MutableTools/Private/MuT/CodeGenerator_Mesh.cpp
+++ b/Engine/Plugins/Experimental/Mutable/Source/MutableTools/Private/MuT/CodeGenerator_Mesh.cpp
@@ -804,6 +804,7 @@ class Node;
 			// Setting bReshapeVertices to false the bind op will remove all mesh members except 
 			// PhysicsBodies and the Skeleton.
             OpBind->bReshapeVertices = false;
+            OpBind->bApplyLaplacian = false;
 		    OpBind->bReshapeSkeleton = node.bReshapeSkeleton;
 		    OpBind->BonesToDeform = node.BonesToDeform;
     	    OpBind->bReshapePhysicsVolumes = node.bReshapePhysicsVolumes; 
@@ -1777,42 +1778,44 @@ class Node;
 
 
 	//---------------------------------------------------------------------------------------------
-	void CodeGenerator::GenerateMesh_Reshape(const FMeshGenerationOptions& InOptions, FMeshGenerationResult& OutResult, const NodeMeshReshape* reshape)
+	void CodeGenerator::GenerateMesh_Reshape(const FMeshGenerationOptions& InOptions, FMeshGenerationResult& OutResult, const NodeMeshReshape* Reshape)
 	{
-		const NodeMeshReshape::Private& node = *reshape->GetPrivate();
+		const NodeMeshReshape::Private& Node = *Reshape->GetPrivate();
 
 		Ptr<ASTOpMeshBindShape> OpBind = new ASTOpMeshBindShape();
 		Ptr<ASTOpMeshApplyShape> OpApply = new ASTOpMeshApplyShape();
 
-		OpBind->bReshapeSkeleton = node.m_reshapeSkeleton;	
-		OpBind->BonesToDeform = node.BonesToDeform;
-    	OpBind->bReshapePhysicsVolumes = node.m_reshapePhysicsVolumes;
-		OpBind->PhysicsToDeform = node.PhysicsToDeform;
-		OpBind->bReshapeVertices = node.m_reshapeVertices;
+		OpBind->bReshapeSkeleton = Node.bReshapeSkeleton;	
+		OpBind->BonesToDeform = Node.BonesToDeform;
+    	OpBind->bReshapePhysicsVolumes = Node.bReshapePhysicsVolumes;
+		OpBind->PhysicsToDeform = Node.PhysicsToDeform;
+		OpBind->bReshapeVertices = Node.bReshapeVertices;
+		OpBind->bApplyLaplacian = Node.bApplyLaplacian;
 		OpBind->BindingMethod = static_cast<uint32>(EShapeBindingMethod::ReshapeClosestProject);
 
-		OpBind->RChannelUsage = node.ColorRChannelUsage;
-		OpBind->GChannelUsage = node.ColorGChannelUsage;
-		OpBind->BChannelUsage = node.ColorBChannelUsage;
-		OpBind->AChannelUsage = node.ColorAChannelUsage;
+		OpBind->RChannelUsage = Node.ColorRChannelUsage;
+		OpBind->GChannelUsage = Node.ColorGChannelUsage;
+		OpBind->BChannelUsage = Node.ColorBChannelUsage;
+		OpBind->AChannelUsage = Node.ColorAChannelUsage;
 
 		OpApply->bReshapeVertices = OpBind->bReshapeVertices;
 		OpApply->bReshapeSkeleton = OpBind->bReshapeSkeleton;
+		OpApply->bApplyLaplacian = OpBind->bApplyLaplacian;
 		OpApply->bReshapePhysicsVolumes = OpBind->bReshapePhysicsVolumes;
 
 		// Base Mesh
-		if (node.m_pBaseMesh)
+		if (Node.BaseMesh)
 		{
 			FMeshGenerationOptions BaseOptions = InOptions;
 			BaseOptions.bUniqueVertexIDs = true;
 
-			GenerateMesh(BaseOptions, OutResult, node.m_pBaseMesh);
+			GenerateMesh(BaseOptions, OutResult, Node.BaseMesh);
 			OpBind->Mesh = OutResult.meshOp;
 		}
 		else
 		{
 			// This argument is required
-			m_pErrorLog->GetPrivate()->Add("Mesh reshape base node is not set.", ELMT_ERROR, node.m_errorContext);
+			m_pErrorLog->GetPrivate()->Add("Mesh reshape base node is not set.", ELMT_ERROR, Node.m_errorContext);
 		}
 
 		// Base and target shapes shouldn't have layouts or modifiers.
@@ -1823,20 +1826,20 @@ class Node;
 		ShapeOptions.ActiveTags.Empty();
 
 		// Base Shape
-		if (node.m_pBaseShape)
+		if (Node.BaseShape)
 		{
 			FMeshGenerationResult baseResult;
-			GenerateMesh(ShapeOptions, baseResult, node.m_pBaseShape);
+			GenerateMesh(ShapeOptions, baseResult, Node.BaseShape);
 			OpBind->Shape = baseResult.meshOp;
 		}
 
 		OpApply->Mesh = OpBind;
 
 		// Target Shape
-		if (node.m_pTargetShape)
+		if (Node.TargetShape)
 		{
 			FMeshGenerationResult targetResult;
-			GenerateMesh(ShapeOptions, targetResult, node.m_pTargetShape);
+			GenerateMesh(ShapeOptions, targetResult, Node.TargetShape);
 			OpApply->Shape = targetResult.meshOp;
 		}
 
diff --git a/Engine/Plugins/Experimental/Mutable/Source/MutableTools/Private/MuT/NodeMeshReshape.cpp b/Engine/Plugins/Experimental/Mutable/Source/MutableTools/Private/MuT/NodeMeshReshape.cpp
--- a/Engine/Plugins/Experimental/Mutable/Source/MutableTools/Private/MuT/NodeMeshReshape.cpp
+++ b/Engine/Plugins/Experimental/Mutable/Source/MutableTools/Private/MuT/NodeMeshReshape.cpp
@@ -18,14 +18,14 @@ namespace mu
 	// Static initialisation
 	//---------------------------------------------------------------------------------------------
 	NODE_TYPE NodeMeshReshape::Private::s_type =
-			NODE_TYPE( "MeshReshape", NodeMesh::GetStaticType() );
+			NODE_TYPE("MeshReshape", NodeMesh::GetStaticType());
 
 
 	//---------------------------------------------------------------------------------------------
 	//!
 	//---------------------------------------------------------------------------------------------
 
-	MUTABLE_IMPLEMENT_NODE( NodeMeshReshape, EType::Reshape, Node, Node::EType::Mesh)
+	MUTABLE_IMPLEMENT_NODE(NodeMeshReshape, EType::Reshape, Node, Node::EType::Mesh)
 
 
 	//---------------------------------------------------------------------------------------------
@@ -46,9 +46,9 @@ namespace mu
 
 		switch (i)
 		{
-		case 0: pResult = m_pD->m_pBaseMesh.get(); break;
-		case 1: pResult = m_pD->m_pBaseShape.get(); break;
-		case 2: pResult = m_pD->m_pTargetShape.get(); break;
+		case 0: pResult = m_pD->BaseMesh.get(); break;
+		case 1: pResult = m_pD->BaseShape.get(); break;
+		case 2: pResult = m_pD->TargetShape.get(); break;
 		default:
 			break;
 		}
@@ -64,9 +64,9 @@ namespace mu
 
 		switch (i)
 		{
-		case 0: m_pD->m_pBaseMesh = dynamic_cast<NodeMesh*>(pNode.get()); break;
-		case 1: m_pD->m_pBaseShape = dynamic_cast<NodeMesh*>(pNode.get()); break;
-		case 2: m_pD->m_pTargetShape = dynamic_cast<NodeMesh*>(pNode.get()); break;
+		case 0: m_pD->BaseMesh = dynamic_cast<NodeMesh*>(pNode.get()); break;
+		case 1: m_pD->BaseShape = dynamic_cast<NodeMesh*>(pNode.get()); break;
+		case 2: m_pD->TargetShape = dynamic_cast<NodeMesh*>(pNode.get()); break;
 		default:
 			break;
 		}
@@ -78,54 +78,61 @@ namespace mu
 	//---------------------------------------------------------------------------------------------
 	const Ptr<NodeMesh>& NodeMeshReshape::GetBaseMesh() const
 	{
-		return m_pD->m_pBaseMesh;
+		return m_pD->BaseMesh;
 	}
 
 
 	//---------------------------------------------------------------------------------------------
 	void NodeMeshReshape::SetBaseMesh(const Ptr<NodeMesh>& pNode)
 	{
-		m_pD->m_pBaseMesh = pNode;
+		m_pD->BaseMesh = pNode;
 	}
 
 
 	//---------------------------------------------------------------------------------------------
 	const Ptr<NodeMesh>& NodeMeshReshape::GetBaseShape() const
 	{
-		return m_pD->m_pBaseShape;
+		return m_pD->BaseShape;
 	}
 
 
 	//---------------------------------------------------------------------------------------------
 	void NodeMeshReshape::SetBaseShape(const Ptr<NodeMesh>& pNode)
 	{
-		m_pD->m_pBaseShape = pNode;
+		m_pD->BaseShape = pNode;
 	}
 
 
 	//---------------------------------------------------------------------------------------------
 	const Ptr<NodeMesh>& NodeMeshReshape::GetTargetShape() const
 	{
-		return m_pD->m_pTargetShape;
+		return m_pD->TargetShape;
 	}
 
 
 	//---------------------------------------------------------------------------------------------
 	void NodeMeshReshape::SetTargetShape(const Ptr<NodeMesh>& pNode)
 	{
-		m_pD->m_pTargetShape = pNode;
+		m_pD->TargetShape = pNode;
 	}
 
 	//---------------------------------------------------------------------------------------------
 	void NodeMeshReshape::SetReshapeVertices(bool bEnable)
 	{
-		m_pD->m_reshapeVertices = bEnable;
+		m_pD->bReshapeVertices = bEnable;
 	}
 
+	//---------------------------------------------------------------------------------------------
+	void NodeMeshReshape::SetApplyLaplacian(bool bEnable)
+	{
+		m_pD->bApplyLaplacian = bEnable;
+	}
+
+
 	//---------------------------------------------------------------------------------------------
 	void NodeMeshReshape::SetReshapeSkeleton(bool bEnable)
 	{
-		m_pD->m_reshapeSkeleton = bEnable;
+		m_pD->bReshapeSkeleton = bEnable;
 	}
 	
 
@@ -152,26 +159,25 @@ namespace mu
 
 	void NodeMeshReshape::SetReshapePhysicsVolumes(bool bEnable)
 	{
-		m_pD->m_reshapePhysicsVolumes = bEnable;
+		m_pD->bReshapePhysicsVolumes = bEnable;
 	}
 	
 	//---------------------------------------------------------------------------------------------
-	Ptr<NodeLayout> NodeMeshReshape::Private::GetLayout( int index ) const
+	Ptr<NodeLayout> NodeMeshReshape::Private::GetLayout(int32 Index) const
 	{
 		NodeLayoutPtr pResult;
 
-		if (m_pBaseMesh)
+		if (BaseMesh)
 		{
 			NodeMesh::Private* pPrivate =
-					dynamic_cast<NodeMesh::Private*>(m_pBaseMesh->GetBasePrivate() );
+					dynamic_cast<NodeMesh::Private*>(BaseMesh->GetBasePrivate() );
 
-			pResult = pPrivate->GetLayout( index );
+			pResult = pPrivate->GetLayout(Index);
 		}
 
 		return pResult;
 	}
 
-
 }
 
 
diff --git a/Engine/Plugins/Experimental/Mutable/Source/MutableTools/Private/MuT/NodeMeshReshapePrivate.h b/Engine/Plugins/Experimental/Mutable/Source/MutableTools/Private/MuT/NodeMeshReshapePrivate.h
--- a/Engine/Plugins/Experimental/Mutable/Source/MutableTools/Private/MuT/NodeMeshReshapePrivate.h
+++ b/Engine/Plugins/Experimental/Mutable/Source/MutableTools/Private/MuT/NodeMeshReshapePrivate.h
@@ -21,12 +21,13 @@ namespace mu
 
 		static NODE_TYPE s_type;
 
-		Ptr<NodeMesh> m_pBaseMesh;
-		Ptr<NodeMesh> m_pBaseShape;
-		Ptr<NodeMesh> m_pTargetShape;
-		bool m_reshapeVertices = true;
-		bool m_reshapeSkeleton = false;
-		bool m_reshapePhysicsVolumes = false;
+		Ptr<NodeMesh> BaseMesh;
+		Ptr<NodeMesh> BaseShape;
+		Ptr<NodeMesh> TargetShape;
+		bool bReshapeVertices = true;
+		bool bApplyLaplacian = false;
+		bool bReshapeSkeleton = false;
+		bool bReshapePhysicsVolumes = false;
 
 		EVertexColorUsage ColorRChannelUsage = EVertexColorUsage::None;
 		EVertexColorUsage ColorGChannelUsage = EVertexColorUsage::None;
@@ -38,18 +39,23 @@ namespace mu
         //!
 		void Serialise( OutputArchive& arch ) const
 		{
-            uint32_t ver = 10;
+            uint32 ver = 11;
 			arch << ver;
 
-			arch << m_pBaseMesh;
-			arch << m_pBaseShape;
-			arch << m_pTargetShape;
-			arch << m_reshapeVertices;
-			arch << m_reshapeSkeleton;
+			arch << BaseMesh;
+			arch << BaseShape;
+			arch << TargetShape;
+
+			arch << bReshapeVertices;
+			arch << bApplyLaplacian;
+			
+			arch << bReshapeSkeleton;
 			arch << BonesToDeform;
-			arch << m_reshapePhysicsVolumes;
+
+			arch << bReshapePhysicsVolumes;
 			arch << PhysicsToDeform;
 
+			
 			arch << ColorRChannelUsage;
 			arch << ColorGChannelUsage;
 			arch << ColorBChannelUsage;
@@ -59,20 +65,25 @@ namespace mu
 		//!
 		void Unserialise( InputArchive& arch )
 		{
-            uint32_t ver;
+            uint32 ver;
 			arch >> ver;
-			check(ver <= 10);
+			check(ver <= 11);
 
-			arch >> m_pBaseMesh;
-			arch >> m_pBaseShape;
-			arch >> m_pTargetShape;
+			arch >> BaseMesh;
+			arch >> BaseShape;
+			arch >> TargetShape;
 			
 			if (ver >= 8)
 			{
-				arch >> m_reshapeVertices;
+				arch >> bReshapeVertices;
+			}
+
+			if (ver >= 11)
+			{
+				arch >> bApplyLaplacian;
 			}
 
-			arch >> m_reshapeSkeleton;
+			arch >> bReshapeSkeleton;
 			if (ver <= 9)
 			{
 				bool bEnableRigidParts_DEPRECATED;
@@ -109,7 +120,7 @@ namespace mu
 
 		  	if (ver >= 4)
 		 	{
-				arch >> m_reshapePhysicsVolumes;
+				arch >> bReshapePhysicsVolumes;
 		 	}
 
 			if (ver >= 5 && ver < 7)
diff --git a/Engine/Plugins/Experimental/Mutable/Source/MutableTools/Public/MuT/NodeMeshReshape.h b/Engine/Plugins/Experimental/Mutable/Source/MutableTools/Public/MuT/NodeMeshReshape.h
--- a/Engine/Plugins/Experimental/Mutable/Source/MutableTools/Public/MuT/NodeMeshReshape.h
+++ b/Engine/Plugins/Experimental/Mutable/Source/MutableTools/Public/MuT/NodeMeshReshape.h
@@ -57,6 +57,9 @@ namespace mu
 
 		/** Also deform the mesh skeleton. Disabled by default. */
 		void SetReshapeSkeleton(bool);
+	
+		/** Apply Laplacian smoothing to the reshaped mesh.  */
+		void SetApplyLaplacian(bool);
 		
 		/** Set vertex color channel usages for Reshape operations. */
 		void SetColorUsages(EVertexColorUsage R, EVertexColorUsage G, EVertexColorUsage B, EVertexColorUsage A);
