[Mutable] Fix issues with NodeVaraition deserialization (and copy).

#rnx

diff --git a/Engine/Plugins/Experimental/Mutable/Source/CustomizableObject/Public/MuCO/CustomizableObjectCustomVersion.h b/Engine/Plugins/Experimental/Mutable/Source/CustomizableObject/Public/MuCO/CustomizableObjectCustomVersion.h
--- a/Engine/Plugins/Experimental/Mutable/Source/CustomizableObject/Public/MuCO/CustomizableObjectCustomVersion.h
+++ b/Engine/Plugins/Experimental/Mutable/Source/CustomizableObject/Public/MuCO/CustomizableObjectCustomVersion.h
@@ -98,6 +98,8 @@ struct CUSTOMIZABLEOBJECT_API FCustomizableObjectCustomVersion
 
 		NodeTextureParameterDefaultToReferenceValue,
 
+		NodeVariationSerializationIssue, // The parent class and child classes had a member with the same name. Unreal was serialized the member from the child class but deserialized it to the parent class.
+
 		// -----<new versions can be added above this line>-------------------------------------------------
 		VersionPlusOne,
 		LatestVersion = VersionPlusOne - 1
diff --git a/Engine/Plugins/Experimental/Mutable/Source/CustomizableObjectEditor/Classes/MuCOE/Nodes/CustomizableObjectNodeVariation.h b/Engine/Plugins/Experimental/Mutable/Source/CustomizableObjectEditor/Classes/MuCOE/Nodes/CustomizableObjectNodeVariation.h
--- a/Engine/Plugins/Experimental/Mutable/Source/CustomizableObjectEditor/Classes/MuCOE/Nodes/CustomizableObjectNodeVariation.h
+++ b/Engine/Plugins/Experimental/Mutable/Source/CustomizableObjectEditor/Classes/MuCOE/Nodes/CustomizableObjectNodeVariation.h
@@ -34,9 +34,10 @@ class CUSTOMIZABLEOBJECTEDITOR_API UCustomizableObjectNodeVariation : public UCu
 
 protected:
 	UPROPERTY(EditAnywhere, Category = CustomizableObject)
-	TArray<FCustomizableObjectVariation> Variations;
+	TArray<FCustomizableObjectVariation> VariationsData; // The variable name can not be Variations due issues with the on to UObject Serialization system
 
 private:
+	UPROPERTY()
 	TArray<FEdGraphPinReference> VariationsPins;
 	
 public:
diff --git a/Engine/Plugins/Experimental/Mutable/Source/CustomizableObjectEditor/Private/MuCOE/Nodes/CustomizableObjectNodeColorVariation.cpp b/Engine/Plugins/Experimental/Mutable/Source/CustomizableObjectEditor/Private/MuCOE/Nodes/CustomizableObjectNodeColorVariation.cpp
--- a/Engine/Plugins/Experimental/Mutable/Source/CustomizableObjectEditor/Private/MuCOE/Nodes/CustomizableObjectNodeColorVariation.cpp
+++ b/Engine/Plugins/Experimental/Mutable/Source/CustomizableObjectEditor/Private/MuCOE/Nodes/CustomizableObjectNodeColorVariation.cpp
@@ -10,14 +10,14 @@ void UCustomizableObjectNodeColorVariation::BackwardsCompatibleFixup()
 {
 	const int32 CustomizableObjectCustomVersion = GetLinkerCustomVersion(FCustomizableObjectCustomVersion::GUID);
 
-	if (CustomizableObjectCustomVersion < FCustomizableObjectCustomVersion::DeduplicateNodeVariant)
+	if (CustomizableObjectCustomVersion < FCustomizableObjectCustomVersion::NodeVariationSerializationIssue)
 	{
 		for (const FCustomizableObjectColorVariation& OldVariation : Variations_DEPRECATED)
 		{
 			FCustomizableObjectVariation Variation;
 			Variation.Tag = OldVariation.Tag;
 			
-			Variations.Add(Variation);
+			VariationsData.Add(Variation);
 		}
 	}
 
diff --git a/Engine/Plugins/Experimental/Mutable/Source/CustomizableObjectEditor/Private/MuCOE/Nodes/CustomizableObjectNodeFloatVariation.cpp b/Engine/Plugins/Experimental/Mutable/Source/CustomizableObjectEditor/Private/MuCOE/Nodes/CustomizableObjectNodeFloatVariation.cpp
--- a/Engine/Plugins/Experimental/Mutable/Source/CustomizableObjectEditor/Private/MuCOE/Nodes/CustomizableObjectNodeFloatVariation.cpp
+++ b/Engine/Plugins/Experimental/Mutable/Source/CustomizableObjectEditor/Private/MuCOE/Nodes/CustomizableObjectNodeFloatVariation.cpp
@@ -10,14 +10,14 @@ void UCustomizableObjectNodeFloatVariation::BackwardsCompatibleFixup()
 {
 	const int32 CustomizableObjectCustomVersion = GetLinkerCustomVersion(FCustomizableObjectCustomVersion::GUID);
 
-	if (CustomizableObjectCustomVersion < FCustomizableObjectCustomVersion::DeduplicateNodeVariant)
+	if (CustomizableObjectCustomVersion < FCustomizableObjectCustomVersion::NodeVariationSerializationIssue)
 	{
 		for (const FCustomizableObjectFloatVariation& OldVariation : Variations_DEPRECATED)
 		{
 			FCustomizableObjectVariation Variation;
 			Variation.Tag = OldVariation.Tag;
 			
-			Variations.Add(Variation);
+			VariationsData.Add(Variation);
 		}
 	}
 
diff --git a/Engine/Plugins/Experimental/Mutable/Source/CustomizableObjectEditor/Private/MuCOE/Nodes/CustomizableObjectNodeMaterialVariation.cpp b/Engine/Plugins/Experimental/Mutable/Source/CustomizableObjectEditor/Private/MuCOE/Nodes/CustomizableObjectNodeMaterialVariation.cpp
--- a/Engine/Plugins/Experimental/Mutable/Source/CustomizableObjectEditor/Private/MuCOE/Nodes/CustomizableObjectNodeMaterialVariation.cpp
+++ b/Engine/Plugins/Experimental/Mutable/Source/CustomizableObjectEditor/Private/MuCOE/Nodes/CustomizableObjectNodeMaterialVariation.cpp
@@ -10,14 +10,14 @@ void UCustomizableObjectNodeMaterialVariation::BackwardsCompatibleFixup()
 {
 	const int32 CustomizableObjectCustomVersion = GetLinkerCustomVersion(FCustomizableObjectCustomVersion::GUID);
 
-	if (CustomizableObjectCustomVersion < FCustomizableObjectCustomVersion::DeduplicateNodeVariant)
+	if (CustomizableObjectCustomVersion < FCustomizableObjectCustomVersion::NodeVariationSerializationIssue)
 	{
 		for (const FCustomizableObjectMaterialVariation& OldVariation : Variations_DEPRECATED)
 		{
 			FCustomizableObjectVariation Variation;
 			Variation.Tag = OldVariation.Tag;
 			
-			Variations.Add(Variation);
+			VariationsData.Add(Variation);
 		}
 	}
 
diff --git a/Engine/Plugins/Experimental/Mutable/Source/CustomizableObjectEditor/Private/MuCOE/Nodes/CustomizableObjectNodeMeshVariation.cpp b/Engine/Plugins/Experimental/Mutable/Source/CustomizableObjectEditor/Private/MuCOE/Nodes/CustomizableObjectNodeMeshVariation.cpp
--- a/Engine/Plugins/Experimental/Mutable/Source/CustomizableObjectEditor/Private/MuCOE/Nodes/CustomizableObjectNodeMeshVariation.cpp
+++ b/Engine/Plugins/Experimental/Mutable/Source/CustomizableObjectEditor/Private/MuCOE/Nodes/CustomizableObjectNodeMeshVariation.cpp
@@ -10,14 +10,14 @@ void UCustomizableObjectNodeMeshVariation::BackwardsCompatibleFixup()
 {
 	const int32 CustomizableObjectCustomVersion = GetLinkerCustomVersion(FCustomizableObjectCustomVersion::GUID);
 
-	if (CustomizableObjectCustomVersion < FCustomizableObjectCustomVersion::DeduplicateNodeVariant)
+	if (CustomizableObjectCustomVersion < FCustomizableObjectCustomVersion::NodeVariationSerializationIssue)
 	{
 		for (const FCustomizableObjectMeshVariation& OldVariation : Variations_DEPRECATED)
 		{
 			FCustomizableObjectVariation Variation;
 			Variation.Tag = OldVariation.Tag;
 			
-			Variations.Add(Variation);
+			VariationsData.Add(Variation);
 		}
 	}
 
diff --git a/Engine/Plugins/Experimental/Mutable/Source/CustomizableObjectEditor/Private/MuCOE/Nodes/CustomizableObjectNodeVariation.cpp b/Engine/Plugins/Experimental/Mutable/Source/CustomizableObjectEditor/Private/MuCOE/Nodes/CustomizableObjectNodeVariation.cpp
--- a/Engine/Plugins/Experimental/Mutable/Source/CustomizableObjectEditor/Private/MuCOE/Nodes/CustomizableObjectNodeVariation.cpp
+++ b/Engine/Plugins/Experimental/Mutable/Source/CustomizableObjectEditor/Private/MuCOE/Nodes/CustomizableObjectNodeVariation.cpp
@@ -32,13 +32,13 @@ void UCustomizableObjectNodeVariation::AllocateDefaultPins(UCustomizableObjectNo
 		CustomCreatePin(EGPD_Output, Category, PinName);
 	}
 	
-	VariationsPins.SetNum(Variations.Num());
-	for (int32 VariationIndex = Variations.Num() - 1; VariationIndex >= 0; --VariationIndex)
+	VariationsPins.SetNum(VariationsData.Num());
+	for (int32 VariationIndex = VariationsData.Num() - 1; VariationIndex >= 0; --VariationIndex)
 	{
 		const FName PinName = FName(FString::Printf( TEXT("Variation %d"), VariationIndex));
 		UEdGraphPin* VariationPin = CustomCreatePin(EGPD_Input, Category, PinName, bIsInputPinArray);
 
-		VariationPin->PinFriendlyName = FText::Format(LOCTEXT("Variation_Pin_FriendlyName", "Variation {0} [{1}]"), VariationIndex, FText::FromString(*Variations[VariationIndex].Tag));
+		VariationPin->PinFriendlyName = FText::Format(LOCTEXT("Variation_Pin_FriendlyName", "Variation {0} [{1}]"), VariationIndex, FText::FromString(*VariationsData[VariationIndex].Tag));
 		
 		VariationsPins[VariationIndex] = VariationPin;
 	}
@@ -55,13 +55,13 @@ bool UCustomizableObjectNodeVariation::IsInputPinArray() const
 
 int32 UCustomizableObjectNodeVariation::GetNumVariations() const
 {
-	return Variations.Num();
+	return VariationsData.Num();
 }
 
 
 const FCustomizableObjectVariation& UCustomizableObjectNodeVariation::GetVariation(int32 Index) const
 {
-	return Variations[Index];
+	return VariationsData[Index];
 }
 
 
@@ -102,7 +102,7 @@ void UCustomizableObjectNodeVariation::BackwardsCompatibleFixup()
 
 	const int32 CustomizableObjectCustomVersion = GetLinkerCustomVersion(FCustomizableObjectCustomVersion::GUID);
 
-	if (CustomizableObjectCustomVersion < FCustomizableObjectCustomVersion::DeduplicateNodeVariant)
+	if (CustomizableObjectCustomVersion < FCustomizableObjectCustomVersion::NodeVariationSerializationIssue)
 	{
 		ReconstructNode();
 	}
