[Mutable] Add support for cylindrical and wrap projection bilinear sampling.

- Refactor raster pixel processor to unify all processors into one.
- Refactor mesh raster to use a function pointer for pixel processing to remove template bloat.
- Fix crash with wrapping projectors when trying to compute best mip.
- Make some parts of the projectors code a bit closer to UE standards.
- Added ParallelFor for Wrapping Projector rasterization.

#rnx

diff --git a/Engine/Plugins/Experimental/Mutable/Source/MutableRuntime/Private/MuR/CodeRunner.cpp b/Engine/Plugins/Experimental/Mutable/Source/MutableRuntime/Private/MuR/CodeRunner.cpp
--- a/Engine/Plugins/Experimental/Mutable/Source/MutableRuntime/Private/MuR/CodeRunner.cpp
+++ b/Engine/Plugins/Experimental/Mutable/Source/MutableRuntime/Private/MuR/CodeRunner.cpp
@@ -4409,9 +4409,9 @@ namespace mu
 				}
 
 				// Allocate memory for the temporary buffers
-				SCRATCH_IMAGE_PROJECT scratch;
-				scratch.vertices.SetNum( pMesh->GetVertexCount() );
-				scratch.culledVertex.SetNum( pMesh->GetVertexCount() );
+				FScratchImageProject Scratch;
+				Scratch.Vertices.SetNum(pMesh->GetVertexCount());
+				Scratch.CulledVertex.SetNum(pMesh->GetVertexCount());
 
 				ESamplingMethod SamplingMethod = Invoke([&]() -> ESamplingMethod
 				{
@@ -4462,7 +4462,7 @@ namespace mu
 							FadeStartRad, FadeEndRad, FMath::Frac(Data.RasterMesh.MipValue),
 							args.LayoutIndex, args.blockId,
 							CropMin, UncroppedSize,
-							&scratch, bUseProjectionVectorImpl);
+							&Scratch, bUseProjectionVectorImpl);
 						break;
 
 					case PROJECTOR_TYPE::WRAPPING:
@@ -4473,18 +4473,19 @@ namespace mu
 							FadeStartRad, FadeEndRad, FMath::Frac(Data.RasterMesh.MipValue),
 							args.LayoutIndex, args.blockId,
 							CropMin, UncroppedSize,
-							&scratch);
+							&Scratch, bUseProjectionVectorImpl);
 						break;
 
 					case PROJECTOR_TYPE::CYLINDRICAL:
 						ImageRasterProjectedCylindrical(pMesh.get(), New.get(),
 							Source.get(), Mask.get(),
 							args.bIsRGBFadingEnabled, args.bIsAlphaFadingEnabled,
-							FadeStartRad, FadeEndRad,
+							SamplingMethod,
+							FadeStartRad, FadeEndRad, FMath::Frac(Data.RasterMesh.MipValue),
 							args.LayoutIndex,
 							Projector.projectionAngle,
 							CropMin, UncroppedSize,
-							&scratch);
+							&Scratch, bUseProjectionVectorImpl);
 						break;
 
 					default:
@@ -4821,8 +4822,8 @@ namespace mu
             break;
         }
     }
-
-    //---------------------------------------------------------------------------------------------
+    
+	//---------------------------------------------------------------------------------------------
     Ptr<RangeIndex> CodeRunner::BuildCurrentOpRangeIndex( const FScheduledOp& item, const Parameters* pParams, const Model* pModel, int32 parameterIndex )
     {
         if (!item.ExecutionIndex)
diff --git a/Engine/Plugins/Experimental/Mutable/Source/MutableRuntime/Private/MuR/OpImageProject.cpp b/Engine/Plugins/Experimental/Mutable/Source/MutableRuntime/Private/MuR/OpImageProject.cpp
--- a/Engine/Plugins/Experimental/Mutable/Source/MutableRuntime/Private/MuR/OpImageProject.cpp
+++ b/Engine/Plugins/Experimental/Mutable/Source/MutableRuntime/Private/MuR/OpImageProject.cpp
@@ -1,6 +1,5 @@
 // Copyright Epic Games, Inc. All Rights Reserved.
 
-
 #include "MuR/OpImageProject.h"
 #include "MuR/ConvertData.h"
 #include "MuR/Image.h"
@@ -23,406 +22,272 @@
 #include "Misc/AssertionMacros.h"
 #include "Math/UnrealMathUtility.h"
 
-#include "Math/VectorRegister.h"
-
 #include "Spatial/PointHashGrid3.h"
 
-namespace mu
-{
+#include "MuR/Image.h"
+#include "MuR/Mesh.h"
 
 
-namespace
+namespace mu
 {
 
+//-------------------------------------------------------------------------------------------------
+//! This format is the one we assumee the meshes optimised for planar and cylindrical projection
+//! will have.
+//! See CreateMeshOptimisedForProjection
+struct FOptimizedVertex
+{
+	FVector2f Uv;
+	FVector3f Position;
+	FVector3f Normal;
+};
 
-    //---------------------------------------------------------------------------------------------
-    template<int PIXEL_SIZE,bool HAS_MASK=true>
-    class RasterProjectedPixelProcessor_UBYTE
-    {
-    public:
-
-        RasterProjectedPixelProcessor_UBYTE( const Image* pSource,
-			const uint8* pTargetData,
-			const uint8* pMaskData,
-			bool bIsRGBFadingEnabled, bool bIsAlphaFadingEnabled,
-			float fadeStart, float fadeEnd )
-        {
-            m_pSourceData = pSource->GetData();
-            m_sourceSizeX = pSource->GetSizeX();
-            m_sourceSizeY = pSource->GetSizeY();
-
-			// \TODO: to template parameters?
-			m_bIsRGBFadingEnabled = bIsRGBFadingEnabled;
-			m_bIsAlphaFadingEnabled = bIsAlphaFadingEnabled;
-
-            m_fadeStartCos = cosf(fadeStart);
-            m_fadeEndCos = cosf(fadeEnd);
-
-            m_pTargetData = pTargetData;
-            m_pMaskData = pMaskData;
-
-            check( GetImageFormatData( pSource->GetFormat() ).BytesPerBlock
-                            ==
-                            PIXEL_SIZE );
-        }
-
-
-        //-----------------------------------------------------------------------------------------
-        inline void ProcessPixel( unsigned char* pBufferPos, float varying[4] ) const
-        {
-			// \TODO: This is very bad for the CPU. Please optimize.
+static_assert(sizeof(FOptimizedVertex)==32, "UNEXPECTED_STRUCT_SIZE" );
 
-			if (!m_pSourceData) return;
 
-			int32 factor_8 = 256;
+//-------------------------------------------------------------------------------------------------
+//! This format is the one we assumee the meshes optimised for wrapping projection will have.
+//! See CreateMeshOptimisedForWrappingProjection
+struct FOptimizedVertexWrapping
+{
+	FVector2f Uv;
+	FVector3f Position;
+	FVector3f Normal;
+	uint32 LayoutBlock;
+};
 
-            // depth clamp
-            if (varying[2]<0.0f || varying[2]>1.0f)
-            {
-                factor_8 = 0;
-            }
+static_assert(sizeof(FOptimizedVertexWrapping)==36, "UNEXPECTED_STRUCT_SIZE");
 
-			int32 factor_unfaded_8 = factor_8;
 
-            float angleCos = varying[3];
+namespace Private
+{
+	struct FImageRasterInvokeArgs
+	{
+		const Mesh* MeshPtr;
+		Image* ImagePtr;
+		const Image* SourcePtr;
+		const Image* MaskPtr;
+		bool bIsRGBFadingEnabled;
+		bool bIsAlphaFadingEnabled;
+		ESamplingMethod SamplingMethod;
+		float FadeStart;
+		float FadeEnd;
+		float ProjectionAngle;
+		float MipInterpolationFactor;
+		int32 Layout;
+		int32 Block;
+		UE::Math::TIntVector2<uint16> CropMin;
+		UE::Math::TIntVector2<uint16> UncroppedSize;
+		FScratchImageProject* Scratch;
+	};
 
-            if (HAS_MASK)
-            {
-                if ( factor_8>0 && m_pMaskData )
-                {
-                    uint8 maskFactor = m_pMaskData[ (pBufferPos-m_pTargetData) / PIXEL_SIZE ];
-                    factor_8 = maskFactor;
-                }
+	enum class EPixelProcessorFeatures
+	{
+		None = 0,
+		WithMask                 = 1 << 1,
+		SamplingPoint            = 1 << 2,
+		SamplingLinear			 = 1 << 3, 
+		ProjectionPlanarAndWrap  = 1 << 4,
+		ProjectionCylindrical    = 1 << 5,
+		FormatRGBA               = 1 << 6,
+		FormatRGB                = 1 << 7,
+		FormatL                  = 1 << 8,
+		VectorizedImpl			 = 1 << 9    
+	};
 
-				factor_unfaded_8 = factor_8;
+	ENUM_CLASS_FLAGS(EPixelProcessorFeatures);
 
-                if ( factor_8>0 )
-                {
-                    if ( angleCos<m_fadeStartCos && angleCos>m_fadeEndCos )
-                    {
-                        factor_8 = ( factor_8
-                                     *
-                                     int32( 256.0f * ( angleCos - m_fadeEndCos ) / ( m_fadeStartCos - m_fadeEndCos ) )
-                                     ) >> 8;
-                    }
-                    else if ( angleCos<=m_fadeEndCos )
-                    {
-                        factor_8 = 0;
-						factor_unfaded_8 = 0;
-                    }
-                }
-            }
-            else
-            {
-                if ( angleCos<m_fadeStartCos && angleCos>m_fadeEndCos )
-                {
-                    factor_8 = int32( 256.0f * ( angleCos - m_fadeEndCos ) / ( m_fadeStartCos - m_fadeEndCos ) );
-                }
-                else if ( angleCos<=m_fadeEndCos )
-                {
-                    factor_8 = 0;
-					factor_unfaded_8 = 0;
-                }
-            }
+	constexpr bool CheckExactlyOneFormatFlag(EPixelProcessorFeatures Features)
+	{
+		return 
+			static_cast<int32>(EnumHasAnyFlags(Features, EPixelProcessorFeatures::FormatL)) +
+			static_cast<int32>(EnumHasAnyFlags(Features, EPixelProcessorFeatures::FormatRGB)) +
+			static_cast<int32>(EnumHasAnyFlags(Features, EPixelProcessorFeatures::FormatRGBA)) == 1;
+	}
 
-            if ( factor_8>0 )
-            {
-                float u = varying[0];
-                float v = varying[1];
+	constexpr bool CheckExactlyOneSamplingFlag(EPixelProcessorFeatures Features)
+	{
+		return
+			static_cast<int32>(EnumHasAnyFlags(Features, EPixelProcessorFeatures::SamplingPoint)) +
+			static_cast<int32>(EnumHasAnyFlags(Features, EPixelProcessorFeatures::SamplingLinear)) == 1;
+	}
 
-                if ( u>=0 && u<1 && v>=0 && v<1 )
-                {
+	constexpr bool CheckExactlyOneProjectionFlag(EPixelProcessorFeatures Features)
+	{
+		return
+			static_cast<int32>(EnumHasAnyFlags(Features, EPixelProcessorFeatures::ProjectionCylindrical)) +
+			static_cast<int32>(EnumHasAnyFlags(Features, EPixelProcessorFeatures::ProjectionPlanarAndWrap)) == 1;
+	}
 
-                    const uint8* pPixel = m_pSourceData;
-                    // TODO: clamp?
-                    pPixel += ( m_sourceSizeX * int(m_sourceSizeY * v) + int( m_sourceSizeX * u ) ) * PIXEL_SIZE;
+	constexpr int32 GetFormatNumChannels(EPixelProcessorFeatures Features)
+	{
+		if (EnumHasAnyFlags(Features, EPixelProcessorFeatures::FormatL))
+		{
+			return 1;
+		}
+		else if (EnumHasAnyFlags(Features, EPixelProcessorFeatures::FormatRGB))
+		{
+			return 3;
+		}
+		else if (EnumHasAnyFlags(Features, EPixelProcessorFeatures::FormatRGBA))
+		{
+			return 4;
+		}
 
-                    // Write result
-					if (!m_bIsRGBFadingEnabled && !m_bIsAlphaFadingEnabled)
-					{
-						for (int i = 0; i < PIXEL_SIZE; ++i)
-						{
-							pBufferPos[i] = uint8((pPixel[i] * factor_unfaded_8) >> 8);
-						}
-					}
-					else if (m_bIsRGBFadingEnabled && m_bIsAlphaFadingEnabled)
-					{
-						for (int i = 0; i < PIXEL_SIZE; ++i)
-						{
-							pBufferPos[i] = uint8((pPixel[i] * factor_8) >> 8);
-						}
-					}
-					else
-					{
-						if (m_bIsRGBFadingEnabled)
-						{
-							for (int i = 0; i < FMath::Min(PIXEL_SIZE, 3); ++i)
-							{
-								pBufferPos[i] = uint8((pPixel[i] * factor_8) >> 8);
-							}
-						}
-						else
-						{
-							for (int i = 0; i < FMath::Min(PIXEL_SIZE, 3); ++i)
-							{
-								pBufferPos[i] = uint8((pPixel[i] * factor_unfaded_8) >> 8);
-							}
-						}
+		return 0;
+	}
 
-						if (PIXEL_SIZE >= 3)
-						{
-							if (m_bIsAlphaFadingEnabled)
-							{
-								pBufferPos[3] = uint8((pPixel[3] * factor_8) >> 8);
-							}
-							else
-							{
-								pBufferPos[3] = uint8((pPixel[3] * factor_unfaded_8) >> 8);
-							}
-						}
-					}
-                }
-            }
-        }
+	// The processor constant data fits into a 64 bytes cache line, align to it so potentially have less cache 
+	// misses.
+	struct alignas(64) FProjectedPixelProcessorContext
+	{
+		const uint8* TargetData = nullptr;	
+		const uint8* Source0Data = nullptr;
+		
+		uint16 Source0SizeX = 0;
+		uint16 Source0SizeY = 0;
+		
+		uint16 Source1SizeX = 0;
+		uint16 Source1SizeY = 0;
+
+		const uint8* Source1Data = nullptr;
+		const uint8* MaskData = nullptr;
 
+        //! Cosine of the fading angle range
+		float FadeEndCos = 0.0f;
+		// precomputation of 255 / (FadeStartCos - FadeEndCos) 
+		float OneOverFadeRangeTimes255 = 0.0f;
 
-        const uint8* m_pSourceData;
-        int m_sourceSizeX, m_sourceSizeY;
+		// Only used for cylindrical projections.
+		float  OneOverProjectionAngle = 0.0f;
 
-        //! Cosine of the fading angle range
-        float m_fadeStartCos, m_fadeEndCos;
-		bool m_bIsRGBFadingEnabled, m_bIsAlphaFadingEnabled;
+		// Only used with linear sampling.
+		float MipInterpolationFactor = 0.0f;
 
-        const uint8* m_pTargetData;
-        const uint8* m_pMaskData;
-    };
+		// Those are the seed for a mask 0xFFFFFFFF (-1) mask for enabled and 0 otherwise.
+		// so, the expected value here is -1 for enabled and 0 for disabled.
+		int32 RGBFadingEnabledMask = 0;
+		int32 AlphaFadingEnabledMask = 0;
+	};
 
+	static_assert(sizeof(FProjectedPixelProcessorContext) <= 64);
 
+	template<EPixelProcessorFeatures Features>
+	class TProjectedPixelProcessor
+	{
+		static_assert(CheckExactlyOneFormatFlag(Features));
+		static_assert(CheckExactlyOneSamplingFlag(Features));
+		static_assert(CheckExactlyOneProjectionFlag(Features));
 
-    //---------------------------------------------------------------------------------------------
-    template<int PIXEL_SIZE, bool bUseVectorImpl = false, bool HAS_MASK = true>
-    class RasterProjectedPixelProcessorLinear_UBYTE
-    {
-    public:
+		static constexpr int32 PIXEL_SIZE = GetFormatNumChannels(Features);
 
-        RasterProjectedPixelProcessorLinear_UBYTE( const Image* Source,
-			const uint8* pTargetData,
-			const uint8* pMaskData,
+	public:
+		static FProjectedPixelProcessorContext MakeContext(
+			const Image* Source,
+			const uint8* InTargetData,
+			const uint8* InMaskData,
 			bool bInIsRGBFadingEnabled, bool bInIsAlphaFadingEnabled,
-			float FadeStart, float FadeEnd, float MipSamplingFactor )
-        {
+			float FadeStart, float FadeEnd, float InProjectionAngle, float InMipInterpolationFactor)
+		{
+			FProjectedPixelProcessorContext Context;
+
+            Context.Source0SizeX = static_cast<uint16>(Source->GetSizeX());
+            Context.Source0SizeY = static_cast<uint16>(Source->GetSizeY());
+			Context.Source0Data = (Context.Source0SizeX > 0 && Context.Source0SizeY > 0) ? Source->GetData() : nullptr;
 
-            Source0SizeX = static_cast<uint16>(Source->GetSizeX());
-            Source0SizeY = static_cast<uint16>(Source->GetSizeY());
+			if constexpr (EnumHasAnyFlags(Features, EPixelProcessorFeatures::SamplingLinear))
+			{
+				// If we don't have the next mip, use the first mip as a fallback.
+				const bool bHasNextMip = Source->GetLODCount() >= 2;
 
-			// If we don't have the next mip, use the first mip as a fallback.
-			const bool bHasNextMip = Source->GetLODCount() >= 2;
+				Context.Source1Data = bHasNextMip ? Source->GetMipData(1) : Context.Source0Data;
 
-            Source0DataPtr = (Source0SizeX > 0 && Source0SizeY > 0) ? Source->GetData() : nullptr;
-			Source1DataPtr = bHasNextMip ? Source->GetMipData(1) : Source0DataPtr;
-	
-			FIntVector2 Source1Size = Source->CalculateMipSize(1);
-            Source1SizeX = bHasNextMip ? static_cast<uint16>(Source1Size.X) : Source0SizeX;
-            Source1SizeY = bHasNextMip ? static_cast<uint16>(Source1Size.Y) : Source0SizeY;
+				FIntVector2 Source1Size = Source->CalculateMipSize(1);
+				Context.Source1SizeX = bHasNextMip ? static_cast<uint16>(Source1Size.X) : Context.Source0SizeX;
+				Context.Source1SizeY = bHasNextMip ? static_cast<uint16>(Source1Size.Y) : Context.Source0SizeY;
 
-			// Invalidate SourceDataPtrs if either of the two is not valid. 
-			if (!(Source0DataPtr && Source1DataPtr))
-			{
-				Source0DataPtr = nullptr;
-				Source1DataPtr = nullptr;
+				// Invalidate SourceData if either of the two is not valid. 
+				if (!(Context.Source0Data && Context.Source1Data))
+				{
+					Context.Source0Data = nullptr;
+					Context.Source1Data = nullptr;
+				}
 			}
 
-			RGBFadingEnabledValue = bInIsRGBFadingEnabled ? -1 : 0;
-			AlphaFadingEnabledValue = bInIsAlphaFadingEnabled ? -1 : 0;
+			Context.MipInterpolationFactor = InMipInterpolationFactor;
+
+			Context.RGBFadingEnabledMask = bInIsRGBFadingEnabled ? -1 : 0;
+			Context.AlphaFadingEnabledMask = bInIsAlphaFadingEnabled ? -1 : 0;
 
 			const float FadeStartCos = FMath::Cos(FadeStart);
-            FadeEndCos = FMath::Cos(FadeEnd);
+            Context.FadeEndCos = FMath::Cos(FadeEnd);
 
 			const float FadeCosRangeSafeDiv = 
-				FMath::IsNearlyZero(FadeStartCos - FadeEndCos, UE_KINDA_SMALL_NUMBER) 
+				FMath::IsNearlyZero(FadeStartCos - Context.FadeEndCos, UE_KINDA_SMALL_NUMBER) 
 				? UE_KINDA_SMALL_NUMBER 
-				: FadeStartCos - FadeEndCos;
+				: FadeStartCos - Context.FadeEndCos;
 
-			OneOverFadeRangeTimes255 = 255.0f / FadeCosRangeSafeDiv;
+			Context.OneOverFadeRangeTimes255 = 255.0f / FadeCosRangeSafeDiv;
 
-            TargetDataPtr = pTargetData;
-            MaskDataPtr = pMaskData;
+			Context.OneOverProjectionAngle = FMath::IsNearlyZero(InProjectionAngle)
+				? 1.0f / UE_KINDA_SMALL_NUMBER
+				: 1.0f / InProjectionAngle;
+
+            Context.TargetData = InTargetData;
+            Context.MaskData = InMaskData;
 
             check(GetImageFormatData(Source->GetFormat() ).BytesPerBlock == PIXEL_SIZE);
-        }
 
-		inline void ProcessPixel(uint8* BufferPosPtr, float Varying[4]) const
+			return Context;
+		}
+
+		static void ProcessPixel(const FProjectedPixelProcessorContext& Context, uint8* BufferPos, float Varying[4])
 		{
-			if constexpr (bUseVectorImpl)
+			if constexpr (EnumHasAnyFlags(Features, EPixelProcessorFeatures::VectorizedImpl))
 			{
-				ProcessPixelVectorImpl(BufferPosPtr, Varying);
+				ProcessPixelVectorImpl(Context, BufferPos, Varying);
 			}
 			else
 			{
-				ProcessPixelImpl(BufferPosPtr, Varying);
+				ProcessPixelImpl(Context, BufferPos, Varying);
 			}
 		}
 
-        //-----------------------------------------------------------------------------------------
-        inline void ProcessPixelImpl( uint8* BufferPosPtr, float Varying[4] ) const
-        {
-			if (!Source0DataPtr)
+	private:
+		static void ProcessPixelVectorImpl(const FProjectedPixelProcessorContext& Context, uint8* BufferPos, float Varying[4])
+		{
+			if (!Context.Source0Data)
 			{
 				return;
 			}
 
-			const bool bDepthClamp = (Varying[2] < 0.0f) | (Varying[2] > 1.0f);
-            const float AngleCos = Varying[3];
-
-			uint16 Factor = static_cast<uint16>(FMath::Clamp((AngleCos - FadeEndCos) * OneOverFadeRangeTimes255, 0.0f, 255.0f));
-			Factor = Factor < 255 ? Factor : static_cast<uint16>(!bDepthClamp) * 255;
-
-			uint16 MaskFactor = AngleCos > FadeEndCos ? 255 : 0;
-
-			if constexpr (HAS_MASK)
+			const bool bDepthClamp = Invoke([&]() -> bool
 			{
-				// Only read from memory if needed.
-				if (Factor > 0)
+				if constexpr (EnumHasAnyFlags(Features, EPixelProcessorFeatures::ProjectionCylindrical))
 				{
-					MaskFactor = ((MaskFactor > 0) && MaskDataPtr) 
-							? MaskDataPtr[(BufferPosPtr - TargetDataPtr) / PIXEL_SIZE] 
-							: MaskFactor;
-					Factor = (Factor * MaskFactor) / 255;
+					const FVector2f CylinderPolarCoords(Varying[1], Varying[2]);
+					return FVector2f::DotProduct(CylinderPolarCoords, CylinderPolarCoords) > 1.0f;
 				}
-			}
-
-            if (Factor > 0)
-            {
-				const FVector2f Uv = { Varying[0], Varying[1] };
-				const FVector2f Source0SizeF = FVector2f(Source0SizeX, Source0SizeY);
-				const FVector2f Source1SizeF = FVector2f(Source1SizeX, Source1SizeY);
-
-                if ((Uv.X >= 0.0f) & (Uv.X < 1.0f) & (Uv.Y >= 0.0f) & (Uv.Y < 1.0f))
-                {
-					using FUInt16Vector2 = UE::Math::TIntVector2<uint16>;
-					struct FPixelData
-					{
-						alignas(8) uint16 Data[PIXEL_SIZE];
-					};
-
-					auto SampleImageBilinear = [](FVector2f Uv, FUInt16Vector2 Size, const uint8* DataPtr)
-					{
-						auto ComputeInterpolator = [](float T) -> uint16
-						{
-							return static_cast<uint16>(255.0f * T);
-						};
-
-						const FVector2f SizeF(Size.X, Size.Y);
-
-						const FVector2f CoordsF = FVector2f(
-							FMath::Clamp(Uv.X * SizeF.X - 0.5f, 0.0f, SizeF.X - 1.0f),
-							FMath::Clamp(Uv.Y * SizeF.Y - 0.5f, 0.0f, SizeF.Y - 1.0f));
-
-						const FUInt16Vector2 Frac = FUInt16Vector2(
-							ComputeInterpolator(FMath::Frac(CoordsF.X)),
-							ComputeInterpolator(FMath::Frac(CoordsF.Y)));
-
-						const FIntVector2 Coords = FIntVector2(CoordsF.X, CoordsF.Y);
-						const FIntVector2 CoordsPlusOne = FIntVector2( 
-								FMath::Min(Size.X - 1, Coords.X + 1),
-								FMath::Min(Size.Y - 1, Coords.Y + 1));
-
-						uint8 const * const Pixel00Ptr = DataPtr + (Coords.Y        * Size.X + Coords.X)        * PIXEL_SIZE;
-						uint8 const * const Pixel10Ptr = DataPtr + (Coords.Y        * Size.X + CoordsPlusOne.X) * PIXEL_SIZE;
-						uint8 const * const Pixel01Ptr = DataPtr + (CoordsPlusOne.Y * Size.X + Coords.X)        * PIXEL_SIZE;
-						uint8 const * const Pixel11Ptr = DataPtr + (CoordsPlusOne.Y * Size.X + CoordsPlusOne.X) * PIXEL_SIZE;
-
-						auto LoadPixel = [](const uint8* Ptr) -> FPixelData
-						{
-							FPixelData Result;
-							if constexpr (PIXEL_SIZE == 4)
-							{
-								const uint32 PackedData = *reinterpret_cast<const uint32*>(Ptr);
-								
-								Result.Data[0] = static_cast<uint16>((PackedData >> (8 * 0)) & 0xFF);
-								Result.Data[1] = static_cast<uint16>((PackedData >> (8 * 1)) & 0xFF);
-								Result.Data[2] = static_cast<uint16>((PackedData >> (8 * 2)) & 0xFF);
-								Result.Data[3] = static_cast<uint16>((PackedData >> (8 * 3)) & 0xFF);
-							}
-							else
-							{
-								for (int32 C = 0; C < PIXEL_SIZE; ++C)
-								{
-									Result.Data[C] = static_cast<uint16>(Ptr[C]);
-								}
-							}
-							return Result;
-						};
-
-						FPixelData PixelData00 = LoadPixel(Pixel00Ptr);
-						FPixelData PixelData10 = LoadPixel(Pixel10Ptr);
-						FPixelData PixelData01 = LoadPixel(Pixel01Ptr);
-						FPixelData PixelData11 = LoadPixel(Pixel11Ptr);
-
-						FPixelData FilteredPixelData;
-					
-						for (int32 C = 0; C < PIXEL_SIZE; ++C)
-						{
-							const uint16 LerpY0 = ((PixelData10.Data[C] * Frac.X) + PixelData00.Data[C] * (255 - Frac.X)) / 255;
-							const uint16 LerpY1 = ((PixelData11.Data[C] * Frac.X) + PixelData01.Data[C] * (255 - Frac.X)) / 255;
-							FilteredPixelData.Data[C] = ((LerpY1 * Frac.Y) + LerpY0*(255 - Frac.Y)) / 255;
-						} 
-
-						return FilteredPixelData;
-					};
-
-					FPixelData Sample0 = SampleImageBilinear(Uv, FUInt16Vector2(Source0SizeX, Source0SizeY), Source0DataPtr);
-					FPixelData Sample1 = SampleImageBilinear(Uv, FUInt16Vector2(Source1SizeX, Source1SizeY), Source1DataPtr);
-
-					const uint16 MipFactor = static_cast<uint16>(MipInterpolationFactor * 255.0f);
-
-					FPixelData& FilteredPixelData = Sample0;
-					for (int32 C = 0; C < PIXEL_SIZE; ++C)
-					{
-						FilteredPixelData.Data[C] = ((Sample1.Data[C] * MipFactor) + Sample0.Data[C]*(255 - MipFactor)) / 255;
-					}
-
-					static_assert(PIXEL_SIZE <= 4);
-
-					// FadingEnabledValue will be -1 if enabled and 0 otherwise.
-					const uint16 FadingMask[4] = { 
-						static_cast<uint16>(-RGBFadingEnabledValue), 
-						static_cast<uint16>(-RGBFadingEnabledValue), 
-						static_cast<uint16>(-RGBFadingEnabledValue), 
-						static_cast<uint16>(-AlphaFadingEnabledValue)};
-
-					for (int32 I = 0; I < PIXEL_SIZE; ++I)
-					{
-						BufferPosPtr[I] = static_cast<uint8>(
-							(FilteredPixelData.Data[I] * (MaskFactor * (1 - FadingMask[I]) + (Factor * FadingMask[I]))) / 255);
-					}				
-                }
-            }
-        }
-
-        //-----------------------------------------------------------------------------------------
-        inline void ProcessPixelVectorImpl( uint8* BufferPosPtr, float Varying[4] ) const
-        {
-			if (!Source0DataPtr)
-			{
-				return;
-			}
+				else
+				{
+					return static_cast<bool>((Varying[2] < 0.0f) | (Varying[2] > 1.0f));
+				}
+			});
 
-			const bool bDepthClamp = (Varying[2] < 0.0f) | (Varying[2] > 1.0f);
             const float AngleCos = Varying[3];
 
-			float Factor = FMath::Clamp((AngleCos - FadeEndCos) * OneOverFadeRangeTimes255, 0.0f, 255.0f);
+			float Factor = FMath::Clamp((AngleCos - Context.FadeEndCos) * Context.OneOverFadeRangeTimes255, 0.0f, 255.0f);
 			Factor = Factor < 255.0f ? Factor : static_cast<float>(!bDepthClamp) * 255.0f;
 
-			float MaskFactor = AngleCos > FadeEndCos ? 255.0f : 0.0f;
+			float MaskFactor = AngleCos > Context.FadeEndCos ? 255.0f : 0.0f;
 
-			if constexpr (HAS_MASK)
+			if constexpr (EnumHasAnyFlags(Features, EPixelProcessorFeatures::WithMask))
 			{
 				// Only read from memory if needed.
 				if (Factor > UE_SMALL_NUMBER)
 				{
-					MaskFactor = ((MaskFactor > 0.0f) && MaskDataPtr) 
-							? static_cast<float>(MaskDataPtr[(BufferPosPtr - TargetDataPtr) / PIXEL_SIZE]) 
+					MaskFactor = ((MaskFactor > 0.0f) && Context.MaskData)
+							? static_cast<float>(Context.MaskData[(BufferPos - Context.TargetData) / PIXEL_SIZE])
 							: MaskFactor;
 					Factor = (Factor * MaskFactor) * (1.0f/255.0f);
 				}
@@ -430,48 +295,26 @@ namespace
 
             if (Factor > UE_SMALL_NUMBER)
             {
-				const FVector2f Uv = { Varying[0], Varying[1] };
+				const FVector2f Uv = Invoke([&]()
+				{ 
+					if constexpr (EnumHasAnyFlags(Features, EPixelProcessorFeatures::ProjectionCylindrical))
+					{
+						return FVector2f{ 0.5f + FMath::Atan2(Varying[2], -Varying[1]) * Context.OneOverProjectionAngle, Varying[0] };
+					}
+					else
+					{
+						return FVector2f{ Varying[0], Varying[1] };
+					}
+				});
 
                 if ((Uv.X >= 0.0f) & (Uv.X < 1.0f) & (Uv.Y >= 0.0f) & (Uv.Y < 1.0f))
                 {
 					using namespace GlobalVectorConstants;
 
-					const VectorRegister4Int SizeI = 
-						MakeVectorRegisterInt(Source0SizeX, Source0SizeY, Source1SizeX, Source1SizeY);
-
-					const VectorRegister4Float SizeF = VectorIntToFloat(SizeI);
-
-					const VectorRegister4Float CoordsF = VectorMin(VectorMax(
-						VectorMultiplyAdd(MakeVectorRegister(Uv.X, Uv.Y, Uv.X, Uv.Y), SizeF, FloatMinusOneHalf),
-						FloatZero), 
-						VectorSubtract(SizeF, FloatOne));
-
-					const VectorRegister4Float Frac = VectorSubtract(CoordsF, VectorFloor(CoordsF));
-					
-					const VectorRegister4Int CoordsI = VectorFloatToInt(CoordsF);
-					const VectorRegister4Int CoordsIPlusOne = VectorIntMin(
-							VectorIntAdd(CoordsI, IntOne), VectorIntSubtract(SizeI, IntOne));
-
-					alignas(VectorRegister4Int) int32 CoordsIData[4];
-					VectorIntStoreAligned(CoordsI, CoordsIData);
-
-					alignas(VectorRegister4Int) int32 CoordsIPlusOneData[4];
-					VectorIntStoreAligned(CoordsIPlusOne, CoordsIPlusOneData);
-
-					uint8 const * const Pixel000Ptr = Source0DataPtr + (CoordsIData[1]        * Source0SizeX + CoordsIData[0])        * PIXEL_SIZE;
-					uint8 const * const Pixel010Ptr = Source0DataPtr + (CoordsIData[1]        * Source0SizeX + CoordsIPlusOneData[0]) * PIXEL_SIZE;
-					uint8 const * const Pixel001Ptr = Source0DataPtr + (CoordsIPlusOneData[1] * Source0SizeX + CoordsIData[0])        * PIXEL_SIZE;
-					uint8 const * const Pixel011Ptr = Source0DataPtr + (CoordsIPlusOneData[1] * Source0SizeX + CoordsIPlusOneData[0]) * PIXEL_SIZE;
-
-					uint8 const * const Pixel100Ptr = Source1DataPtr + (CoordsIData[3]        * Source1SizeX + CoordsIData[2])        * PIXEL_SIZE;
-					uint8 const * const Pixel110Ptr = Source1DataPtr + (CoordsIData[3]        * Source1SizeX + CoordsIPlusOneData[2]) * PIXEL_SIZE;
-					uint8 const * const Pixel101Ptr = Source1DataPtr + (CoordsIPlusOneData[3] * Source1SizeX + CoordsIData[2])        * PIXEL_SIZE;
-					uint8 const * const Pixel111Ptr = Source1DataPtr + (CoordsIPlusOneData[3] * Source1SizeX + CoordsIPlusOneData[2]) * PIXEL_SIZE;
-
 					auto LoadPixel = [](const uint8* Ptr) -> VectorRegister4Float
 					{
-						constexpr VectorRegister4Float OneOver255 = 
-							MakeVectorRegisterFloatConstant(1.0f/255.0f, 1.0f/255.0f, 1.0f/255.0f, 1.0f/255.0f);
+						constexpr VectorRegister4Float OneOver255 =
+							MakeVectorRegisterFloatConstant(1.0f / 255.0f, 1.0f / 255.0f, 1.0f / 255.0f, 1.0f / 255.0f);
 
 						if constexpr (PIXEL_SIZE == 4)
 						{
@@ -494,58 +337,102 @@ namespace
 						}
 					};
 
-					const VectorRegister4Float Pixel000 = LoadPixel(Pixel000Ptr);
-					const VectorRegister4Float Pixel010 = LoadPixel(Pixel010Ptr);
-					const VectorRegister4Float Pixel001 = LoadPixel(Pixel001Ptr);
-					const VectorRegister4Float Pixel011 = LoadPixel(Pixel011Ptr);
-
-					const VectorRegister4Float Pixel100 = LoadPixel(Pixel100Ptr);
-					const VectorRegister4Float Pixel110 = LoadPixel(Pixel110Ptr);
-					const VectorRegister4Float Pixel101 = LoadPixel(Pixel101Ptr);
-					const VectorRegister4Float Pixel111 = LoadPixel(Pixel111Ptr);
-					
-					const VectorRegister4Float Frac0X = VectorReplicate(Frac, 0);
-					const VectorRegister4Float Frac0Y = VectorReplicate(Frac, 1);
-
-					// Bilerp image 0
-					VectorRegister4Float Sample0 = VectorMultiplyAdd(Frac0X, VectorSubtract(Pixel010, Pixel000), Pixel000);
-					Sample0 = VectorMultiplyAdd(
-						Frac0Y, 
-						VectorSubtract(
-							VectorMultiplyAdd(Frac0X, VectorSubtract(Pixel011, Pixel001), Pixel001), 
-							Sample0), 
-						Sample0);
-
-					const VectorRegister4Float Frac1X = VectorReplicate(Frac, 2);
-					const VectorRegister4Float Frac1Y = VectorReplicate(Frac, 3);
-
-					// Bilerp image 1
-					VectorRegister4Float Sample1 = VectorMultiplyAdd(Frac1X, VectorSubtract(Pixel110, Pixel100), Pixel100);
-					Sample1 = VectorMultiplyAdd(
-						Frac1Y, 
-						VectorSubtract(
-							VectorMultiplyAdd(Frac1X, VectorSubtract(Pixel111, Pixel101), Pixel101), 
-							Sample1), 
-						Sample1);
-
-					VectorRegister4Float Result = VectorMultiply(
-						// Lerp mip
-						VectorMultiplyAdd(VectorSetFloat1(MipInterpolationFactor), VectorSubtract(Sample1, Sample0), Sample0),
+					VectorRegister4Float Result;
+					if constexpr (EnumHasAnyFlags(Features, EPixelProcessorFeatures::SamplingLinear))
+					{
+						const VectorRegister4Int SizeI =
+							MakeVectorRegisterInt(Context.Source0SizeX, Context.Source0SizeY, Context.Source1SizeX, Context.Source1SizeY);
+
+						const VectorRegister4Float SizeF = VectorIntToFloat(SizeI);
+
+						const VectorRegister4Float CoordsF = VectorMin(VectorMax(
+							VectorMultiplyAdd(MakeVectorRegister(Uv.X, Uv.Y, Uv.X, Uv.Y), SizeF, FloatMinusOneHalf),
+							FloatZero),
+							VectorSubtract(SizeF, FloatOne));
+
+						const VectorRegister4Float Frac = VectorSubtract(CoordsF, VectorFloor(CoordsF));
+
+						const VectorRegister4Int CoordsI = VectorFloatToInt(CoordsF);
+						const VectorRegister4Int CoordsIPlusOne = VectorIntMin(
+							VectorIntAdd(CoordsI, IntOne), VectorIntSubtract(SizeI, IntOne));
+
+						alignas(VectorRegister4Int) int32 CoordsIData[4];
+						VectorIntStoreAligned(CoordsI, CoordsIData);
+
+						alignas(VectorRegister4Int) int32 CoordsIPlusOneData[4];
+						VectorIntStoreAligned(CoordsIPlusOne, CoordsIPlusOneData);
+
+						uint8 const* const Pixel000Ptr = Context.Source0Data + (CoordsIData[1] * Context.Source0SizeX + CoordsIData[0]) * PIXEL_SIZE;
+						uint8 const* const Pixel010Ptr = Context.Source0Data + (CoordsIData[1] * Context.Source0SizeX + CoordsIPlusOneData[0]) * PIXEL_SIZE;
+						uint8 const* const Pixel001Ptr = Context.Source0Data + (CoordsIPlusOneData[1] * Context.Source0SizeX + CoordsIData[0]) * PIXEL_SIZE;
+						uint8 const* const Pixel011Ptr = Context.Source0Data + (CoordsIPlusOneData[1] * Context.Source0SizeX + CoordsIPlusOneData[0]) * PIXEL_SIZE;
+
+						uint8 const* const Pixel100Ptr = Context.Source1Data + (CoordsIData[3] * Context.Source1SizeX + CoordsIData[2]) * PIXEL_SIZE;
+						uint8 const* const Pixel110Ptr = Context.Source1Data + (CoordsIData[3] * Context.Source1SizeX + CoordsIPlusOneData[2]) * PIXEL_SIZE;
+						uint8 const* const Pixel101Ptr = Context.Source1Data + (CoordsIPlusOneData[3] * Context.Source1SizeX + CoordsIData[2]) * PIXEL_SIZE;
+						uint8 const* const Pixel111Ptr = Context.Source1Data + (CoordsIPlusOneData[3] * Context.Source1SizeX + CoordsIPlusOneData[2]) * PIXEL_SIZE;
+
+						const VectorRegister4Float Pixel000 = LoadPixel(Pixel000Ptr);
+						const VectorRegister4Float Pixel010 = LoadPixel(Pixel010Ptr);
+						const VectorRegister4Float Pixel001 = LoadPixel(Pixel001Ptr);
+						const VectorRegister4Float Pixel011 = LoadPixel(Pixel011Ptr);
+
+						const VectorRegister4Float Pixel100 = LoadPixel(Pixel100Ptr);
+						const VectorRegister4Float Pixel110 = LoadPixel(Pixel110Ptr);
+						const VectorRegister4Float Pixel101 = LoadPixel(Pixel101Ptr);
+						const VectorRegister4Float Pixel111 = LoadPixel(Pixel111Ptr);
+
+						const VectorRegister4Float Frac0X = VectorReplicate(Frac, 0);
+						const VectorRegister4Float Frac0Y = VectorReplicate(Frac, 1);
+
+						// Bilerp image 0
+						VectorRegister4Float Sample0 = VectorMultiplyAdd(Frac0X, VectorSubtract(Pixel010, Pixel000), Pixel000);
+						Sample0 = VectorMultiplyAdd(
+							Frac0Y,
+							VectorSubtract(
+								VectorMultiplyAdd(Frac0X, VectorSubtract(Pixel011, Pixel001), Pixel001),
+								Sample0),
+							Sample0);
+
+						const VectorRegister4Float Frac1X = VectorReplicate(Frac, 2);
+						const VectorRegister4Float Frac1Y = VectorReplicate(Frac, 3);
+
+						// Bilerp image 1
+						VectorRegister4Float Sample1 = VectorMultiplyAdd(Frac1X, VectorSubtract(Pixel110, Pixel100), Pixel100);
+						Sample1 = VectorMultiplyAdd(
+							Frac1Y,
+							VectorSubtract(
+								VectorMultiplyAdd(Frac1X, VectorSubtract(Pixel111, Pixel101), Pixel101),
+								Sample1),
+							Sample1);
+
+						Result = VectorMultiplyAdd(VectorSetFloat1(Context.MipInterpolationFactor), VectorSubtract(Sample1, Sample0), Sample0);
+					}
+					else if (EnumHasAnyFlags(Features, EPixelProcessorFeatures::SamplingPoint))
+					{						
+						const int32 CoordX = FMath::FloorToInt32(Context.Source0SizeX * Uv.X);
+						const int32 CoordY = FMath::FloorToInt32(Context.Source0SizeY * Uv.Y);
+
+						Result = LoadPixel(Context.Source0Data + (CoordY*Context.Source0SizeX + CoordX) * PIXEL_SIZE);
+					}
+
+					Result = VectorMultiply(
+						Result,
 						// Fadding factor
 						VectorSelect(
 							// The selction mask needs to be 0xFFFFFFFF if enabled 0 otherwise. FadingEnabledValue will be -1
 							// if enabled and 0 otherwise.
-							VectorCastIntToFloat(MakeVectorRegisterInt( 
-								static_cast<int32>(RGBFadingEnabledValue), 
-								static_cast<int32>(RGBFadingEnabledValue), 
-								static_cast<int32>(RGBFadingEnabledValue), 
-								static_cast<int32>(AlphaFadingEnabledValue))), 
-							VectorSetFloat1(Factor), 
+							VectorCastIntToFloat(MakeVectorRegisterInt(
+								static_cast<int32>(Context.RGBFadingEnabledMask),
+								static_cast<int32>(Context.RGBFadingEnabledMask),
+								static_cast<int32>(Context.RGBFadingEnabledMask),
+								static_cast<int32>(Context.AlphaFadingEnabledMask))),
+							VectorSetFloat1(Factor),
 							VectorSetFloat1(MaskFactor)));
 
 					if constexpr (PIXEL_SIZE == 4)
 					{
-						VectorStoreByte4(Result, BufferPosPtr);
+						VectorStoreByte4(Result, BufferPos);
 						VectorResetFloatRegisters();
 					}
 					else
@@ -553,586 +440,878 @@ namespace
 						const AlignedFloat4 ResultData(VectorMin(VectorMax(Result, FloatZero), Float255));
 						for (int32 C = 0; C < PIXEL_SIZE; ++C)
 						{
-							BufferPosPtr[C] = static_cast<uint8>(ResultData[C]);		
+							BufferPos[C] = static_cast<uint8>(ResultData[C]);		
 						}				
 					}
                 }
             }
-        }
-
-        const uint8* Source0DataPtr;
-        const uint8* Source1DataPtr;
-        uint16 Source0SizeX, Source0SizeY;
-        uint16 Source1SizeX, Source1SizeY;
-
-        //! Cosine of the fading angle range
-		float FadeEndCos;
-		float MipInterpolationFactor;
-
-		// precomputation of 255 / (FadeStartCos - FadeEndCos) 
-		float OneOverFadeRangeTimes255;
-
-		int8 RGBFadingEnabledValue, AlphaFadingEnabledValue;
-
-        const uint8* TargetDataPtr;
-        const uint8* MaskDataPtr;
-    };
-
+		}
 
-    //---------------------------------------------------------------------------------------------
-    template<int PIXEL_SIZE,bool HAS_MASK=true>
-    class RasterCylindricalProjectedPixelProcessor_UBYTE
-    {
-    public:
-
-        RasterCylindricalProjectedPixelProcessor_UBYTE( const Image* pSource,
-                                     const uint8* pTargetData,
-                                     const uint8* pMaskData,
-                                     float fadeStart,
-                                     float fadeEnd,
-                                                        float projectionAngle )
+        static void ProcessPixelImpl(const FProjectedPixelProcessorContext& Context, uint8* BufferPos, float Varying[4])
         {
-            m_pSourceData = pSource->GetData();
-            m_sourceSizeX = pSource->GetSizeX();
-            m_sourceSizeY = pSource->GetSizeY();
-
-            m_fadeStartCos = cosf(fadeStart);
-            m_fadeEndCos = cosf(fadeEnd);
-            m_projectionAngle = projectionAngle;
+			if (!Context.Source0Data)
+			{
+				return;
+			}
 
-            m_pTargetData = pTargetData;
-            m_pMaskData = pMaskData;
+			const bool bDepthClamp = Invoke([&]() -> bool
+			{
+				if constexpr (EnumHasAnyFlags(Features, EPixelProcessorFeatures::ProjectionCylindrical))
+				{
+					const FVector2f CylinderPolarCoords(Varying[1], Varying[2]);
+					return FVector2f::DotProduct(CylinderPolarCoords, CylinderPolarCoords) > 1.0f;
+				}
+				else
+				{
+					return static_cast<bool>((Varying[2] < 0.0f) | (Varying[2] > 1.0f));
+				}
+			});
 
-            check( GetImageFormatData( pSource->GetFormat() ).BytesPerBlock
-                            ==
-                            PIXEL_SIZE );
-        }
+            const float AngleCos = Varying[3];
 
+			uint16 Factor = static_cast<uint16>(FMath::Clamp((AngleCos - Context.FadeEndCos) * Context.OneOverFadeRangeTimes255, 0.0f, 255.0f));
+			Factor = Factor < 255 ? Factor : static_cast<uint16>(!bDepthClamp) * 255;
 
-        //-----------------------------------------------------------------------------------------
-        inline void ProcessPixel( unsigned char* pBufferPos, float varying[4] ) const
-        {
-            int32 factor_8 = 256;
+			uint16 MaskFactor = AngleCos > Context.FadeEndCos ? 255 : 0;
 
-            // Position in unit cylinder space
-            float x_cyl = varying[0];
-            float y_cyl = varying[1];
-            float z_cyl = varying[2];
-            float r = y_cyl*y_cyl+z_cyl*z_cyl;
+			if constexpr (EnumHasAnyFlags(Features, EPixelProcessorFeatures::WithMask))
+			{
+				// Only read from memory if needed.
+				if (Factor > 0)
+				{
+					MaskFactor = ((MaskFactor > 0) && Context.MaskData) 
+							? Context.MaskData[(BufferPos - Context.TargetData) / PIXEL_SIZE] 
+							: MaskFactor;
+					Factor = (Factor * MaskFactor) / 255;
+				}
+			}
 
-            // depth clamp
-            if (r>1.0f)
+            if (Factor > 0)
             {
-                factor_8 = 0;
-            }
+				const FVector2f Uv = Invoke([&]()
+				{ 
+					if constexpr (EnumHasAnyFlags(Features, EPixelProcessorFeatures::ProjectionCylindrical))
+					{
+						return FVector2f{ 0.5f + FMath::Atan2(Varying[2], -Varying[1])*Context.OneOverProjectionAngle, Varying[0] };
+					}
+					else
+					{
+						return FVector2f{ Varying[0], Varying[1] };
+					}
+				});
 
-            float angleCos = varying[3];
+				const FVector2f Source0SizeF = FVector2f(Context.Source0SizeX, Context.Source0SizeY);
+				const FVector2f Source1SizeF = FVector2f(Context.Source1SizeX, Context.Source1SizeY);
 
-            if (HAS_MASK)
-            {
-                if ( factor_8>0 && m_pMaskData )
-                {
-                    uint8 maskFactor = m_pMaskData[ (pBufferPos-m_pTargetData) / PIXEL_SIZE ];
-                    factor_8 = maskFactor;
-                }
+				if ((Uv.X >= 0.0f) & (Uv.X < 1.0f) & (Uv.Y >= 0.0f) & (Uv.Y < 1.0f))
+				{
+					using FUInt16Vector2 = UE::Math::TIntVector2<uint16>;
+					struct FPixelData
+					{
+						alignas(8) uint16 Data[PIXEL_SIZE];
+					};
 
-                if ( factor_8>0 )
-                {
-                    if ( angleCos<m_fadeStartCos && angleCos>m_fadeEndCos )
-                    {
-                        factor_8 = ( factor_8
-                                     *
-                                     int32( 256.0f * ( angleCos - m_fadeEndCos ) / ( m_fadeStartCos - m_fadeEndCos ) )
-                                     ) >> 8;
-                    }
-                    else if ( angleCos<=m_fadeEndCos )
-                    {
-                        factor_8 = 0;
-                    }
-                }
-            }
-            else
-            {
-                if ( angleCos<m_fadeStartCos && angleCos>m_fadeEndCos )
-                {
-                    factor_8 = int32( 256.0f * ( angleCos - m_fadeEndCos ) / ( m_fadeStartCos - m_fadeEndCos ) );
-                }
-                else if ( angleCos<=m_fadeEndCos )
-                {
-                    factor_8 = 0;
-                }
-            }
 
+					auto LoadPixel = [](const uint8* Ptr) -> FPixelData
+					{
+						FPixelData Result;
+						if constexpr (PIXEL_SIZE == 4)
+						{
+							const uint32 PackedData = *reinterpret_cast<const uint32*>(Ptr);
 
-            if ( factor_8>0 )
-            {
-                // Project
-                float u = 0.5f + atan2f(z_cyl,-y_cyl) / m_projectionAngle;
-                float v = x_cyl;
+							Result.Data[0] = static_cast<uint16>((PackedData >> (8 * 0)) & 0xFF);
+							Result.Data[1] = static_cast<uint16>((PackedData >> (8 * 1)) & 0xFF);
+							Result.Data[2] = static_cast<uint16>((PackedData >> (8 * 2)) & 0xFF);
+							Result.Data[3] = static_cast<uint16>((PackedData >> (8 * 3)) & 0xFF);
+						}
+						else
+						{
+							for (int32 C = 0; C < PIXEL_SIZE; ++C)
+							{
+								Result.Data[C] = static_cast<uint16>(Ptr[C]);
+							}
+						}
 
-                if ( u>=0 && u<1 && v>=0 && v<1 )
-                {
+						return Result;
+					};
 
-                    const uint8* pPixel = m_pSourceData;
-                    // TODO: clamp?
-                    pPixel += ( m_sourceSizeX * int(m_sourceSizeY * v)
-                                + int( m_sourceSizeX * u ) )
-                            * PIXEL_SIZE;
+					FPixelData Result;
+					if constexpr (EnumHasAnyFlags(Features, EPixelProcessorFeatures::SamplingLinear))
+					{
+						auto SampleImageBilinear = [&](FVector2f Uv, FUInt16Vector2 Size, const uint8* DataPtr)
+						{
+							auto ComputeInterpolator = [](float T) -> uint16
+							{
+								return static_cast<uint16>(255.0f * T);
+							};
 
-                    // Write result
-                    for ( int i=0; i<PIXEL_SIZE; ++i )
-                    {
-                        pBufferPos[i] = uint8( ( pPixel[i] * factor_8 ) >> 8 );
-                    }
-                }
-            }
-        }
+							const FVector2f SizeF(Size.X, Size.Y);
 
+							const FVector2f CoordsF = FVector2f(
+								FMath::Clamp(Uv.X * SizeF.X - 0.5f, 0.0f, SizeF.X - 1.0f),
+								FMath::Clamp(Uv.Y * SizeF.Y - 0.5f, 0.0f, SizeF.Y - 1.0f));
 
-        const uint8* m_pSourceData;
-        int m_sourceSizeX, m_sourceSizeY;
+							const FUInt16Vector2 Frac = FUInt16Vector2(
+								ComputeInterpolator(FMath::Frac(CoordsF.X)),
+								ComputeInterpolator(FMath::Frac(CoordsF.Y)));
 
-        //! Cosine of the fading angle range
-        float m_fadeStartCos, m_fadeEndCos;
+							const FIntVector2 Coords = FIntVector2(CoordsF.X, CoordsF.Y);
+							const FIntVector2 CoordsPlusOne = FIntVector2(
+								FMath::Min(Size.X - 1, Coords.X + 1),
+								FMath::Min(Size.Y - 1, Coords.Y + 1));
 
-        float m_projectionAngle;
+							uint8 const* const Pixel00Ptr = DataPtr + (Coords.Y * Size.X + Coords.X) * PIXEL_SIZE;
+							uint8 const* const Pixel10Ptr = DataPtr + (Coords.Y * Size.X + CoordsPlusOne.X) * PIXEL_SIZE;
+							uint8 const* const Pixel01Ptr = DataPtr + (CoordsPlusOne.Y * Size.X + Coords.X) * PIXEL_SIZE;
+							uint8 const* const Pixel11Ptr = DataPtr + (CoordsPlusOne.Y * Size.X + CoordsPlusOne.X) * PIXEL_SIZE;
 
-        const uint8* m_pTargetData;
-        const uint8* m_pMaskData;
-    };
+							FPixelData PixelData00 = LoadPixel(Pixel00Ptr);
+							FPixelData PixelData10 = LoadPixel(Pixel10Ptr);
+							FPixelData PixelData01 = LoadPixel(Pixel01Ptr);
+							FPixelData PixelData11 = LoadPixel(Pixel11Ptr);
 
+							FPixelData FilteredPixelData;
 
-} // anon namespace
+							for (int32 C = 0; C < PIXEL_SIZE; ++C)
+							{
+								const uint16 LerpY0 = ((PixelData10.Data[C] * Frac.X) + PixelData00.Data[C] * (255 - Frac.X)) / 255;
+								const uint16 LerpY1 = ((PixelData11.Data[C] * Frac.X) + PixelData01.Data[C] * (255 - Frac.X)) / 255;
+								FilteredPixelData.Data[C] = ((LerpY1 * Frac.Y) + LerpY0 * (255 - Frac.Y)) / 255;
+							}
 
+							return FilteredPixelData;
+						};
 
+						Result = SampleImageBilinear(Uv, FUInt16Vector2(Context.Source0SizeX, Context.Source0SizeY), Context.Source0Data);
+						FPixelData Sample1 = SampleImageBilinear(Uv, FUInt16Vector2(Context.Source1SizeX, Context.Source1SizeY), Context.Source1Data);
 
-//-------------------------------------------------------------------------------------------------
-//! This format is the one we assumee the meshes optimised for planar and cylindrical projection
-//! will have.
-//! See CreateMeshOptimisedForProjection
-struct OPTIMISED_VERTEX
-{
-	FVector2f uv;
-	FVector3f pos;
-	FVector3f nor;
-};
+						const uint16 MipFactor = static_cast<uint16>(Context.MipInterpolationFactor * 255.0f);
 
-static_assert( sizeof(OPTIMISED_VERTEX)==32, "UNEXPECTED_STRUCT_SIZE" );
+						for (int32 C = 0; C < PIXEL_SIZE; ++C)
+						{
+							Result.Data[C] = ((Sample1.Data[C] * MipFactor) + Result.Data[C] * (255 - MipFactor)) / 255;
+						}
+					}
+					else if (EnumHasAnyFlags(Features, EPixelProcessorFeatures::SamplingPoint))
+					{
+						const int32 CoordX = FMath::FloorToInt32(Context.Source0SizeX * Uv.X);
+						const int32 CoordY = FMath::FloorToInt32(Context.Source0SizeY * Uv.Y);
 
+						Result = LoadPixel(Context.Source0Data + (CoordY*Context.Source0SizeX + CoordX) * PIXEL_SIZE);
+					}
+					static_assert(EnumHasAnyFlags(Features, EPixelProcessorFeatures::SamplingPoint | EPixelProcessorFeatures::SamplingLinear));
 
-//-------------------------------------------------------------------------------------------------
-//! This format is the one we assumee the meshes optimised for wrapping projection will have.
-//! See CreateMeshOptimisedForWrappingProjection
-struct OPTIMISED_VERTEX_WRAPPING
-{
-	FVector2f uv;
-	FVector3f pos;
-	FVector3f nor;
-    uint32 layoutBlock;
-};
+					static_assert(PIXEL_SIZE <= 4);
+					// FadingEnabledMaskSeed will be -1 if enabled and 0 otherwise.
+					const int16 FadingMask[4] = {
+						static_cast<int16>(Context.RGBFadingEnabledMask),
+						static_cast<int16>(Context.RGBFadingEnabledMask),
+						static_cast<int16>(Context.RGBFadingEnabledMask),
+						static_cast<int16>(Context.AlphaFadingEnabledMask) };
 
-static_assert( sizeof(OPTIMISED_VERTEX_WRAPPING)==36, "UNEXPECTED_STRUCT_SIZE" );
+					for (int32 I = 0; I < PIXEL_SIZE; ++I)
+					{
+						BufferPos[I] = static_cast<uint8>(
+							(Result.Data[I] * ((MaskFactor & ~FadingMask[I]) + (Factor & FadingMask[I]))) / 255);
+					}
+				}
+            }
+        }
+	};
+	
+	//-------------------------------------------------------------------------------------------------
+	FORCENOINLINE void ImageRasterProjected_Optimised(const Mesh* pMesh, Image* pImage,
+		TTriangleRasterPixelProcRefType<4> PixelProc, 
+		float FadeEnd,
+		UE::Math::TIntVector2<uint16> CropMin, UE::Math::TIntVector2<uint16> UncroppedSize,
+		FScratchImageProject* Scratch)
+	{
+		MUTABLE_CPUPROFILER_SCOPE(ImageRasterProjected_Optimised);
 
+		if (!pMesh || !pMesh->GetFaceCount())
+		{
+			return;
+		}
 
-//-------------------------------------------------------------------------------------------------
-template<class PIXEL_PROCESSOR>
-void ImageRasterProjected_Optimised( const Mesh* pMesh, Image* pImage,
-	PIXEL_PROCESSOR& pixelProc, float fadeEnd,
-	UE::Math::TIntVector2<uint16> CropMin, UE::Math::TIntVector2<uint16> UncroppedSize,
-	SCRATCH_IMAGE_PROJECT* scratch )
-{
-	MUTABLE_CPUPROFILER_SCOPE(ImageRasterProjected_Optimised);
+		EImageFormat Format = pImage->GetFormat();
+		int32 PixelSize = GetImageFormatData(Format).BytesPerBlock;
 
-    if (!pMesh || !pMesh->GetFaceCount())
-    {
-        return;
-    }
+		int32 SizeX = pImage->GetSizeX();
+		int32 SizeY = pImage->GetSizeY();
+
+		bool bUseCropping = UncroppedSize[0] > 0;
 
-	EImageFormat format = pImage->GetFormat();
-    int pixelSize = GetImageFormatData( format ).BytesPerBlock;
+		// Get the vertices
+		int32 VertexCount = pMesh->GetVertexCount();
 
-    int32 sizeX = pImage->GetSizeX();
-    int32 sizeY = pImage->GetSizeY();
+		check((int32)Scratch->Vertices.Num() == VertexCount);
+		check((int32)Scratch->CulledVertex.Num() == VertexCount);
 
-	bool bUseCropping = UncroppedSize[0] > 0;
+		check(pMesh->GetVertexBuffers().GetElementSize(0) == sizeof(FOptimizedVertex));
+		const FOptimizedVertex* pVertices = 
+				reinterpret_cast<const FOptimizedVertex*>(pMesh->GetVertexBuffers().GetBufferData(0));
 
-    // Get the vertices
-    int vertexCount = pMesh->GetVertexCount();
+		float FadeEndCos = FMath::Cos(FadeEnd);
+		for (int32 V = 0; V < VertexCount; ++V)
+		{
+			if (bUseCropping)
+			{
+				Scratch->Vertices[V].x = pVertices[V].Uv[0] * UncroppedSize[0] - CropMin[0];
+				Scratch->Vertices[V].y = pVertices[V].Uv[1] * UncroppedSize[1] - CropMin[1];
+			}
+			else
+			{
+				Scratch->Vertices[V].x = pVertices[V].Uv[0] * SizeX;
+				Scratch->Vertices[V].y = pVertices[V].Uv[1] * SizeY;
+			}
 
-    check( (int)scratch->vertices.Num()==vertexCount );
-    check( (int)scratch->culledVertex.Num()==vertexCount );
+			// TODO: No need to copy all. use scratch for the rest only.
+			Scratch->Vertices[V].interpolators[0] = pVertices[V].Position[0];
+			Scratch->Vertices[V].interpolators[1] = pVertices[V].Position[1];
+			Scratch->Vertices[V].interpolators[2] = pVertices[V].Position[2];
+			Scratch->Vertices[V].interpolators[3] = pVertices[V].Normal[0];
+			Scratch->CulledVertex[V] = pVertices[V].Normal[0] < FadeEndCos;
+		}
 
-    check( pMesh->GetVertexBuffers().GetElementSize(0)==sizeof(OPTIMISED_VERTEX) );
-	const OPTIMISED_VERTEX* pVertices = reinterpret_cast<const OPTIMISED_VERTEX*>( pMesh->GetVertexBuffers().GetBufferData(0) );
+		// Get the indices
+		check(pMesh->GetIndexBuffers().GetElementSize(0) == 4);
+		const uint32* pIndices = reinterpret_cast<const uint32*>(pMesh->GetIndexBuffers().GetBufferData(0));
 
-    float fadeEndCos = cosf(fadeEnd);
-    for ( int v=0; v<vertexCount; ++v )
-    {
-		if (bUseCropping)
+		// The mesh is supposed to contain only the faces in the selected block
+		int32 FaceCount = pMesh->GetFaceCount();
+		//for (int32 Face = 0; Face < FaceCount; ++Face)
+		ParallelFor(FaceCount, [pIndices, Scratch, SizeX, SizeY, PixelSize, pImage, PixelProc](int32 Face)
+		{
+			int32 Index0 = pIndices[Face*3 + 0];
+			int32 Index1 = pIndices[Face*3 + 1];
+			int32 Index2 = pIndices[Face*3 + 2];
+
+			// TODO: This optimisation may remove projection in the center of the face, if the angle
+			// range is small. Make it optional or more sophisticated (cross product may help).
+			if (!Scratch->CulledVertex[Index0] ||
+				!Scratch->CulledVertex[Index1] ||
+				!Scratch->CulledVertex[Index2])
+			{
+				constexpr int32 NumInterpolators = 4;
+				Triangle<NumInterpolators>(pImage->GetData(), pImage->GetDataSize(),
+					SizeX, SizeY,
+					PixelSize,
+					Scratch->Vertices[Index0],
+					Scratch->Vertices[Index1],
+					Scratch->Vertices[Index2],
+					PixelProc,
+					false);
+			}
+		});
+
+		// Update the relevancy map
+		// \TODO: fix it to work with cropping.
+		if (!bUseCropping)
 		{
-			scratch->vertices[v].x = pVertices[v].uv[0] * UncroppedSize[0] - CropMin[0];
-			scratch->vertices[v].y = pVertices[v].uv[1] * UncroppedSize[1] - CropMin[1];
+			MUTABLE_CPUPROFILER_SCOPE(ImageRasterProjected_Optimised_Relevancy);
+
+			float MinY = float(SizeY) - 1.0f;
+			float MaxY = 0.0f;
+			for (int32 F = 0; F < FaceCount; ++F)
+			{
+				int32 Index0 = pIndices[F*3 + 0];
+				int32 Index1 = pIndices[F*3 + 1];
+				int32 Index2 = pIndices[F*3 + 2];
+
+				// A bit ugly, probably can be improved if integrated in above loops and made more precise
+				// inside the pixel processor to account for masked out pixels?
+				if (!Scratch->CulledVertex[Index0] ||
+					!Scratch->CulledVertex[Index1] ||
+					!Scratch->CulledVertex[Index2])
+				{
+					MinY = FMath::Min(MinY, Scratch->Vertices[Index0].y);
+					MinY = FMath::Min(MinY, Scratch->Vertices[Index1].y);
+					MinY = FMath::Min(MinY, Scratch->Vertices[Index2].y);
+					MaxY = FMath::Max(MaxY, Scratch->Vertices[Index0].y);
+					MaxY = FMath::Max(MaxY, Scratch->Vertices[Index1].y);
+					MaxY = FMath::Max(MaxY, Scratch->Vertices[Index2].y);
+				}
+			}
+
+			pImage->m_flags |= Image::IF_HAS_RELEVANCY_MAP;
+			pImage->RelevancyMinY = uint16(FMath::FloorToFloat(MinY));
+			pImage->RelevancyMaxY = uint16(FMath::Min( int32(FMath::CeilToFloat(MaxY)), SizeY - 1));
 		}
-		else
+	}
+
+
+
+	//-------------------------------------------------------------------------------------------------
+	FORCENOINLINE void ImageRasterProjected_OptimisedWrapping(const Mesh* pMesh, Image* pImage, 
+		TTriangleRasterPixelProcRefType<4> PixelProc, 
+		float FadeEnd,
+		int32 Block,
+		UE::Math::TIntVector2<uint16> CropMin, UE::Math::TIntVector2<uint16> UncroppedSize,
+		FScratchImageProject* Scratch)
+	{
+		MUTABLE_CPUPROFILER_SCOPE(ImageRasterProjected_OptimisedWrapping);
+
+		if (!pMesh || !pMesh->GetFaceCount())
 		{
-			scratch->vertices[v].x = pVertices[v].uv[0] * sizeX;
-			scratch->vertices[v].y = pVertices[v].uv[1] * sizeY;
+			return;
 		}
 
-        // TODO: No need to copy all. use scratch for the rest only.
-        scratch->vertices[v].interpolators[0] = pVertices[v].pos[0];
-        scratch->vertices[v].interpolators[1] = pVertices[v].pos[1];
-        scratch->vertices[v].interpolators[2] = pVertices[v].pos[2];
-        scratch->vertices[v].interpolators[3] = pVertices[v].nor[0];
-        scratch->culledVertex[v] = pVertices[v].nor[0] < fadeEndCos;
-    }
+		EImageFormat Format = pImage->GetFormat();
+		int32 PixelSize = GetImageFormatData(Format).BytesPerBlock;
 
-    // Get the indices
-    check( pMesh->GetIndexBuffers().GetElementSize(0)==4 );
-	const uint32* pIndices = reinterpret_cast<const uint32*>( pMesh->GetIndexBuffers().GetBufferData(0) );
+		int32 SizeX = pImage->GetSizeX();
+		int32 SizeY = pImage->GetSizeY();
 
+		// Get the vertices
+		int32 VertexCount = pMesh->GetVertexCount();
 
-    // The mesh is supposed to contain only the faces in the selected block
-    int faceCount = pMesh->GetFaceCount();
-    //for ( int f=0; f<faceCount; ++f )
-	ParallelFor(faceCount, [pIndices, scratch, sizeX, sizeY, pixelSize, pImage, pixelProc](int f)
+		check( (int32)Scratch->Vertices.Num() == VertexCount );
+		check( (int32)Scratch->CulledVertex.Num() == VertexCount );
+
+		check(pMesh->GetVertexBuffers().GetElementSize(0) == sizeof(FOptimizedVertexWrapping));
+		const FOptimizedVertexWrapping* pVertices = 
+			reinterpret_cast<const FOptimizedVertexWrapping*>(pMesh->GetVertexBuffers().GetBufferData(0));
+
+		float FadeEndCos = FMath::Cos(FadeEnd);
+		for (int32 V = 0; V < VertexCount; ++V)
+		{
+			bool bUseCropping = UncroppedSize[0] > 0;
+			if (bUseCropping)
+			{
+				Scratch->Vertices[V].x = pVertices[V].Uv[0] * UncroppedSize[0] - CropMin[0];
+				Scratch->Vertices[V].y = pVertices[V].Uv[1] * UncroppedSize[1] - CropMin[1];
+			}
+			else
+			{
+				Scratch->Vertices[V].x = pVertices[V].Uv[0] * SizeX;
+				Scratch->Vertices[V].y = pVertices[V].Uv[1] * SizeY;
+			}
+
+			// TODO: No need to copy all. use scratch for the rest only.
+			Scratch->Vertices[V].interpolators[0] = pVertices[V].Position[0];
+			Scratch->Vertices[V].interpolators[1] = pVertices[V].Position[1];
+			Scratch->Vertices[V].interpolators[2] = pVertices[V].Position[2];
+			Scratch->Vertices[V].interpolators[3] = pVertices[V].Normal[0];
+			Scratch->CulledVertex[V] = pVertices[V].Normal[0] < FadeEndCos;
+
+			// Cull vertices that don't belong to the current layout block.
+			if (pVertices[V].LayoutBlock != uint32(Block))
+			{
+				Scratch->CulledVertex[V] = true;
+			}
+		}
+
+		// Get the indices
+		check(pMesh->GetIndexBuffers().GetElementSize(0) == 4);
+		const uint32* pIndices = reinterpret_cast<const uint32*>(pMesh->GetIndexBuffers().GetBufferData(0));
+
+		// The mesh is supposed to contain only the faces in the selected block
+		int32 FaceCount = pMesh->GetFaceCount();
+		//for (int32 Face = 0; Face < FaceCount; ++Face)
+		ParallelFor(FaceCount, [pIndices, Scratch, SizeX, SizeY, PixelSize, pImage, PixelProc](int32 Face)
 		{
-			int i0 = pIndices[f * 3 + 0];
-			int i1 = pIndices[f * 3 + 1];
-			int i2 = pIndices[f * 3 + 2];
+			int32 Index0 = pIndices[Face*3 + 0];
+			int32 Index1 = pIndices[Face*3 + 1];
+			int32 Index2 = pIndices[Face*3 + 2];
 
 			// TODO: This optimisation may remove projection in the center of the face, if the angle
 			// range is small. Make it optional or more sophisticated (cross product may help).
-			if (!scratch->culledVertex[i0] ||
-				!scratch->culledVertex[i1] ||
-				!scratch->culledVertex[i2])
+			if (!Scratch->CulledVertex[Index0] ||
+				!Scratch->CulledVertex[Index1] ||
+				!Scratch->CulledVertex[Index2])
 			{
-				Triangle(pImage->GetData(), pImage->GetDataSize(),
-					sizeX, sizeY,
-					pixelSize,
-					scratch->vertices[i0],
-					scratch->vertices[i1],
-					scratch->vertices[i2],
-					pixelProc,
+				constexpr int32 NumInterpolators = 4;
+				Triangle<NumInterpolators>(pImage->GetData(), pImage->GetDataSize(),
+					SizeX, SizeY,
+					PixelSize,
+					Scratch->Vertices[Index0], Scratch->Vertices[Index1], Scratch->Vertices[Index2],
+					PixelProc,
 					false);
 			}
 		});
+	}
 
-	// Update the relevancy map
-	// \TODO: fix it to work with cropping.
-	if (!bUseCropping)
+	template <EPixelProcessorFeatures Features> 
+	FORCENOINLINE void InvokePlanarRasterizerImpl(const FImageRasterInvokeArgs& Args)
 	{
-		MUTABLE_CPUPROFILER_SCOPE(ImageRasterProjected_Optimised_Relevancy);
+		static_assert(EnumHasAnyFlags(Features, EPixelProcessorFeatures::ProjectionPlanarAndWrap));
 
-		float MinY = float(sizeY) - 1.0f;
-		float MaxY = 0.0f;
-		for (int f = 0; f < faceCount; ++f)
+		const Private::FProjectedPixelProcessorContext Context = 
+			Private::TProjectedPixelProcessor<Features>::MakeContext(
+				Args.SourcePtr, Args.ImagePtr->GetData(), Args.MaskPtr ? Args.MaskPtr->GetData() : nullptr,
+				Args.bIsRGBFadingEnabled, Args.bIsAlphaFadingEnabled, Args.FadeStart, Args.FadeEnd, 
+				Args.ProjectionAngle, Args.MipInterpolationFactor);
+
+		auto PixelProc = [&Context](uint8* Buffer, float Varying[4])
 		{
-			int i0 = pIndices[f * 3 + 0];
-			int i1 = pIndices[f * 3 + 1];
-			int i2 = pIndices[f * 3 + 2];
+			Private::TProjectedPixelProcessor<Features>::ProcessPixel(Context, Buffer, Varying);
+		};
 
-			// A bit ugly, probably can be improved if integrated in above loops and made more precise
-			// inside the pixel processor to account for masked out pixels?
-			if (!scratch->culledVertex[i0] ||
-				!scratch->culledVertex[i1] ||
-				!scratch->culledVertex[i2])
-			{
-				MinY = FMath::Min(MinY, scratch->vertices[i0].y);
-				MinY = FMath::Min(MinY, scratch->vertices[i1].y);
-				MinY = FMath::Min(MinY, scratch->vertices[i2].y);
-				MaxY = FMath::Max(MaxY, scratch->vertices[i0].y);
-				MaxY = FMath::Max(MaxY, scratch->vertices[i1].y);
-				MaxY = FMath::Max(MaxY, scratch->vertices[i2].y);
-			}
-		}
+		ImageRasterProjected_Optimised(Args.MeshPtr, Args.ImagePtr, PixelProc, Args.FadeEnd, Args.CropMin, Args.UncroppedSize, Args.Scratch);
+	}
+
+	template <EPixelProcessorFeatures Features> 
+	FORCENOINLINE void InvokeCylindricalRasterizerImpl(const FImageRasterInvokeArgs& Args)
+	{
+		static_assert(EnumHasAnyFlags(Features, EPixelProcessorFeatures::ProjectionCylindrical));
+
+		const Private::FProjectedPixelProcessorContext Context = 
+			Private::TProjectedPixelProcessor<Features>::MakeContext(
+				Args.SourcePtr, Args.ImagePtr->GetData(), Args.MaskPtr ? Args.MaskPtr->GetData() : nullptr,
+				Args.bIsRGBFadingEnabled, Args.bIsAlphaFadingEnabled, Args.FadeStart, Args.FadeEnd, 
+				Args.ProjectionAngle, Args.MipInterpolationFactor);
+
+		auto PixelProc = [&Context](uint8* Buffer, float Varying[4])
+		{
+			Private::TProjectedPixelProcessor<Features>::ProcessPixel(Context, Buffer, Varying);
+		};
 
-		pImage->m_flags |= Image::IF_HAS_RELEVANCY_MAP;
-		pImage->RelevancyMinY = uint16(FMath::FloorToFloat(MinY));
-		pImage->RelevancyMaxY = uint16(FMath::Min( int32(FMath::CeilToFloat(MaxY)), sizeY-1 ));
+		ImageRasterProjected_Optimised(Args.MeshPtr, Args.ImagePtr, PixelProc, Args.FadeEnd, Args.CropMin, Args.UncroppedSize, Args.Scratch);
 	}
-}
+
+	template <EPixelProcessorFeatures Features> 
+	FORCENOINLINE void InvokeWrappingRasterizerImpl(const FImageRasterInvokeArgs& Args)
+	{
+		static_assert(EnumHasAnyFlags(Features, EPixelProcessorFeatures::ProjectionPlanarAndWrap));
+
+		const Private::FProjectedPixelProcessorContext Context = 
+			Private::TProjectedPixelProcessor<Features>::MakeContext(
+				Args.SourcePtr, Args.ImagePtr->GetData(), Args.MaskPtr ? Args.MaskPtr->GetData() : nullptr,
+				Args.bIsRGBFadingEnabled, Args.bIsAlphaFadingEnabled, Args.FadeStart, Args.FadeEnd, 
+				Args.ProjectionAngle, Args.MipInterpolationFactor);
+
+		auto PixelProc = [&Context](uint8* Buffer, float Varying[4])
+		{
+			Private::TProjectedPixelProcessor<Features>::ProcessPixel(Context, Buffer, Varying);
+		};
+
+		ImageRasterProjected_OptimisedWrapping(Args.MeshPtr, Args.ImagePtr, PixelProc, Args.FadeEnd, Args.Block, Args.CropMin, Args.UncroppedSize, Args.Scratch);
+	}
+} // namespace Private
+
 
 
 //-------------------------------------------------------------------------------------------------
-template<class PIXEL_PROCESSOR>
-void ImageRasterProjected_OptimisedWrapping( const Mesh* pMesh, Image* pImage, PIXEL_PROCESSOR& pixelProc,
-	float fadeEnd,
-	int block,
+void ImageRasterProjectedPlanar(const Mesh* pMesh, Image* pImage,
+	const Image* pSource,
+	const Image* pMask,
+	bool bIsRGBFadingEnabled, bool bIsAlphaFadingEnabled,
+	ESamplingMethod SamplingMethod,
+	float FadeStart, float FadeEnd, float MipInterpolationFactor,
+	int32 Layout, int32 Block,
 	UE::Math::TIntVector2<uint16> CropMin, UE::Math::TIntVector2<uint16> UncroppedSize,
-	SCRATCH_IMAGE_PROJECT* scratch )
+	FScratchImageProject* Scratch, bool bUseVectorImplementation)
 {
-	MUTABLE_CPUPROFILER_SCOPE(ImageRasterProjected_OptimisedWrapping);
+	using namespace mu::Private;
 
-    if (!pMesh || !pMesh->GetFaceCount())
-    {
-        return;
-    }
+    check(!pMask || pMask->GetSizeX() == pImage->GetSizeX());
+    check(!pMask || pMask->GetSizeY() == pImage->GetSizeY());
+    check(!pMask || pMask->GetFormat()==EImageFormat::IF_L_UBYTE);
 
-	EImageFormat format = pImage->GetFormat();
-    int pixelSize = GetImageFormatData( format ).BytesPerBlock;
+	MUTABLE_CPUPROFILER_SCOPE(ImageProject);
 
-    int sizeX = pImage->GetSizeX();
-    int sizeY = pImage->GetSizeY();
+	EPixelProcessorFeatures ProcessorFeatures = EPixelProcessorFeatures::ProjectionPlanarAndWrap;
 
-    // Get the vertices
-    int vertexCount = pMesh->GetVertexCount();
+	// Disable vectorized implementation if Point sampling.
+	if (bUseVectorImplementation && SamplingMethod != ESamplingMethod::Point)
+	{
+		ProcessorFeatures |= EPixelProcessorFeatures::VectorizedImpl;
+	}
 
-    check( (int)scratch->vertices.Num()==vertexCount );
-    check( (int)scratch->culledVertex.Num()==vertexCount );
+	if (pSource->GetFormat() == EImageFormat::IF_L_UBYTE)
+	{
+		ProcessorFeatures |= EPixelProcessorFeatures::FormatL;
+	}
+	else if (pSource->GetFormat() == EImageFormat::IF_RGB_UBYTE)
+	{
+		ProcessorFeatures |= EPixelProcessorFeatures::FormatRGB;
+	}
+	else if (pSource->GetFormat() == EImageFormat::IF_RGBA_UBYTE || pSource->GetFormat() == EImageFormat::IF_BGRA_UBYTE)
+	{
+		ProcessorFeatures |= EPixelProcessorFeatures::FormatRGBA;
+	}
 
-    check( pMesh->GetVertexBuffers().GetElementSize(0)==sizeof(OPTIMISED_VERTEX_WRAPPING) );
-	const OPTIMISED_VERTEX_WRAPPING* pVertices = reinterpret_cast<const OPTIMISED_VERTEX_WRAPPING*>( pMesh->GetVertexBuffers().GetBufferData(0) );
+	if (SamplingMethod == ESamplingMethod::Point)
+	{
+		ProcessorFeatures |= EPixelProcessorFeatures::SamplingPoint;
+	}
+	else if (SamplingMethod == ESamplingMethod::BiLinear)
+	{
+		ProcessorFeatures |= EPixelProcessorFeatures::SamplingLinear;
+	}
 
-    float fadeEndCos = cosf(fadeEnd);
-    for ( int v=0; v<vertexCount; ++v )
-    {
-		bool bUseCropping = UncroppedSize[0] > 0;
-		if (bUseCropping)
+	if (pMask)
+	{
+		ProcessorFeatures |= EPixelProcessorFeatures::WithMask;
+	}
+
+	if (pMesh->m_staticFormatFlags & (1 << SMF_PROJECT))
+	{
+		float UnusedProjectionAngle = 0;
+		FImageRasterInvokeArgs RasterArgs = 
+		{
+			pMesh, pImage, pSource, pMask, 
+			bIsRGBFadingEnabled, bIsAlphaFadingEnabled,
+			SamplingMethod,
+			FadeStart, FadeEnd,  UnusedProjectionAngle, MipInterpolationFactor, 
+			Layout, Block, CropMin, UncroppedSize, Scratch
+		};
+
+		using EPPF = EPixelProcessorFeatures;
+
+		// This contains all the permutations, this is way more than what it was supported before
+		// but maybe not needed. Mask permutations can be removed safely
+		switch (ProcessorFeatures)
 		{
-			scratch->vertices[v].x = pVertices[v].uv[0] * UncroppedSize[0] - CropMin[0];
-			scratch->vertices[v].y = pVertices[v].uv[1] * UncroppedSize[1] - CropMin[1];
+		case (EPPF::FormatL | EPPF::ProjectionPlanarAndWrap | EPPF::SamplingPoint):
+		{
+			InvokePlanarRasterizerImpl<EPPF::FormatL | EPPF::ProjectionPlanarAndWrap | EPPF::SamplingPoint>(RasterArgs);
+			break;
+		}
+		case (EPPF::FormatRGB | EPPF::ProjectionPlanarAndWrap | EPPF::SamplingPoint):
+		{
+			InvokePlanarRasterizerImpl<EPPF::FormatRGB | EPPF::ProjectionPlanarAndWrap | EPPF::SamplingPoint>(RasterArgs);
+			break;
 		}
-		else
+		case (EPPF::FormatRGBA | EPPF::ProjectionPlanarAndWrap | EPPF::SamplingPoint):
 		{
-			scratch->vertices[v].x = pVertices[v].uv[0] * sizeX;
-			scratch->vertices[v].y = pVertices[v].uv[1] * sizeY;
+			InvokePlanarRasterizerImpl<EPPF::FormatRGBA | EPPF::ProjectionPlanarAndWrap | EPPF::SamplingPoint>(RasterArgs);
+			break;
 		}
 
-        // TODO: No need to copy all. use scratch for the rest only.
-        scratch->vertices[v].interpolators[0] = pVertices[v].pos[0];
-        scratch->vertices[v].interpolators[1] = pVertices[v].pos[1];
-        scratch->vertices[v].interpolators[2] = pVertices[v].pos[2];
-        scratch->vertices[v].interpolators[3] = pVertices[v].nor[0];
-        scratch->culledVertex[v] = pVertices[v].nor[0] < fadeEndCos;
+		case (EPPF::FormatL | EPPF::ProjectionPlanarAndWrap | EPPF::SamplingLinear):
+		{
+			InvokePlanarRasterizerImpl<EPPF::FormatL | EPPF::ProjectionPlanarAndWrap | EPPF::SamplingLinear>(RasterArgs);
+			break;
+		}
+		case (EPPF::FormatRGB | EPPF::ProjectionPlanarAndWrap | EPPF::SamplingLinear):
+		{
+			InvokePlanarRasterizerImpl<EPPF::FormatRGB | EPPF::ProjectionPlanarAndWrap | EPPF::SamplingLinear>(RasterArgs);
+			break;
+		}
+		case (EPPF::FormatRGBA | EPPF::ProjectionPlanarAndWrap | EPPF::SamplingLinear):
+		{
+			InvokePlanarRasterizerImpl<EPPF::FormatRGBA | EPPF::ProjectionPlanarAndWrap | EPPF::SamplingLinear>(RasterArgs);
+			break;
+		}
 
-        // Cull vertices that don't belong to the current layout block.
-        if (pVertices[v].layoutBlock!=uint32(block))
-        {
-            scratch->culledVertex[v] = true;
-        }
-    }
+		case (EPPF::FormatL | EPPF::ProjectionPlanarAndWrap | EPPF::SamplingPoint | EPPF::WithMask):
+		{
+			InvokePlanarRasterizerImpl<EPPF::FormatL | EPPF::ProjectionPlanarAndWrap | EPPF::SamplingPoint | EPPF::WithMask>(RasterArgs);
+			break;
+		}
+		case (EPPF::FormatRGB | EPPF::ProjectionPlanarAndWrap | EPPF::SamplingPoint | EPPF::WithMask):
+		{
+			InvokePlanarRasterizerImpl<EPPF::FormatRGB | EPPF::ProjectionPlanarAndWrap | EPPF::SamplingPoint | EPPF::WithMask>(RasterArgs);
+			break;
+		}
+		case (EPPF::FormatRGBA | EPPF::ProjectionPlanarAndWrap | EPPF::SamplingPoint | EPPF::WithMask):
+		{
+			InvokePlanarRasterizerImpl<EPPF::FormatRGBA | EPPF::ProjectionPlanarAndWrap | EPPF::SamplingPoint>(RasterArgs);
+			break;
+		}
 
-    // Get the indices
-    check( pMesh->GetIndexBuffers().GetElementSize(0)==4 );
-	const uint32* pIndices = reinterpret_cast<const uint32*>( pMesh->GetIndexBuffers().GetBufferData(0) );
+		case (EPPF::FormatL | EPPF::ProjectionPlanarAndWrap | EPPF::SamplingLinear | EPPF::WithMask):
+		{
+			InvokePlanarRasterizerImpl<EPPF::FormatL | EPPF::ProjectionPlanarAndWrap | EPPF::SamplingLinear | EPPF::WithMask>(RasterArgs);
+			break;
+		}
+		case (EPPF::FormatRGB | EPPF::ProjectionPlanarAndWrap | EPPF::SamplingLinear | EPPF::WithMask):
+		{
+			InvokePlanarRasterizerImpl<EPPF::FormatRGB | EPPF::ProjectionPlanarAndWrap | EPPF::SamplingLinear | EPPF::WithMask>(RasterArgs);
+			break;
+		}
+		case (EPPF::FormatRGBA | EPPF::ProjectionPlanarAndWrap | EPPF::SamplingLinear | EPPF::WithMask):
+		{
+			InvokePlanarRasterizerImpl<EPPF::FormatRGBA | EPPF::ProjectionPlanarAndWrap | EPPF::SamplingLinear | EPPF::WithMask>(RasterArgs);
+			break;
+		}
+	
+		//case (EPPF::FormatL | EPPF::ProjectionPlanarAndWrap | EPPF::SamplingPoint | EPPF::VectorizedImpl):
+		//{
+		//	InvokePlanarRasterizerImpl<EPPF::FormatL | EPPF::ProjectionPlanarAndWrap | EPPF::SamplingPoint | EPPF::VectorizedImpl>(RasterArgs);
+		//	break;
+		//}
+		//case (EPPF::FormatRGB | EPPF::ProjectionPlanarAndWrap | EPPF::SamplingPoint | EPPF::VectorizedImpl):
+		//{
+		//	InvokePlanarRasterizerImpl<EPPF::FormatRGB | EPPF::ProjectionPlanarAndWrap | EPPF::SamplingPoint | EPPF::VectorizedImpl>(RasterArgs);
+		//	break;
+		//}
+		//case (EPPF::FormatRGBA | EPPF::ProjectionPlanarAndWrap | EPPF::SamplingPoint | EPPF::VectorizedImpl):
+		//{
+		//	InvokePlanarRasterizerImpl<EPPF::FormatRGBA | EPPF::ProjectionPlanarAndWrap | EPPF::SamplingPoint | EPPF::VectorizedImpl>(RasterArgs);
+		//	break;
+		//}
 
+		case (EPPF::FormatL | EPPF::ProjectionPlanarAndWrap | EPPF::SamplingLinear | EPPF::VectorizedImpl):
+		{
+			InvokePlanarRasterizerImpl<EPPF::FormatL | EPPF::ProjectionPlanarAndWrap | EPPF::SamplingLinear | EPPF::VectorizedImpl>(RasterArgs);
+			break;
+		}
+		case (EPPF::FormatRGB | EPPF::ProjectionPlanarAndWrap | EPPF::SamplingLinear | EPPF::VectorizedImpl):
+		{
+			InvokePlanarRasterizerImpl<EPPF::FormatRGB | EPPF::ProjectionPlanarAndWrap | EPPF::SamplingLinear | EPPF::VectorizedImpl>(RasterArgs);
+			break;
+		}
+		case (EPPF::FormatRGBA | EPPF::ProjectionPlanarAndWrap | EPPF::SamplingLinear | EPPF::VectorizedImpl):
+		{
+			InvokePlanarRasterizerImpl<EPPF::FormatRGBA | EPPF::ProjectionPlanarAndWrap | EPPF::SamplingLinear | EPPF::VectorizedImpl>(RasterArgs);
+			break;
+		}
 
-    // The mesh is supposed to contain only the faces in the selected block
-    int faceCount = pMesh->GetFaceCount();
-    for ( int f=0; f<faceCount; ++f )
-    {
-        int i0 = pIndices[f*3+0];
-        int i1 = pIndices[f*3+1];
-        int i2 = pIndices[f*3+2];
+		//case (EPPF::FormatL | EPPF::ProjectionPlanarAndWrap | EPPF::SamplingPoint | EPPF::WithMask | EPPF::VectorizedImpl):
+		//{
+		//	InvokePlanarRasterizerImpl<EPPF::FormatL | EPPF::ProjectionPlanarAndWrap | EPPF::SamplingPoint | EPPF::WithMask | EPPF::VectorizedImpl>(RasterArgs);
+		//	break;
+		//}
+		//case (EPPF::FormatRGB | EPPF::ProjectionPlanarAndWrap | EPPF::SamplingPoint | EPPF::WithMask | EPPF::VectorizedImpl):
+		//{
+		//	InvokePlanarRasterizerImpl<EPPF::FormatRGB | EPPF::ProjectionPlanarAndWrap | EPPF::SamplingPoint | EPPF::WithMask | EPPF::VectorizedImpl>(RasterArgs);
+		//	break;
+		//}
+		//case (EPPF::FormatRGBA | EPPF::ProjectionPlanarAndWrap | EPPF::SamplingPoint | EPPF::WithMask | EPPF::VectorizedImpl):
+		//{
+		//	InvokePlanarRasterizerImpl<EPPF::FormatRGBA | EPPF::ProjectionPlanarAndWrap | EPPF::SamplingPoint | EPPF::VectorizedImpl>(RasterArgs);
+		//	break;
+		//}
 
-        // TODO: This optimisation may remove projection in the center of the face, if the angle
-        // range is small. Make it optional or more sophisticated (cross product may help).
-        if ( !scratch->culledVertex[i0] ||
-             !scratch->culledVertex[i1] ||
-             !scratch->culledVertex[i2] )
-        {
-            Triangle( pImage->GetData(), pImage->GetDataSize(),
-				sizeX, sizeY,
-				pixelSize,
-				scratch->vertices[i0], scratch->vertices[i1], scratch->vertices[i2],
-				pixelProc,
-				false );
-        }
+		case (EPPF::FormatL | EPPF::ProjectionPlanarAndWrap | EPPF::SamplingLinear | EPPF::WithMask | EPPF::VectorizedImpl):
+		{
+			InvokePlanarRasterizerImpl<EPPF::FormatL | EPPF::ProjectionPlanarAndWrap | EPPF::SamplingLinear | EPPF::WithMask | EPPF::VectorizedImpl>(RasterArgs);
+			break;
+		}
+		case (EPPF::FormatRGB | EPPF::ProjectionPlanarAndWrap | EPPF::SamplingLinear | EPPF::WithMask | EPPF::VectorizedImpl):
+		{
+			InvokePlanarRasterizerImpl<EPPF::FormatRGB | EPPF::ProjectionPlanarAndWrap | EPPF::SamplingLinear | EPPF::WithMask>(RasterArgs);
+			break;
+		}
+		case (EPPF::FormatRGBA | EPPF::ProjectionPlanarAndWrap | EPPF::SamplingLinear | EPPF::WithMask | EPPF::VectorizedImpl):
+		{
+			InvokePlanarRasterizerImpl<EPPF::FormatRGBA | EPPF::ProjectionPlanarAndWrap | EPPF::SamplingLinear | EPPF::WithMask | EPPF::VectorizedImpl>(RasterArgs);
+			break;
+		}
+
+		default:
+		{
+			checkf(false, TEXT("Planar raster pixel processor permutation not implemented."));
+		}
+		};
+	}
+
+    else
+    {
+        check(false);
     }
 
 }
 
 
 //-------------------------------------------------------------------------------------------------
-void ImageRasterProjectedPlanar( const Mesh* pMesh,
-	Image* pImage,
-	const Image* pSource,
-	const Image* pMask,
+void ImageRasterProjectedWrapping( const Mesh* pMesh, Image* pImage,
+	const Image* pSource, const Image* pMask,
 	bool bIsRGBFadingEnabled, bool bIsAlphaFadingEnabled,
 	ESamplingMethod SamplingMethod,
-	float fadeStart, float fadeEnd, float MipInterpolationFactor,
-	int layout, int block,
+	float FadeStart, float FadeEnd, float MipInterpolationFactor,
+	int32 Layout, int32 Block,
 	UE::Math::TIntVector2<uint16> CropMin, UE::Math::TIntVector2<uint16> UncroppedSize,
-	SCRATCH_IMAGE_PROJECT* scratch, bool bUseVectorImplementation)
+	FScratchImageProject* Scratch, bool bUseVectorImplementation)
 {
-    check( !pMask || pMask->GetSizeX() == pImage->GetSizeX() );
-    check( !pMask || pMask->GetSizeY() == pImage->GetSizeY() );
-    check( !pMask || pMask->GetFormat()==EImageFormat::IF_L_UBYTE );
 
-	MUTABLE_CPUPROFILER_SCOPE(ImageProject);
+	using namespace mu::Private;
+    
+	check(!pMask || pMask->GetSizeX() == pImage->GetSizeX());
+    check(!pMask || pMask->GetSizeY() == pImage->GetSizeY());
+    check(!pMask || pMask->GetFormat()==EImageFormat::IF_L_UBYTE);
 
-    if ( pMesh->m_staticFormatFlags & (1<<SMF_PROJECT) )
-    {
-        // Mesh-optimised version
-        if ( pSource->GetFormat() == EImageFormat::IF_RGB_UBYTE )
-        {
-			if (SamplingMethod == ESamplingMethod::Point)
-			{
-				const RasterProjectedPixelProcessor_UBYTE<3> pixelProc(pSource,
-					pImage->GetData(),
-					pMask ? pMask->GetData() : nullptr,
-					bIsRGBFadingEnabled, bIsAlphaFadingEnabled,
-					fadeStart, fadeEnd);
-				ImageRasterProjected_Optimised(pMesh, pImage, pixelProc, fadeEnd, CropMin, UncroppedSize, scratch);
-			}
-			else if (SamplingMethod == ESamplingMethod::BiLinear)
-			{
-				if (bUseVectorImplementation)
-				{
-					constexpr bool bUseVectorImpl = true;
-					const RasterProjectedPixelProcessorLinear_UBYTE<3, bUseVectorImpl> pixelProc(pSource,
-						pImage->GetData(),
-						pMask ? pMask->GetData() : nullptr,
-						bIsRGBFadingEnabled, bIsAlphaFadingEnabled,
-						fadeStart, fadeEnd, MipInterpolationFactor);
-					ImageRasterProjected_Optimised(pMesh, pImage, pixelProc, fadeEnd, CropMin, UncroppedSize, scratch);
-				}
-				else
-				{
-					constexpr bool bUseVectorImpl = false;
-					const RasterProjectedPixelProcessorLinear_UBYTE<3, bUseVectorImpl> pixelProc(pSource,
-						pImage->GetData(),
-						pMask ? pMask->GetData() : nullptr,
-						bIsRGBFadingEnabled, bIsAlphaFadingEnabled,
-						fadeStart, fadeEnd, MipInterpolationFactor);
-					ImageRasterProjected_Optimised(pMesh, pImage, pixelProc, fadeEnd, CropMin, UncroppedSize, scratch);
-				}
-			}
-			else
-			{
-				check(false);
-			}
-        }
+	MUTABLE_CPUPROFILER_SCOPE(ImageProjectWrapping);
 
-        else if ( pSource->GetFormat() == EImageFormat::IF_RGBA_UBYTE ||
-                  pSource->GetFormat() == EImageFormat::IF_BGRA_UBYTE )
-        {
-            if (pMask != nullptr)
-            {
-				if (SamplingMethod == ESamplingMethod::Point)
-				{
-					const RasterProjectedPixelProcessor_UBYTE<4, true> pixelProc(pSource,
-						pImage->GetData(),
-						pMask->GetData(),
-						bIsRGBFadingEnabled, bIsAlphaFadingEnabled,
-						fadeStart, fadeEnd);
-					ImageRasterProjected_Optimised(pMesh, pImage, pixelProc, fadeEnd, CropMin, UncroppedSize, scratch);
-				}
-				else if (SamplingMethod == ESamplingMethod::BiLinear)
-				{
-					if (bUseVectorImplementation)
-					{
-						constexpr bool bUseVectorImpl = true;
-						const RasterProjectedPixelProcessorLinear_UBYTE<4, bUseVectorImpl, true> pixelProc(pSource,
-							pImage->GetData(),
-							pMask->GetData(),
-							bIsRGBFadingEnabled, bIsAlphaFadingEnabled,
-							fadeStart, fadeEnd, MipInterpolationFactor);
-						ImageRasterProjected_Optimised(pMesh, pImage, pixelProc, fadeEnd, CropMin, UncroppedSize, scratch);
-					}
-					else
-					{
-						constexpr bool bUseVectorImpl = false;
-						const RasterProjectedPixelProcessorLinear_UBYTE<4, bUseVectorImpl, true> pixelProc(pSource,
-							pImage->GetData(),
-							pMask->GetData(),
-							bIsRGBFadingEnabled, bIsAlphaFadingEnabled,
-							fadeStart, fadeEnd, MipInterpolationFactor);
-						ImageRasterProjected_Optimised(pMesh, pImage, pixelProc, fadeEnd, CropMin, UncroppedSize, scratch);
-					}
-				}
-				else
-				{
-					check(false);
-				}
-            }
-            else
-            {
-				if (SamplingMethod == ESamplingMethod::Point) 
-				{
-					const RasterProjectedPixelProcessor_UBYTE<4, false> pixelProc(pSource,
-						pImage->GetData(),
-						nullptr,
-						bIsRGBFadingEnabled, bIsAlphaFadingEnabled,
-						fadeStart, fadeEnd);
-					ImageRasterProjected_Optimised(pMesh, pImage, pixelProc, fadeEnd, CropMin, UncroppedSize, scratch);
-				}
-				else if (SamplingMethod == ESamplingMethod::BiLinear)
-				{
-					if (bUseVectorImplementation)
-					{
-						constexpr bool bUseVectorImpl = true;
-						const RasterProjectedPixelProcessorLinear_UBYTE<4, bUseVectorImpl, false> pixelProc(pSource,
-							pImage->GetData(),
-							nullptr,
-							bIsRGBFadingEnabled, bIsAlphaFadingEnabled,
-							fadeStart, fadeEnd, MipInterpolationFactor);
-						ImageRasterProjected_Optimised(pMesh, pImage, pixelProc, fadeEnd, CropMin, UncroppedSize, scratch);
-					}
-					else
-					{
-						constexpr bool bUseVectorImpl = false;
-						const RasterProjectedPixelProcessorLinear_UBYTE<4, bUseVectorImpl, false> pixelProc(pSource,
-							pImage->GetData(),
-							nullptr,
-							bIsRGBFadingEnabled, bIsAlphaFadingEnabled,
-							fadeStart, fadeEnd, MipInterpolationFactor);
-						ImageRasterProjected_Optimised(pMesh, pImage, pixelProc, fadeEnd, CropMin, UncroppedSize, scratch);
-					}
-				}
-				else
-				{
-					check(false);
-				}
-            }
-        }
-        else if ( pSource->GetFormat() == EImageFormat::IF_L_UBYTE )
-        {
-			if (SamplingMethod == ESamplingMethod::Point)
-			{
-				const RasterProjectedPixelProcessor_UBYTE<1> pixelProc(pSource,
-					pImage->GetData(),
-					pMask ? pMask->GetData() : nullptr,
-					bIsRGBFadingEnabled, bIsAlphaFadingEnabled,
-					fadeStart, fadeEnd);
-				ImageRasterProjected_Optimised(pMesh, pImage, pixelProc, fadeEnd, CropMin, UncroppedSize, scratch);
-			}
-			else if (SamplingMethod == ESamplingMethod::BiLinear)
-			{
-				if (bUseVectorImplementation)
-				{
-					constexpr bool bUseVectorImpl = true;
-					const RasterProjectedPixelProcessorLinear_UBYTE<1, bUseVectorImpl> pixelProc(pSource,
-						pImage->GetData(),
-						pMask ? pMask->GetData() : nullptr,
-						bIsRGBFadingEnabled, bIsAlphaFadingEnabled,
-						fadeStart, fadeEnd, MipInterpolationFactor);
-					ImageRasterProjected_Optimised(pMesh, pImage, pixelProc, fadeEnd, CropMin, UncroppedSize, scratch);
-				}
-				else
-				{
-					constexpr bool bUseVectorImpl = false;
-					const RasterProjectedPixelProcessorLinear_UBYTE<1, bUseVectorImpl> pixelProc(pSource,
-						pImage->GetData(),
-						pMask ? pMask->GetData() : nullptr,
-						bIsRGBFadingEnabled, bIsAlphaFadingEnabled,
-						fadeStart, fadeEnd, MipInterpolationFactor);
-					ImageRasterProjected_Optimised(pMesh, pImage, pixelProc, fadeEnd, CropMin, UncroppedSize, scratch);
-				}
-			}
-			else
-			{
-				check(false);
-			}
-        }
+	EPixelProcessorFeatures ProcessorFeatures = EPixelProcessorFeatures::ProjectionPlanarAndWrap;
+
+	if (bUseVectorImplementation && SamplingMethod != ESamplingMethod::Point)
+	{
+		ProcessorFeatures |= EPixelProcessorFeatures::VectorizedImpl;
+	}
+
+	if (pSource->GetFormat() == EImageFormat::IF_L_UBYTE)
+	{
+		ProcessorFeatures |= EPixelProcessorFeatures::FormatL;
+	}
+	else if (pSource->GetFormat() == EImageFormat::IF_RGB_UBYTE)
+	{
+		ProcessorFeatures |= EPixelProcessorFeatures::FormatRGB;
+	}
+	else if (pSource->GetFormat() == EImageFormat::IF_RGBA_UBYTE || pSource->GetFormat() == EImageFormat::IF_BGRA_UBYTE)
+	{
+		ProcessorFeatures |= EPixelProcessorFeatures::FormatRGBA;
+	}
+
+	if (SamplingMethod == ESamplingMethod::Point)
+	{
+		ProcessorFeatures |= EPixelProcessorFeatures::SamplingPoint;
+	}
+	else if (SamplingMethod == ESamplingMethod::BiLinear)
+	{
+		ProcessorFeatures |= EPixelProcessorFeatures::SamplingLinear;
+	}
+
+	if (pMask)
+	{
+		ProcessorFeatures |= EPixelProcessorFeatures::WithMask;
+	}
 
-        else
-        {
-			check(false);
-        }
-    }
+    if ( ( pMesh->m_staticFormatFlags & (1<<SMF_PROJECTWRAPPING) ) )
+    {
+		float UnusedProjectionAngle = 0;
+
+		FImageRasterInvokeArgs RasterArgs =
+		{
+			pMesh, pImage, pSource, pMask,
+			bIsRGBFadingEnabled, bIsAlphaFadingEnabled,
+			SamplingMethod,
+			FadeStart, FadeEnd, UnusedProjectionAngle, MipInterpolationFactor,
+			Layout, Block, CropMin, UncroppedSize, Scratch
+		};
+
+		using EPPF = EPixelProcessorFeatures;
+
+		// This contains all the permutations, this is way more than what it was supported before
+		// but maybe not needed. Mask permutations can be removed safely
+		switch (ProcessorFeatures)
+		{
+		case (EPPF::FormatL | EPPF::ProjectionPlanarAndWrap | EPPF::SamplingPoint):
+		{
+			InvokeWrappingRasterizerImpl<EPPF::FormatL | EPPF::ProjectionPlanarAndWrap | EPPF::SamplingPoint>(RasterArgs);
+			break;
+		}
+		case (EPPF::FormatRGB | EPPF::ProjectionPlanarAndWrap | EPPF::SamplingPoint):
+		{
+			InvokeWrappingRasterizerImpl<EPPF::FormatRGB | EPPF::ProjectionPlanarAndWrap | EPPF::SamplingPoint>(RasterArgs);
+			break;
+		}
+		case (EPPF::FormatRGBA | EPPF::ProjectionPlanarAndWrap | EPPF::SamplingPoint):
+		{
+			InvokeWrappingRasterizerImpl<EPPF::FormatRGBA | EPPF::ProjectionPlanarAndWrap | EPPF::SamplingPoint>(RasterArgs);
+			break;
+		}
+
+		case (EPPF::FormatL | EPPF::ProjectionPlanarAndWrap | EPPF::SamplingLinear):
+		{
+			InvokeWrappingRasterizerImpl<EPPF::FormatL | EPPF::ProjectionPlanarAndWrap | EPPF::SamplingLinear>(RasterArgs);
+			break;
+		}
+		case (EPPF::FormatRGB | EPPF::ProjectionPlanarAndWrap | EPPF::SamplingLinear):
+		{
+			InvokeWrappingRasterizerImpl<EPPF::FormatRGB | EPPF::ProjectionPlanarAndWrap | EPPF::SamplingLinear>(RasterArgs);
+			break;
+		}
+		case (EPPF::FormatRGBA | EPPF::ProjectionPlanarAndWrap | EPPF::SamplingLinear):
+		{
+			InvokeWrappingRasterizerImpl<EPPF::FormatRGBA | EPPF::ProjectionPlanarAndWrap | EPPF::SamplingLinear>(RasterArgs);
+			break;
+		}
+
+		case (EPPF::FormatL | EPPF::ProjectionPlanarAndWrap | EPPF::SamplingPoint | EPPF::WithMask):
+		{
+			InvokeWrappingRasterizerImpl<EPPF::FormatL | EPPF::ProjectionPlanarAndWrap | EPPF::SamplingPoint | EPPF::WithMask>(RasterArgs);
+			break;
+		}
+		case (EPPF::FormatRGB | EPPF::ProjectionPlanarAndWrap | EPPF::SamplingPoint | EPPF::WithMask):
+		{
+			InvokeWrappingRasterizerImpl<EPPF::FormatRGB | EPPF::ProjectionPlanarAndWrap | EPPF::SamplingPoint | EPPF::WithMask>(RasterArgs);
+			break;
+		}
+		case (EPPF::FormatRGBA | EPPF::ProjectionPlanarAndWrap | EPPF::SamplingPoint | EPPF::WithMask):
+		{
+			InvokeWrappingRasterizerImpl<EPPF::FormatRGBA | EPPF::ProjectionPlanarAndWrap | EPPF::SamplingPoint>(RasterArgs);
+			break;
+		}
+
+		case (EPPF::FormatL | EPPF::ProjectionPlanarAndWrap | EPPF::SamplingLinear | EPPF::WithMask):
+		{
+			InvokeWrappingRasterizerImpl<EPPF::FormatL | EPPF::ProjectionPlanarAndWrap | EPPF::SamplingLinear | EPPF::WithMask>(RasterArgs);
+			break;
+		}
+		case (EPPF::FormatRGB | EPPF::ProjectionPlanarAndWrap | EPPF::SamplingLinear | EPPF::WithMask):
+		{
+			InvokeWrappingRasterizerImpl<EPPF::FormatRGB | EPPF::ProjectionPlanarAndWrap | EPPF::SamplingLinear | EPPF::WithMask>(RasterArgs);
+			break;
+		}
+		case (EPPF::FormatRGBA | EPPF::ProjectionPlanarAndWrap | EPPF::SamplingLinear | EPPF::WithMask):
+		{
+			InvokeWrappingRasterizerImpl<EPPF::FormatRGBA | EPPF::ProjectionPlanarAndWrap | EPPF::SamplingLinear | EPPF::WithMask>(RasterArgs);
+			break;
+		}
+	
+		//case (EPPF::FormatL | EPPF::ProjectionPlanarAndWrap | EPPF::SamplingPoint | EPPF::VectorizedImpl):
+		//{
+		//	InvokeWrappingRasterizerImpl<EPPF::FormatL | EPPF::ProjectionPlanarAndWrap | EPPF::SamplingPoint | EPPF::VectorizedImpl>(RasterArgs);
+		//	break;
+		//}
+		//case (EPPF::FormatRGB | EPPF::ProjectionPlanarAndWrap | EPPF::SamplingPoint | EPPF::VectorizedImpl):
+		//{
+		//	InvokeWrappingRasterizerImpl<EPPF::FormatRGB | EPPF::ProjectionPlanarAndWrap | EPPF::SamplingPoint | EPPF::VectorizedImpl>(RasterArgs);
+		//	break;
+		//}
+		//case (EPPF::FormatRGBA | EPPF::ProjectionPlanarAndWrap | EPPF::SamplingPoint | EPPF::VectorizedImpl):
+		//{
+		//	InvokeWrappingRasterizerImpl<EPPF::FormatRGBA | EPPF::ProjectionPlanarAndWrap | EPPF::SamplingPoint | EPPF::VectorizedImpl>(RasterArgs);
+		//	break;
+		//}
+
+		case (EPPF::FormatL | EPPF::ProjectionPlanarAndWrap | EPPF::SamplingLinear | EPPF::VectorizedImpl):
+		{
+			InvokeWrappingRasterizerImpl<EPPF::FormatL | EPPF::ProjectionPlanarAndWrap | EPPF::SamplingLinear | EPPF::VectorizedImpl>(RasterArgs);
+			break;
+		}
+		case (EPPF::FormatRGB | EPPF::ProjectionPlanarAndWrap | EPPF::SamplingLinear | EPPF::VectorizedImpl):
+		{
+			InvokeWrappingRasterizerImpl<EPPF::FormatRGB | EPPF::ProjectionPlanarAndWrap | EPPF::SamplingLinear | EPPF::VectorizedImpl>(RasterArgs);
+			break;
+		}
+		case (EPPF::FormatRGBA | EPPF::ProjectionPlanarAndWrap | EPPF::SamplingLinear | EPPF::VectorizedImpl):
+		{
+			InvokeWrappingRasterizerImpl<EPPF::FormatRGBA | EPPF::ProjectionPlanarAndWrap | EPPF::SamplingLinear | EPPF::VectorizedImpl>(RasterArgs);
+			break;
+		}
+
+		//case (EPPF::FormatL | EPPF::ProjectionPlanarAndWrap | EPPF::SamplingPoint | EPPF::WithMask | EPPF::VectorizedImpl):
+		//{
+		//	InvokeWrappingRasterizerImpl<EPPF::FormatL | EPPF::ProjectionPlanarAndWrap | EPPF::SamplingPoint | EPPF::WithMask | EPPF::VectorizedImpl>(RasterArgs);
+		//	break;
+		//}
+		//case (EPPF::FormatRGB | EPPF::ProjectionPlanarAndWrap | EPPF::SamplingPoint | EPPF::WithMask | EPPF::VectorizedImpl):
+		//{
+		//	InvokeWrappingRasterizerImpl<EPPF::FormatRGB | EPPF::ProjectionPlanarAndWrap | EPPF::SamplingPoint | EPPF::WithMask | EPPF::VectorizedImpl>(RasterArgs);
+		//	break;
+		//}
+		//case (EPPF::FormatRGBA | EPPF::ProjectionPlanarAndWrap | EPPF::SamplingPoint | EPPF::WithMask | EPPF::VectorizedImpl):
+		//{
+		//	InvokeWrappingRasterizerImpl<EPPF::FormatRGBA | EPPF::ProjectionPlanarAndWrap | EPPF::SamplingPoint | EPPF::VectorizedImpl>(RasterArgs);
+		//	break;
+		//}
 
+		case (EPPF::FormatL | EPPF::ProjectionPlanarAndWrap | EPPF::SamplingLinear | EPPF::WithMask | EPPF::VectorizedImpl):
+		{
+			InvokeWrappingRasterizerImpl<EPPF::FormatL | EPPF::ProjectionPlanarAndWrap | EPPF::SamplingLinear | EPPF::WithMask | EPPF::VectorizedImpl>(RasterArgs);
+			break;
+		}
+		case (EPPF::FormatRGB | EPPF::ProjectionPlanarAndWrap | EPPF::SamplingLinear | EPPF::WithMask | EPPF::VectorizedImpl):
+		{
+			InvokeWrappingRasterizerImpl<EPPF::FormatRGB | EPPF::ProjectionPlanarAndWrap | EPPF::SamplingLinear | EPPF::WithMask>(RasterArgs);
+			break;
+		}
+		case (EPPF::FormatRGBA | EPPF::ProjectionPlanarAndWrap | EPPF::SamplingLinear | EPPF::WithMask | EPPF::VectorizedImpl):
+		{
+			InvokeWrappingRasterizerImpl<EPPF::FormatRGBA | EPPF::ProjectionPlanarAndWrap | EPPF::SamplingLinear | EPPF::WithMask | EPPF::VectorizedImpl>(RasterArgs);
+			break;
+		}
+		
+		default:
+		{
+			checkf(false, TEXT("Wrapping projection raster pixel processor permutation not implemented."));
+		}
+		};
+	}
     else
     {
         check(false);
@@ -1140,159 +1319,226 @@ void ImageRasterProjectedPlanar( const Mesh* pMesh,
 
 }
 
-
 //-------------------------------------------------------------------------------------------------
-void ImageRasterProjectedWrapping( const Mesh* pMesh, Image* pImage,
+void ImageRasterProjectedCylindrical( const Mesh* pMesh, Image* pImage,
 	const Image* pSource, const Image* pMask,
 	bool bIsRGBFadingEnabled, bool bIsAlphaFadingEnabled,
 	ESamplingMethod SamplingMethod,
-	float fadeStart, float fadeEnd, float MipInterpolationFactor,
-	int layout, int block,
+	float FadeStart, float FadeEnd, float MipInterpolationFactor,
+	int32 /*layout*/,
+	float ProjectionAngle,
 	UE::Math::TIntVector2<uint16> CropMin, UE::Math::TIntVector2<uint16> UncroppedSize,
-	SCRATCH_IMAGE_PROJECT* scratch )
+	FScratchImageProject* Scratch, bool bUseVectorImplementation)
 {
-    check( !pMask || pMask->GetSizeX() == pImage->GetSizeX() );
-    check( !pMask || pMask->GetSizeY() == pImage->GetSizeY() );
-    check( !pMask || pMask->GetFormat()==EImageFormat::IF_L_UBYTE );
 
-	MUTABLE_CPUPROFILER_SCOPE(ImageProjectWrapping);
+	using namespace mu::Private;
 
+    check(!pMask || pMask->GetSizeX() == pImage->GetSizeX());
+    check(!pMask || pMask->GetSizeY() == pImage->GetSizeY());
+    check(!pMask || pMask->GetFormat()==EImageFormat::IF_L_UBYTE);
 
-	constexpr bool bUseSmoothInterpolator = false;
+	MUTABLE_CPUPROFILER_SCOPE(ImageProjectCylindrical);
 
-    if ( ( pMesh->m_staticFormatFlags & (1<<SMF_PROJECTWRAPPING) ) )
+	EPixelProcessorFeatures ProcessorFeatures = EPixelProcessorFeatures::ProjectionCylindrical;
+
+	if (bUseVectorImplementation && SamplingMethod != ESamplingMethod::Point)
+	{
+		ProcessorFeatures |= EPixelProcessorFeatures::VectorizedImpl;
+	}
+
+	if (pSource->GetFormat() == EImageFormat::IF_L_UBYTE)
+	{
+		ProcessorFeatures |= EPixelProcessorFeatures::FormatL;
+	}
+	else if (pSource->GetFormat() == EImageFormat::IF_RGB_UBYTE)
+	{
+		ProcessorFeatures |= EPixelProcessorFeatures::FormatRGB;
+	}
+	else if (pSource->GetFormat() == EImageFormat::IF_RGBA_UBYTE || pSource->GetFormat() == EImageFormat::IF_BGRA_UBYTE)
+	{
+		ProcessorFeatures |= EPixelProcessorFeatures::FormatRGBA;
+	}
+
+	if (SamplingMethod == ESamplingMethod::Point)
+	{
+		ProcessorFeatures |= EPixelProcessorFeatures::SamplingPoint;
+	}
+	else if (SamplingMethod == ESamplingMethod::BiLinear)
+	{
+		ProcessorFeatures |= EPixelProcessorFeatures::SamplingLinear;
+	}
+
+	if (pMask)
+	{
+		ProcessorFeatures |= EPixelProcessorFeatures::WithMask;
+	}
+
+    if ((pMesh->m_staticFormatFlags & (1 << SMF_PROJECT)))
     {
-        // Mesh-optimised version
-        if ( pSource->GetFormat() == EImageFormat::IF_RGB_UBYTE )
-        {
-			if (SamplingMethod == ESamplingMethod::Point)
-			{
-				RasterProjectedPixelProcessor_UBYTE<3> pixelProc(pSource,
-					pImage->GetData(),
-					pMask ? pMask->GetData() : 0,
-					bIsRGBFadingEnabled, bIsAlphaFadingEnabled,
-					fadeStart, fadeEnd);
+		int32 UnusedLayout = 0;
+		int32 UnusedBlock = 0;
 
-				ImageRasterProjected_OptimisedWrapping(pMesh, pImage, pixelProc, fadeEnd, block, CropMin, UncroppedSize, scratch);
-			}
-			else if (SamplingMethod == ESamplingMethod::BiLinear)
-			{
-				RasterProjectedPixelProcessorLinear_UBYTE<3, bUseSmoothInterpolator> pixelProc(pSource,
-					pImage->GetData(),
-					pMask ? pMask->GetData() : 0,
-					bIsRGBFadingEnabled, bIsAlphaFadingEnabled,
-					fadeStart, fadeEnd, MipInterpolationFactor);
+		FImageRasterInvokeArgs RasterArgs =
+		{
+			pMesh, pImage, pSource, pMask, 
+			bIsRGBFadingEnabled, bIsAlphaFadingEnabled,
+			SamplingMethod,
+			FadeStart, FadeEnd, ProjectionAngle, MipInterpolationFactor, 
+			UnusedLayout, UnusedBlock, CropMin, UncroppedSize, Scratch
+		};
+
+		using EPPF = EPixelProcessorFeatures;
+
+		// This contains all the permutations, this is way more than what it was supported before
+		// but maybe not needed. Mask permutations can be removed safely
+		switch (ProcessorFeatures)
+		{
+		case (EPPF::FormatL | EPPF::ProjectionCylindrical | EPPF::SamplingPoint):
+		{
+			InvokeCylindricalRasterizerImpl<EPPF::FormatL | EPPF::ProjectionCylindrical | EPPF::SamplingPoint>(RasterArgs);
+			break;
+		}
+		case (EPPF::FormatRGB | EPPF::ProjectionCylindrical | EPPF::SamplingPoint):
+		{
+			InvokeCylindricalRasterizerImpl<EPPF::FormatRGB | EPPF::ProjectionCylindrical | EPPF::SamplingPoint>(RasterArgs);
+			break;
+		}
+		case (EPPF::FormatRGBA | EPPF::ProjectionCylindrical | EPPF::SamplingPoint):
+		{
+			InvokeCylindricalRasterizerImpl<EPPF::FormatRGBA | EPPF::ProjectionCylindrical | EPPF::SamplingPoint>(RasterArgs);
+			break;
+		}
 
-				ImageRasterProjected_OptimisedWrapping(pMesh, pImage, pixelProc, fadeEnd, block, CropMin, UncroppedSize, scratch);
-			}
-			else
-			{
-				check(false);
-			}
-        }
+		case (EPPF::FormatL | EPPF::ProjectionCylindrical | EPPF::SamplingLinear):
+		{
+			InvokeCylindricalRasterizerImpl<EPPF::FormatL | EPPF::ProjectionCylindrical | EPPF::SamplingLinear>(RasterArgs);
+			break;
+		}
+		case (EPPF::FormatRGB | EPPF::ProjectionCylindrical | EPPF::SamplingLinear):
+		{
+			InvokeCylindricalRasterizerImpl<EPPF::FormatRGB | EPPF::ProjectionCylindrical | EPPF::SamplingLinear>(RasterArgs);
+			break;
+		}
+		case (EPPF::FormatRGBA | EPPF::ProjectionCylindrical | EPPF::SamplingLinear):
+		{
+			InvokeCylindricalRasterizerImpl<EPPF::FormatRGBA | EPPF::ProjectionCylindrical | EPPF::SamplingLinear>(RasterArgs);
+			break;
+		}
 
-        else if ( pSource->GetFormat() == EImageFormat::IF_RGBA_UBYTE ||
-                  pSource->GetFormat() == EImageFormat::IF_BGRA_UBYTE )
-        {
-            if (pMask!=nullptr)
-            {
-				if (SamplingMethod == ESamplingMethod::Point)
-				{
-					RasterProjectedPixelProcessor_UBYTE<4, true> pixelProc(pSource,
-						pImage->GetData(),
-						pMask->GetData(),
-						bIsRGBFadingEnabled, bIsAlphaFadingEnabled,
-						fadeStart, fadeEnd);
-					ImageRasterProjected_OptimisedWrapping(pMesh, pImage, pixelProc, fadeEnd, block, CropMin, UncroppedSize, scratch);
-				}
-				else if (SamplingMethod == ESamplingMethod::BiLinear)
-				{
-					RasterProjectedPixelProcessorLinear_UBYTE<4, bUseSmoothInterpolator, true> pixelProc(pSource,
-						pImage->GetData(),
-						pMask->GetData(),
-						bIsRGBFadingEnabled, bIsAlphaFadingEnabled,
-						fadeStart, fadeEnd, MipInterpolationFactor);
-					ImageRasterProjected_OptimisedWrapping(pMesh, pImage, pixelProc, fadeEnd, block, CropMin, UncroppedSize, scratch);
-				}
-				else
-				{
-					check(false);
-				}
-            }
-            else
-            {
-				if (SamplingMethod == ESamplingMethod::Point)
-				{
-					RasterProjectedPixelProcessor_UBYTE<4, false> pixelProc(pSource,
-						pImage->GetData(),
-						nullptr,
-						bIsRGBFadingEnabled, bIsAlphaFadingEnabled,
-						fadeStart, fadeEnd);
-					ImageRasterProjected_OptimisedWrapping(pMesh, pImage, pixelProc, fadeEnd, block, CropMin, UncroppedSize, scratch);
-				}
-				else if (SamplingMethod == ESamplingMethod::BiLinear)
-				{
-					RasterProjectedPixelProcessorLinear_UBYTE<4, bUseSmoothInterpolator, false> pixelProc(pSource,
-						pImage->GetData(),
-						nullptr,
-						bIsRGBFadingEnabled, bIsAlphaFadingEnabled,
-						fadeStart, fadeEnd, MipInterpolationFactor);
-					ImageRasterProjected_OptimisedWrapping(pMesh, pImage, pixelProc, fadeEnd, block, CropMin, UncroppedSize, scratch);
-				}
-				else
-				{
-					check(false);
-				}
-            }
-        }
+		case (EPPF::FormatL | EPPF::ProjectionCylindrical | EPPF::SamplingPoint | EPPF::WithMask):
+		{
+			InvokeCylindricalRasterizerImpl<EPPF::FormatL | EPPF::ProjectionCylindrical | EPPF::SamplingPoint | EPPF::WithMask>(RasterArgs);
+			break;
+		}
+		case (EPPF::FormatRGB | EPPF::ProjectionCylindrical | EPPF::SamplingPoint | EPPF::WithMask):
+		{
+			InvokeCylindricalRasterizerImpl<EPPF::FormatRGB | EPPF::ProjectionCylindrical | EPPF::SamplingPoint | EPPF::WithMask>(RasterArgs);
+			break;
+		}
+		case (EPPF::FormatRGBA | EPPF::ProjectionCylindrical | EPPF::SamplingPoint | EPPF::WithMask):
+		{
+			InvokeCylindricalRasterizerImpl<EPPF::FormatRGBA | EPPF::ProjectionCylindrical | EPPF::SamplingPoint>(RasterArgs);
+			break;
+		}
 
-        else if ( pSource->GetFormat() == EImageFormat::IF_L_UBYTE )
-        {
-			if (SamplingMethod == ESamplingMethod::Point)
-			{
-				RasterProjectedPixelProcessor_UBYTE<1> pixelProc(pSource,
-					pImage->GetData(),
-					pMask ? pMask->GetData() : 0,
-					bIsRGBFadingEnabled, bIsAlphaFadingEnabled,
-					fadeStart, fadeEnd);
+		case (EPPF::FormatL | EPPF::ProjectionCylindrical | EPPF::SamplingLinear | EPPF::WithMask):
+		{
+			InvokeCylindricalRasterizerImpl<EPPF::FormatL | EPPF::ProjectionCylindrical | EPPF::SamplingLinear | EPPF::WithMask>(RasterArgs);
+			break;
+		}
+		case (EPPF::FormatRGB | EPPF::ProjectionCylindrical | EPPF::SamplingLinear | EPPF::WithMask):
+		{
+			InvokeCylindricalRasterizerImpl<EPPF::FormatRGB | EPPF::ProjectionCylindrical | EPPF::SamplingLinear | EPPF::WithMask>(RasterArgs);
+			break;
+		}
+		case (EPPF::FormatRGBA | EPPF::ProjectionCylindrical | EPPF::SamplingLinear | EPPF::WithMask):
+		{
+			InvokeCylindricalRasterizerImpl<EPPF::FormatRGBA | EPPF::ProjectionCylindrical | EPPF::SamplingLinear | EPPF::WithMask>(RasterArgs);
+			break;
+		}
+	
+		//case (EPPF::FormatL | EPPF::ProjectionCylindrical | EPPF::SamplingPoint | EPPF::VectorizedImpl):
+		//{
+		//	InvokeCylindricalRasterizerImpl<EPPF::FormatL | EPPF::ProjectionCylindrical | EPPF::SamplingPoint | EPPF::VectorizedImpl>(RasterArgs);
+		//	break;
+		//}
+		//case (EPPF::FormatRGB | EPPF::ProjectionCylindrical | EPPF::SamplingPoint | EPPF::VectorizedImpl):
+		//{
+		//	InvokeCylindricalRasterizerImpl<EPPF::FormatRGB | EPPF::ProjectionCylindrical | EPPF::SamplingPoint | EPPF::VectorizedImpl>(RasterArgs);
+		//	break;
+		//}
+		//case (EPPF::FormatRGBA | EPPF::ProjectionCylindrical | EPPF::SamplingPoint | EPPF::VectorizedImpl):
+		//{
+		//	InvokeCylindricalRasterizerImpl<EPPF::FormatRGBA | EPPF::ProjectionCylindrical | EPPF::SamplingPoint | EPPF::VectorizedImpl>(RasterArgs);
+		//	break;
+		//}
 
-				ImageRasterProjected_OptimisedWrapping(pMesh, pImage, pixelProc, fadeEnd, block, CropMin, UncroppedSize, scratch);
-			}
-			else if (SamplingMethod == ESamplingMethod::BiLinear)
-			{	
-				RasterProjectedPixelProcessorLinear_UBYTE<1, bUseSmoothInterpolator> pixelProc(pSource,
-					pImage->GetData(),
-					pMask ? pMask->GetData() : 0,
-					bIsRGBFadingEnabled, bIsAlphaFadingEnabled,
-					fadeStart, fadeEnd, MipInterpolationFactor);
-
-				ImageRasterProjected_OptimisedWrapping(pMesh, pImage, pixelProc, fadeEnd, block, CropMin, UncroppedSize, scratch);
-			}
-			else
-			{
-				check(false);
-			}
-        }
+		case (EPPF::FormatL | EPPF::ProjectionCylindrical | EPPF::SamplingLinear | EPPF::VectorizedImpl):
+		{
+			InvokeCylindricalRasterizerImpl<EPPF::FormatL | EPPF::ProjectionCylindrical | EPPF::SamplingLinear | EPPF::VectorizedImpl>(RasterArgs);
+			break;
+		}
+		case (EPPF::FormatRGB | EPPF::ProjectionCylindrical | EPPF::SamplingLinear | EPPF::VectorizedImpl):
+		{
+			InvokeCylindricalRasterizerImpl<EPPF::FormatRGB | EPPF::ProjectionCylindrical | EPPF::SamplingLinear | EPPF::VectorizedImpl>(RasterArgs);
+			break;
+		}
+		case (EPPF::FormatRGBA | EPPF::ProjectionCylindrical | EPPF::SamplingLinear | EPPF::VectorizedImpl):
+		{
+			InvokeCylindricalRasterizerImpl<EPPF::FormatRGBA | EPPF::ProjectionCylindrical | EPPF::SamplingLinear | EPPF::VectorizedImpl>(RasterArgs);
+			break;
+		}
 
-        else
-        {
-			check(false);
-         }
-    }
+		//case (EPPF::FormatL | EPPF::ProjectionCylindrical | EPPF::SamplingPoint | EPPF::WithMask | EPPF::VectorizedImpl):
+		//{
+		//	InvokeCylindricalRasterizerImpl<EPPF::FormatL | EPPF::ProjectionCylindrical | EPPF::SamplingPoint | EPPF::WithMask | EPPF::VectorizedImpl>(RasterArgs);
+		//	break;
+		//}
+		//case (EPPF::FormatRGB | EPPF::ProjectionCylindrical | EPPF::SamplingPoint | EPPF::WithMask | EPPF::VectorizedImpl):
+		//{
+		//	InvokeCylindricalRasterizerImpl<EPPF::FormatRGB | EPPF::ProjectionCylindrical | EPPF::SamplingPoint | EPPF::WithMask | EPPF::VectorizedImpl>(RasterArgs);
+		//	break;
+		//}
+		//case (EPPF::FormatRGBA | EPPF::ProjectionCylindrical | EPPF::SamplingPoint | EPPF::WithMask | EPPF::VectorizedImpl):
+		//{
+		//	InvokeCylindricalRasterizerImpl<EPPF::FormatRGBA | EPPF::ProjectionCylindrical | EPPF::SamplingPoint | EPPF::VectorizedImpl>(RasterArgs);
+		//	break;
+		//}
 
+		case (EPPF::FormatL | EPPF::ProjectionCylindrical | EPPF::SamplingLinear | EPPF::WithMask | EPPF::VectorizedImpl):
+		{
+			InvokeCylindricalRasterizerImpl<EPPF::FormatL | EPPF::ProjectionCylindrical | EPPF::SamplingLinear | EPPF::WithMask | EPPF::VectorizedImpl>(RasterArgs);
+			break;
+		}
+		case (EPPF::FormatRGB | EPPF::ProjectionCylindrical | EPPF::SamplingLinear | EPPF::WithMask | EPPF::VectorizedImpl):
+		{
+			InvokeCylindricalRasterizerImpl<EPPF::FormatRGB | EPPF::ProjectionCylindrical | EPPF::SamplingLinear | EPPF::WithMask>(RasterArgs);
+			break;
+		}
+		case (EPPF::FormatRGBA | EPPF::ProjectionCylindrical | EPPF::SamplingLinear | EPPF::WithMask | EPPF::VectorizedImpl):
+		{
+			InvokeCylindricalRasterizerImpl<EPPF::FormatRGBA | EPPF::ProjectionCylindrical | EPPF::SamplingLinear | EPPF::WithMask | EPPF::VectorizedImpl>(RasterArgs);
+			break;
+		}
+		default:
+		{
+			checkf(false, TEXT("Cylindrical projection raster pixel processor permutation not implemented."));
+		}
+		}
+    }
     else
     {
         check(false);
     }
-
 }
 
 float ComputeProjectedFootprintBestMip(
 	const Mesh* pMesh, const FProjector& Projector, const FVector2f& TargetSize, const FVector2f& SourceSize)
 {
-	// Compute projected mesh footprint on the traget image and extract a desired size for source image.
-	check(pMesh->GetVertexBuffers().GetElementSize(0) == sizeof(OPTIMISED_VERTEX));
-	const OPTIMISED_VERTEX* VerticesPtr = reinterpret_cast<const OPTIMISED_VERTEX*>(pMesh->GetVertexBuffers().GetBufferData(0));
+	// Compute projected mesh footprint on the traget image and extract a optimal mip for source image.
+	const bool bIsPlanarProjection = Projector.type == PROJECTOR_TYPE::PLANAR;
+	const bool bIsCylindicalProjection = Projector.type == PROJECTOR_TYPE::CYLINDRICAL;
+	const bool bIsWrappingProjection = Projector.type == PROJECTOR_TYPE::WRAPPING;
 
 	check(pMesh->GetIndexBuffers().GetElementSize(0) == sizeof(uint32));
 	const uint32* IndicesPtr = reinterpret_cast<const uint32*>(pMesh->GetIndexBuffers().GetBufferData(0));
@@ -1309,125 +1555,104 @@ float ComputeProjectedFootprintBestMip(
 		return (A.X * (B.Y - C.Y) + B.X * (C.Y - A.Y) + C.X * (A.Y - B.Y)) * 0.5f;
 	};
 
-	const bool bIsPlanarProjection = Projector.type == PROJECTOR_TYPE::PLANAR;
-
-	// TODO: evaluate other strategies to get a better mip estimate. Maybe use the median of area change?
 	check(NumIndices % 3 == 0);
-	for (uint32 I = 0; I < NumIndices; I += 3)
+	
+	// TODO: evaluate other strategies to get a better mip estimate. Maybe use the median of area change?
+	if (bIsPlanarProjection)
 	{
-		const OPTIMISED_VERTEX& A = VerticesPtr[IndicesPtr[I + 0]];
-		const OPTIMISED_VERTEX& B = VerticesPtr[IndicesPtr[I + 1]];
-		const OPTIMISED_VERTEX& C = VerticesPtr[IndicesPtr[I + 2]];
+		check(pMesh->GetVertexBuffers().GetElementSize(0) == sizeof(FOptimizedVertex));
+		const FOptimizedVertex* VerticesPtr = 
+			reinterpret_cast<const FOptimizedVertex*>(pMesh->GetVertexBuffers().GetBufferData(0));
 
-		const float TriangleSourceArea = ComputeTriangleArea(
-			FVector2f(A.pos.X, A.pos.Y), FVector2f(B.pos.X, B.pos.Y), FVector2f(C.pos.X, C.pos.Y));
+		check(NumIndices % 3 == 0);
+		for (uint32 I = 0; I < NumIndices; I += 3)
+		{
+			const FOptimizedVertex& A = VerticesPtr[IndicesPtr[I + 0]];
+			const FOptimizedVertex& B = VerticesPtr[IndicesPtr[I + 1]];
+			const FOptimizedVertex& C = VerticesPtr[IndicesPtr[I + 2]];
 
-		const float TriangleTargetArea = ComputeTriangleArea(A.uv, B.uv, C.uv);
+			const float TriangleSourceArea = ComputeTriangleArea(
+				FVector2f(A.Position.X, A.Position.Y), FVector2f(B.Position.X, B.Position.Y), FVector2f(C.Position.X, C.Position.Y));
 
-		// Set weight to zero if source or target area are close to zero to remove outliers.
-		float TriangleWeight = static_cast<float>(!FMath::IsNearlyZero(TriangleSourceArea)) *
-							   static_cast<float>(!FMath::IsNearlyZero(TriangleTargetArea));
-		if (bIsPlanarProjection)
-		{
-			// Scale contribution to total area based on the mean of the projected triangle fade factor, e.i.,
-			// the cosinus of the angle between the projection direction and the vertices normal.
-			TriangleWeight *= FMath::Clamp((A.nor.X + B.nor.X + C.nor.X) * (1.0f / 3.0f), 0.0f, 1.0f);
-		}
+			const float TriangleTargetArea = ComputeTriangleArea(A.Uv, B.Uv, C.Uv);
 
-		SourceArea += TriangleSourceArea * TriangleWeight;
-		TargetArea += TriangleTargetArea * TriangleWeight;
-	}
+			// Set weight to zero if source or target area are close to zero to remove outliers.
+			float TriangleWeight = static_cast<float>(!FMath::IsNearlyZero(TriangleSourceArea)) *
+								   static_cast<float>(!FMath::IsNearlyZero(TriangleTargetArea));
+		
+			TriangleWeight *= FMath::Clamp((A.Normal.X + B.Normal.X + C.Normal.X) * (1.0f / 3.0f), 0.0f, 1.0f);
 
-	const float SourceFootprintAreaOnTarget = TargetArea * (1.0f / SourceArea);
-	const float SourceFootprintAreaOnTargetInPixelsSquared = 
-		SourceFootprintAreaOnTarget * (TargetSize.X * TargetSize.Y);
+			SourceArea += TriangleSourceArea * TriangleWeight;
+			TargetArea += TriangleTargetArea * TriangleWeight;
+		}
+	}
+	else if (bIsCylindicalProjection)
+	{
+		check(pMesh->GetVertexBuffers().GetElementSize(0) == sizeof(FOptimizedVertex));
+		const FOptimizedVertex* VerticesPtr = 
+			reinterpret_cast<const FOptimizedVertex*>(pMesh->GetVertexBuffers().GetBufferData(0));
 
-	// Find the mip that better adapts to the footprint.
-	float BestMip = FMath::Log2(FMath::Max(1.0f, SourceSize.X * SourceSize.Y)) * 0.5f - 
-				    FMath::Log2(FMath::Max(1.0f, SourceFootprintAreaOnTargetInPixelsSquared)) * 0.5f;
+		check(NumIndices % 3 == 0);
+		for (uint32 I = 0; I < NumIndices; I += 3)
+		{
+			const FOptimizedVertex& A = VerticesPtr[IndicesPtr[I + 0]];
+			const FOptimizedVertex& B = VerticesPtr[IndicesPtr[I + 1]];
+			const FOptimizedVertex& C = VerticesPtr[IndicesPtr[I + 2]];
 
-	return BestMip;
-}
+			const float UnwindedProjectorAngleRad = FMath::UnwindRadians(Projector.projectionAngle);
+			const float OneOverProjectionAngleSafe = !FMath::IsNearlyZero(UnwindedProjectorAngleRad) 
+				? 1.0f / Projector.projectionAngle
+				: 1.0f / KINDA_SMALL_NUMBER;
 
-//-------------------------------------------------------------------------------------------------
-void ImageRasterProjectedCylindrical( const Mesh* pMesh, Image* pImage,
-	const Image* pSource, const Image* pMask,
-	bool bIsRGBFadingEnabled, bool bIsAlphaFadingEnabled,
-	float fadeStart, float fadeEnd,
-	int /*layout*/,
-	float projectionAngle,
-	UE::Math::TIntVector2<uint16> CropMin, UE::Math::TIntVector2<uint16> UncroppedSize,
-	SCRATCH_IMAGE_PROJECT* scratch )
-{
-    check( !pMask || pMask->GetSizeX() == pImage->GetSizeX() );
-    check( !pMask || pMask->GetSizeY() == pImage->GetSizeY() );
-    check( !pMask || pMask->GetFormat()==EImageFormat::IF_L_UBYTE );
+			const FVector2f PosA = FVector2f(0.5f + FMath::Atan2(A.Position.Z, -A.Position.Y) * OneOverProjectionAngleSafe, A.Position.X);
+			const FVector2f PosB = FVector2f(0.5f + FMath::Atan2(B.Position.Z, -B.Position.Y) * OneOverProjectionAngleSafe, B.Position.X);
+			const FVector2f PosC = FVector2f(0.5f + FMath::Atan2(C.Position.Z, -C.Position.Y) * OneOverProjectionAngleSafe, C.Position.X);
 
-	MUTABLE_CPUPROFILER_SCOPE(ImageProjectCylindrical);
+			const float TriangleSourceArea = ComputeTriangleArea(PosA, PosB, PosC);
+			const float TriangleTargetArea = ComputeTriangleArea(A.Uv, B.Uv, C.Uv);
 
-    if ( ( pMesh->m_staticFormatFlags & (1<<SMF_PROJECT) )
-         )
-    {
-        // Mesh-optimised version
-        if ( pSource->GetFormat() == EImageFormat::IF_RGB_UBYTE )
-        {
-            RasterCylindricalProjectedPixelProcessor_UBYTE<3> pixelProc(  pSource,
-                                                       pImage->GetData(),
-                                                       pMask ? pMask->GetData() : 0,
-                                                       fadeStart, fadeEnd,
-                                                                          projectionAngle);
+			// Set weight to zero if source or target area are close to zero to remove outliers.
+			float TriangleWeight = static_cast<float>(!FMath::IsNearlyZero(TriangleSourceArea)) *
+								   static_cast<float>(!FMath::IsNearlyZero(TriangleTargetArea));
 
-            ImageRasterProjected_Optimised( pMesh, pImage, pixelProc, fadeEnd, CropMin, UncroppedSize, scratch );
-        }
+			SourceArea += TriangleSourceArea * TriangleWeight;
+			TargetArea += TriangleTargetArea * TriangleWeight;
+		}
+	}
+	else if (bIsWrappingProjection)
+	{
+		check(pMesh->GetVertexBuffers().GetElementSize(0) == sizeof(FOptimizedVertexWrapping));
+		const FOptimizedVertexWrapping* VerticesPtr = 
+			reinterpret_cast<const FOptimizedVertexWrapping*>(pMesh->GetVertexBuffers().GetBufferData(0));
 
-        else if ( pSource->GetFormat() == EImageFormat::IF_RGBA_UBYTE ||
-                  pSource->GetFormat() == EImageFormat::IF_BGRA_UBYTE )
-        {
-            if (pMask!=nullptr)
-            {
-                RasterCylindricalProjectedPixelProcessor_UBYTE<4,true> pixelProc( pSource,
-                                                           pImage->GetData(),
-                                                           pMask->GetData(),
-                                                           fadeStart, fadeEnd,
-                                                                                  projectionAngle );
-                ImageRasterProjected_Optimised( pMesh, pImage, pixelProc, fadeEnd, CropMin, UncroppedSize, scratch );
-            }
-            else
-            {
-                RasterCylindricalProjectedPixelProcessor_UBYTE<4,false> pixelProc( pSource,
-                                                           pImage->GetData(),
-                                                           nullptr,
-                                                           fadeStart, fadeEnd,
-                                                                                   projectionAngle );
-                ImageRasterProjected_Optimised( pMesh, pImage, pixelProc, fadeEnd, CropMin, UncroppedSize, scratch );
-            }
+		check(NumIndices % 3 == 0);
+		for (uint32 I = 0; I < NumIndices; I += 3)
+		{
+			const FOptimizedVertexWrapping& A = VerticesPtr[IndicesPtr[I + 0]];
+			const FOptimizedVertexWrapping& B = VerticesPtr[IndicesPtr[I + 1]];
+			const FOptimizedVertexWrapping& C = VerticesPtr[IndicesPtr[I + 2]];
 
-        }
+			TargetArea += ComputeTriangleArea(A.Uv, B.Uv, C.Uv);
+		}
 
-        else if ( pSource->GetFormat() == EImageFormat::IF_L_UBYTE )
-        {
-            RasterCylindricalProjectedPixelProcessor_UBYTE<1> pixelProc( pSource,
-                                                     pImage->GetData(),
-                                                     pMask ? pMask->GetData() : 0,
-                                                     fadeStart, fadeEnd,
-                                                                         projectionAngle );
+		SourceArea = TargetArea;
+	}
+	else
+	{
+		check(false);
+	}
 
-            ImageRasterProjected_Optimised( pMesh, pImage, pixelProc, fadeEnd, CropMin, UncroppedSize, scratch );
-        }
+	const float SourceFootprintAreaOnTarget = TargetArea * (1.0f / SourceArea);
+	const float SourceFootprintAreaOnTargetInPixelsSquared = 
+		SourceFootprintAreaOnTarget * (TargetSize.X * TargetSize.Y);
 
-        else
-        {
-            check(false);
-        }
-    }
-    else
-    {
-        check(false);
-    }
+	// Find the mip that better adapts to the footprint.
+	float BestMip = FMath::Log2(FMath::Max(1.0f, SourceSize.X * SourceSize.Y)) * 0.5f - 
+				    FMath::Log2(FMath::Max(1.0f, SourceFootprintAreaOnTargetInPixelsSquared)) * 0.5f;
 
+	return BestMip;
 }
 
-
 //#define DEBUG_PROJECTION 1
 
 #ifdef DEBUG_PROJECTION
@@ -1646,15 +1871,15 @@ struct NeighborFace
 
 
 void getEdgeHorizontalLength( int edgeVert0, int edgeVert1, int opposedVert,
-                              const OPTIMISED_VERTEX_WRAPPING* pVertices,
+                              const FOptimizedVertexWrapping* pVertices,
                               float &out_uvSpaceLen, float &out_objSpaceLen,
                               float& out_midEdgePointFraction)
 {
 	const int oldVertices[2] = { edgeVert0, edgeVert1 };
 
-	const FVector2f& edge0_oldUVSpace = pVertices[oldVertices[0]].uv;
-	const FVector2f& edge1_oldUVSpace = pVertices[oldVertices[1]].uv;
-	const FVector2f& opposedVert_oldUVSpace = pVertices[opposedVert].uv;
+	const FVector2f& edge0_oldUVSpace = pVertices[oldVertices[0]].Uv;
+	const FVector2f& edge1_oldUVSpace = pVertices[oldVertices[1]].Uv;
+	const FVector2f& opposedVert_oldUVSpace = pVertices[opposedVert].Uv;
 
 	FVector2f edgeVector_oldUVSpace = edge1_oldUVSpace - edge0_oldUVSpace;
 	FVector2f sideVector_oldUVSpace = opposedVert_oldUVSpace - edge0_oldUVSpace;
@@ -1674,9 +1899,9 @@ void getEdgeHorizontalLength( int edgeVert0, int edgeVert1, int opposedVert,
 	out_midEdgePointFraction = midEdgePointFraction_oldUVSpace;
 
 	// Do the same in object space to be able to extract the scale of the original uv space
-	const FVector3f& edge0_objSpace = pVertices[oldVertices[0]].pos;
-	const FVector3f& edge1_objSpace = pVertices[oldVertices[1]].pos;
-	const FVector3f& opposedVert_objSpace = pVertices[opposedVert].pos;
+	const FVector3f& edge0_objSpace = pVertices[oldVertices[0]].Position;
+	const FVector3f& edge1_objSpace = pVertices[oldVertices[1]].Position;
+	const FVector3f& opposedVert_objSpace = pVertices[opposedVert].Position; // FIXME
 
 	FVector3f edgeVector_objSpace = edge1_objSpace - edge0_objSpace;
 	FVector3f sideVector_objSpace = opposedVert_objSpace - edge0_objSpace;
@@ -1724,13 +1949,13 @@ static_assert( sizeof(PROJECTED_VERTEX)==16, "Unexpected struct size" );
 
 
 //-------------------------------------------------------------------------------------------------
-void MeshProject_Optimised_Planar( const OPTIMISED_VERTEX* pVertices, int vertexCount,
-                                   const uint32* pIndices, int faceCount,
-                                   const FVector3f& projectorPosition, const FVector3f& projectorDirection,
-                                   const FVector3f& projectorSide, const FVector3f& projectorUp,
-                                   const FVector3f& projectorScale,
-                                   OPTIMISED_VERTEX* pResultVertices, int& currentVertex,
-                                   uint32* pResultIndices, int& currentIndex  )
+void MeshProject_Optimised_Planar(const FOptimizedVertex* pVertices, int vertexCount,
+                                  const uint32* pIndices, int faceCount,
+                                  const FVector3f& projectorPosition, const FVector3f& projectorDirection,
+                                  const FVector3f& projectorSide, const FVector3f& projectorUp,
+                                  const FVector3f& projectorScale,
+                                  FOptimizedVertex* pResultVertices, int& currentVertex,
+                                  uint32* pResultIndices, int& currentIndex)
 {
 	MUTABLE_CPUPROFILER_SCOPE(MeshProject_Optimised_Planar)
 
@@ -1742,10 +1967,10 @@ void MeshProject_Optimised_Planar( const OPTIMISED_VERTEX* pVertices, int vertex
 
     for ( int v=0; v<vertexCount; ++v )
     {
-        float x = FVector3f::DotProduct( pVertices[v].pos-projectorPosition, projectorSide ) / projectorScale[0] + 0.5f;
-        float y = FVector3f::DotProduct( pVertices[v].pos-projectorPosition, projectorUp ) / projectorScale[1] + 0.5f;
+        float x = FVector3f::DotProduct(pVertices[v].Position - projectorPosition, projectorSide) / projectorScale[0] + 0.5f;
+        float y = FVector3f::DotProduct(pVertices[v].Position - projectorPosition, projectorUp) / projectorScale[1] + 0.5f;
         y = 1.0f-y;
-        float z = FVector3f::DotProduct( pVertices[v].pos-projectorPosition, projectorDirection ) / projectorScale[2];
+        float z = FVector3f::DotProduct(pVertices[v].Position - projectorPosition, projectorDirection) / projectorScale[2];
 
         // Plane mask with bits for each plane discarding the vertex
         uint32 planeMask =
@@ -1787,15 +2012,15 @@ void MeshProject_Optimised_Planar( const OPTIMISED_VERTEX* pVertices, int vertex
                 {
                     pResultVertices[currentVertex] = pVertices[i];
 
-                    pResultVertices[currentVertex].pos[0] = projectedPositions[i].pos0;
-                    pResultVertices[currentVertex].pos[1] = projectedPositions[i].pos1;
-                    pResultVertices[currentVertex].pos[2] = projectedPositions[i].pos2;
+                    pResultVertices[currentVertex].Position[0] = projectedPositions[i].pos0;
+                    pResultVertices[currentVertex].Position[1] = projectedPositions[i].pos1;
+                    pResultVertices[currentVertex].Position[2] = projectedPositions[i].pos2;
 
                     // Normal is actually the fade factor
-                    float angleCos = FVector3f::DotProduct( pVertices[i].nor, projectorDirection * -1.0f );
-                    pResultVertices[currentVertex].nor[0] = angleCos;
-                    pResultVertices[currentVertex].nor[1] = angleCos;
-                    pResultVertices[currentVertex].nor[2] = angleCos;
+                    float angleCos = FVector3f::DotProduct(pVertices[i].Normal, projectorDirection * -1.0f);
+                    pResultVertices[currentVertex].Normal[0] = angleCos;
+                    pResultVertices[currentVertex].Normal[1] = angleCos;
+                    pResultVertices[currentVertex].Normal[2] = angleCos;
 
                     oldToNewVertex[i] = currentVertex++;
                 }
@@ -1808,13 +2033,13 @@ void MeshProject_Optimised_Planar( const OPTIMISED_VERTEX* pVertices, int vertex
 
 
 //-------------------------------------------------------------------------------------------------
-void MeshProject_Optimised_Cylindrical( const OPTIMISED_VERTEX* pVertices, int vertexCount,
-                                        const uint32* pIndices, int faceCount,
-                                        const FVector3f& projectorPosition, const FVector3f& projectorDirection,
-                                        const FVector3f& projectorSide, const FVector3f& projectorUp,
-                                        const FVector3f& projectorScale,
-                                        OPTIMISED_VERTEX* pResultVertices, int& currentVertex,
-                                        uint32* pResultIndices, int& currentIndex  )
+void MeshProject_Optimised_Cylindrical(const FOptimizedVertex* pVertices, int vertexCount,
+                                       const uint32* pIndices, int faceCount,
+                                       const FVector3f& projectorPosition, const FVector3f& projectorDirection,
+                                       const FVector3f& projectorSide, const FVector3f& projectorUp,
+                                       const FVector3f& projectorScale,
+                                       FOptimizedVertex* pResultVertices, int& currentVertex,
+                                       uint32* pResultIndices, int& currentIndex)
 {
 	MUTABLE_CPUPROFILER_SCOPE(MeshProject_Optimised_Cylindrical)
 
@@ -1835,7 +2060,7 @@ void MeshProject_Optimised_Cylindrical( const OPTIMISED_VERTEX* pVertices, int v
         // Cylinder is along the X axis
 
         // Project
-		vec3f relPos = vec3f(pVertices[v].pos - projectorPosition);
+		vec3f relPos = vec3f(pVertices[v].Position - projectorPosition);
         vec3f vertexPos_cylinder = worldToCylinder.Transform(relPos);
 
         // This final projection needs to be done per pixel
@@ -1881,16 +2106,16 @@ void MeshProject_Optimised_Cylindrical( const OPTIMISED_VERTEX* pVertices, int v
                 {
                     pResultVertices[currentVertex] = pVertices[i];
 
-                    pResultVertices[currentVertex].pos[0] = projectedPositions[i].pos0;
-                    pResultVertices[currentVertex].pos[1] = projectedPositions[i].pos1;
-                    pResultVertices[currentVertex].pos[2] = projectedPositions[i].pos2;
+                    pResultVertices[currentVertex].Position[0] = projectedPositions[i].pos0;
+                    pResultVertices[currentVertex].Position[1] = projectedPositions[i].pos1;
+                    pResultVertices[currentVertex].Position[2] = projectedPositions[i].pos2;
 
                     // Normal is actually the fade factor
                     //vec3 vertexCylinderDirection =
-                    float angleCos = 1.0f; //dot( pVertices[i].nor, vertexCylinderDirection * -1.0f );
-                    pResultVertices[currentVertex].nor[0] = angleCos;
-                    pResultVertices[currentVertex].nor[1] = angleCos;
-                    pResultVertices[currentVertex].nor[2] = angleCos;
+                    float angleCos = 1.0f; //dot( pVertices[i].Normal, vertexCylinderDirection * -1.0f );
+                    pResultVertices[currentVertex].Normal[0] = angleCos;
+                    pResultVertices[currentVertex].Normal[1] = angleCos;
+                    pResultVertices[currentVertex].Normal[2] = angleCos;
 
                     oldToNewVertex[i] = currentVertex++;
                 }
@@ -1903,23 +2128,24 @@ void MeshProject_Optimised_Cylindrical( const OPTIMISED_VERTEX* pVertices, int v
 
 
 //-------------------------------------------------------------------------------------------------
-void MeshProject_Optimised_Wrapping( const Mesh* pMesh,
-                                     const FVector3f& projectorPosition, const FVector3f& projectorDirection,
-                                     const FVector3f& projectorSide, const FVector3f& projectorUp,
-                                     const FVector3f& projectorScale,
-                                     OPTIMISED_VERTEX_WRAPPING* pResultVertices, int& currentVertex,
-                                     uint32* pResultIndices, int& currentIndex )
+void MeshProject_Optimised_Wrapping(const Mesh* pMesh,
+                                    const FVector3f& projectorPosition, const FVector3f& projectorDirection,
+                                    const FVector3f& projectorSide, const FVector3f& projectorUp,
+                                    const FVector3f& projectorScale,
+                                    FOptimizedVertexWrapping* pResultVertices, int& currentVertex,
+                                    uint32* pResultIndices, int& currentIndex)
 {
 	MUTABLE_CPUPROFILER_SCOPE(MeshProject_Optimised_Wrapping)
-
-    // Get the vertices
+    
+		// Get the vertices
     int vertexCount = pMesh->GetVertexCount();
-    check( pMesh->GetVertexBuffers().GetElementSize(0)==sizeof(OPTIMISED_VERTEX_WRAPPING) );
-	const OPTIMISED_VERTEX_WRAPPING* pVertices = reinterpret_cast<const OPTIMISED_VERTEX_WRAPPING*>( pMesh->GetVertexBuffers().GetBufferData(0) );
+    check(pMesh->GetVertexBuffers().GetElementSize(0) == sizeof(FOptimizedVertexWrapping));
+	const FOptimizedVertexWrapping* pVertices = 
+		reinterpret_cast<const FOptimizedVertexWrapping*>(pMesh->GetVertexBuffers().GetBufferData(0));
 
     // Get the indices
-    check( pMesh->GetIndexBuffers().GetElementSize(0)==4 );
-	const uint32* pIndices = reinterpret_cast<const uint32*>( pMesh->GetIndexBuffers().GetBufferData(0) );
+    check(pMesh->GetIndexBuffers().GetElementSize(0) == 4);
+	const uint32* pIndices = reinterpret_cast<const uint32*>(pMesh->GetIndexBuffers().GetBufferData(0));
     int faceCount = pMesh->GetFaceCount();
 
 	TArray<PROJECTED_VERTEX> projectedPositions;
@@ -1980,7 +2206,7 @@ void MeshProject_Optimised_Wrapping( const Mesh* pMesh,
 
         rayLength = (rayEnd - rayStart).Length();
 
-        bool intersects = rayIntersectsFace2(rayStart, rayEnd, pVertices[i0].pos, pVertices[i1].pos, pVertices[i2].pos,
+        bool intersects = rayIntersectsFace2(rayStart, rayEnd, pVertices[i0].Position, pVertices[i1].Position, pVertices[i2].Position,
             aux_out_intersection, out_intersected_vert, out_intersected_edge_v0, out_intersected_edge_v1, t);
 
         if (intersects && t < min_t)
@@ -1989,9 +2215,9 @@ void MeshProject_Optimised_Wrapping( const Mesh* pMesh,
             min_t = t;
             out_intersection = aux_out_intersection;
 
-            FVector3f v0 = pVertices[i0].pos;
-			FVector3f v1 = pVertices[i1].pos;
-			FVector3f v2 = pVertices[i2].pos;
+            FVector3f v0 = pVertices[i0].Position;
+			FVector3f v1 = pVertices[i1].Position;
+			FVector3f v2 = pVertices[i2].Position;
             projectionPlaneNormal = FVector3f::CrossProduct(v1 - v0, v2 - v0).GetSafeNormal();
         }
 
@@ -2150,16 +2376,16 @@ void MeshProject_Optimised_Wrapping( const Mesh* pMesh,
                 int i0 = pIndices[currentFace * 3 + 0];
                 int i1 = pIndices[currentFace * 3 + 1];
                 int i2 = pIndices[currentFace * 3 + 2];
-				FVector3f v3_0 = pVertices[i0].pos;
-				FVector3f v3_1 = pVertices[i1].pos;
-				FVector3f v3_2 = pVertices[i2].pos;
+				FVector3f v3_0 = pVertices[i0].Position;
+				FVector3f v3_1 = pVertices[i1].Position;
+				FVector3f v3_2 = pVertices[i2].Position;
                 GetBarycentricCoords(out_intersection, v3_0, v3_1, v3_2, outIntersectionBaricentric[0], outIntersectionBaricentric[1], outIntersectionBaricentric[2]);
-				FVector2f outIntersectionUV = pVertices[i0].uv * outIntersectionBaricentric[0] + pVertices[i1].uv * outIntersectionBaricentric[1] + pVertices[i2].uv * outIntersectionBaricentric[2];
+				FVector2f outIntersectionUV = pVertices[i0].Uv * outIntersectionBaricentric[0] + pVertices[i1].Uv * outIntersectionBaricentric[1] + pVertices[i2].Uv * outIntersectionBaricentric[2];
 
 				FVector4f proj4D_v0, proj4D_v1, proj4D_v2;
-                PlanarlyProjectVertex(pVertices[i0].pos, proj4D_v0, projector2, projectorPosition2, projectorDirection2, s2, u2);
-                PlanarlyProjectVertex(pVertices[i1].pos, proj4D_v1, projector2, projectorPosition2, projectorDirection2, s2, u2);
-                PlanarlyProjectVertex(pVertices[i2].pos, proj4D_v2, projector2, projectorPosition2, projectorDirection2, s2, u2);
+                PlanarlyProjectVertex(pVertices[i0].Position, proj4D_v0, projector2, projectorPosition2, projectorDirection2, s2, u2);
+                PlanarlyProjectVertex(pVertices[i1].Position, proj4D_v1, projector2, projectorPosition2, projectorDirection2, s2, u2);
+                PlanarlyProjectVertex(pVertices[i2].Position, proj4D_v2, projector2, projectorPosition2, projectorDirection2, s2, u2);
 
 				FVector2f proj_v0 = FVector2f(proj4D_v0[0], proj4D_v0[1]);
 				FVector2f proj_v1 = FVector2f(proj4D_v1[0], proj4D_v1[1]);
@@ -2175,9 +2401,9 @@ void MeshProject_Optimised_Wrapping( const Mesh* pMesh,
                 GetBarycentricCoords(FVector2f(0.f, 1.f), proj_v0, proj_v1, proj_v2, BaricentricBaseV[0], BaricentricBaseV[1], BaricentricBaseV[2]);
                 GetBarycentricCoords(FVector2f(0.f, 0.f), proj_v0, proj_v1, proj_v2, BaricentricOrigin[0], BaricentricOrigin[1], BaricentricOrigin[2]);
 
-                FVector2f OrigUVs_BaseU_EndPoint = pVertices[i0].uv * BaricentricBaseU[0] + pVertices[i1].uv * BaricentricBaseU[1] + pVertices[i2].uv * BaricentricBaseU[2];
-                FVector2f OrigUVs_BaseV_EndPoint = pVertices[i0].uv * BaricentricBaseV[0] + pVertices[i1].uv * BaricentricBaseV[1] + pVertices[i2].uv * BaricentricBaseV[2];
-                FVector2f OrigUVs_Origin = pVertices[i0].uv * BaricentricOrigin[0] + pVertices[i1].uv * BaricentricOrigin[1] + pVertices[i2].uv * BaricentricOrigin[2];
+                FVector2f OrigUVs_BaseU_EndPoint = pVertices[i0].Uv * BaricentricBaseU[0] + pVertices[i1].Uv * BaricentricBaseU[1] + pVertices[i2].Uv * BaricentricBaseU[2];
+                FVector2f OrigUVs_BaseV_EndPoint = pVertices[i0].Uv * BaricentricBaseV[0] + pVertices[i1].Uv * BaricentricBaseV[1] + pVertices[i2].Uv * BaricentricBaseV[2];
+                FVector2f OrigUVs_Origin = pVertices[i0].Uv * BaricentricOrigin[0] + pVertices[i1].Uv * BaricentricOrigin[1] + pVertices[i2].Uv * BaricentricOrigin[2];
 
                 FVector2f OrigUVs_BaseU = OrigUVs_BaseU_EndPoint - OrigUVs_Origin;
                 FVector2f OrigUVs_BaseV_Aux = OrigUVs_BaseV_EndPoint - OrigUVs_Origin;
@@ -2193,7 +2419,7 @@ void MeshProject_Optimised_Wrapping( const Mesh* pMesh,
 #endif
 
                 OrigUVs_Origin = outIntersectionUV;
-                GetBarycentricCoords(outIntersectionUV, pVertices[i0].uv, pVertices[i1].uv, pVertices[i2].uv, outIntersectionBaricentric[0], outIntersectionBaricentric[1], outIntersectionBaricentric[2]);
+                GetBarycentricCoords(outIntersectionUV, pVertices[i0].Uv, pVertices[i1].Uv, pVertices[i2].Uv, outIntersectionBaricentric[0], outIntersectionBaricentric[1], outIntersectionBaricentric[2]);
 
                 for (int i = 0; i < 3; ++i)
                 {
@@ -2207,7 +2433,7 @@ void MeshProject_Optimised_Wrapping( const Mesh* pMesh,
                     processedVertices.Add(collapsedVert);
 
                     //PlanarlyProjectVertex(pVertices[collapsedVert].pos, projectedPositions[collapsedVert], projector2, projectorPosition2, projectorDirection2, s2, u2);
-                    FVector2f projectedVertex = ChangeBase2D(pVertices[v].uv, OrigUVs_Origin, OrigUVs_BaseU, OrigUVs_BaseV);
+                    FVector2f projectedVertex = ChangeBase2D(pVertices[v].Uv, OrigUVs_Origin, OrigUVs_BaseU, OrigUVs_BaseV);
 
                     projectedPositions[collapsedVert].pos0 = projectedVertex[0];
                     projectedPositions[collapsedVert].pos1 = projectedVertex[1];
@@ -2374,9 +2600,9 @@ void MeshProject_Optimised_Wrapping( const Mesh* pMesh,
                         float objSpaceToCurrentUVSpaceConversion = currentFace_uvSpaceLen / currentFace_objSpaceLen;
                         float previousFaceWidthInCurrentUVspace = previousFace_objSpaceLen * objSpaceToCurrentUVSpaceConversion;
 
-                        const FVector2f& edge0_currentFace_oldUVSpace = pVertices[oldVertices[0]].uv;
-                        const FVector2f& edge1_currentFace_oldUVSpace = pVertices[oldVertices[1]].uv;
-                        const FVector2f& newVertex_currentFace_oldUVSpace = pVertices[newVertex].uv;
+                        const FVector2f& edge0_currentFace_oldUVSpace = pVertices[oldVertices[0]].Uv;
+                        const FVector2f& edge1_currentFace_oldUVSpace = pVertices[oldVertices[1]].Uv;
+                        const FVector2f& newVertex_currentFace_oldUVSpace = pVertices[newVertex].Uv;
 
                         FVector2f edge_currentFace_oldUVSpace = edge1_currentFace_oldUVSpace - edge0_currentFace_oldUVSpace;
                         FVector2f midpoint_currentFace_oldUVSpace = edge0_currentFace_oldUVSpace + edge_currentFace_oldUVSpace * midEdgePointFraction_currentFace;
@@ -2411,16 +2637,16 @@ void MeshProject_Optimised_Wrapping( const Mesh* pMesh,
                     else
                     {
                         // If there's no UV island border, just use the previous uvs
-                        pv0 = pVertices[pi0].uv;
-                        pv1 = pVertices[pi1].uv;
-                        pv2 = pVertices[pi2].uv;
+                        pv0 = pVertices[pi0].Uv;
+                        pv1 = pVertices[pi1].Uv;
+                        pv2 = pVertices[pi2].Uv;
 
                         pv0_proj = projectedPositions[pi0].xy();
                         pv1_proj = projectedPositions[pi1].xy();
                         pv2_proj = projectedPositions[pi2].xy();
                     }
 
-					FVector2f v2 = pVertices[newVertex].uv;
+					FVector2f v2 = pVertices[newVertex].Uv;
 
                     // New vertex baricentric coords in respect to old triangle
                     float a, b, c;
@@ -2511,7 +2737,7 @@ void MeshProject_Optimised_Wrapping( const Mesh* pMesh,
             for (int i = 0; i < 3; ++i)
             {
                 int v = pIndices[currentFace * 3 + i];
-                FVector3f r = pVertices[v].pos - out_intersection;
+                FVector3f r = pVertices[v].Position - out_intersection;
                 float squaredDist = FVector3f::DotProduct(r, r);
 
                 if (squaredDist <= maxDistSquared / 4.f) // Is it inside?
@@ -2609,19 +2835,19 @@ void MeshProject_Optimised_Wrapping( const Mesh* pMesh,
                 {
                     pResultVertices[currentVertex] = pVertices[i];
 
-                    pResultVertices[currentVertex].pos[0] = projectedPositions[i].pos0;
-                    pResultVertices[currentVertex].pos[1] = projectedPositions[i].pos1;
-                    pResultVertices[currentVertex].pos[2] = projectedPositions[i].pos2;
+                    pResultVertices[currentVertex].Position[0] = projectedPositions[i].pos0;
+                    pResultVertices[currentVertex].Position[1] = projectedPositions[i].pos1;
+                    pResultVertices[currentVertex].Position[2] = projectedPositions[i].pos2;
 
                     // Normal is actually the fade factor
                     int step = faceStep[f];
                     const float maxGradient = 10.f;
                     float stepGradient = step / maxGradient;
                     stepGradient = stepGradient > maxGradient ? maxGradient : stepGradient;
-                    float angleCos = stepGradient; //1.f; // dot(pVertices[i].nor, projectorDirection * -1.0f);
-                    pResultVertices[currentVertex].nor[0] = angleCos;
-                    pResultVertices[currentVertex].nor[1] = angleCos;
-                    pResultVertices[currentVertex].nor[2] = angleCos;
+                    float angleCos = stepGradient; //1.f; // dot(pVertices[i].Normal, projectorDirection * -1.0f);
+                    pResultVertices[currentVertex].Normal[0] = angleCos;
+                    pResultVertices[currentVertex].Normal[1] = angleCos;
+                    pResultVertices[currentVertex].Normal[2] = angleCos;
 
                     oldToNewVertex[i] = currentVertex++;
                 }
@@ -2637,7 +2863,7 @@ void MeshProject_Optimised_Wrapping( const Mesh* pMesh,
 
     //for (int i = 0; i < currentVertex; ++i)
     //{
-    //	LogDebug("v %f %f %f\n", pResultVertices[i].pos[0], pResultVertices[i].pos[1], pResultVertices[i].pos[2]);
+    //	LogDebug("v %f %f %f\n", pResultVertices[i].Position[0], pResultVertices[i].Position[1], pResultVertices[i].Position[2]);
     //}
 
     //for (int i = 0; i < currentIndex; ++i)
@@ -2649,12 +2875,12 @@ void MeshProject_Optimised_Wrapping( const Mesh* pMesh,
     {
         for (int i = 0; i < currentVertex; ++i)
         {
-            outfile << "v " << pResultVertices[i].pos[0] << " " << pResultVertices[i].pos[1] << " " << pResultVertices[i].pos[2] << "\n";
+            outfile << "v " << pResultVertices[i].Position[0] << " " << pResultVertices[i].Position[1] << " " << pResultVertices[i].Position[2] << "\n";
         }
 
         for (int i = 0; i < currentVertex; ++i)
         {
-            outfile << "vt " << pResultVertices[i].pos[0] << " " << 1.f - pResultVertices[i].pos[1] << "\n";
+            outfile << "vt " << pResultVertices[i].Position[0] << " " << 1.f - pResultVertices[i].Position[1] << "\n";
         }
 
         for (int i = 0; i < currentIndex / 3; ++i)
@@ -2665,9 +2891,9 @@ void MeshProject_Optimised_Wrapping( const Mesh* pMesh,
 
             outfile << "f " << i0 << "/" << i0 << " " << i1 << "/" << i1 << " " << i2 << "/" << i2 << "\n";
 
-            FVector2f v0 = pResultVertices[i0].pos.xy();
-            FVector2f v1 = pResultVertices[i1].pos.xy();
-            FVector2f v2 = pResultVertices[i2].pos.xy();
+            FVector2f v0 = pResultVertices[i0].Position.xy();
+            FVector2f v1 = pResultVertices[i1].Position.xy();
+            FVector2f v2 = pResultVertices[i2].Position.xy();
 
             outfile << "# Triangle " << i << ", area = " << getTriangleArea(v0, v1, v2) << "\n";
         }
@@ -2712,11 +2938,11 @@ void MeshProject_Optimised(Mesh* Result, const Mesh* pMesh, const FProjector& pr
         Result->GetVertexBuffers().SetElementCount(vertexCount);
         Result->GetIndexBuffers().SetElementCount(indexCount);
 		uint32* pResultIndices = reinterpret_cast<uint32*>(Result->GetIndexBuffers().GetBufferData(0));
-		OPTIMISED_VERTEX* pResultVertices = reinterpret_cast<OPTIMISED_VERTEX*>(Result->GetVertexBuffers().GetBufferData(0));
+		FOptimizedVertex* pResultVertices = reinterpret_cast<FOptimizedVertex*>(Result->GetVertexBuffers().GetBufferData(0));
 
         // Get the vertices
-        check(pMesh->GetVertexBuffers().GetElementSize(0)==sizeof(OPTIMISED_VERTEX));
-		const OPTIMISED_VERTEX* pVertices = reinterpret_cast<const OPTIMISED_VERTEX*>(pMesh->GetVertexBuffers().GetBufferData(0));
+        check(pMesh->GetVertexBuffers().GetElementSize(0)==sizeof(FOptimizedVertex));
+		const FOptimizedVertex* pVertices = reinterpret_cast<const FOptimizedVertex*>(pMesh->GetVertexBuffers().GetBufferData(0));
 
         // Get the indices
         check(pMesh->GetIndexBuffers().GetElementSize(0) == 4);
@@ -2754,7 +2980,7 @@ void MeshProject_Optimised(Mesh* Result, const Mesh* pMesh, const FProjector& pr
         Result->GetVertexBuffers().SetElementCount(vertexCount);
         Result->GetIndexBuffers().SetElementCount(indexCount);
 		uint32* pResultIndices = reinterpret_cast<uint32*>(Result->GetIndexBuffers().GetBufferData(0));
-		OPTIMISED_VERTEX_WRAPPING* pResultVertices = reinterpret_cast<OPTIMISED_VERTEX_WRAPPING*>(Result->GetVertexBuffers().GetBufferData(0));
+		FOptimizedVertexWrapping* pResultVertices = reinterpret_cast<FOptimizedVertexWrapping*>(Result->GetVertexBuffers().GetBufferData(0));
 
         MeshProject_Optimised_Wrapping(pMesh,
                                        projectorPosition, projectorDirection,
diff --git a/Engine/Plugins/Experimental/Mutable/Source/MutableRuntime/Private/MuR/OpImageProject.h b/Engine/Plugins/Experimental/Mutable/Source/MutableRuntime/Private/MuR/OpImageProject.h
--- a/Engine/Plugins/Experimental/Mutable/Source/MutableRuntime/Private/MuR/OpImageProject.h
+++ b/Engine/Plugins/Experimental/Mutable/Source/MutableRuntime/Private/MuR/OpImageProject.h
@@ -2,54 +2,57 @@
 
 #pragma once
 
-#include "MuR/Mesh.h"
-#include "MuR/Image.h"
+
+#include "Math/IntVector.h"
 
 namespace mu
 {
 	class Image;
+	class Mesh;
 	struct FProjector;
+	enum class ESamplingMethod : uint8;
 	template <int NUM_INTERPOLATORS> class RasterVertex;
 
-    struct SCRATCH_IMAGE_PROJECT
-    {
-        TArray< RasterVertex<4> > vertices;
-		TArray<uint8> culledVertex;
-    };
+	struct FScratchImageProject
+	{
+		TArray<RasterVertex<4>> Vertices;
+		TArray<uint8> CulledVertex;
+	};
 
     extern void ImageRasterProjectedPlanar( const Mesh* pMesh, Image* pTargetImage,
 		const Image* pSource, const Image* pMask,
 		bool bIsRGBFadingEnabled, bool bIsAlphaFadingEnabled,
 		ESamplingMethod SamplingMethod,
-		float fadeStart, float fadeEnd, float MipInterpolationFactor,
-		int layout, int block,
+		float FadeStart, float FadeEnd, float MipInterpolationFactor,
+		int Layout, int32 Block,
 		UE::Math::TIntVector2<uint16> CropMin, UE::Math::TIntVector2<uint16> UncroppedSize,
-		SCRATCH_IMAGE_PROJECT* scratch, bool bUseVectorImplementation = false);
+		FScratchImageProject* scratch, bool bUseVectorImplementation = false);
 
     extern void ImageRasterProjectedCylindrical( const Mesh* pMesh, Image* pTargetImage,
 		const Image* pSource, const Image* pMask,
 		bool bIsRGBFadingEnabled, bool bIsAlphaFadingEnabled,
-		float fadeStart, float fadeEnd,
-		int layout,
-		float projectionAngle,
+		ESamplingMethod SamplingMethod,
+		float FadeStart, float FadeEnd, float MipInterpolationFactor,
+		int32 Layout,
+		float ProjectionAngle,
 		UE::Math::TIntVector2<uint16> CropMin, UE::Math::TIntVector2<uint16> UncroppedSize,
-		SCRATCH_IMAGE_PROJECT* scratch );
+		FScratchImageProject* Scratch, bool bUseVectorImplementation = false);
 
     extern void ImageRasterProjectedWrapping( const Mesh* pMesh, Image* pTargetImage,
 		const Image* pSource, const Image* pMask,
 		bool bIsRGBFadingEnabled, bool bIsAlphaFadingEnabled,
 		ESamplingMethod SamplingMethod,
-		float fadeStart, float fadeEnd, float MipInterpolationFactor,
-		int layout, int block,
+		float FadeStart, float FadeEnd, float MipInterpolationFactor,
+		int32 Layout, int32 Block,
 		UE::Math::TIntVector2<uint16> CropMin, UE::Math::TIntVector2<uint16> UncroppedSize,
-		SCRATCH_IMAGE_PROJECT* scratch );
+		FScratchImageProject* Scratch, bool bUseVectorImplementation = false);
 
 	extern float ComputeProjectedFootprintBestMip(
 			const Mesh* pMesh, const FProjector& Projector, const FVector2f& TargetSize, const FVector2f& SourceSize);
 
-    extern void MeshProject(Mesh* Result, const Mesh* pMesh, const FProjector& projector, bool& bOutSuccess);
+    extern void MeshProject(Mesh* Result, const Mesh* pMesh, const FProjector& Projector, bool& bOutSuccess);
 
-	MUTABLERUNTIME_API extern void CreateMeshOptimisedForProjection(Mesh* Result, int layout);
-	MUTABLERUNTIME_API extern void CreateMeshOptimisedForWrappingProjection(Mesh* Result, int layout);
+	MUTABLERUNTIME_API extern void CreateMeshOptimisedForProjection(Mesh* Result, int32 Layout);
+	MUTABLERUNTIME_API extern void CreateMeshOptimisedForWrappingProjection(Mesh* Result, int32 Layout);
 
 }
diff --git a/Engine/Plugins/Experimental/Mutable/Source/MutableRuntime/Private/MuR/OpImageRasterMesh.h b/Engine/Plugins/Experimental/Mutable/Source/MutableRuntime/Private/MuR/OpImageRasterMesh.h
--- a/Engine/Plugins/Experimental/Mutable/Source/MutableRuntime/Private/MuR/OpImageRasterMesh.h
+++ b/Engine/Plugins/Experimental/Mutable/Source/MutableRuntime/Private/MuR/OpImageRasterMesh.h
@@ -15,10 +15,15 @@ namespace mu
 	class WhitePixelProcessor
 	{
 	public:
-		inline void ProcessPixel( unsigned char* pBufferPos, float[1] ) const
+		inline void ProcessPixel(uint8* pBufferPos, float[1]) const
 		{
 			pBufferPos[0] = 255;
 		}
+
+		inline void operator()(uint8* BufferPos, float Interpolators[1]) const
+		{
+			ProcessPixel(BufferPos, Interpolators);
+		}
 	};
 
 
@@ -84,7 +89,8 @@ namespace mu
 				vertices, indices, pImage, sizeX, sizeY, pixelProc
 			] (int32 f)
 			{
-				Triangle(pImage->GetData(), pImage->GetDataSize(),
+				constexpr int32 NumInterpolators = 1;
+				Triangle<NumInterpolators>(pImage->GetData(), pImage->GetDataSize(),
 					sizeX, sizeY,
 					1,
 					vertices[indices[f * 3 + 0]],
@@ -122,7 +128,8 @@ namespace mu
 				// TODO: Select faces outside for loop?
 				if (blocks[indices[f * 3 + 0]] == BlockId)
 				{
-					Triangle(pImage->GetData(), pImage->GetDataSize(),
+					constexpr int32 NumInterpolators = 1;
+					Triangle<NumInterpolators>(pImage->GetData(), pImage->GetDataSize(),
 						sizeX, sizeY,
 						1,
 						vertices[indices[f * 3 + 0]],
diff --git a/Engine/Plugins/Experimental/Mutable/Source/MutableRuntime/Private/MuR/Raster.h b/Engine/Plugins/Experimental/Mutable/Source/MutableRuntime/Private/MuR/Raster.h
--- a/Engine/Plugins/Experimental/Mutable/Source/MutableRuntime/Private/MuR/Raster.h
+++ b/Engine/Plugins/Experimental/Mutable/Source/MutableRuntime/Private/MuR/Raster.h
@@ -2,7 +2,9 @@
 
 #pragma once
 
+#include "HAL/Platform.h"
 #include "Math/UnrealMathUtility.h"
+#include "Templates/Function.h"
 
 namespace mu
 {
@@ -27,25 +29,33 @@ namespace mu
 	class ColorPixelProcessor
 	{
 	public:
-		inline void ProcessPixel( unsigned char* pBufferPos, float interpolatedValues[3] ) const
+		inline void ProcessPixel(uint8* pBufferPos, float interpolatedValues[3]) const
 		{
-			pBufferPos[0] = (unsigned char)(interpolatedValues[0]*0xFF);
-			pBufferPos[1] = (unsigned char)(interpolatedValues[1]*0xFF);
-			pBufferPos[2] = (unsigned char)(interpolatedValues[2]*0xFF);
+			pBufferPos[0] = (uint8)(interpolatedValues[0]*0xFF);
+			pBufferPos[1] = (uint8)(interpolatedValues[1]*0xFF);
+			pBufferPos[2] = (uint8)(interpolatedValues[2]*0xFF);
+		}
+
+		inline void operator()(uint8* BufferPos, float InterpolatedValues[3]) const
+		{
+			ProcessPixel(BufferPos, InterpolatedValues);
 		}
 	};
 
+	template<int32 NumInterpolatedValues>
+	using TTriangleRasterPixelProcRefType = TFunctionRef<void(uint8*, float[NumInterpolatedValues])>;
+
 	//! Templatized rasterizer. It is not very efficient on the interpolating side, but the cost
     //! per pixel is supposed to be huge compared to the raster process (for lightmaps or
     //! normalmaps).
 	//! TODO Optimize anyway.
 	//! TODO Option for two sided? right now it culls back facing.
-    template<class PIXEL_PROCESSOR, int NUM_INTERPOLATORS, typename INT=int64_t>
+    template<int32 NUM_INTERPOLATORS, typename INT=int64>
 	void Triangle( uint8* Buffer, int32 BufferSize, int Width, int Height, int PixelSize
 				 , const RasterVertex<NUM_INTERPOLATORS> &av1
 				 , const RasterVertex<NUM_INTERPOLATORS> &av2
 				 , const RasterVertex<NUM_INTERPOLATORS> &av3
-				 , PIXEL_PROCESSOR &processor
+				 , TTriangleRasterPixelProcRefType<NUM_INTERPOLATORS> PixelProc
 				 , bool cullBackface = true )
 	{
 		check(Width * Height * PixelSize <= BufferSize);
@@ -190,7 +200,7 @@ namespace mu
 							}
 
 	                        // Process pixel
-							processor.ProcessPixel( &CurrentBuffer[ix*PixelSize], interpolatedValues );
+							PixelProc(&CurrentBuffer[ix*PixelSize], interpolatedValues);
 	                    }
 
 						CurrentBuffer += Stride;
@@ -228,7 +238,7 @@ namespace mu
 															+ beta * v3.interpolators[ii];
 								}
 
-								processor.ProcessPixel( &CurrentBuffer[ix*PixelSize], interpolatedValues );
+								PixelProc(&CurrentBuffer[ix*PixelSize], interpolatedValues);
 	                        }
 
 	                        CX1 -= FDY12;
@@ -248,5 +258,4 @@ namespace mu
 			Buffer += q * Stride;
 	    }
 	}
-
 }
