#rnx

diff --git a/Engine/Plugins/Experimental/Mutable/Source/CustomizableObjectEditor/Private/MuCOE/CustomizableObjectInstanceEditor.cpp b/Engine/Plugins/Experimental/Mutable/Source/CustomizableObjectEditor/Private/MuCOE/CustomizableObjectInstanceEditor.cpp
--- a/Engine/Plugins/Experimental/Mutable/Source/CustomizableObjectEditor/Private/MuCOE/CustomizableObjectInstanceEditor.cpp
+++ b/Engine/Plugins/Experimental/Mutable/Source/CustomizableObjectEditor/Private/MuCOE/CustomizableObjectInstanceEditor.cpp
@@ -334,7 +334,7 @@ void FCustomizableObjectInstanceEditor::UpdatePreviewVisibility()
 
 void FCustomizableObjectInstanceEditor::SaveAsset_Execute()
 {
-	if (CustomizableObjectInstance)
+	if (!CustomizableObjectInstance)
 	{
 		return;
 	}
