[Mutable] Fix static analyzer uninitialized memory warning.

#rnx

diff --git a/Engine/Plugins/Experimental/Mutable/Source/MutableRuntime/Private/MuR/OpImageProject.cpp b/Engine/Plugins/Experimental/Mutable/Source/MutableRuntime/Private/MuR/OpImageProject.cpp
--- a/Engine/Plugins/Experimental/Mutable/Source/MutableRuntime/Private/MuR/OpImageProject.cpp
+++ b/Engine/Plugins/Experimental/Mutable/Source/MutableRuntime/Private/MuR/OpImageProject.cpp
@@ -536,6 +536,8 @@ namespace Private
 					};
 
 					FPixelData Result;
+					FMemory::Memzero(Result);
+
 					if constexpr (EnumHasAnyFlags(Features, EPixelProcessorFeatures::SamplingLinear))
 					{
 						auto SampleImageBilinear = [&](FVector2f Uv, FUInt16Vector2 Size, const uint8* DataPtr)
