[mutable] Prevent double loading of source image data.
Enforce the use of image proxies in tables.
Apply Unreal's coding standards to mutable's Table class.
#rnx

diff --git a/Engine/Plugins/Experimental/Mutable/Source/CustomizableObjectEditor/Private/MuCOE/CustomizableObjectCompiler.cpp b/Engine/Plugins/Experimental/Mutable/Source/CustomizableObjectEditor/Private/MuCOE/CustomizableObjectCompiler.cpp
--- a/Engine/Plugins/Experimental/Mutable/Source/CustomizableObjectEditor/Private/MuCOE/CustomizableObjectCompiler.cpp
+++ b/Engine/Plugins/Experimental/Mutable/Source/CustomizableObjectEditor/Private/MuCOE/CustomizableObjectCompiler.cpp
@@ -862,7 +862,8 @@ void FCustomizableObjectCompiler::UpdatePendingTextureConversion(bool UseTimeLim
 						int32 ColumnIndx = ArrayTextureUnrealToMutableTask[j].TableColumn;
 						int32 RowIndx = ArrayTextureUnrealToMutableTask[j].TableRow;
 
-						ArrayTextureUnrealToMutableTask[j].TableNode->SetCell(ColumnIndx, RowIndx, Image.get());
+						mu::Ptr<mu::ResourceProxyMemory<mu::Image>> ImageProxy = new mu::ResourceProxyMemory<mu::Image>(Image);
+						ArrayTextureUnrealToMutableTask[j].TableNode->SetCell(ColumnIndx, RowIndx, ImageProxy.get());
 
 						ArrayTextureUnrealToMutableTask[j].Texture = nullptr;
 					}
diff --git a/Engine/Plugins/Experimental/Mutable/Source/CustomizableObjectEditor/Private/MuCOE/GenerateMutableSource/GenerateMutableSource.h b/Engine/Plugins/Experimental/Mutable/Source/CustomizableObjectEditor/Private/MuCOE/GenerateMutableSource/GenerateMutableSource.h
--- a/Engine/Plugins/Experimental/Mutable/Source/CustomizableObjectEditor/Private/MuCOE/GenerateMutableSource/GenerateMutableSource.h
+++ b/Engine/Plugins/Experimental/Mutable/Source/CustomizableObjectEditor/Private/MuCOE/GenerateMutableSource/GenerateMutableSource.h
@@ -566,7 +566,7 @@ struct FMutableGraphGenerationContext
 		FKey Key;
 
 		/** Generated mesh. */
-		mu::MeshPtr Generated;
+		mu::Ptr<mu::Mesh> Generated;
 	};
 	TArray<FGeneratedMeshData> GeneratedMeshes;
 
@@ -575,7 +575,7 @@ struct FMutableGraphGenerationContext
 	TArray<EMutableMeshConversionFlags> MeshGenerationFlags;
 
 	/** Find a mesh if already generated for a given source and flags. */
-	mu::MeshPtr FindGeneratedMesh(const FGeneratedMeshData::FKey& Key);
+	mu::Ptr<mu::Mesh> FindGeneratedMesh(const FGeneratedMeshData::FKey& Key);
 
 	/** Adds to ParameterNamesMap the node Node to the array of elements with name Name */
 	void AddParameterNameUnique(const UCustomizableObjectNode* Node, FString Name);
@@ -791,7 +791,7 @@ mu::NodeMeshApplyPosePtr CreateNodeMeshApplyPose(FMutableGraphGenerationContext&
 /** Adds Tag to MutableMesh uniquely, returns the index were the tag has been inserted or the index where an intance of the tag has been found */
 int32 AddTagToMutableMeshUnique(mu::Mesh& MutableMesh, const FString& Tag);
 
-void AddSocketTagsToMesh(const USkeletalMesh* SourceMesh, mu::MeshPtr MutableMesh, FMutableGraphGenerationContext& GenerationContext);
+void AddSocketTagsToMesh(const USkeletalMesh* SourceMesh, mu::Ptr<mu::Mesh> MutableMesh, FMutableGraphGenerationContext& GenerationContext);
 
 // Generates the tag for an animation instance
 FString GenerateAnimationInstanceTag(const FString& AnimInstance, const FName& SlotIndex);
diff --git a/Engine/Plugins/Experimental/Mutable/Source/CustomizableObjectEditor/Private/MuCOE/GenerateMutableSource/GenerateMutableSourceMesh.cpp b/Engine/Plugins/Experimental/Mutable/Source/CustomizableObjectEditor/Private/MuCOE/GenerateMutableSource/GenerateMutableSourceMesh.cpp
--- a/Engine/Plugins/Experimental/Mutable/Source/CustomizableObjectEditor/Private/MuCOE/GenerateMutableSource/GenerateMutableSourceMesh.cpp
+++ b/Engine/Plugins/Experimental/Mutable/Source/CustomizableObjectEditor/Private/MuCOE/GenerateMutableSource/GenerateMutableSourceMesh.cpp
@@ -2672,7 +2672,7 @@ mu::NodeMeshPtr GenerateMorphMesh(const UEdGraphPin* Pin,
 
 			if (ColumnIndex == INDEX_NONE)
 			{
-				ColumnIndex = Table->AddColumn(StringCast<ANSICHAR>(*ColumnName).Get(), mu::TABLE_COLUMN_TYPE::TCT_MESH);
+				ColumnIndex = Table->AddColumn(StringCast<ANSICHAR>(*ColumnName).Get(), mu::ETableColumnType::Mesh);
 			}
 
 			mu::MeshPtr MorphedSourceTableMesh = BuildMorphedMutableMesh(Pin, TypedNodeMorphs[MorphIndex].MorphTargetName, GenerationContext, bOnlyConnectedLOD, RowName);
diff --git a/Engine/Plugins/Experimental/Mutable/Source/CustomizableObjectEditor/Private/MuCOE/GenerateMutableSource/GenerateMutableSourceMesh.h b/Engine/Plugins/Experimental/Mutable/Source/CustomizableObjectEditor/Private/MuCOE/GenerateMutableSource/GenerateMutableSourceMesh.h
--- a/Engine/Plugins/Experimental/Mutable/Source/CustomizableObjectEditor/Private/MuCOE/GenerateMutableSource/GenerateMutableSourceMesh.h
+++ b/Engine/Plugins/Experimental/Mutable/Source/CustomizableObjectEditor/Private/MuCOE/GenerateMutableSource/GenerateMutableSourceMesh.h
@@ -48,14 +48,14 @@ void GetLODAndSectionForAutomaticLODs(const FMutableGraphGenerationContext& Cont
  * @param GenerationContext 
  * @param CurrentNode 
  * @return Mutable Mesh. Nullptr if there has been an error. Empty mesh if the Skeletal Mesh does not contain the requested LOD + Section. */
-mu::MeshPtr ConvertSkeletalMeshToMutable(const USkeletalMesh* InSkeletalMesh, const TSoftClassPtr<UAnimInstance>& AnimBp,
+mu::Ptr<mu::Mesh> ConvertSkeletalMeshToMutable(const USkeletalMesh* InSkeletalMesh, const TSoftClassPtr<UAnimInstance>& AnimBp,
                                          int32 LODIndexConnected,  int32 SectionIndexConnected,
                                          int32 LODIndex, int32 SectionIndex,
                                          FMutableGraphGenerationContext& GenerationContext,
                                          const UCustomizableObjectNode* CurrentNode);
 
 
-mu::MeshPtr ConvertStaticMeshToMutable(const UStaticMesh* StaticMesh, int32 LODIndex, int32 SectionIndex,
+mu::Ptr<mu::Mesh> ConvertStaticMeshToMutable(const UStaticMesh* StaticMesh, int32 LODIndex, int32 SectionIndex,
 	FMutableGraphGenerationContext& GenerationContext, const UCustomizableObjectNode* CurrentNode);
 
 
@@ -71,13 +71,13 @@ mu::MeshPtr ConvertStaticMeshToMutable(const UStaticMesh* StaticMesh, int32 LODI
  * @param GenerationContext 
  * @param CurrentNode 
  * @return Mutable Mesh. Nullptr if there has been an error. Empty mesh if the Skeletal Mesh does not contain the requested LOD + Section. */
-mu::MeshPtr GenerateMutableMesh(const UObject* Mesh, const TSoftClassPtr<UAnimInstance>& AnimBp,
+mu::Ptr<mu::Mesh> GenerateMutableMesh(const UObject* Mesh, const TSoftClassPtr<UAnimInstance>& AnimBp,
                                 int32 LODIndexConnected, int32 SectionIndexConnected,
                                 int32 LODIndex, int32 SectionIndex, const FString& MeshUniqueTags,
                                 FMutableGraphGenerationContext& GenerationContext, const UCustomizableObjectNode* CurrentNode);
 
 
-mu::MeshPtr BuildMorphedMutableMesh(const UEdGraphPin* BaseSourcePin, const FString& MorphTargetName,
+mu::Ptr<mu::Mesh> BuildMorphedMutableMesh(const UEdGraphPin* BaseSourcePin, const FString& MorphTargetName,
 	FMutableGraphGenerationContext& GenerationContext, bool bOnlyConnectedLOD, const FName& RowName = "");
 
 
diff --git a/Engine/Plugins/Experimental/Mutable/Source/CustomizableObjectEditor/Private/MuCOE/GenerateMutableSource/GenerateMutableSourceTable.cpp b/Engine/Plugins/Experimental/Mutable/Source/CustomizableObjectEditor/Private/MuCOE/GenerateMutableSource/GenerateMutableSourceTable.cpp
--- a/Engine/Plugins/Experimental/Mutable/Source/CustomizableObjectEditor/Private/MuCOE/GenerateMutableSource/GenerateMutableSourceTable.cpp
+++ b/Engine/Plugins/Experimental/Mutable/Source/CustomizableObjectEditor/Private/MuCOE/GenerateMutableSource/GenerateMutableSourceTable.cpp
@@ -44,7 +44,7 @@ bool FillTableColumn(const UCustomizableObjectNodeTable* TableNode,	mu::TablePtr
 
 				if (CurrentColumn == -1)
 				{
-					CurrentColumn = MutableTable->AddColumn(StringCast<ANSICHAR>(*MutableColumnName).Get(), mu::TABLE_COLUMN_TYPE::TCT_MESH);
+					CurrentColumn = MutableTable->AddColumn(StringCast<ANSICHAR>(*MutableColumnName).Get(), mu::ETableColumnType::Mesh);
 				}
 
 				mu::MeshPtr EmptySkeletalMesh = nullptr;
@@ -215,7 +215,7 @@ bool FillTableColumn(const UCustomizableObjectNodeTable* TableNode,	mu::TablePtr
 
 			if (CurrentColumn == -1)
 			{
-				CurrentColumn = MutableTable->AddColumn(StringCast<ANSICHAR>(*MutableColumnName).Get(), mu::TABLE_COLUMN_TYPE::TCT_MESH);
+				CurrentColumn = MutableTable->AddColumn(StringCast<ANSICHAR>(*MutableColumnName).Get(), mu::ETableColumnType::Mesh);
 			}
 
 			// First process the mesh tags that are going to make the mesh unique and affect whether it's repeated in 
@@ -235,7 +235,7 @@ bool FillTableColumn(const UCustomizableObjectNodeTable* TableNode,	mu::TablePtr
 			}
 
 			//TODO: Add AnimBp physics to Tables.
-			mu::MeshPtr MutableMesh = GenerateMutableMesh(SkeletalMesh, TSoftClassPtr<UAnimInstance>(), LODIndexConnected, SectionIndexConnected, LODIndex, SectionIndex, MeshUniqueTags, GenerationContext, TableNode);
+			mu::Ptr<mu::Mesh> MutableMesh = GenerateMutableMesh(SkeletalMesh, TSoftClassPtr<UAnimInstance>(), LODIndexConnected, SectionIndexConnected, LODIndex, SectionIndex, MeshUniqueTags, GenerationContext, TableNode);
 
 			if (MutableMesh)
 			{
@@ -347,7 +347,7 @@ bool FillTableColumn(const UCustomizableObjectNodeTable* TableNode,	mu::TablePtr
 
 			if (CurrentColumn == -1)
 			{
-				CurrentColumn = MutableTable->AddColumn(StringCast<ANSICHAR>(*MutableColumnName).Get(), mu::TABLE_COLUMN_TYPE::TCT_MESH);
+				CurrentColumn = MutableTable->AddColumn(StringCast<ANSICHAR>(*MutableColumnName).Get(), mu::ETableColumnType::Mesh);
 			}
 
 			mu::MeshPtr MutableMesh = GenerateMutableMesh(StaticMesh, TSoftClassPtr<UAnimInstance>(), CurrentLOD, SectionIndex, CurrentLOD, SectionIndex, FString(), GenerationContext, TableNode); // TODO GMT
@@ -390,7 +390,7 @@ bool FillTableColumn(const UCustomizableObjectNodeTable* TableNode,	mu::TablePtr
 
 			if (CurrentColumn == INDEX_NONE)
 			{
-				CurrentColumn = MutableTable->AddColumn(StringCast<ANSICHAR>(*ColumnName).Get(), mu::TABLE_COLUMN_TYPE::TCT_IMAGE);
+				CurrentColumn = MutableTable->AddColumn(StringCast<ANSICHAR>(*ColumnName).Get(), mu::ETableColumnType::Image);
 			}
 
 			if (TableNode->GetColumnImageMode(ColumnName) == ETableTextureType::PASSTHROUGH_TEXTURE)
@@ -411,7 +411,8 @@ bool FillTableColumn(const UCustomizableObjectNodeTable* TableNode,	mu::TablePtr
 					ImageReferenceID = *FoundIndex;
 				}
 
-				MutableTable->SetCell(CurrentColumn, RowIdx, mu::Image::CreateAsReference(ImageReferenceID).get());
+				mu::Ptr<mu::ResourceProxyMemory<mu::Image>> Proxy = new mu::ResourceProxyMemory<mu::Image>(mu::Image::CreateAsReference(ImageReferenceID));
+				MutableTable->SetCell(CurrentColumn, RowIdx, Proxy.get());
 			}
 			else
 			{
@@ -462,7 +463,7 @@ bool FillTableColumn(const UCustomizableObjectNodeTable* TableNode,	mu::TablePtr
 
 				if (CurrentColumn == -1)
 				{
-					CurrentColumn = MutableTable->AddColumn(StringCast<ANSICHAR>(*ColumnName).Get(), mu::TABLE_COLUMN_TYPE::TCT_SCALAR);
+					CurrentColumn = MutableTable->AddColumn(StringCast<ANSICHAR>(*ColumnName).Get(), mu::ETableColumnType::Scalar);
 				}
 
 				const int32 lastMaterialAmount = GenerationContext.ReferencedMaterials.Num();
@@ -501,7 +502,7 @@ bool FillTableColumn(const UCustomizableObjectNodeTable* TableNode,	mu::TablePtr
 				if (ColumnIndex == INDEX_NONE)
 				{
 					// If there is no column with the parameters name, we generate a new one
-					ColumnIndex = MutableTable->AddColumn(StringCast<ANSICHAR>(*ColumnName).Get(), mu::TABLE_COLUMN_TYPE::TCT_IMAGE);
+					ColumnIndex = MutableTable->AddColumn(StringCast<ANSICHAR>(*ColumnName).Get(), mu::ETableColumnType::Image);
 				}
 
 				UTexture* ParentTextureValue = nullptr;
@@ -556,12 +557,12 @@ bool FillTableColumn(const UCustomizableObjectNodeTable* TableNode,	mu::TablePtr
 
 			if (CurrentColumn == INDEX_NONE)
 			{
-				CurrentColumn = MutableTable->AddColumn(StringCast<ANSICHAR>(*ColumnName).Get(), mu::TABLE_COLUMN_TYPE::TCT_COLOUR);
+				CurrentColumn = MutableTable->AddColumn(StringCast<ANSICHAR>(*ColumnName).Get(), mu::ETableColumnType::Color);
 			}
 
 			// Setting cell value
 			FLinearColor Value = *(FLinearColor*)CellData;
-			MutableTable->SetCell(CurrentColumn, RowIdx, Value.R, Value.G, Value.B, Value.A);
+			MutableTable->SetCell(CurrentColumn, RowIdx, Value);
 		}
 		
 		else
@@ -577,7 +578,7 @@ bool FillTableColumn(const UCustomizableObjectNodeTable* TableNode,	mu::TablePtr
 
 		if (CurrentColumn == INDEX_NONE)
 		{
-			CurrentColumn = MutableTable->AddColumn(StringCast<ANSICHAR>(*ColumnName).Get(), mu::TABLE_COLUMN_TYPE::TCT_SCALAR);
+			CurrentColumn = MutableTable->AddColumn(StringCast<ANSICHAR>(*ColumnName).Get(), mu::ETableColumnType::Scalar);
 		}
 
 		// Setting cell value
@@ -591,7 +592,7 @@ bool FillTableColumn(const UCustomizableObjectNodeTable* TableNode,	mu::TablePtr
 	
 		if (CurrentColumn == INDEX_NONE)
 		{
-			CurrentColumn = MutableTable->AddColumn(StringCast<ANSICHAR>(*ColumnName).Get(), mu::TABLE_COLUMN_TYPE::TCT_SCALAR);
+			CurrentColumn = MutableTable->AddColumn(StringCast<ANSICHAR>(*ColumnName).Get(), mu::ETableColumnType::Scalar);
 		}
 	
 		// Setting cell value
@@ -684,7 +685,7 @@ mu::TablePtr GenerateMutableSourceTable(const FString& TableName, const UEdGraph
 			TArray<FName> RowNames = TypedTable->GetRowNames();
 
 			// Adding and filling Name Column
-			MutableTable->AddColumn("Name", mu::TABLE_COLUMN_TYPE::TCT_STRING);
+			MutableTable->AddColumn("Name", mu::ETableColumnType::String);
 
 			// Add metadata
 			FParameterUIData ParameterUIData(
diff --git a/Engine/Plugins/Experimental/Mutable/Source/CustomizableObjectEditor/Private/MuCOE/UnrealToMutableTextureConversionUtils.cpp b/Engine/Plugins/Experimental/Mutable/Source/CustomizableObjectEditor/Private/MuCOE/UnrealToMutableTextureConversionUtils.cpp
--- a/Engine/Plugins/Experimental/Mutable/Source/CustomizableObjectEditor/Private/MuCOE/UnrealToMutableTextureConversionUtils.cpp
+++ b/Engine/Plugins/Experimental/Mutable/Source/CustomizableObjectEditor/Private/MuCOE/UnrealToMutableTextureConversionUtils.cpp
@@ -160,6 +160,8 @@ void BlurNormalForComposite(FImage& Image)
 
 TTuple<mu::ImagePtr, EUnrealToMutableConversionError> ConvertTextureUnrealToMutable(UTexture2D* Texture, bool bIsNormalComposite)
 {
+	MUTABLE_CPUPROFILER_SCOPE(ConvertTextureUnrealToMutableTuple);
+
     using namespace UnrealToMutableImageConversion_Interanl;
 
     const int32 LODs = 1;
@@ -192,6 +194,7 @@ TTuple<mu::ImagePtr, EUnrealToMutableConversionError> ConvertTextureUnrealToMuta
     // If any post processes of the image is needed, convert to RGBA32F
     if (bFlipGreenChannel || bIsNormalComposite)
     { 
+		MUTABLE_CPUPROFILER_SCOPE(FlipOrComposite);
 
         TempImage0.CopyTo(TempImage1, ERawImageFormat::RGBA32F, EGammaSpace::Linear);
 
@@ -239,6 +242,8 @@ TTuple<mu::ImagePtr, EUnrealToMutableConversionError> ConvertTextureUnrealToMuta
     // Convert to RGBA8 in place if needed 
     if (MutableCompatibleFormat == ERawImageFormat::BGRA8)
     {
+		MUTABLE_CPUPROFILER_SCOPE(ToRGBA8);
+
         TArrayView64<FColor> ImageDataView = TempImage1.AsBGRA8();
 
         for (FColor& Color : ImageDataView)
@@ -251,32 +256,35 @@ TTuple<mu::ImagePtr, EUnrealToMutableConversionError> ConvertTextureUnrealToMuta
 	switch (MutableCompatibleFormat)
 	{
 	case ERawImageFormat::G8:
+	{
+		MUTABLE_CPUPROFILER_SCOPE(NoConvert);
 		Image = new mu::Image(SizeX, SizeY, LODs, mu::EImageFormat::IF_L_UBYTE, mu::EInitializationType::NotInitialized);
 		FMemory::Memcpy(Image->GetData(), TempImage1.RawData.GetData(), Image->GetDataSize());
 		break;
+	}
 
 	case ERawImageFormat::BGRA8:
 	{
 		// Try to find out if the texture has and actually makes use of the alpha channel
-		bool bHasAlphaChannel = true;
-		FImage SourceImage;
-		if (FImageUtils::GetTexture2DSourceImage(Texture, SourceImage))
-		{
-			bHasAlphaChannel = Texture->AdjustMinAlpha != Texture->AdjustMaxAlpha
-				&& Texture->CompressionSettings != TextureCompressionSettings::TC_Normalmap
-				&& !Texture->CompressionNoAlpha
-				&& (Texture->CompressionForceAlpha
-					|| FImageCore::DetectAlphaChannel(SourceImage));
-		}
+		bool bHasAlphaChannel = 
+			Texture->AdjustMinAlpha != Texture->AdjustMaxAlpha
+			&& Texture->CompressionSettings != TextureCompressionSettings::TC_Normalmap
+			&& !Texture->CompressionNoAlpha
+			&& (Texture->CompressionForceAlpha 
+				||
+				FImageCore::DetectAlphaChannel(TempImage0));
 
 		// TODO: If we ever manage to get Pixel Format data on cook compilation time, remove the code that sets bHasAlphaChannel and just use Texture->HasAlphaChannel() here. Currently unreliable, it always returns EPixelFormat::PF_Unknown when cooking, which returns always "false" to HasAlphaChannel().
 		if (bHasAlphaChannel)
 		{
+			MUTABLE_CPUPROFILER_SCOPE(NoConvert);
 			Image = new mu::Image(SizeX, SizeY, LODs, mu::EImageFormat::IF_RGBA_UBYTE, mu::EInitializationType::NotInitialized);
 			FMemory::Memcpy(Image->GetData(), TempImage1.RawData.GetData(), Image->GetDataSize());
 		}
 		else
 		{
+			MUTABLE_CPUPROFILER_SCOPE(ToRGB);
+
 			Image = new mu::Image(SizeX, SizeY, LODs, mu::EImageFormat::IF_RGB_UBYTE, mu::EInitializationType::NotInitialized);
 			// Manual copy
 			const uint8* DataSource = TempImage1.RawData.GetData();
diff --git a/Engine/Plugins/Experimental/Mutable/Source/MutableRuntime/Public/MuR/Serialisation.h b/Engine/Plugins/Experimental/Mutable/Source/MutableRuntime/Public/MuR/Serialisation.h
--- a/Engine/Plugins/Experimental/Mutable/Source/MutableRuntime/Public/MuR/Serialisation.h
+++ b/Engine/Plugins/Experimental/Mutable/Source/MutableRuntime/Public/MuR/Serialisation.h
@@ -373,9 +373,9 @@ namespace mu
 	MUTABLE_DEFINE_POD_VECTOR_SERIALISABLE(FVector4f);
 	
 	//---------------------------------------------------------------------------------------------
-	void operator<< (OutputArchive& arch, const FString& t);
+	MUTABLERUNTIME_API void operator<<(OutputArchive& arch, const FString& t);
 	
-	void operator>> (InputArchive& arch, FString& t);
+	MUTABLERUNTIME_API void operator>>(InputArchive& arch, FString& t);
 
 
 	//---------------------------------------------------------------------------------------------
diff --git a/Engine/Plugins/Experimental/Mutable/Source/MutableTools/Private/MuT/CodeGenerator.cpp b/Engine/Plugins/Experimental/Mutable/Source/MutableTools/Private/MuT/CodeGenerator.cpp
--- a/Engine/Plugins/Experimental/Mutable/Source/MutableTools/Private/MuT/CodeGenerator.cpp
+++ b/Engine/Plugins/Experimental/Mutable/Source/MutableTools/Private/MuT/CodeGenerator.cpp
@@ -304,7 +304,7 @@ namespace mu
         param.m_name = strName;
         if ( param.m_name.size()==0 )
         {
-            param.m_name = pTable->GetName();
+            param.m_name = TCHAR_TO_ANSI(*pTable->GetName());
         }
         param.m_type = PARAMETER_TYPE::T_INT;
         param.m_defaultValue.Set<ParamIntType>(0);
@@ -314,16 +314,16 @@ namespace mu
 			// See if there is a string column. If there is one, we will use it as names for the
 			// options. Only the first string column will be used.
 			int nameCol = -1;
-			int32 cols = pTable->GetPrivate()->m_columns.Num();
+			int32 cols = pTable->GetPrivate()->Columns.Num();
 			for (int32 c = 0; c < cols && nameCol < 0; ++c)
 			{
-				if (pTable->GetPrivate()->m_columns[c].m_type == TCT_STRING)
+				if (pTable->GetPrivate()->Columns[c].Type == ETableColumnType::String)
 				{
 					nameCol = c;
 				}
 			}
 
-			if (pTable->GetPrivate()->m_NoneOption)
+			if (pTable->GetPrivate()->bNoneOption)
 			{
 				FParameterDesc::INT_VALUE_DESC nullValue;
 				nullValue.m_value = -1;
@@ -333,15 +333,15 @@ namespace mu
 			}
 
 			// Add every row
-			int32 rows = pTable->GetPrivate()->m_rows.Num();
+			int32 rows = pTable->GetPrivate()->Rows.Num();
 			for (size_t i = 0; i < rows; ++i)
 			{
 				FParameterDesc::INT_VALUE_DESC value;
-				value.m_value = (int16_t)pTable->GetPrivate()->m_rows[i].m_id;
+				value.m_value = (int16_t)pTable->GetPrivate()->Rows[i].Id;
 
 				if (nameCol > -1)
 				{
-					value.m_name = pTable->GetPrivate()->m_rows[i].m_values[nameCol].m_string;
+					value.m_name = TCHAR_TO_ANSI(*pTable->GetPrivate()->Rows[i].Values[nameCol].String);
 				}
 
 				param.m_possibleValues.Add(value);
diff --git a/Engine/Plugins/Experimental/Mutable/Source/MutableTools/Private/MuT/CodeGenerator.h b/Engine/Plugins/Experimental/Mutable/Source/MutableTools/Private/MuT/CodeGenerator.h
--- a/Engine/Plugins/Experimental/Mutable/Source/MutableTools/Private/MuT/CodeGenerator.h
+++ b/Engine/Plugins/Experimental/Mutable/Source/MutableTools/Private/MuT/CodeGenerator.h
@@ -310,7 +310,7 @@ namespace mu
         Ptr<ASTOp> GenerateMissingBoolCode(const TCHAR* strWhere, bool value, const void* errorContext );
 
         //!
-		template<class NODE_TABLE_PRIVATE, TABLE_COLUMN_TYPE TYPE, OP_TYPE OPTYPE, typename F>
+		template<class NODE_TABLE_PRIVATE, ETableColumnType TYPE, OP_TYPE OPTYPE, typename F>
 		Ptr<ASTOp> GenerateTableSwitch( const NODE_TABLE_PRIVATE& node, F&& GenerateOption );
 
 
@@ -755,7 +755,7 @@ namespace mu
 
 	
     //---------------------------------------------------------------------------------------------
-    template<class NODE_TABLE_PRIVATE, TABLE_COLUMN_TYPE TYPE, OP_TYPE OPTYPE, typename F>
+    template<class NODE_TABLE_PRIVATE, ETableColumnType TYPE, OP_TYPE OPTYPE, typename F>
     Ptr<ASTOp> CodeGenerator::GenerateTableSwitch
         (
             const NODE_TABLE_PRIVATE& node, F&& GenerateOption
@@ -791,7 +791,7 @@ namespace mu
             return nullptr;
         }
 
-        if ( pTable->GetPrivate()->m_columns[ colIndex ].m_type != TYPE )
+        if ( pTable->GetPrivate()->Columns[ colIndex ].Type != TYPE )
         {
             m_pErrorLog->GetPrivate()->Add("Table column type is not the right type.",
                                            ELMT_ERROR, node.m_errorContext);
@@ -800,20 +800,20 @@ namespace mu
 
         // Create the switch to cover all the options
         Ptr<ASTOp> lastSwitch;
-        std::size_t rows = pTable->GetPrivate()->m_rows.Num();
+        int32 rows = pTable->GetPrivate()->Rows.Num();
 
         Ptr<ASTOpSwitch> SwitchOp = new ASTOpSwitch();
 		SwitchOp->type = OPTYPE;
 		SwitchOp->variable = variable;
 		SwitchOp->def = nullptr;
 
-		for (size_t i = 0; i < rows; ++i)
+		for (int32 i = 0; i < rows; ++i)
         {
-            check( pTable->GetPrivate()->m_rows[i].m_id <= 0xFFFF);
-            auto condition = (uint16)pTable->GetPrivate()->m_rows[i].m_id;
+            check( pTable->GetPrivate()->Rows[i].Id <= 0xFFFF);
+            auto condition = (uint16)pTable->GetPrivate()->Rows[i].Id;
             Ptr<ASTOp> Branch = GenerateOption( node, colIndex, (int)i, m_pErrorLog.get() );
 
-			if (Branch || TYPE != TCT_MESH)
+			if (Branch || TYPE != ETableColumnType::Mesh)
 			{
 				SwitchOp->cases.Add(ASTOpSwitch::FCase(condition, SwitchOp, Branch));
 			}
diff --git a/Engine/Plugins/Experimental/Mutable/Source/MutableTools/Private/MuT/CodeGenerator_Color.cpp b/Engine/Plugins/Experimental/Mutable/Source/MutableTools/Private/MuT/CodeGenerator_Color.cpp
--- a/Engine/Plugins/Experimental/Mutable/Source/MutableTools/Private/MuT/CodeGenerator_Color.cpp
+++ b/Engine/Plugins/Experimental/Mutable/Source/MutableTools/Private/MuT/CodeGenerator_Color.cpp
@@ -471,11 +471,11 @@ namespace mu
 	{
 		const NodeColourTable::Private& node = *Typed->GetPrivate();
 
-		result.op = GenerateTableSwitch<NodeColourTable::Private, TCT_COLOUR, OP_TYPE::CO_SWITCH>(node,
+		result.op = GenerateTableSwitch<NodeColourTable::Private, ETableColumnType::Color, OP_TYPE::CO_SWITCH>(node,
 			[this](const NodeColourTable::Private& node, int colIndex, int row, ErrorLog* pErrorLog)
 			{
 				NodeColourConstantPtr CellData = new NodeColourConstant();
-				FVector4f Colour = node.m_pTable->GetPrivate()->m_rows[row].m_values[colIndex].m_colour;
+				FVector4f Colour = node.m_pTable->GetPrivate()->Rows[row].Values[colIndex].Color;
 				CellData->SetValue(Colour);
 				return Generate(CellData);
 			});
diff --git a/Engine/Plugins/Experimental/Mutable/Source/MutableTools/Private/MuT/CodeGenerator_Image.cpp b/Engine/Plugins/Experimental/Mutable/Source/MutableTools/Private/MuT/CodeGenerator_Image.cpp
--- a/Engine/Plugins/Experimental/Mutable/Source/MutableTools/Private/MuT/CodeGenerator_Image.cpp
+++ b/Engine/Plugins/Experimental/Mutable/Source/MutableTools/Private/MuT/CodeGenerator_Image.cpp
@@ -1700,14 +1700,14 @@ namespace mu
 	{
 		const NodeImageTable::Private& node = *InNode->GetPrivate();
 
-		Result.op = GenerateTableSwitch<NodeImageTable::Private, TCT_IMAGE, OP_TYPE::IM_SWITCH>(node,
+		Result.op = GenerateTableSwitch<NodeImageTable::Private, ETableColumnType::Image, OP_TYPE::IM_SWITCH>(node,
 			[this,Options](const NodeImageTable::Private& node, int colIndex, int row, ErrorLog* pErrorLog)
 			{
-				TABLE_VALUE CellData = node.m_pTable->GetPrivate()->m_rows[row].m_values[colIndex];
+				const FTableValue& CellData = node.m_pTable->GetPrivate()->Rows[row].Values[colIndex];
 				ImagePtrConst pImage = nullptr;
 				NodeImagePtr CellImage = nullptr;
 
-				if (Ptr<ResourceProxy<Image>> pProxyImage = CellData.m_pProxyImage)
+				if (Ptr<ResourceProxy<Image>> pProxyImage = CellData.ProxyImage)
 				{
 					pImage = pProxyImage->Get();
 				}
diff --git a/Engine/Plugins/Experimental/Mutable/Source/MutableTools/Private/MuT/CodeGenerator_Mesh.cpp b/Engine/Plugins/Experimental/Mutable/Source/MutableTools/Private/MuT/CodeGenerator_Mesh.cpp
--- a/Engine/Plugins/Experimental/Mutable/Source/MutableTools/Private/MuT/CodeGenerator_Mesh.cpp
+++ b/Engine/Plugins/Experimental/Mutable/Source/MutableTools/Private/MuT/CodeGenerator_Mesh.cpp
@@ -1117,10 +1117,10 @@ class Node;
 		FMeshGenerationResult NewResult = OutResult;
 		int t = 0;
 
-		Ptr<ASTOp> Op = GenerateTableSwitch<NodeMeshTable::Private, TCT_MESH, OP_TYPE::ME_SWITCH>(*TableNode->GetPrivate(), 
+		Ptr<ASTOp> Op = GenerateTableSwitch<NodeMeshTable::Private, ETableColumnType::Mesh, OP_TYPE::ME_SWITCH>(*TableNode->GetPrivate(),
 			[this, &NewResult, &t, &InOptions] (const NodeMeshTable::Private& node, int colIndex, int row, ErrorLog* pErrorLog)
 			{
-				MeshPtr pMesh = node.m_pTable->GetPrivate()->m_rows[row].m_values[colIndex].m_pMesh;
+				mu::Ptr<mu::Mesh> pMesh = node.m_pTable->GetPrivate()->Rows[row].Values[colIndex].Mesh;
 				FMeshGenerationResult BranchResults;
 
 				if (pMesh)
diff --git a/Engine/Plugins/Experimental/Mutable/Source/MutableTools/Private/MuT/CodeGenerator_Scalar.cpp b/Engine/Plugins/Experimental/Mutable/Source/MutableTools/Private/MuT/CodeGenerator_Scalar.cpp
--- a/Engine/Plugins/Experimental/Mutable/Source/MutableTools/Private/MuT/CodeGenerator_Scalar.cpp
+++ b/Engine/Plugins/Experimental/Mutable/Source/MutableTools/Private/MuT/CodeGenerator_Scalar.cpp
@@ -417,11 +417,11 @@ namespace mu
 	{
 		const NodeScalarTable::Private& node = *Typed->GetPrivate();
 
-		Ptr<ASTOp> Op = GenerateTableSwitch<NodeScalarTable::Private, TCT_SCALAR, OP_TYPE::SC_SWITCH>(node,
+		Ptr<ASTOp> Op = GenerateTableSwitch<NodeScalarTable::Private, ETableColumnType::Scalar, OP_TYPE::SC_SWITCH>(node,
 			[this](const NodeScalarTable::Private& node, int colIndex, int row, ErrorLog* pErrorLog)
 			{
 				NodeScalarConstantPtr pCell = new NodeScalarConstant();
-				float scalar = node.m_pTable->GetPrivate()->m_rows[row].m_values[colIndex].m_scalar;
+				float scalar = node.m_pTable->GetPrivate()->Rows[row].Values[colIndex].Scalar;
 				pCell->SetValue(scalar);
 				return Generate(pCell);
 			});
diff --git a/Engine/Plugins/Experimental/Mutable/Source/MutableTools/Private/MuT/ImageDescGenerator.cpp b/Engine/Plugins/Experimental/Mutable/Source/MutableTools/Private/MuT/ImageDescGenerator.cpp
--- a/Engine/Plugins/Experimental/Mutable/Source/MutableTools/Private/MuT/ImageDescGenerator.cpp
+++ b/Engine/Plugins/Experimental/Mutable/Source/MutableTools/Private/MuT/ImageDescGenerator.cpp
@@ -130,15 +130,15 @@ namespace mu
         Ptr<const Image> pImage;
 
         std::size_t i = 0;
-        while ( !pImage && i<node.m_pTable->GetPrivate()->m_rows.Num() )
+        while ( !pImage && i<node.m_pTable->GetPrivate()->Rows.Num() )
         {
-			if (node.m_pTable->GetPrivate()->m_rows[i].m_values[colIndex].m_pProxyImage->Get()->IsReference())
+			if (node.m_pTable->GetPrivate()->Rows[i].Values[colIndex].ProxyImage->Get()->IsReference())
 			{
 				// Image References do not need an image desc
 				break;
 			}
 
-            pImage = node.m_pTable->GetPrivate()->m_rows[i].m_values[ colIndex ].m_pProxyImage->Get();
+            pImage = node.m_pTable->GetPrivate()->Rows[i].Values[ colIndex ].ProxyImage->Get();
             ++i;
         }
 
diff --git a/Engine/Plugins/Experimental/Mutable/Source/MutableTools/Private/MuT/Table.cpp b/Engine/Plugins/Experimental/Mutable/Source/MutableTools/Private/MuT/Table.cpp
--- a/Engine/Plugins/Experimental/Mutable/Source/MutableTools/Private/MuT/Table.cpp
+++ b/Engine/Plugins/Experimental/Mutable/Source/MutableTools/Private/MuT/Table.cpp
@@ -19,7 +19,7 @@
 
 namespace mu
 {
-	MUTABLE_IMPLEMENT_ENUM_SERIALISABLE(TABLE_COLUMN_TYPE)
+	MUTABLE_IMPLEMENT_ENUM_SERIALISABLE(ETableColumnType)
 
 
 	//---------------------------------------------------------------------------------------------
@@ -46,70 +46,62 @@ namespace mu
 
 
 	//---------------------------------------------------------------------------------------------
-	TablePtr Table::StaticUnserialise( InputArchive& arch )
+	Ptr<Table> Table::StaticUnserialise( InputArchive& arch )
 	{
-		TablePtr pResult = new Table();
+		Ptr<Table> pResult = new Table();
 		arch >> *pResult->m_pD;
 		return pResult;
 	}
 
 
-	//---------------------------------------------------------------------------------------------
 	Table::Private* Table::GetPrivate() const
 	{
 		return m_pD;
 	}
 
 
-	//---------------------------------------------------------------------------------------------
-	void Table::SetName( const char* strName )
+	void Table::SetName( const FString& InName )
 	{
-		m_pD->m_name = strName;
+		m_pD->Name = InName;
 	}
 
 
-	//---------------------------------------------------------------------------------------------
-	const char* Table::GetName() const
+	const FString& Table::GetName() const
 	{
-		return m_pD->m_name.c_str();
+		return m_pD->Name;
 	}
 
 
-	//---------------------------------------------------------------------------------------------
-	int Table::AddColumn( const char* strName, TABLE_COLUMN_TYPE type )
+	int32 Table::AddColumn(const FString& Name, ETableColumnType type )
 	{
-		int result = m_pD->m_columns.Num();
+		int32 result = m_pD->Columns.Num();
 
-		TABLE_COLUMN c;
+		FTableColumn c;
 
-		if (strName)
-		{
-			c.m_name = strName;
-		}
-		c.m_type = type;
+		c.Name = Name;
+		c.Type = type;
 
-		m_pD->m_columns.Add(c);
+		m_pD->Columns.Add(c);
 
 		// Add it to all rows
-		for ( std::size_t r=0; r<m_pD->m_rows.Num(); ++r )
+		for ( int32 r=0; r<m_pD->Rows.Num(); ++r )
 		{
-			m_pD->m_rows[r].m_values.Add( TABLE_VALUE() );
+			m_pD->Rows[r].Values.Add( FTableValue() );
 		}
 
 		return result;
 	}
 
 
-	//---------------------------------------------------------------------------------------------
-	int Table::FindColumn( const char* strName ) const
+	int32 Table::FindColumn(const FString& Name ) const
 	{
-		int res = -1;
+		int32 res = -1;
 
-		for ( std::size_t c=0; c<m_pD->m_columns.Num(); ++c )
+		for ( int32 c=0; c<m_pD->Columns.Num(); ++c )
 		{
-			if ( m_pD->m_columns[c].m_name == strName )
+			if ( m_pD->Columns[c].Name == Name )
 			{
-				res = (int)c;
+				res = c;
 			}
 		}
 
@@ -117,87 +109,80 @@ namespace mu
 	}
 
 
-	//---------------------------------------------------------------------------------------------
-    void Table::AddRow( uint32_t id )
+    void Table::AddRow( uint32 id )
 	{
 		check( m_pD->FindRow(id)<0 );
 
-		TABLE_ROW r;
-		r.m_id = id;
-		r.m_values.SetNum( m_pD->m_columns.Num() );
-		m_pD->m_rows.Add( r );
+		FTableRow r;
+		r.Id = id;
+		r.Values.SetNum( m_pD->Columns.Num() );
+		m_pD->Rows.Add( r );
 	}
 
 
-	//---------------------------------------------------------------------------------------------
 	void Table::SetNoneOption(bool bAddOption)
 	{
-		m_pD->m_NoneOption = bAddOption;
+		m_pD->bNoneOption = bAddOption;
 	}
 
 
-	//---------------------------------------------------------------------------------------------
-    void Table::SetCell( int column, uint32_t rowId, float value )
+    void Table::SetCell( int32 Column, uint32 RowId, float Value)
 	{
-		int row = m_pD->FindRow(rowId);
-		check( row>=0 );
-		check( column < (int)m_pD->m_columns.Num() );
-		check( column < (int)m_pD->m_rows[row].m_values.Num() );
-		check( m_pD->m_columns[column].m_type == TCT_SCALAR );
+		int32 Row = m_pD->FindRow(RowId);
+		check( Row>=0 );
+		check( Column < m_pD->Columns.Num() );
+		check( Column < m_pD->Rows[Row].Values.Num() );
+		check( m_pD->Columns[Column].Type == ETableColumnType::Scalar );
 
-		m_pD->m_rows[ row ].m_values[column].m_scalar = value;
+		m_pD->Rows[ Row ].Values[Column].Scalar = Value;
 	}
 
 
-	//---------------------------------------------------------------------------------------------
-    void Table::SetCell( int column, uint32_t rowId, float r, float g, float b, float a)
+    void Table::SetCell( int32 Column, uint32 RowId, const FVector4f& Value)
 	{
-		int row = m_pD->FindRow(rowId);
-		check( row>=0 );
-		check( column < (int)m_pD->m_columns.Num() );
-		check( column < (int)m_pD->m_rows[row].m_values.Num() );
-		check( m_pD->m_columns[column].m_type == TCT_COLOUR );
+		int32 Row = m_pD->FindRow(RowId);
+		check( Row>=0 );
+		check( Column < m_pD->Columns.Num() );
+		check( Column < m_pD->Rows[Row].Values.Num() );
+		check( m_pD->Columns[Column].Type == ETableColumnType::Color );
 
-		m_pD->m_rows[ row ].m_values[column].m_colour = FVector4f(r, g, b, a);
+		m_pD->Rows[ Row ].Values[Column].Color = Value;
 	}
 
 
-	//---------------------------------------------------------------------------------------------
-    void Table::SetCell( int column, uint32_t rowId, Image* pImage )
+    void Table::SetCell(int32 Column, uint32 RowId, ResourceProxy<Image>* Value)
 	{
-		int row = m_pD->FindRow(rowId);
-		check( row>=0 );
-		check( column < (int)m_pD->m_columns.Num() );
-		check( column < (int)m_pD->m_rows[row].m_values.Num() );
-		check( m_pD->m_columns[column].m_type == TCT_IMAGE );
+		int32 Row = m_pD->FindRow(RowId);
+		check( Row>=0 );
+		check( Column < m_pD->Columns.Num() );
+		check( Column < m_pD->Rows[Row].Values.Num() );
+		check( m_pD->Columns[Column].Type == ETableColumnType::Image );
 
-		m_pD->m_rows[ row ].m_values[column].m_pProxyImage = new ResourceProxyMemory<Image>(pImage);
+		m_pD->Rows[ Row ].Values[Column].ProxyImage = Value;
 	}
 
 
-	//---------------------------------------------------------------------------------------------
-    void Table::SetCell( int column, uint32_t rowId, Mesh* pMesh )
+    void Table::SetCell( int32 Column, uint32 RowId, Mesh* Value )
 	{
-		int row = m_pD->FindRow(rowId);
-		check( row>=0 );
-		check( column < (int)m_pD->m_columns.Num() );
-		check( column < (int)m_pD->m_rows[row].m_values.Num() );
-		check( m_pD->m_columns[column].m_type == TCT_MESH );
+		int32 Row = m_pD->FindRow(RowId);
+		check( Row>=0 );
+		check( Column < m_pD->Columns.Num() );
+		check( Column < m_pD->Rows[Row].Values.Num() );
+		check( m_pD->Columns[Column].Type == ETableColumnType::Mesh );
 
-		m_pD->m_rows[ row ].m_values[column].m_pMesh = pMesh;
+		m_pD->Rows[ Row ].Values[Column].Mesh = Value;
 	}
 
 
-	//---------------------------------------------------------------------------------------------
-    void Table::SetCell( int column, uint32_t rowId, const char* strValue )
+    void Table::SetCell( int32 Column, uint32 RowId, const FString& Value )
 	{
-		int row = m_pD->FindRow(rowId);
-		check( row>=0 );
-		check( column < (int)m_pD->m_columns.Num() );
-		check( column < (int)m_pD->m_rows[row].m_values.Num() );
-		check( m_pD->m_columns[column].m_type == TCT_STRING );
+		int32 Row = m_pD->FindRow(RowId);
+		check( Row>=0 );
+		check( Column < m_pD->Columns.Num() );
+		check( Column < m_pD->Rows[Row].Values.Num() );
+		check( m_pD->Columns[Column].Type == ETableColumnType::String );
 
-		m_pD->m_rows[ row ].m_values[column].m_string = strValue;
+		m_pD->Rows[ Row ].Values[Column].String = Value;
 	}
 
 
diff --git a/Engine/Plugins/Experimental/Mutable/Source/MutableTools/Private/MuT/TablePrivate.h b/Engine/Plugins/Experimental/Mutable/Source/MutableTools/Private/MuT/TablePrivate.h
--- a/Engine/Plugins/Experimental/Mutable/Source/MutableTools/Private/MuT/TablePrivate.h
+++ b/Engine/Plugins/Experimental/Mutable/Source/MutableTools/Private/MuT/TablePrivate.h
@@ -13,32 +13,32 @@
 
 namespace mu
 {
-	MUTABLE_DEFINE_ENUM_SERIALISABLE(TABLE_COLUMN_TYPE)
+	MUTABLE_DEFINE_ENUM_SERIALISABLE(ETableColumnType)
 
 
-	struct TABLE_COLUMN
+	struct FTableColumn
 	{
-		string m_name;
+		FString Name;
 
-		TABLE_COLUMN_TYPE m_type;
+		ETableColumnType Type;
 	};
 
 
-	struct TABLE_VALUE
+	struct FTableValue
 	{
 		// TODO: Union
-		float m_scalar;
-		FVector4f m_colour;
-		Ptr<ResourceProxy<Image>> m_pProxyImage;
-		MeshPtr m_pMesh;
-		string m_string;
+		float Scalar;
+		FVector4f Color;
+		Ptr<ResourceProxy<Image>> ProxyImage;
+		Ptr<Mesh> Mesh;
+		FString String;
 	};
 
 
-	struct TABLE_ROW
+	struct FTableRow
 	{
-        uint32_t m_id;
-		TArray<TABLE_VALUE> m_values;
+        uint32 Id;
+		TArray<FTableValue> Values;
 	};
 
 
@@ -47,118 +47,148 @@ namespace mu
 	{
 	public:
 
-		string m_name;
-		TArray<TABLE_COLUMN> m_columns;
-		TArray<TABLE_ROW> m_rows;
-		bool m_NoneOption = false;
+		FString Name;
+		TArray<FTableColumn> Columns;
+		TArray<FTableRow> Rows;
+		bool bNoneOption = false;
 
 		//!
 		void Serialise( OutputArchive& arch ) const
 		{
-            uint32_t ver = 2;
+            uint32 ver = 3;
 			arch << ver;
 
-            arch << (uint32_t)m_columns.Num();
-			for ( int32 c=0; c<m_columns.Num(); ++c )
+            arch << Columns.Num();
+			for ( int32 c=0; c< Columns.Num(); ++c )
 			{
-				arch << m_columns[c].m_name;
-				arch << m_columns[c].m_type;
+				arch << Columns[c].Name;
+				arch << Columns[c].Type;
 			}
 
-            arch << (uint32_t)m_rows.Num();
-			for (int32 r=0; r<m_rows.Num(); ++r )
+            arch << Rows.Num();
+			for (int32 r=0; r< Rows.Num(); ++r )
 			{
-				arch << m_rows[r].m_id;
+				arch << Rows[r].Id;
 
-				for (int32 c=0; c<m_columns.Num(); ++c )
+				for (int32 c=0; c<Columns.Num(); ++c )
 				{
-					const TABLE_VALUE& v = m_rows[r].m_values[c];
+					const FTableValue& v = Rows[r].Values[c];
 
-					switch (m_columns[c].m_type)
+					switch (Columns[c].Type)
 					{
-					case TCT_SCALAR: 	arch << v.m_scalar; break;
-					case TCT_COLOUR: 	arch << v.m_colour; break;
-					case TCT_MESH: 		arch << v.m_pMesh; break;
-					case TCT_IMAGE:		
+					case ETableColumnType::Scalar: 	arch << v.Scalar; break;
+					case ETableColumnType::Color: 	arch << v.Color; break;
+					case ETableColumnType::Mesh: 	arch << v.Mesh; break;
+					case ETableColumnType::Image:
 					{
-						Ptr<const Image> image;
-						if (v.m_pProxyImage)
+						Ptr<const Image> Value;
+						if (v.ProxyImage)
 						{
-							image = v.m_pProxyImage->Get();
+							Value = v.ProxyImage->Get();
 						}
-						arch << image;
+						arch << Value;
 						break;
 					}
-					case TCT_STRING:	arch << v.m_string; break;
+					case ETableColumnType::String:	arch << v.String; break;
 					default: check( false);
 					}
 				}
 			}
 
-			arch << m_NoneOption;
+			arch << bNoneOption;
 		}
 
 		//!
 		void Unserialise( InputArchive& arch )
 		{
-            uint32_t ver;
+            uint32 ver;
 			arch >> ver;
-			check(ver<=2);
+			check(ver<=3);
 
-            uint32_t columnCount;
-			arch >> columnCount;
-			m_columns.SetNum( columnCount );
-            for (uint32_t c=0; c<columnCount; ++c )
+            uint32 ColumnCount;
+			arch >> ColumnCount;
+			Columns.SetNum(ColumnCount);
+            for (uint32 c=0; c<ColumnCount; ++c )
 			{
-				arch >> m_columns[c].m_name;
-				arch >> m_columns[c].m_type;
+				if (ver <= 2)
+				{
+					string LegacyString;
+					arch >> LegacyString;
+					Columns[c].Name = LegacyString.c_str();
+				}
+				else
+				{
+					arch >> Columns[c].Name;
+				}
+				arch >> Columns[c].Type;
 			}
 
-            uint32_t rowCount;
-			arch >> rowCount;
-			m_rows.SetNum( rowCount );
-            for ( uint32_t r=0; r<rowCount; ++r )
+            uint32 RowCount;
+			arch >> RowCount;
+			Rows.SetNum( RowCount );
+            for ( uint32 r=0; r<RowCount; ++r )
 			{
-				arch >> m_rows[r].m_id;
-				m_rows[r].m_values.SetNum( columnCount );
+				arch >> Rows[r].Id;
+				Rows[r].Values.SetNum( ColumnCount );
 
-                for (uint32_t c=0; c<columnCount; ++c )
+                for (uint32_t c=0; c<ColumnCount; ++c )
 				{
-					TABLE_VALUE& v = m_rows[r].m_values[c];
+					FTableValue& v = Rows[r].Values[c];
 
-					switch (m_columns[c].m_type)
+					switch (Columns[c].Type)
+					{
+					case ETableColumnType::Scalar:	
 					{
-					case TCT_SCALAR:	arch >> v.m_scalar; break;
-					case TCT_COLOUR:
+						arch >> v.Scalar;
+						break;
+					}
+					case ETableColumnType::Color:
 					{
 						if (ver <= 1)
 						{
 							vec3<float> Value;
 							arch >> Value;
 
-							v.m_colour = FVector4f(Value[0], Value[1], Value[2], 1.0f);
+							v.Color = FVector4f(Value[0], Value[1], Value[2], 1.0f);
 						}
 						else
 						{
-							arch >> v.m_colour;
+							arch >> v.Color;
 						}
 						break;
 					}
-					case TCT_MESH:		arch >> v.m_pMesh; break;
-					case TCT_IMAGE:		
+					case ETableColumnType::Mesh:		
+					{
+						arch >> v.Mesh;
+						break;
+					}
+					case ETableColumnType::Image:
 					{
 						// Are we using proxies?
-						v.m_pProxyImage = arch.NewImageProxy();
-						if (!v.m_pProxyImage)
+						v.ProxyImage = arch.NewImageProxy();
+						if (!v.ProxyImage)
 						{
 							// Normal serialisation
-							ImagePtr image;
-							arch >> image;
-							v.m_pProxyImage = new ResourceProxyMemory<Image>(image.get());
+							Ptr<Image> Value;
+							arch >> Value;
+							v.ProxyImage = new ResourceProxyMemory<Image>(Value.get());
+						}
+						break;
+					}
+					case ETableColumnType::String:
+					{
+						if (ver <= 2)
+						{
+							string LegacyString;
+							arch >> LegacyString;
+							v.String = LegacyString.c_str();
+						}
+						else
+						{
+							arch >> v.String;
 						}
 						break;
 					}
-					case TCT_STRING:	arch >> v.m_string; break;
 					default: check( false);
 					}
 				}
@@ -166,21 +196,21 @@ namespace mu
 
 			if (ver >= 1)
 			{
-				arch >> m_NoneOption;
+				arch >> bNoneOption;
 			}
 		}
 
 
 		//! Find a row in the table by id. Return -1 if not found.
-        int FindRow( uint32_t id ) const
+        int32 FindRow( uint32 id ) const
 		{
-			int res = -1;
+			int32 res = -1;
 
-			for ( std::size_t r=0; res<0 && r<m_rows.Num(); ++r )
+			for ( int32 r=0; res<0 && r<Rows.Num(); ++r )
 			{
-				if ( m_rows[r].m_id==id )
+				if ( Rows[r].Id==id )
 				{
-					res = (int)r;
+					res = (int32)r;
 				}
 			}
 
diff --git a/Engine/Plugins/Experimental/Mutable/Source/MutableTools/Public/MuT/Table.h b/Engine/Plugins/Experimental/Mutable/Source/MutableTools/Public/MuT/Table.h
--- a/Engine/Plugins/Experimental/Mutable/Source/MutableTools/Public/MuT/Table.h
+++ b/Engine/Plugins/Experimental/Mutable/Source/MutableTools/Public/MuT/Table.h
@@ -4,6 +4,7 @@
 
 #include "MuR/Ptr.h"
 #include "MuR/RefCounted.h"
+#include "MuR/Serialisation.h"
 
 
 namespace mu
@@ -14,30 +15,20 @@ namespace mu
 	typedef Ptr<const Table> TablePtrConst;
 
 	class Mesh;
-	typedef Ptr<Mesh> MeshPtr;
-	typedef Ptr<const Mesh> MeshPtrConst;
-
 	class Image;
-	typedef Ptr<Image> ImagePtr;
-	typedef Ptr<const Image> ImagePtrConst;
-
 	class InputArchive;
 	class OutputArchive;
 
-	//! Types of the values for the table cells.
-	//! \ingroup runtime
-	typedef enum
+	/** Types of the values for the table cells. */
+	enum class ETableColumnType : uint32
 	{
-		TCT_NONE,
-		TCT_SCALAR,
-		TCT_COLOUR,
-		TCT_IMAGE,
-		TCT_MESH,
-		TCT_STRING,
-		TCT_COUNT,
-
-		_TCT_FORCE32BITS = 0xFFFFFFFF
-	} TABLE_COLUMN_TYPE;
+		None,
+		Scalar,
+		Color,
+		Image,
+		Mesh,
+		String
+	};
 
 
 	//! A table that contains many rows and defines attributes like meshes, images,
@@ -54,21 +45,21 @@ namespace mu
 		Table();
 
 		static void Serialise( const Table* p, OutputArchive& arch );
-		static TablePtr StaticUnserialise( InputArchive& arch );
+		static Ptr<Table> StaticUnserialise( InputArchive& arch );
 
 		//-----------------------------------------------------------------------------------------
 		// Own interface
 		//-----------------------------------------------------------------------------------------
 
 		//!
-		void SetName( const char* strName );
-		const char* GetName() const;
+		void SetName(const FString&);
+		const FString& GetName() const;
 
 		//!
-		int AddColumn( const char* strName, TABLE_COLUMN_TYPE type );
+		int32 AddColumn(const FString&, ETableColumnType );
 
 		//! Return the column index with the given name. -1 if not found.
-		int FindColumn( const char* strName ) const;
+		int32 FindColumn(const FString&) const;
 
 		//!
         void AddRow( uint32 id );
@@ -77,11 +68,11 @@ namespace mu
 		void SetNoneOption(bool bAddOption);
 
 		//!
-        void SetCell( int column, uint32 rowId, float value );
-        void SetCell( int column, uint32 rowId, float r, float g, float b, float a);
-        void SetCell( int column, uint32 rowId, Image* pImage );
-        void SetCell( int column, uint32 rowId, Mesh* pMesh );
-        void SetCell( int column, uint32 rowId, const char* strValue );
+        void SetCell( int32 Column, uint32 RowId, float );
+        void SetCell( int32 Column, uint32 RowId, const FVector4f& );
+		void SetCell( int32 Column, uint32 RowId, ResourceProxy<Image>* );
+		void SetCell( int32 Column, uint32 RowId, Mesh* );
+        void SetCell( int32 Column, uint32 RowId, const FString& );
 
 
 		//-----------------------------------------------------------------------------------------
