[Mutable] Add peak memory tracking to Mutable tracking allocation policy.

- The tracking is enabled by setting the UE_MUTABLE_TRACK_ALLOCATOR_MEMORY_PEAK preprocessor define to 1. Diabled by default.

#rnx

diff --git a/Engine/Plugins/Experimental/Mutable/Source/MutableRuntime/Public/MuR/MemoryTrackingAllocationPolicy.h b/Engine/Plugins/Experimental/Mutable/Source/MutableRuntime/Public/MuR/MemoryTrackingAllocationPolicy.h
--- a/Engine/Plugins/Experimental/Mutable/Source/MutableRuntime/Public/MuR/MemoryTrackingAllocationPolicy.h
+++ b/Engine/Plugins/Experimental/Mutable/Source/MutableRuntime/Public/MuR/MemoryTrackingAllocationPolicy.h
@@ -4,123 +4,112 @@
 
 #include "Containers/ContainerAllocationPolicies.h"
 #include "Containers/ContainersFwd.h"
+#include "HAL/CriticalSection.h"
 
 #include <atomic>
 #include <type_traits>
 
-namespace mu
-{
-	template<typename TagType>
-	struct TMemoryCounter
-	{
-		alignas(8) static inline std::atomic<SSIZE_T> Counter{ 0 };
-	};
 
-	template<typename Type>
-	struct TIsAMemoryCounter : std::false_type {};
+#ifndef UE_MUTABLE_TRACK_ALLOCATOR_MEMORY_PEAK
+#define UE_MUTABLE_TRACK_ALLOCATOR_MEMORY_PEAK 0
+#endif
 
-	template<typename TagType>
-	struct TIsAMemoryCounter<TMemoryCounter<TagType>> : std::true_type {};
+namespace mu
+{
 
-	template<typename BaseAlloc, typename CounterType>
-	class FMemoryTrackingAllocatorWrapper
+#if UE_MUTABLE_TRACK_ALLOCATOR_MEMORY_PEAK
+	struct FGlobalMemoryCounter
 	{
-		static_assert(TIsAMemoryCounter<CounterType>::value);
+	private:
+		 static inline SSIZE_T AbsoluteCounter { 0 };
+		 static inline SSIZE_T AbsolutePeakValue { 0 };
+		 static inline SSIZE_T Counter { 0 };
+		 static inline SSIZE_T PeakValue { 0 };
+		
+		 static inline FCriticalSection Mutex {};
 
 	public:
-		using SizeType = typename BaseAlloc::SizeType;
-
-		enum { NeedsElementType = BaseAlloc::NeedsElementType };
-		enum { RequireRangeCheck = BaseAlloc::RequireRangeCheck };
-
-		class ForAnyElementType : public BaseAlloc::ForAnyElementType
+		static void Update(SSIZE_T Differential)
 		{
-			SSIZE_T AllocSize = 0;
-
-		public:
-			ForAnyElementType() 
-				: BaseAlloc::ForAnyElementType()
-			{
-			}
+			FScopeLock Lock(&Mutex);
 
-			/** Destructor. */
-			FORCEINLINE ~ForAnyElementType()
-			{
-				CounterType::Counter.fetch_sub(AllocSize, std::memory_order_relaxed);
+			Counter += Differential;
+			PeakValue = FMath::Max(PeakValue, Counter);
 
-				AllocSize = 0;
-			}
+			AbsoluteCounter += Differential;
+			AbsolutePeakValue = FMath::Max(AbsolutePeakValue, AbsoluteCounter);
+		}
 
-			template <typename OtherAllocator>
-			FORCEINLINE void MoveToEmptyFromOtherAllocator(typename OtherAllocator::ForAnyElementType& Other) = delete;
+		static void Zero()
+		{
+			FScopeLock Lock(&Mutex);
 
-			FORCEINLINE void MoveToEmpty(ForAnyElementType& Other)
-			{
-				BaseAlloc::ForAnyElementType::MoveToEmpty(Other);
-				
-				CounterType::Counter.fetch_sub(AllocSize, std::memory_order_relaxed);
+			Counter = 0;	
+			PeakValue = 0;
+		}
 
-				AllocSize = Other.AllocSize;
-				Other.AllocSize = 0;
-			}
+		static void Restore()
+		{
+			FScopeLock Lock(&Mutex);
 
-			FORCEINLINE void ResizeAllocation(
-				SizeType PreviousNumElements,
-				SizeType NumElements,
-				SIZE_T NumBytesPerElement
-				)
-			{
-				BaseAlloc::ForAnyElementType::ResizeAllocation(
-					PreviousNumElements, NumElements, NumBytesPerElement);
+			Counter = AbsoluteCounter;
+			PeakValue = AbsolutePeakValue;
+		}
 
-				const SSIZE_T Differential = SSIZE_T(NumElements * NumBytesPerElement) - AllocSize;
-				const SSIZE_T PrevCounterValue = CounterType::Counter.fetch_add(Differential, std::memory_order_relaxed);
-				check(PrevCounterValue >= AllocSize);
+		static SSIZE_T GetPeak()
+		{
+			FScopeLock Lock(&Mutex);
 
-				AllocSize = NumElements * NumBytesPerElement;
-			}
+			const volatile SSIZE_T Result = PeakValue;
+			return Result;
+		}
 
-			FORCEINLINE typename TEnableIf<TAllocatorTraits<BaseAlloc>::SupportsElementAlignment, void>::Type ResizeAllocation(
-				SizeType PreviousNumElements,
-				SizeType NumElements,
-				SIZE_T NumBytesPerElement,
-				uint32 AlignmentOfElement
-				)
-			{
-				BaseAlloc::ForAnyElementType::ResizeAllocation(
-					PreviousNumElements, NumElements, NumBytesPerElement, AlignmentOfElement);
 
-				const SSIZE_T Differential = SSIZE_T(NumElements * NumBytesPerElement) - AllocSize;
-				const SSIZE_T PrevCounterValue = CounterType::Counter.fetch_add(Differential, std::memory_order_relaxed);
-				check(PrevCounterValue >= AllocSize);
+		static SSIZE_T GetCounter()
+		{
+			FScopeLock Lock(&Mutex);
 
-				AllocSize = NumElements * NumBytesPerElement;
-			}
+			const volatile SSIZE_T Result = Counter;
+			return Result;
+		}
+	};
+#else
+	struct FGlobalMemoryCounter
+	{
+		static void Zero()
+		{
+		}
 
-		};
+		static void Restore()
+		{
+		}
 
-		template<typename ElementType>
-		class ForElementType : public ForAnyElementType
+		static SSIZE_T GetPeak()
 		{
-			// Some Allocators, e.g, TAlignedHeapAllocator do some static checks. Instanciate the BaseAlloc::ForElementType
-			// at compile time so those warnings are emitted. 
-			using BaseTypeInstantaitionType = decltype(DeclVal<typename BaseAlloc::template ForElementType<ElementType>>());
+			return 0;
+		}
 
-		public:
-			ForElementType()
-			{
-			}
+		static SSIZE_T GetCounter()
+		{
+			return 0;
+		}
+	};
+#endif
 
-			FORCEINLINE ElementType* GetAllocation() const
-			{
-				return (ElementType*)ForAnyElementType::GetAllocation();
-			}
-		};
+	template<typename TagType>
+	struct TMemoryCounter
+	{
+		alignas(8) static inline std::atomic<SSIZE_T> Counter{ 0 };
 	};
 
+	template<typename Type>
+	struct TIsAMemoryCounter : std::false_type {};
+
+	template<typename TagType>
+	struct TIsAMemoryCounter<TMemoryCounter<TagType>> : std::true_type {};
 
 	template<typename BaseAlloc, typename CounterType>
-	class FMemoryTrackingAllocatorWrapper2
+	class FMemoryTrackingAllocatorWrapper
 	{
 		static_assert(TIsAMemoryCounter<CounterType>::value);
 
@@ -141,6 +130,11 @@ namespace mu
 			FORCEINLINE ~ForAnyElementType()
 			{
 				CounterType::Counter.fetch_sub(AllocSize, std::memory_order_relaxed);
+
+#if UE_MUTABLE_TRACK_ALLOCATOR_MEMORY_PEAK
+				FGlobalMemoryCounter::Update(-AllocSize);
+#endif
+
 				AllocSize = 0;
 			}
 
@@ -158,6 +152,10 @@ namespace mu
 
 				CounterType::Counter.fetch_sub(AllocSize, std::memory_order_relaxed);
 
+#if UE_MUTABLE_TRACK_ALLOCATOR_MEMORY_PEAK
+				FGlobalMemoryCounter::Update(-AllocSize);
+#endif
+
 				AllocSize = Other.AllocSize;
 				Other.AllocSize = 0;
 			}
@@ -174,6 +172,10 @@ namespace mu
 				const SSIZE_T PrevCounterValue = CounterType::Counter.fetch_add(Differential, std::memory_order_relaxed);
 				check(PrevCounterValue >= AllocSize);
 
+#if UE_MUTABLE_TRACK_ALLOCATOR_MEMORY_PEAK
+				FGlobalMemoryCounter::Update(Differential);
+#endif
+
 				AllocSize = NumElements * NumBytesPerElement;
 			}
 
@@ -190,6 +192,10 @@ namespace mu
 				const SSIZE_T PrevCounterValue = CounterType::Counter.fetch_add(Differential, std::memory_order_relaxed);
 				check(PrevCounterValue >= AllocSize);
 
+#if UE_MUTABLE_TRACK_ALLOCATOR_MEMORY_PEAK
+				FGlobalMemoryCounter::Update(Differential);
+#endif
+
 				AllocSize = NumElements * NumBytesPerElement;
 			}
 
@@ -293,7 +299,7 @@ namespace mu
 	// Default memory tarcking allocator needed for TArray and TMap.
 
 	template<typename CounterType>
-	using FDefaultMemoryTrackingAllocator = FMemoryTrackingAllocatorWrapper2<FDefaultAllocator, CounterType>;
+	using FDefaultMemoryTrackingAllocator = FMemoryTrackingAllocatorWrapper<FDefaultAllocator, CounterType>;
 
 	template<typename CounterType>
 	using FDefaultMemoryTrackingBitArrayAllocator = TInlineAllocator<4, FDefaultMemoryTrackingAllocator<CounterType>>;
@@ -312,7 +318,6 @@ namespace mu
 		DEFAULT_MIN_NUMBER_OF_HASHED_ELEMENTS>;
 }
 
-
 template<typename BaseAlloc, typename Counter>
 struct TAllocatorTraits<mu::FMemoryTrackingAllocatorWrapper<BaseAlloc, Counter>> : public TAllocatorTraitsBase<mu::FMemoryTrackingAllocatorWrapper<BaseAlloc, Counter>>
 {
@@ -320,11 +325,4 @@ struct TAllocatorTraits<mu::FMemoryTrackingAllocatorWrapper<BaseAlloc, Counter>>
 	enum { SupportsMoveFromOtherAllocator = false };
 };
 
-template<typename BaseAlloc, typename Counter>
-struct TAllocatorTraits<mu::FMemoryTrackingAllocatorWrapper2<BaseAlloc, Counter>> : public TAllocatorTraitsBase<mu::FMemoryTrackingAllocatorWrapper2<BaseAlloc, Counter>>
-{
-	enum { SupportsElementAlignment = TAllocatorTraits<BaseAlloc>::SupportsElementAlignment };
-	enum { SupportsMoveFromOtherAllocator = false };
-};
-
 
