[Mutable] Fix IM_PLAINCOLOUR operation for BGRA formats.

#rnx

diff --git a/Engine/Plugins/Experimental/Mutable/Source/MutableRuntime/Private/MuR/Image.cpp b/Engine/Plugins/Experimental/Mutable/Source/MutableRuntime/Private/MuR/Image.cpp
--- a/Engine/Plugins/Experimental/Mutable/Source/MutableRuntime/Private/MuR/Image.cpp
+++ b/Engine/Plugins/Experimental/Mutable/Source/MutableRuntime/Private/MuR/Image.cpp
@@ -1283,7 +1283,7 @@ namespace mu
 			uint32 B = uint32(FMath::Clamp(255.0f * Color[2], 0.0f, 255.0f));
 			uint32 A = uint32(FMath::Clamp(255.0f * Color[3], 0.0f, 255.0f));
 	
-			const uint32 PixelData = R | (G << 8) | (B << 16) | (A << 24);  
+			const uint32 PixelData = B | (G << 8) | (R << 16) | (A << 24);  
 
 			if (PixelData == 0)
 			{
