[mutable] Avoid pixel copies and iterations if possible when converting images to mutable.
#rnx

diff --git a/Engine/Plugins/Experimental/Mutable/Source/CustomizableObjectEditor/Private/MuCOE/GenerateMutableSource/GenerateMutableSourceImage.cpp b/Engine/Plugins/Experimental/Mutable/Source/CustomizableObjectEditor/Private/MuCOE/GenerateMutableSource/GenerateMutableSourceImage.cpp
--- a/Engine/Plugins/Experimental/Mutable/Source/CustomizableObjectEditor/Private/MuCOE/GenerateMutableSource/GenerateMutableSourceImage.cpp
+++ b/Engine/Plugins/Experimental/Mutable/Source/CustomizableObjectEditor/Private/MuCOE/GenerateMutableSource/GenerateMutableSourceImage.cpp
@@ -56,6 +56,8 @@
 
 mu::ImagePtr ConvertTextureUnrealToMutable(UTexture2D* Texture, const UCustomizableObjectNode* Node, FCustomizableObjectCompiler* Compiler, bool bIsNormalComposite)
 {     
+	MUTABLE_CPUPROFILER_SCOPE(ConvertTextureUnrealToMutable);
+
 	mu::ImagePtr MutableImage;
 	EUnrealToMutableConversionError Error = EUnrealToMutableConversionError::Unknown;
 
diff --git a/Engine/Plugins/Experimental/Mutable/Source/CustomizableObjectEditor/Private/MuCOE/UnrealToMutableTextureConversionUtils.cpp b/Engine/Plugins/Experimental/Mutable/Source/CustomizableObjectEditor/Private/MuCOE/UnrealToMutableTextureConversionUtils.cpp
--- a/Engine/Plugins/Experimental/Mutable/Source/CustomizableObjectEditor/Private/MuCOE/UnrealToMutableTextureConversionUtils.cpp
+++ b/Engine/Plugins/Experimental/Mutable/Source/CustomizableObjectEditor/Private/MuCOE/UnrealToMutableTextureConversionUtils.cpp
@@ -110,14 +110,24 @@ EUnrealToMutableConversionError ApplyCompositeTexture(
     return EUnrealToMutableConversionError::Success;
 }
 
-void FlipGreenChannel(FImage& Image)
+void FlipGreenChannelRGBA32F(FImage& Image)
 {
-    TArrayView64<FLinearColor> ImageView = Image.AsRGBA32F();
+	TArrayView64<FLinearColor> ImageDataView = Image.AsRGBA32F();
+	ParallelFor(ImageDataView.Num(),
+		[&ImageDataView](uint32 p)
+		{
+			ImageDataView[p].G = 1.0f - FMath::Clamp(ImageDataView[p].G, 0.0f, 1.0f);
+		});
+}
 
-    for (FLinearColor& Color : ImageView)
-    {
-        Color.G = 1.0f - FMath::Clamp(Color.G, 0.0f, 1.0f); 
-    }
+void FlipGreenChannelBGRA8(FImage& Image)
+{
+	TArrayView64<FColor> ImageDataView = Image.AsBGRA8();
+	ParallelFor(ImageDataView.Num(),
+		[&ImageDataView](uint32 p)
+		{
+			ImageDataView[p].G = 255 - ImageDataView[p].G;
+		});
 }
 
 void Normalize(FImage& Image)
@@ -169,88 +179,65 @@ TTuple<mu::ImagePtr, EUnrealToMutableConversionError> ConvertTextureUnrealToMuta
     const int32 SizeY = Texture->Source.GetSizeY();
     ETextureSourceFormat Format = Texture->Source.GetFormat();
  
-    const ERawImageFormat::Type RawFormat = ConvertFormatSourceToRaw(Format);
+    ERawImageFormat::Type RawFormat = ConvertFormatSourceToRaw(Format);
 
-    if (RawFormat == ERawImageFormat::RGBA32F)
-    {
-        return MakeTuple(nullptr, EUnrealToMutableConversionError::UnsupportedFormat);
-    }
+	// Not true, we will convert it.
+	// \TODO: Warn?
+    //if (RawFormat == ERawImageFormat::RGBA32F)
+    //{
+    //    return MakeTuple(nullptr, EUnrealToMutableConversionError::UnsupportedFormat);
+    //}
 
+	// What if source data is not linear?
+    FImage TempImage(SizeX, SizeY, 1, RawFormat, EGammaSpace::Linear);
+	FImage TempImage2;
 
-    FImage TempImage0(SizeX, SizeY, 1, RawFormat, EGammaSpace::Linear);
-    FImage TempImage1;
-    
-    if (!Texture->Source.GetMipData(TempImage0.RawData, 0))
+    if (!Texture->Source.GetMipData(TempImage.RawData, 0))
     {
         return MakeTuple(nullptr, EUnrealToMutableConversionError::Unknown);
     }
 
-    const bool bFlipGreenChannel = Texture->bFlipGreenChannel;
-    
-	//const bool bApplyCompositeTexture =
-    //        static_cast<bool>(Texture->CompositeTexture) &&
-    //        Texture->CompositeTextureMode != ECompositeTextureMode::CTM_Disabled;
+    bool bFlipGreenChannel = Texture->bFlipGreenChannel;
 
     // If any post processes of the image is needed, convert to RGBA32F
-    if (bFlipGreenChannel || bIsNormalComposite)
+    if (bIsNormalComposite)
     { 
 		MUTABLE_CPUPROFILER_SCOPE(FlipOrComposite);
 
-        TempImage0.CopyTo(TempImage1, ERawImageFormat::RGBA32F, EGammaSpace::Linear);
+        TempImage.CopyTo(TempImage2, ERawImageFormat::RGBA32F, EGammaSpace::Linear);
+		RawFormat = ERawImageFormat::RGBA32F;
 
         if (bFlipGreenChannel)
         {
-            FlipGreenChannel(TempImage1);
-        }
+			FlipGreenChannelRGBA32F(TempImage2);
+			// Don't flip again below.
+			bFlipGreenChannel = false;
+		}
 
 		// Prepare texture for use as normal composite.
-		if (bIsNormalComposite)
-		{
-			Normalize(TempImage1);
-			BlurNormalForComposite(TempImage1);
-		}
+		Normalize(TempImage2);
+		BlurNormalForComposite(TempImage2);
 
-		// Don't do this here as the result depends on the generation of the mips  
-        //if (bApplyCompositeTexture)
-        //{
-        //    EUnrealToMutableConversionError CompositeTextureError = 
-        //        ApplyCompositeTexture(
-        //            TempImage1, 
-        //            Texture->CompositeTexture, 
-        //            Texture->CompositeTextureMode,
-        //            Texture->CompositePower);
-
-        //    // This should probably only be a warning indicating the composite could not
-        //    // be applied.
-        //    if (CompositeTextureError != EUnrealToMutableConversionError::Success)
-        //    {
-        //        return MakeTuple(nullptr, CompositeTextureError);
-        //    }
-        //}
-
-        // The result is needed to TempImage0
+        // The result is needed to TempImage
         // Swap internals so the memory allocations is potentially reused.
-        Exchange(TempImage1, TempImage0); 
+		TempImage2.Swap(TempImage);
     }
    
     const ERawImageFormat::Type MutableCompatibleFormat = Format == TSF_G8 
             ? ERawImageFormat::G8 
             : ERawImageFormat::BGRA8;
 
-    TempImage0.CopyTo(TempImage1, MutableCompatibleFormat, EGammaSpace::Linear);
-
-    // Convert to RGBA8 in place if needed 
-    if (MutableCompatibleFormat == ERawImageFormat::BGRA8)
-    {
-		MUTABLE_CPUPROFILER_SCOPE(ToRGBA8);
-
-        TArrayView64<FColor> ImageDataView = TempImage1.AsBGRA8();
-
-        for (FColor& Color : ImageDataView)
-        {
-            Color = FColor(Color.ToPackedABGR()); 
-        } 
-    }
+	if (MutableCompatibleFormat != RawFormat)
+	{
+		MUTABLE_CPUPROFILER_SCOPE(ToCompatibleFormat);
+		TempImage.CopyTo(TempImage2, MutableCompatibleFormat, EGammaSpace::Linear);
+		TempImage2.Swap(TempImage);
+	}
+	
+	if (bFlipGreenChannel)
+	{
+		FlipGreenChannelBGRA8(TempImage);
+	}
 
 	mu::ImagePtr Image;
 	switch (MutableCompatibleFormat)
@@ -259,7 +246,7 @@ TTuple<mu::ImagePtr, EUnrealToMutableConversionError> ConvertTextureUnrealToMuta
 	{
 		MUTABLE_CPUPROFILER_SCOPE(NoConvert);
 		Image = new mu::Image(SizeX, SizeY, LODs, mu::EImageFormat::IF_L_UBYTE, mu::EInitializationType::NotInitialized);
-		FMemory::Memcpy(Image->GetData(), TempImage1.RawData.GetData(), Image->GetDataSize());
+		Image->m_data = MoveTemp(TempImage.RawData);
 		break;
 	}
 
@@ -272,32 +259,43 @@ TTuple<mu::ImagePtr, EUnrealToMutableConversionError> ConvertTextureUnrealToMuta
 			&& !Texture->CompressionNoAlpha
 			&& (Texture->CompressionForceAlpha 
 				||
-				FImageCore::DetectAlphaChannel(TempImage0));
+				FImageCore::DetectAlphaChannel(TempImage));
 
 		// TODO: If we ever manage to get Pixel Format data on cook compilation time, remove the code that sets bHasAlphaChannel and just use Texture->HasAlphaChannel() here. Currently unreliable, it always returns EPixelFormat::PF_Unknown when cooking, which returns always "false" to HasAlphaChannel().
 		if (bHasAlphaChannel)
 		{
-			MUTABLE_CPUPROFILER_SCOPE(NoConvert);
+			MUTABLE_CPUPROFILER_SCOPE(ToRGBA);
 			Image = new mu::Image(SizeX, SizeY, LODs, mu::EImageFormat::IF_RGBA_UBYTE, mu::EInitializationType::NotInitialized);
-			FMemory::Memcpy(Image->GetData(), TempImage1.RawData.GetData(), Image->GetDataSize());
+			uint8* DataDest = Image->GetData();
+
+			// Convert to RGBA8 in place
+			TArrayView64<FColor> ImageDataView = TempImage.AsBGRA8();
+			ParallelFor(ImageDataView.Num(),
+				[DataDest, &ImageDataView](uint32 p)
+				{
+					DataDest[4 * p + 0] = ImageDataView[p].R;
+					DataDest[4 * p + 1] = ImageDataView[p].G;
+					DataDest[4 * p + 2] = ImageDataView[p].B;
+					DataDest[4 * p + 3] = ImageDataView[p].A;
+				});
 		}
 		else
 		{
 			MUTABLE_CPUPROFILER_SCOPE(ToRGB);
 
+			// TODO: add support for a mu::IF_RGBX_UBYTE?
 			Image = new mu::Image(SizeX, SizeY, LODs, mu::EImageFormat::IF_RGB_UBYTE, mu::EInitializationType::NotInitialized);
-			// Manual copy
-			const uint8* DataSource = TempImage1.RawData.GetData();
 			uint8* DataDest = Image->GetData();
-			uint8* DataDestEnd = DataDest + Image->GetDataSize();
-			for (int32 Pixel=0; DataDest!= DataDestEnd; )
-			{
-				DataDest[0] = DataSource[0];
-				DataDest[1] = DataSource[1];
-				DataDest[2] = DataSource[2];
-				DataDest += 3;
-				DataSource += 4;
-			}
+
+			// Convert to RGB8 in place
+			TArrayView64<FColor> ImageDataView = TempImage.AsBGRA8();
+			ParallelFor(ImageDataView.Num(),
+				[DataDest, &ImageDataView](uint32 p)
+				{
+					DataDest[3 * p + 0] = ImageDataView[p].R;
+					DataDest[3 * p + 1] = ImageDataView[p].G;
+					DataDest[3 * p + 2] = ImageDataView[p].B;
+				});
 		}
 
 		//FString Msg = FString::Printf(TEXT("Alpha channel is %s for %s"), bHasAlphaChannel ? TEXT("enabled") : TEXT("disabled"), *Texture->GetName());
