#rnx

diff --git a/Engine/Plugins/Experimental/Mutable/Source/CustomizableObjectEditor/Classes/MuCOE/Nodes/CustomizableObjectNodePassThroughTextureVariation.h b/Engine/Plugins/Experimental/Mutable/Source/CustomizableObjectEditor/Classes/MuCOE/Nodes/CustomizableObjectNodePassThroughTextureVariation.h
new file mode 100644
--- /dev/null
+++ b/Engine/Plugins/Experimental/Mutable/Source/CustomizableObjectEditor/Classes/MuCOE/Nodes/CustomizableObjectNodePassThroughTextureVariation.h
@@ -0,0 +1,21 @@
+// Copyright Epic Games, Inc. All Rights Reserved.
+
+#pragma once
+
+#include "MuCOE/Nodes/CustomizableObjectNodeTextureVariation.h"
+
+#include "CustomizableObjectNodePassThroughTextureVariation.generated.h"
+
+
+UCLASS()
+class CUSTOMIZABLEOBJECTEDITOR_API UCustomizableObjectNodePassThroughTextureVariation : public UCustomizableObjectNodeTextureVariation
+{
+public:
+	GENERATED_BODY()
+
+	// EdGraphNode interface
+	virtual FText GetNodeTitle(ENodeTitleType::Type TitleType) const override;
+
+	// UCustomizableObjectNode interface
+	virtual void AllocateDefaultPins(UCustomizableObjectNodeRemapPins* RemapPins) override;
+};
diff --git a/Engine/Plugins/Experimental/Mutable/Source/CustomizableObjectEditor/Classes/MuCOE/Nodes/CustomizableObjectNodeTextureVariation.h b/Engine/Plugins/Experimental/Mutable/Source/CustomizableObjectEditor/Classes/MuCOE/Nodes/CustomizableObjectNodeTextureVariation.h
--- a/Engine/Plugins/Experimental/Mutable/Source/CustomizableObjectEditor/Classes/MuCOE/Nodes/CustomizableObjectNodeTextureVariation.h
+++ b/Engine/Plugins/Experimental/Mutable/Source/CustomizableObjectEditor/Classes/MuCOE/Nodes/CustomizableObjectNodeTextureVariation.h
@@ -39,12 +39,12 @@ public:
 	void PostEditChangeProperty(FPropertyChangedEvent& PropertyChangedEvent) override;
 
 	// EdGraphNode interface
-	FText GetNodeTitle(ENodeTitleType::Type TitleType) const override;
+	virtual FText GetNodeTitle(ENodeTitleType::Type TitleType) const override;
 	FLinearColor GetNodeTitleColor() const override;
 	FText GetTooltipText() const override;
 	
 	// UCustomizableObjectNode interface
-	void AllocateDefaultPins(UCustomizableObjectNodeRemapPins* RemapPins) override;
+	virtual void AllocateDefaultPins(UCustomizableObjectNodeRemapPins* RemapPins) override;
 
 	UEdGraphPin* OutputPin() const
 	{
diff --git a/Engine/Plugins/Experimental/Mutable/Source/CustomizableObjectEditor/Private/MuCOE/EdGraphSchema_CustomizableObject.cpp b/Engine/Plugins/Experimental/Mutable/Source/CustomizableObjectEditor/Private/MuCOE/EdGraphSchema_CustomizableObject.cpp
--- a/Engine/Plugins/Experimental/Mutable/Source/CustomizableObjectEditor/Private/MuCOE/EdGraphSchema_CustomizableObject.cpp
+++ b/Engine/Plugins/Experimental/Mutable/Source/CustomizableObjectEditor/Private/MuCOE/EdGraphSchema_CustomizableObject.cpp
@@ -70,6 +70,7 @@
 #include "MuCOE/Nodes/CustomizableObjectNodeTextureTransform.h"
 #include "MuCOE/Nodes/CustomizableObjectNodeTextureSaturate.h"
 #include "MuCOE/Nodes/CustomizableObjectNodeTextureVariation.h"
+#include "MuCOE/Nodes/CustomizableObjectNodePassThroughTextureVariation.h"
 #include "ScopedTransaction.h"
 #include "Settings/EditorStyleSettings.h"
 #include "ToolMenu.h"
@@ -421,6 +422,7 @@ void UEdGraphSchema_CustomizableObject::GetGraphContextActions(FGraphContextMenu
 			ContextMenuBuilder.CreateTemplateNode<UCustomizableObjectNodeTextureSwitch>(),
 			ContextMenuBuilder.CreateTemplateNode<UCustomizableObjectNodePassThroughTextureSwitch>(),
 			ContextMenuBuilder.CreateTemplateNode<UCustomizableObjectNodeTextureVariation>(),
+			ContextMenuBuilder.CreateTemplateNode<UCustomizableObjectNodePassThroughTextureVariation>(),
 			ContextMenuBuilder.CreateTemplateNode<UCustomizableObjectNodeTextureToChannels>(),
 			ContextMenuBuilder.CreateTemplateNode<UCustomizableObjectNodeTextureFromChannels>(),
 			ContextMenuBuilder.CreateTemplateNode<UCustomizableObjectNodeTextureFromColor>(),
diff --git a/Engine/Plugins/Experimental/Mutable/Source/CustomizableObjectEditor/Private/MuCOE/GenerateMutableSource/GenerateMutableSourceImage.cpp b/Engine/Plugins/Experimental/Mutable/Source/CustomizableObjectEditor/Private/MuCOE/GenerateMutableSource/GenerateMutableSourceImage.cpp
--- a/Engine/Plugins/Experimental/Mutable/Source/CustomizableObjectEditor/Private/MuCOE/GenerateMutableSource/GenerateMutableSourceImage.cpp
+++ b/Engine/Plugins/Experimental/Mutable/Source/CustomizableObjectEditor/Private/MuCOE/GenerateMutableSource/GenerateMutableSourceImage.cpp
@@ -411,6 +411,7 @@ mu::NodeImagePtr GenerateMutableSourceImage(const UEdGraphPin* Pin, FMutableGrap
 
 	else if (const UCustomizableObjectNodeTextureVariation* TypedNodeImageVar = Cast<const UCustomizableObjectNodeTextureVariation>(Node))
 	{
+		// UCustomizableObjectNodePassThroughTextureVariation nodes are also handled here
 		mu::NodeImageVariationPtr TextureNode = new mu::NodeImageVariation();
 		Result = TextureNode;
 
@@ -432,10 +433,13 @@ mu::NodeImagePtr GenerateMutableSourceImage(const UEdGraphPin* Pin, FMutableGrap
 		}
 
 		TextureNode->SetVariationCount(TypedNodeImageVar->Variations.Num());
-		for (int VariationIndex = 0; VariationIndex < TypedNodeImageVar->Variations.Num(); ++VariationIndex)
+		for (int32 VariationIndex = 0; VariationIndex < TypedNodeImageVar->Variations.Num(); ++VariationIndex)
 		{
 			const UEdGraphPin* VariationPin = TypedNodeImageVar->VariationPin(VariationIndex);
-			if (!VariationPin) continue;
+			if (!VariationPin)
+			{
+				continue;
+			}
 
 			TextureNode->SetVariationTag(VariationIndex, StringCast<ANSICHAR>(*TypedNodeImageVar->Variations[VariationIndex].Tag).Get());
 			if (const UEdGraphPin* ConnectedPin = FollowInputPin(*VariationPin))
diff --git a/Engine/Plugins/Experimental/Mutable/Source/CustomizableObjectEditor/Private/MuCOE/Nodes/CustomizableObjectNodePassThroughTextureVariation.cpp b/Engine/Plugins/Experimental/Mutable/Source/CustomizableObjectEditor/Private/MuCOE/Nodes/CustomizableObjectNodePassThroughTextureVariation.cpp
new file mode 100644
--- /dev/null
+++ b/Engine/Plugins/Experimental/Mutable/Source/CustomizableObjectEditor/Private/MuCOE/Nodes/CustomizableObjectNodePassThroughTextureVariation.cpp
@@ -0,0 +1,37 @@
+// Copyright Epic Games, Inc. All Rights Reserved.
+
+#include "MuCOE/Nodes/CustomizableObjectNodePassThroughTextureVariation.h"
+
+#include "MuCOE/EdGraphSchema_CustomizableObject.h"
+
+#define LOCTEXT_NAMESPACE "CustomizableObjectEditor"
+
+
+void UCustomizableObjectNodePassThroughTextureVariation::AllocateDefaultPins(UCustomizableObjectNodeRemapPins* RemapPins)
+{
+	const UEdGraphSchema_CustomizableObject* Schema = GetDefault<UEdGraphSchema_CustomizableObject>();
+
+	UEdGraphPin* OutputPin = CustomCreatePin(EGPD_Output, Schema->PC_PassThroughImage, FName(TEXT("Texture")));
+	OutputPin->bDefaultValueIsIgnored = true;
+
+	for (int VariationIndex = Variations.Num() - 1; VariationIndex >= 0; --VariationIndex)
+	{
+		FString PinName = FString::Printf(TEXT("Variation %d"), VariationIndex);
+		UEdGraphPin* VariationPin = CustomCreatePin(EGPD_Input, Schema->PC_PassThroughImage, FName(*PinName), false);
+		VariationPin->bDefaultValueIsIgnored = true;
+
+		FString FriendlyName = FString::Printf(TEXT("Variation %d [%s]"), VariationIndex, *Variations[VariationIndex].Tag);
+		VariationPin->PinFriendlyName = FText::FromString(FriendlyName);
+	}
+
+	UEdGraphPin* DefaultVariation = CustomCreatePin(EGPD_Input, Schema->PC_PassThroughImage, FName(TEXT("Default")), false);
+	DefaultVariation->bDefaultValueIsIgnored = true;
+}
+
+
+FText UCustomizableObjectNodePassThroughTextureVariation::GetNodeTitle(ENodeTitleType::Type TitleType) const
+{
+	return LOCTEXT("PassThrough_Texture_Variation", "PassThrough Texture Variation");
+}
+
+#undef LOCTEXT_NAMESPACE
