[mutable] Fixed bad MipsToSkip value.
- Added commandlet argument to enable the loading of the Mips data for the platform being tested.
- Simplified .xml with the editor tests.

diff --git a/Engine/Plugins/Experimental/Mutable/Source/CustomizableObject/Private/MuCO/UnrealMutableImageProvider.cpp b/Engine/Plugins/Experimental/Mutable/Source/CustomizableObject/Private/MuCO/UnrealMutableImageProvider.cpp
--- a/Engine/Plugins/Experimental/Mutable/Source/CustomizableObject/Private/MuCO/UnrealMutableImageProvider.cpp
+++ b/Engine/Plugins/Experimental/Mutable/Source/CustomizableObject/Private/MuCO/UnrealMutableImageProvider.cpp
@@ -199,6 +199,7 @@ mu::EImageFormat GetMutablePixelFormat(EPixelFormat InTextureFormat)
 #endif
 
 			int32 MipIndex = MipmapsToSkip < TextureToLoad->GetPlatformData()->Mips.Num() ? MipmapsToSkip : TextureToLoad->GetPlatformData()->Mips.Num() - 1;
+			check (MipIndex >= 0);
 
 			// Mips in the mip tail are inlined and can't be streamed, find the smallest mip available.
 			for (; MipIndex > 0; --MipIndex)
