[mutable] Fix performance loss on background image loading when editor is not focused.
Fix bugs related to using floats for platform time.
#rnx

diff --git a/Engine/Plugins/Experimental/Mutable/Source/CustomizableObject/Private/MuCO/FMutableTaskGraph.cpp b/Engine/Plugins/Experimental/Mutable/Source/CustomizableObject/Private/MuCO/FMutableTaskGraph.cpp
--- a/Engine/Plugins/Experimental/Mutable/Source/CustomizableObject/Private/MuCO/FMutableTaskGraph.cpp
+++ b/Engine/Plugins/Experimental/Mutable/Source/CustomizableObject/Private/MuCO/FMutableTaskGraph.cpp
@@ -142,8 +142,8 @@ void FMutableTaskGraph::TryLaunchMutableTaskLowPriority(bool bFromMutableTask)
 		return;
 	}
 
-	float TimeLimit;
-	float TimeElapsed;
+	double TimeLimit;
+	double TimeElapsed;
 	
 	bool bTimeLimit;
 	
diff --git a/Engine/Plugins/Experimental/Mutable/Source/CustomizableObject/Public/MuCO/FMutableTaskGraph.h b/Engine/Plugins/Experimental/Mutable/Source/CustomizableObject/Public/MuCO/FMutableTaskGraph.h
--- a/Engine/Plugins/Experimental/Mutable/Source/CustomizableObject/Public/MuCO/FMutableTaskGraph.h
+++ b/Engine/Plugins/Experimental/Mutable/Source/CustomizableObject/Public/MuCO/FMutableTaskGraph.h
@@ -36,7 +36,7 @@ class FMutableTaskGraph
 		uint32 Id;
 		FString DebugName;
 		TFunction<void()> Body;
-		float CreationTime = FPlatformTime::Seconds();
+		double CreationTime = FPlatformTime::Seconds();
 	};
 	
 public:
diff --git a/Engine/Plugins/Experimental/Mutable/Source/CustomizableObjectEditor/Private/MuCOE/CustomizableObjectCompiler.cpp b/Engine/Plugins/Experimental/Mutable/Source/CustomizableObjectEditor/Private/MuCOE/CustomizableObjectCompiler.cpp
--- a/Engine/Plugins/Experimental/Mutable/Source/CustomizableObjectEditor/Private/MuCOE/CustomizableObjectCompiler.cpp
+++ b/Engine/Plugins/Experimental/Mutable/Source/CustomizableObjectEditor/Private/MuCOE/CustomizableObjectCompiler.cpp
@@ -22,6 +22,7 @@
 #include "UObject/UObjectIterator.h"
 #include "Widgets/Notifications/SNotificationList.h"
 #include "Interfaces/ITargetPlatform.h"
+#include "Misc/App.h"
 
 class UTexture2D;
 
@@ -827,7 +828,15 @@ void FCustomizableObjectCompiler::UpdatePendingTextureConversion(bool UseTimeLim
 	// In editor, when compiling a CO, referencer assets and Unreal to Mutable texture conversion are performed asynchronously
 	if (PendingTexturesToLoad)
 	{
-		float initialTime = FPlatformTime::Seconds();
+		double InitialTimeSeconds = FPlatformTime::Seconds();
+
+		float CurrentMaxConvertToMutableTextureTimeSeconds = MaxConvertToMutableTextureTime;
+		if (!FApp::HasFocus())
+		{
+			// If the app is not focused, it will reduce the framerate to 10fps. Do 2 more seconds of work to compensate for that.
+			CurrentMaxConvertToMutableTextureTimeSeconds += 2.0f;
+		}
+
 		while ( CompletedUnrealToMutableTask < ArrayTextureUnrealToMutableTask.Num() )
 		{
 			UTexture2D* Texture = ArrayTextureUnrealToMutableTask[CompletedUnrealToMutableTask].Texture;
@@ -872,9 +881,13 @@ void FCustomizableObjectCompiler::UpdatePendingTextureConversion(bool UseTimeLim
 
 			CompletedUnrealToMutableTask++;
 
-			if (UseTimeLimit && ((FPlatformTime::Seconds() - initialTime) > MaxConvertToMutableTextureTime))
+			if (UseTimeLimit)
 			{
-				break;
+				double UsedTimeSeconds = (FPlatformTime::Seconds() - InitialTimeSeconds);
+				if (UsedTimeSeconds > CurrentMaxConvertToMutableTextureTimeSeconds)
+				{
+					break;
+				}
 			}
 		}
 
