[mutable] Improve runtime ASTC compression quality a little bit.
#rnx

diff --git a/Engine/Plugins/Experimental/Mutable/Source/MutableRuntime/Private/MuR/BlockCompression/Miro/Miro.cpp b/Engine/Plugins/Experimental/Mutable/Source/MutableRuntime/Private/MuR/BlockCompression/Miro/Miro.cpp
--- a/Engine/Plugins/Experimental/Mutable/Source/MutableRuntime/Private/MuR/BlockCompression/Miro/Miro.cpp
+++ b/Engine/Plugins/Experimental/Mutable/Source/MutableRuntime/Private/MuR/BlockCompression/Miro/Miro.cpp
@@ -1,6 +1,7 @@
 // Copyright Epic Games, Inc. All Rights Reserved.
 
 #include "MuR/BlockCompression/Miro/Miro.h"
+#include "MuR/Platform.h"
 
 #include "Async/ParallelFor.h"
 #include "HAL/UnrealMemory.h"
@@ -8,8 +9,28 @@
 #include "Math/UnrealMathUtility.h"
 #include "Misc/AssertionMacros.h"
 
-// debug
-//static bool s_debuglog = false;
+//#define UE_MIRO_DEBUG
+
+#ifdef UE_MIRO_DEBUG
+UE_DISABLE_OPTIMIZATION
+
+static int32 s_DebugBlock = -1;
+static int32 s_CurrentBlock = -1;
+
+inline constexpr void miro_check(bool x)
+{
+	(void)x;
+	check(x);
+}
+
+#else
+
+inline constexpr void miro_check(bool x)
+{
+	(void)x;
+}
+
+#endif //UE_MIRO_DEBUG
 
 #ifdef _MSC_VER
 #define MIRO_ALIGN __declspec(align(16))
@@ -2116,25 +2137,18 @@ namespace impl
 	// Modified for additional format support (alpha), and integration with UE platform libs.
 	// See: [REDACTED_PROFILE_URL]/astcrt
 	//-------------------------------------------------------------------------------------------------
-	inline constexpr void DCHECK(bool x)
-	{
-		(void)x;
-		 if (!x)
-		 {
-		     check(x);
-		 }
-	}
-
 
 	namespace astcrt
 	{
-		constexpr int APPROX_COLOR_EPSILON = 50;
+		constexpr int APPROX_COLOR_EPSILON = 0; // 50
 		constexpr size_t BLOCK_BYTES = 16;
 
 		constexpr size_t MAXIMUM_ENCODED_WEIGHT_BITS = 96;
 		constexpr size_t MAXIMUM_ENCODED_WEIGHT_BYTES = 12;
 
 		constexpr size_t MAXIMUM_ENCODED_COLOR_ENDPOINT_BYTES = 12;
+		constexpr size_t MAX_BLOCK_WIDTH = 12;
+		constexpr size_t MAX_BLOCK_HEIGHT = 12;
 
 
 		inline bool getbit(size_t number, size_t n) {
@@ -2155,8 +2169,8 @@ namespace impl
 			size_t bitoffset,
 			size_t number,
 			size_t bitcount) {
-			DCHECK(bitcount <= 8);
-			DCHECK((number >> bitcount) == 0);
+			miro_check(bitcount <= 8);
+			miro_check((number >> bitcount) == 0);
 
 			size_t index = bitoffset / 8;
 			size_t shift = bitoffset % 8;
@@ -2172,8 +2186,8 @@ namespace impl
 			uint8* p = ptr + index;
 			size_t mask = number << shift;
 
-			DCHECK((p[0] & mask) == 0);
-			DCHECK((p[1] & (mask >> 8)) == 0);
+			miro_check((p[0] & mask) == 0);
+			miro_check((p[1] & (mask >> 8)) == 0);
 
 			p[0] |= static_cast<uint8>(mask & 0xFF);
 			p[1] |= static_cast<uint8>((mask >> 8) & 0xFF);
@@ -2190,7 +2204,7 @@ namespace impl
 		}
 
 		inline void split_high_low(uint8 n, size_t i, uint8& high, uint8& low) {
-			DCHECK(i < 8);
+			miro_check(i < 8);
 
 			uint8 low_mask = static_cast<uint8>((1 << i) - 1);
 
@@ -2266,7 +2280,7 @@ namespace impl
 		{
 			for (int i = 0; i < static_cast<int>(bytecount); ++i)
 			{
-				DCHECK((reverse_byte(source[i]) & target[-i]) == 0);
+				miro_check((reverse_byte(source[i]) & target[-i]) == 0);
 				target[-i] = target[-i] | reverse_byte(source[i]);
 			}
 		}
@@ -2323,7 +2337,8 @@ namespace impl
 			vec4_t() {}
 			vec4_t(T x_, T y_, T z_, T w_) : r(x_), g(y_), b(z_), a(w_) {}
 			vec3_t<T> rgb() const { return vec3_t<T>(r, g, b); }
-			T components(size_t i) { return ((T*)this)[i]; }
+			T& components(size_t i) { return ((T*)this)[i]; }
+			const T& components(size_t i) const { return ((T*)this)[i]; }
 			T r, g, b, a;
 		};
 
@@ -2331,7 +2346,8 @@ namespace impl
 		typedef vec4_t<int> vec4i_t;
 
 		template <typename T>
-		vec3_t<T> operator+(vec3_t<T> a, vec3_t<T> b) {
+		vec3_t<T> operator+(vec3_t<T> a, vec3_t<T> b) 
+		{
 			vec3_t<T> result;
 			result.r = a.r + b.r;
 			result.g = a.g + b.g;
@@ -2340,7 +2356,19 @@ namespace impl
 		}
 
 		template <typename T>
-		vec3_t<T> operator-(vec3_t<T> a, vec3_t<T> b) {
+		vec4_t<T> operator+(vec4_t<T> a, vec4_t<T> b) 
+		{
+			vec4_t<T> result;
+			result.r = a.r + b.r;
+			result.g = a.g + b.g;
+			result.b = a.b + b.b;
+			result.a = a.a + b.a;
+			return result;
+		}
+
+		template <typename T>
+		vec3_t<T> operator-(vec3_t<T> a, vec3_t<T> b) 
+		{
 			vec3_t<T> result;
 			result.r = a.r - b.r;
 			result.g = a.g - b.g;
@@ -2349,7 +2377,8 @@ namespace impl
 		}
 
 		template <typename T>
-		vec4_t<T> operator-(vec4_t<T> a, vec4_t<T> b) {
+		vec4_t<T> operator-(vec4_t<T> a, vec4_t<T> b) 
+		{
 			vec4_t<T> result;
 			result.r = a.r - b.r;
 			result.g = a.g - b.g;
@@ -2359,7 +2388,8 @@ namespace impl
 		}
 
 		template <typename T>
-		vec3_t<T> operator*(vec3_t<T> a, vec3_t<T> b) {
+		vec3_t<T> operator*(vec3_t<T> a, vec3_t<T> b) 
+		{
 			vec3_t<T> result;
 			result.r = a.r * b.r;
 			result.g = a.g * b.g;
@@ -2367,6 +2397,28 @@ namespace impl
 			return result;
 		}
 
+		template <typename T>
+		vec4_t<T> operator*(vec4_t<T> a, vec4_t<T> b) 
+		{
+			vec4_t<T> result;
+			result.r = a.r * b.r;
+			result.g = a.g * b.g;
+			result.b = a.b * b.b;
+			result.a = a.a * b.a;
+			return result;
+		}
+
+		template <typename T>
+		vec4_t<T> operator/(vec4_t<T> a, vec4_t<T> b) 
+		{
+			vec4_t<T> result;
+			result.r = a.r / b.r;
+			result.g = a.g / b.g;
+			result.b = a.b / b.b;
+			result.a = a.a / b.a;
+			return result;
+		}
+
 		template <typename T>
 		vec3_t<T> operator*(vec3_t<T> a, T b)
 		{
@@ -2377,6 +2429,17 @@ namespace impl
 			return result;
 		}
 
+		template <typename T>
+		vec4_t<T> operator*(vec4_t<T> a, T b)
+		{
+			vec4_t<T> result;
+			result.r = a.r * b;
+			result.g = a.g * b;
+			result.b = a.b * b;
+			result.a = a.a * b;
+			return result;
+		}
+
 		template <typename T>
 		vec3_t<T> operator/(vec3_t<T> a, T b)
 		{
@@ -2388,7 +2451,19 @@ namespace impl
 		}
 
 		template <typename T>
-		vec3_t<T> operator/(vec3_t<T> a, vec3_t<T> b) {
+		vec4_t<T> operator/(vec4_t<T> a, T b)
+		{
+			vec4_t<T> result;
+			result.r = a.r / b;
+			result.g = a.g / b;
+			result.b = a.b / b;
+			result.a = a.a / b;
+			return result;
+		}
+
+		template <typename T>
+		vec3_t<T> operator/(vec3_t<T> a, vec3_t<T> b) 
+		{
 			vec3_t<T> result;
 			result.x = a.x / b.x;
 			result.y = a.y / b.y;
@@ -2438,27 +2513,37 @@ namespace impl
 		}
 
 		template <typename T>
-		T norm(vec3_t<T> a) {
+		T norm(vec3_t<T> a) 
+		{
+			return static_cast<T>(sqrt(quadrance(a)));
+		}
+
+		template <typename T>
+		T norm(vec4_t<T> a) 
+		{
 			return static_cast<T>(sqrt(quadrance(a)));
 		}
 
 		template <typename T>
-		T distance(vec3_t<T> a, vec3_t<T> b) {
+		T distance(vec3_t<T> a, vec3_t<T> b)
+		{
 			return norm(a - b);
 		}
 
 		template <typename T>
-		T qd(vec3_t<T> a, vec3_t<T> b) {
+		T qd(vec3_t<T> a, vec3_t<T> b) 
+		{
 			return quadrance(a - b);
 		}
 
 		template <typename T>
-		vec3_t<T> signorm(vec3_t<T> a) {
+		vec3_t<T> signorm(vec3_t<T> a)
+		{
 			T x = norm(a);
 
 			// Safety fix for degenerated cases.
 			// \todo This should be intercepted earlier.
-			// DCHECK(x != 0.0);
+			// miro_check(x != 0.0);
 			if (x == 0.0)
 				return vec3_t<T>(0, 1, 0);
 
@@ -2466,7 +2551,22 @@ namespace impl
 		}
 
 		template <typename T>
-		vec3_t<T> vecmin(vec3_t<T> a, vec3_t<T> b) {
+		vec4_t<T> signorm(vec4_t<T> a)
+		{
+			T x = norm(a);
+
+			// Safety fix for degenerated cases.
+			// \todo This should be intercepted earlier.
+			// miro_check(x != 0.0);
+			if (x == 0.0)
+				return vec4_t<T>(0, 1, 0,0);
+
+			return a / x;
+		}
+
+		template <typename T>
+		vec3_t<T> vecmin(vec3_t<T> a, vec3_t<T> b) 
+		{
 			vec3_t<T> result;
 			result.x = FMath::Min(a.x, b.x);
 			result.y = FMath::Min(a.y, b.y);
@@ -2475,7 +2575,8 @@ namespace impl
 		}
 
 		template <typename T>
-		vec3_t<T> vecmax(vec3_t<T> a, vec3_t<T> b) {
+		vec3_t<T> vecmax(vec3_t<T> a, vec3_t<T> b) 
+		{
 			vec3_t<T> result;
 			result.x = FMath::Max(a.x, b.x);
 			result.y = FMath::Max(a.y, b.y);
@@ -2484,55 +2585,75 @@ namespace impl
 		}
 
 		template <typename T>
-		T qd_to_line(vec3_t<T> m, vec3_t<T> k, T kk, vec3_t<T> p) {
+		T qd_to_line(vec3_t<T> m, vec3_t<T> k, T kk, vec3_t<T> p) 
+		{
 			T t = dot(p - m, k) / kk;
 			vec3_t<T> q = k * t + m;
 			return qd(p, q);
 		}
 
 
-		inline bool is_greyscale(vec3i_t color) {
+		inline bool is_greyscale(vec3i_t color) 
+		{
 			// integer equality is transitive
 			return color.r == color.g && color.g == color.b;
 		}
 
-		inline int luminance(vec3i_t color) {
+		inline int luminance(vec3i_t color) 
+		{
 			return (color.r + color.g + color.b) / 3;
 		}
 
-		inline bool approx_equal(vec3i_t a, vec3i_t b) {
-			return quadrance(a - b) <= APPROX_COLOR_EPSILON;
+		inline bool approx_equal(vec3i_t a, vec3i_t b, int32 Epsilon ) 
+		{
+			return quadrance(a - b) <= Epsilon;
 		}
 
-		inline bool approx_equal(vec4i_t a, vec4i_t b) {
-			return quadrance(a - b) <= APPROX_COLOR_EPSILON;
+		inline bool approx_equal(vec4i_t a, vec4i_t b, int32 Epsilon) 
+		{
+			return quadrance(a - b) <= Epsilon;
 		}
 
-		inline vec3i_t clamp_rgb(vec3i_t color) {
+		inline vec3i_t clamp_rgb(vec3i_t color) 
+		{
 			vec3i_t result;
-			result.r = FMath::Min(255, FMath::Max(0, color.r));
-			result.g = FMath::Min(255, FMath::Max(0, color.g));
-			result.b = FMath::Min(255, FMath::Max(0, color.b));
+			result.r = FMath::Clamp(color.r, 0, 255);
+			result.g = FMath::Clamp(color.g, 0, 255);
+			result.b = FMath::Clamp(color.b, 0, 255);
 			return result;
 		}
 
-		inline vec3f_t clamp_rgb(vec3f_t color) {
+		inline vec3f_t clamp_rgb(vec3f_t color) 
+		{
 			vec3f_t result;
-			result.r = FMath::Min(255.0f, FMath::Max(0.0f, color.r));
-			result.g = FMath::Min(255.0f, FMath::Max(0.0f, color.g));
-			result.b = FMath::Min(255.0f, FMath::Max(0.0f, color.b));
+			result.r = FMath::Clamp(color.r, 0.0f, 255.0f);
+			result.g = FMath::Clamp(color.g, 0.0f, 255.0f);
+			result.b = FMath::Clamp(color.b, 0.0f, 255.0f);
+			return result;
+		}
+
+		inline vec4f_t clamp_rgba(vec4f_t color) 
+		{
+			vec4f_t result;
+			result.r = FMath::Clamp(color.r, 0.0f, 255.0f);
+			result.g = FMath::Clamp(color.g, 0.0f, 255.0f);
+			result.b = FMath::Clamp(color.b, 0.0f, 255.0f);
+			result.a = FMath::Clamp(color.a, 0.0f, 255.0f);
 			return result;
 		}
 
-		inline bool is_rgb(float color) {
+		inline bool is_rgb(float color) 
+		{
 			return color >= 0.0f && color <= 255.0f;
 		}
 
-		inline bool is_rgb(vec3f_t color) {
+		inline bool is_rgb(vec3f_t color) 
+		{
 			return is_rgb(color.r) && is_rgb(color.g) && is_rgb(color.b);
 		}
 
-		inline vec3i_t floor(vec3f_t color) {
+		inline vec3i_t floor(vec3f_t color) 
+		{
 			vec3i_t result;
 			result.r = static_cast<int>(FMath::Floor(color.r));
 			result.g = static_cast<int>(FMath::Floor(color.g));
@@ -2540,7 +2661,8 @@ namespace impl
 			return result;
 		}
 
-		inline vec3i_t round(vec3f_t color) {
+		inline vec3i_t round(vec3f_t color) 
+		{
 			vec3i_t result;
 			result.r = static_cast<int>(FMath::RoundToInt32(color.r));
 			result.g = static_cast<int>(FMath::RoundToInt32(color.g));
@@ -2548,7 +2670,8 @@ namespace impl
 			return result;
 		}
 
-		inline vec4i_t round(vec4f_t color) {
+		inline vec4i_t round(vec4f_t color) 
+		{
 			vec4i_t result;
 			result.r = static_cast<int>(FMath::RoundToInt32(color.r));
 			result.g = static_cast<int>(FMath::RoundToInt32(color.g));
@@ -2557,7 +2680,8 @@ namespace impl
 			return result;
 		}
 
-		inline vec3i_t to_vec3i(unorm8_t color) {
+		inline vec3i_t to_vec3i(unorm8_t color) 
+		{
 			vec3i_t result;
 			result.r = color.channels.r;
 			result.g = color.channels.g;
@@ -2565,7 +2689,8 @@ namespace impl
 			return result;
 		}
 
-		inline vec4i_t to_vec4i(unorm8_t color) {
+		inline vec4i_t to_vec4i(unorm8_t color) 
+		{
 			vec4i_t result;
 			result.r = color.channels.r;
 			result.g = color.channels.g;
@@ -2574,7 +2699,8 @@ namespace impl
 			return result;
 		}
 
-		inline vec3i_t to_vec3i(vec3f_t color) {
+		inline vec3i_t to_vec3i(vec3f_t color) 
+		{
 			vec3i_t result;
 			result.r = static_cast<int>(color.r);
 			result.g = static_cast<int>(color.g);
@@ -2582,7 +2708,8 @@ namespace impl
 			return result;
 		}
 
-		inline vec3f_t to_vec3f(unorm8_t color) {
+		inline vec3f_t to_vec3f(unorm8_t color) 
+		{
 			vec3f_t result;
 			result.r = color.channels.r;
 			result.g = color.channels.g;
@@ -2590,7 +2717,18 @@ namespace impl
 			return result;
 		}
 
-		inline vec3f_t to_vec3f(vec3i_t color) {
+		inline vec4f_t to_vec4f(unorm8_t color) 
+		{
+			vec4f_t result;
+			result.r = color.channels.r;
+			result.g = color.channels.g;
+			result.b = color.channels.b;
+			result.a = color.channels.a;
+			return result;
+		}
+
+		inline vec3f_t to_vec3f(vec3i_t color) 
+		{
 			vec3f_t result;
 			result.r = static_cast<float>(color.r);
 			result.g = static_cast<float>(color.g);
@@ -2598,7 +2736,18 @@ namespace impl
 			return result;
 		}
 
-		inline unorm8_t to_unorm8(vec3i_t color) {
+		inline vec4f_t to_vec4f(vec4i_t color) 
+		{
+			vec4f_t result;
+			result.r = static_cast<float>(color.r);
+			result.g = static_cast<float>(color.g);
+			result.b = static_cast<float>(color.b);
+			result.a = static_cast<float>(color.a);
+			return result;
+		}
+
+		inline unorm8_t to_unorm8(vec3i_t color) 
+		{
 			unorm8_t result;
 			result.channels.r = static_cast<uint8>(color.r);
 			result.channels.g = static_cast<uint8>(color.g);
@@ -2607,7 +2756,8 @@ namespace impl
 			return result;
 		}
 
-		inline unorm8_t to_unorm8(vec4i_t color) {
+		inline unorm8_t to_unorm8(vec4i_t color) 
+		{
 			unorm8_t result;
 			result.channels.r = static_cast<uint8>(color.r);
 			result.channels.g = static_cast<uint8>(color.g);
@@ -2616,7 +2766,8 @@ namespace impl
 			return result;
 		}
 
-		inline unorm16_t unorm8_to_unorm16(unorm8_t c8) {
+		inline unorm16_t unorm8_to_unorm16(unorm8_t c8) 
+		{
 			// (x / 255) * (2^16-1) = x * 65535 / 255 = x * 257
 			unorm16_t result;
 			result.channels.r = static_cast<uint16>(c8.channels.r * 257);
@@ -2669,6 +2820,8 @@ namespace impl
 			CEM_MAX = 16
 		};
 
+		constexpr uint8 cem_value_count[CEM_MAX] = { 2,2,2,2, 4,4,4,4, 6,6,6,6, 8,8,8,8 };
+
 
 		/**
 		 * Define normalized (starting at zero) numeric ranges that can be represented
@@ -2700,10 +2853,10 @@ namespace impl
 			RANGE_MAX
 		};
 
+#ifdef UE_MIRO_DEBUG
 		/**
 		 * Table of maximum value for each range, minimum is always zero.
 		 */
-#ifndef NDEBUG
 		const uint8 range_max_table[RANGE_MAX] = { 1,  2,  3,  4,   5,   7,   9,
 													11, 15, 19, 23,  31,  39,  47,
 													63, 79, 95, 127, 159, 191, 255 };
@@ -2740,112 +2893,6 @@ namespace impl
 			{{102,103,70,71,38},{110,111,78,79,46},{118,119,86,87,54},{126,127,94,95,62},{39,47,55,63,31}}
 		};
 
-		constexpr int8 color_endpoint_range_table4x4[2][12][16] =
-		{
-			{
-				{-1,-1,-1,-1,-1,-1,-1,-1,-1,-1,-1,-1,-1,-1,-1,-1},
-				{20,20,20,20,20,20,20,20,20,20,20,20,20,20,20,20},
-				{20,20,20,20,20,20,20,20,20,20,20,20,20,20,20,20},
-				{20,20,20,20,20,20,20,20,20,20,20,20,20,20,20,20},
-				{20,20,20,20,20,20,20,20,20,20,20,20,20,20,20,20},
-				{20,20,20,20,20,20,20,20,20,20,20,20,19,19,19,19},
-				{20,20,20,20,20,20,20,20,20,20,20,20,17,17,17,17},
-				{20,20,20,20,20,20,20,20,20,20,20,20,16,16,16,16},
-				{20,20,20,20,20,20,20,20,19,19,19,19,13,13,13,13},
-				{20,20,20,20,20,20,20,20,16,16,16,16,11,11,11,11},
-				{20,20,20,20,20,20,20,20,14,14,14,14,10,10,10,10},
-				{20,20,20,20,19,19,19,19,11,11,11,11,7,7,7,7}
-			},
-			{
-				{-1,-1,-1,-1,-1,-1,-1,-1,-1,-1,-1,-1,-1,-1,-1,-1},
-				{20,20,20,20,20,20,20,20,14,14,14,14,9,9,9,9},
-				{20,20,20,20,20,20,20,20,12,12,12,12,8,8,8,8},
-				{20,20,20,20,19,19,19,19,11,11,11,11,7,7,7,7},
-				{20,20,20,20,17,17,17,17,10,10,10,10,6,6,6,6},
-				{20,20,20,20,15,15,15,15,8,8,8,8,5,5,5,5},
-				{20,20,20,20,13,13,13,13,7,7,7,7,4,4,4,4},
-				{20,20,20,20,11,11,11,11,6,6,6,6,3,3,3,3},
-				{20,20,20,20,9,9,9,9,4,4,4,4,2,2,2,2},
-				{17,17,17,17,7,7,7,7,3,3,3,3,1,1,1,1},
-				{14,14,14,14,5,5,5,5,2,2,2,2,0,0,0,0},
-				{10,10,10,10,3,3,3,3,0,0,0,0,0,0,0,0}
-			}
-		};
-
-		constexpr int8 color_endpoint_range_table6x6[2][12][16] = 
-		{ 
-			{
-				{20,20,20,20,20,20,20,20,20,20,20,20,20,20,20,20},
-				{20,20,20,20,20,20,20,20,20,20,20,20,16,16,16,16},
-				{20,20,20,20,20,20,20,20,15,15,15,15,10,10,10,10},
-				{20,20,20,20,16,16,16,16,9,9,9,9,6,6,6,6},
-				{20,20,20,20,8,8,8,8,4,4,4,4,2,2,2,2},
-				{-1,-1,-1,-1,-1,-1,-1,-1,-1,-1,-1,-1,-1,-1,-1,-1},
-				{-1,-1,-1,-1,-1,-1,-1,-1,-1,-1,-1,-1,-1,-1,-1,-1},
-				{-1,-1,-1,-1,-1,-1,-1,-1,-1,-1,-1,-1,-1,-1,-1,-1},
-				{-1,-1,-1,-1,-1,-1,-1,-1,-1,-1,-1,-1,-1,-1,-1,-1},
-				{-1,-1,-1,-1,-1,-1,-1,-1,-1,-1,-1,-1,-1,-1,-1,-1},
-				{-1,-1,-1,-1,-1,-1,-1,-1,-1,-1,-1,-1,-1,-1,-1,-1},
-				{-1,-1,-1,-1,-1,-1,-1,-1,-1,-1,-1,-1,-1,-1,-1,-1}
-			},
-			{
-				{20,20,20,20,19,19,19,19,11,11,11,11,7,7,7,7},
-				{20,20,20,20,11,11,11,11,6,6,6,6,3,3,3,3},
-				{16,16,16,16,6,6,6,6,2,2,2,2,1,1,1,1},
-				{7,7,7,7,1,1,1,1,0,0,0,0,-1,-1,-1,-1},
-				{0,0,0,0,-1,-1,-1,-1,-1,-1,-1,-1,-1,-1,-1,-1},
-				{-1,-1,-1,-1,-1,-1,-1,-1,-1,-1,-1,-1,-1,-1,-1,-1},
-				{-1,-1,-1,-1,-1,-1,-1,-1,-1,-1,-1,-1,-1,-1,-1,-1},
-				{-1,-1,-1,-1,-1,-1,-1,-1,-1,-1,-1,-1,-1,-1,-1,-1},
-				{-1,-1,-1,-1,-1,-1,-1,-1,-1,-1,-1,-1,-1,-1,-1,-1},
-				{-1,-1,-1,-1,-1,-1,-1,-1,-1,-1,-1,-1,-1,-1,-1,-1},
-				{-1,-1,-1,-1,-1,-1,-1,-1,-1,-1,-1,-1,-1,-1,-1,-1},
-				{-1,-1,-1,-1,-1,-1,-1,-1,-1,-1,-1,-1,-1,-1,-1,-1}
-			}
-		};
-
-
-		const int8_t color_endpoint_range_table8x8[2][12][16] = { 
-			{
-				{20,20,20,20,20,20,20,20,19,19,19,19,13,13,13,13},
-				{-1,-1,-1,-1,-1,-1,-1,-1,-1,-1,-1,-1,-1,-1,-1,-1},
-				{-1,-1,-1,-1,-1,-1,-1,-1,-1,-1,-1,-1,-1,-1,-1,-1},
-				{-1,-1,-1,-1,-1,-1,-1,-1,-1,-1,-1,-1,-1,-1,-1,-1},
-				{-1,-1,-1,-1,-1,-1,-1,-1,-1,-1,-1,-1,-1,-1,-1,-1},
-				{-1,-1,-1,-1,-1,-1,-1,-1,-1,-1,-1,-1,-1,-1,-1,-1},
-				{-1,-1,-1,-1,-1,-1,-1,-1,-1,-1,-1,-1,-1,-1,-1,-1},
-				{-1,-1,-1,-1,-1,-1,-1,-1,-1,-1,-1,-1,-1,-1,-1,-1},
-				{-1,-1,-1,-1,-1,-1,-1,-1,-1,-1,-1,-1,-1,-1,-1,-1},
-				{-1,-1,-1,-1,-1,-1,-1,-1,-1,-1,-1,-1,-1,-1,-1,-1},
-				{-1,-1,-1,-1,-1,-1,-1,-1,-1,-1,-1,-1,-1,-1,-1,-1},
-				{-1,-1,-1,-1,-1,-1,-1,-1,-1,-1,-1,-1,-1,-1,-1,-1}
-			},
-			{
-				{20,20,20,20,9,9,9,9,4,4,4,4,2,2,2,2},
-				{-1,-1,-1,-1,-1,-1,-1,-1,-1,-1,-1,-1,-1,-1,-1,-1},
-				{-1,-1,-1,-1,-1,-1,-1,-1,-1,-1,-1,-1,-1,-1,-1,-1},
-				{-1,-1,-1,-1,-1,-1,-1,-1,-1,-1,-1,-1,-1,-1,-1,-1},
-				{-1,-1,-1,-1,-1,-1,-1,-1,-1,-1,-1,-1,-1,-1,-1,-1},
-				{-1,-1,-1,-1,-1,-1,-1,-1,-1,-1,-1,-1,-1,-1,-1,-1},
-				{-1,-1,-1,-1,-1,-1,-1,-1,-1,-1,-1,-1,-1,-1,-1,-1},
-				{-1,-1,-1,-1,-1,-1,-1,-1,-1,-1,-1,-1,-1,-1,-1,-1},
-				{-1,-1,-1,-1,-1,-1,-1,-1,-1,-1,-1,-1,-1,-1,-1,-1},
-				{-1,-1,-1,-1,-1,-1,-1,-1,-1,-1,-1,-1,-1,-1,-1,-1},
-				{-1,-1,-1,-1,-1,-1,-1,-1,-1,-1,-1,-1,-1,-1,-1,-1},
-				{-1,-1,-1,-1,-1,-1,-1,-1,-1,-1,-1,-1,-1,-1,-1,-1}
-			} 
-		};
-
-		//const int8_t color_endpoint_range_table10x10[2][12][16] = { 
-		//	{{-1,-1,-1,-1,-1,-1,-1,-1,-1,-1,-1,-1,-1,-1,-1,-1},{-1,-1,-1,-1,-1,-1,-1,-1,-1,-1,-1,-1,-1,-1,-1,-1},{-1,-1,-1,-1,-1,-1,-1,-1,-1,-1,-1,-1,-1,-1,-1,-1},{-1,-1,-1,-1,-1,-1,-1,-1,-1,-1,-1,-1,-1,-1,-1,-1},{-1,-1,-1,-1,-1,-1,-1,-1,-1,-1,-1,-1,-1,-1,-1,-1},{-1,-1,-1,-1,-1,-1,-1,-1,-1,-1,-1,-1,-1,-1,-1,-1},{-1,-1,-1,-1,-1,-1,-1,-1,-1,-1,-1,-1,-1,-1,-1,-1},{-1,-1,-1,-1,-1,-1,-1,-1,-1,-1,-1,-1,-1,-1,-1,-1},{-1,-1,-1,-1,-1,-1,-1,-1,-1,-1,-1,-1,-1,-1,-1,-1},{-1,-1,-1,-1,-1,-1,-1,-1,-1,-1,-1,-1,-1,-1,-1,-1},{-1,-1,-1,-1,-1,-1,-1,-1,-1,-1,-1,-1,-1,-1,-1,-1},{-1,-1,-1,-1,-1,-1,-1,-1,-1,-1,-1,-1,-1,-1,-1,-1}},
-		//	{{-1,-1,-1,-1,-1,-1,-1,-1,-1,-1,-1,-1,-1,-1,-1,-1},{-1,-1,-1,-1,-1,-1,-1,-1,-1,-1,-1,-1,-1,-1,-1,-1},{-1,-1,-1,-1,-1,-1,-1,-1,-1,-1,-1,-1,-1,-1,-1,-1},{-1,-1,-1,-1,-1,-1,-1,-1,-1,-1,-1,-1,-1,-1,-1,-1},{-1,-1,-1,-1,-1,-1,-1,-1,-1,-1,-1,-1,-1,-1,-1,-1},{-1,-1,-1,-1,-1,-1,-1,-1,-1,-1,-1,-1,-1,-1,-1,-1},{-1,-1,-1,-1,-1,-1,-1,-1,-1,-1,-1,-1,-1,-1,-1,-1},{-1,-1,-1,-1,-1,-1,-1,-1,-1,-1,-1,-1,-1,-1,-1,-1},{-1,-1,-1,-1,-1,-1,-1,-1,-1,-1,-1,-1,-1,-1,-1,-1},{-1,-1,-1,-1,-1,-1,-1,-1,-1,-1,-1,-1,-1,-1,-1,-1},{-1,-1,-1,-1,-1,-1,-1,-1,-1,-1,-1,-1,-1,-1,-1,-1},{-1,-1,-1,-1,-1,-1,-1,-1,-1,-1,-1,-1,-1,-1,-1,-1}} 
-		//};
-
-		//const int8_t color_endpoint_range_table12x12[2][12][16] = { 
-		//	{{-1,-1,-1,-1,-1,-1,-1,-1,-1,-1,-1,-1,-1,-1,-1,-1},{-1,-1,-1,-1,-1,-1,-1,-1,-1,-1,-1,-1,-1,-1,-1,-1},{-1,-1,-1,-1,-1,-1,-1,-1,-1,-1,-1,-1,-1,-1,-1,-1},{-1,-1,-1,-1,-1,-1,-1,-1,-1,-1,-1,-1,-1,-1,-1,-1},{-1,-1,-1,-1,-1,-1,-1,-1,-1,-1,-1,-1,-1,-1,-1,-1},{-1,-1,-1,-1,-1,-1,-1,-1,-1,-1,-1,-1,-1,-1,-1,-1},{-1,-1,-1,-1,-1,-1,-1,-1,-1,-1,-1,-1,-1,-1,-1,-1},{-1,-1,-1,-1,-1,-1,-1,-1,-1,-1,-1,-1,-1,-1,-1,-1},{-1,-1,-1,-1,-1,-1,-1,-1,-1,-1,-1,-1,-1,-1,-1,-1},{-1,-1,-1,-1,-1,-1,-1,-1,-1,-1,-1,-1,-1,-1,-1,-1},{-1,-1,-1,-1,-1,-1,-1,-1,-1,-1,-1,-1,-1,-1,-1,-1},{-1,-1,-1,-1,-1,-1,-1,-1,-1,-1,-1,-1,-1,-1,-1,-1}},
-		//	{{-1,-1,-1,-1,-1,-1,-1,-1,-1,-1,-1,-1,-1,-1,-1,-1},{-1,-1,-1,-1,-1,-1,-1,-1,-1,-1,-1,-1,-1,-1,-1,-1},{-1,-1,-1,-1,-1,-1,-1,-1,-1,-1,-1,-1,-1,-1,-1,-1},{-1,-1,-1,-1,-1,-1,-1,-1,-1,-1,-1,-1,-1,-1,-1,-1},{-1,-1,-1,-1,-1,-1,-1,-1,-1,-1,-1,-1,-1,-1,-1,-1},{-1,-1,-1,-1,-1,-1,-1,-1,-1,-1,-1,-1,-1,-1,-1,-1},{-1,-1,-1,-1,-1,-1,-1,-1,-1,-1,-1,-1,-1,-1,-1,-1},{-1,-1,-1,-1,-1,-1,-1,-1,-1,-1,-1,-1,-1,-1,-1,-1},{-1,-1,-1,-1,-1,-1,-1,-1,-1,-1,-1,-1,-1,-1,-1,-1},{-1,-1,-1,-1,-1,-1,-1,-1,-1,-1,-1,-1,-1,-1,-1,-1},{-1,-1,-1,-1,-1,-1,-1,-1,-1,-1,-1,-1,-1,-1,-1,-1},{-1,-1,-1,-1,-1,-1,-1,-1,-1,-1,-1,-1,-1,-1,-1,-1}} 
-		//};
-
 
 		const uint8 color_unquantize_table[21][256] =
 		{
@@ -3507,11 +3554,11 @@ namespace impl
 			split_high_low(b3, bits, t3, m3);
 			split_high_low(b4, bits, t4, m4);
 
-			DCHECK(t0 < 3);
-			DCHECK(t1 < 3);
-			DCHECK(t2 < 3);
-			DCHECK(t3 < 3);
-			DCHECK(t4 < 3);
+			miro_check(t0 < 3);
+			miro_check(t1 < 3);
+			miro_check(t2 < 3);
+			miro_check(t3 < 3);
+			miro_check(t4 < 3);
 
 			if (t0 < 3 && t1 < 3 && t2 < 3 && t3 < 3 && t4 < 3)
 			{
@@ -3545,9 +3592,9 @@ namespace impl
 			split_high_low(b1, bits, q1, m1);
 			split_high_low(b2, bits, q2, m2);
 
-			DCHECK(q0 < 5);
-			DCHECK(q1 < 5);
-			DCHECK(q2 < 5);
+			miro_check(q0 < 5);
+			miro_check(q1 < 5);
+			miro_check(q2 < 5);
 
 			if (q0 < 5 && q1 < 5 && q2 < 5)
 			{
@@ -3605,7 +3652,7 @@ namespace impl
 			size_t count,
 			bitwriter& writer,
 			size_t bits) {
-			DCHECK(count > 0);
+			miro_check(count > 0);
 			for (size_t i = 0; i < count; ++i) {
 				writer.write8(numbers[i], bits);
 			}
@@ -3620,9 +3667,10 @@ namespace impl
 			size_t count,
 			range_t range,
 			bitwriter writer) {
-#ifndef NDEBUG
+
+#ifdef UE_MIRO_DEBUG
 			for (size_t i = 0; i < count; ++i) {
-				DCHECK(numbers[i] <= range_max_table[range]);
+				miro_check(numbers[i] <= range_max_table[range]);
 			}
 #endif
 
@@ -3670,7 +3718,7 @@ namespace impl
 
 		inline uint8 quantize_color(range_t quant, int c)
 		{
-			DCHECK(c >= 0 && c <= 255);
+			miro_check(c >= 0 && c <= 255);
 			if (c >= 0 && c <= 255)
 			{
 				return color_quantize_table[quant][c];
@@ -3699,7 +3747,7 @@ namespace impl
 
 		uint8 unquantize_color(range_t quant, int c)
 		{
-			DCHECK(c >= 0 && c <= 255);
+			miro_check(c >= 0 && c <= 255);
 			if (c >= 0 && c <= 255)
 			{
 				return color_unquantize_table[quant][c];
@@ -3749,14 +3797,36 @@ namespace impl
 			endpoint_unquantized[1] = unquantize_color(endpoint_quant, endpoint_quantized[1]);
 		}
 
-		void encode_rgb_direct(range_t endpoint_quant,
-			vec3i_t e0,
+		void encode_rgb_base_scale(range_t endpoint_quant,
 			vec3i_t e1,
-			uint8 endpoint_quantized[6],
+			uint8 scale,
+			uint8 endpoint_quantized[4],
 			vec3i_t endpoint_unquantized[2])
 		{
-			vec3i_t e0q = quantize_color(endpoint_quant, e0);
 			vec3i_t e1q = quantize_color(endpoint_quant, e1);
+			vec3i_t e1u = unquantize_color(endpoint_quant, e1q);
+
+			uint8 scaleq = color_quantize_table[endpoint_quant][scale];
+			uint8 scaleu = color_unquantize_table[endpoint_quant][scaleq];
+			vec3i_t e0u = (e1u * int32(scaleu)) / 256;
+
+			endpoint_quantized[0] = static_cast<uint8>(e1q.r);
+			endpoint_quantized[1] = static_cast<uint8>(e1q.b);
+			endpoint_quantized[2] = static_cast<uint8>(e1q.g);
+			endpoint_quantized[3] = static_cast<uint8>(scaleq);
+
+			endpoint_unquantized[0] = e0u;
+			endpoint_unquantized[1] = e1u;
+		}
+
+		void encode_rgb_direct(range_t endpoint_quant,
+			vec3i_t& InOutE0,
+			vec3i_t& InOutE1,
+			uint8 endpoint_quantized[6],
+			vec3i_t endpoint_unquantized[2])
+		{
+			vec3i_t e0q = quantize_color(endpoint_quant, InOutE0);
+			vec3i_t e1q = quantize_color(endpoint_quant, InOutE1);
 			vec3i_t e0u = unquantize_color(endpoint_quant, e0q);
 			vec3i_t e1u = unquantize_color(endpoint_quant, e1q);
 
@@ -3774,6 +3844,9 @@ namespace impl
 
 				endpoint_unquantized[0] = e1u;
 				endpoint_unquantized[1] = e0u;
+
+				// We need them swapped in the calling code too.
+				Swap(InOutE0, InOutE1);
 			}
 			else
 			{
@@ -3790,11 +3863,54 @@ namespace impl
 		}
 
 
+		void encode_rgb_direct(range_t endpoint_quant,
+			vec4i_t& e0,
+			vec4i_t& e1,
+			uint8 endpoint_quantized[6],
+			vec3i_t endpoint_unquantized[2])
+		{
+			vec3i_t e0q = quantize_color(endpoint_quant, e0.rgb());
+			vec3i_t e1q = quantize_color(endpoint_quant, e1.rgb());
+			vec3i_t e0u = unquantize_color(endpoint_quant, e0q);
+			vec3i_t e1u = unquantize_color(endpoint_quant, e1q);
+
+			// ASTC uses a different blue contraction encoding when the sum of values for
+			// the first endpoint is larger than the sum of values in the second
+			// endpoint. Sort the endpoints to ensure that the normal encoding is used.
+			if (color_channel_sum(e0u) > color_channel_sum(e1u))
+			{
+				endpoint_quantized[0] = static_cast<uint8>(e1q.r);
+				endpoint_quantized[1] = static_cast<uint8>(e0q.r);
+				endpoint_quantized[2] = static_cast<uint8>(e1q.g);
+				endpoint_quantized[3] = static_cast<uint8>(e0q.g);
+				endpoint_quantized[4] = static_cast<uint8>(e1q.b);
+				endpoint_quantized[5] = static_cast<uint8>(e0q.b);
+
+				endpoint_unquantized[0] = e1u;
+				endpoint_unquantized[1] = e0u;
+
+				Swap(e0, e1);
+			}
+			else
+			{
+				endpoint_quantized[0] = static_cast<uint8>(e0q.r);
+				endpoint_quantized[1] = static_cast<uint8>(e1q.r);
+				endpoint_quantized[2] = static_cast<uint8>(e0q.g);
+				endpoint_quantized[3] = static_cast<uint8>(e1q.g);
+				endpoint_quantized[4] = static_cast<uint8>(e0q.b);
+				endpoint_quantized[5] = static_cast<uint8>(e1q.b);
+
+				endpoint_unquantized[0] = e0u;
+				endpoint_unquantized[1] = e1u;
+			}
+		}
+
 		void encode_rgba_direct(range_t endpoint_quant,
-			vec4i_t e0,
-			vec4i_t e1,
+			vec4i_t& e0,
+			vec4i_t& e1,
 			uint8 endpoint_quantized[8],
-			vec4i_t endpoint_unquantized[2])
+			vec4i_t endpoint_unquantized[2],
+			bool bDualPlane)
 		{
 			vec4i_t e0q = quantize_color(endpoint_quant, e0);
 			vec4i_t e1q = quantize_color(endpoint_quant, e1);
@@ -3812,9 +3928,13 @@ namespace impl
 				endpoint_quantized[3] = static_cast<uint8>(e0q.g);
 				endpoint_quantized[4] = static_cast<uint8>(e1q.b);
 				endpoint_quantized[5] = static_cast<uint8>(e0q.b);
+				endpoint_quantized[6] = static_cast<uint8>(e1q.a);
+				endpoint_quantized[7] = static_cast<uint8>(e0q.a);
 
 				endpoint_unquantized[0] = e1u;
 				endpoint_unquantized[1] = e0u;
+				
+				Swap(e0, e1);
 			}
 			else
 			{
@@ -3824,25 +3944,30 @@ namespace impl
 				endpoint_quantized[3] = static_cast<uint8>(e1q.g);
 				endpoint_quantized[4] = static_cast<uint8>(e0q.b);
 				endpoint_quantized[5] = static_cast<uint8>(e1q.b);
+				endpoint_quantized[6] = static_cast<uint8>(e0q.a);
+				endpoint_quantized[7] = static_cast<uint8>(e1q.a);
 
 				endpoint_unquantized[0] = e0u;
 				endpoint_unquantized[1] = e1u;
 			}
 
 			// Sort alpha endpoints
-			if (e0.a > e1.a)
-			{
-				endpoint_quantized[6] = static_cast<uint8>(e1q.a);
-				endpoint_quantized[7] = static_cast<uint8>(e0q.a);
-				endpoint_unquantized[0].a = e1u.a;
-				endpoint_unquantized[1].a = e0u.a;
-			}
-			else
+			if (bDualPlane)
 			{
-				endpoint_quantized[6] = static_cast<uint8>(e0q.a);
-				endpoint_quantized[7] = static_cast<uint8>(e1q.a);
-				endpoint_unquantized[0].a = e0u.a;
-				endpoint_unquantized[1].a = e1u.a;
+				if (e0.a > e1.a)
+				{
+					endpoint_quantized[6] = static_cast<uint8>(e1q.a);
+					endpoint_quantized[7] = static_cast<uint8>(e0q.a);
+					endpoint_unquantized[0].a = e1u.a;
+					endpoint_unquantized[1].a = e0u.a;
+				}
+				else
+				{
+					endpoint_quantized[6] = static_cast<uint8>(e0q.a);
+					endpoint_quantized[7] = static_cast<uint8>(e1q.a);
+					endpoint_unquantized[0].a = e0u.a;
+					endpoint_unquantized[1].a = e1u.a;
+				}
 			}
 		}
 
@@ -3865,17 +3990,11 @@ namespace impl
 			bool dual_plane = false
 		)
 		{
-			DCHECK(weight_quant <= RANGE_32);
-			DCHECK(endpoint_quant < RANGE_MAX);
-			DCHECK(color_endpoint_mode < CEM_MAX);
-			DCHECK(partition_count == 1);
-			DCHECK(compute_ise_bitcount(BLOCK_WIDTH*BLOCK_HEIGHT, weight_quant) < MAXIMUM_ENCODED_WEIGHT_BITS);
-
-			//        if (s_debuglog)
-			//        {
-			//            UE_LOG(LogMutableCore,Warning," physicalising interesting block" );
-			//        }
-
+			miro_check(weight_quant <= RANGE_32);
+			miro_check(endpoint_quant < RANGE_MAX);
+			miro_check(color_endpoint_mode < CEM_MAX);
+			miro_check(partition_count == 1);
+			miro_check(compute_ise_bitcount(BLOCK_WIDTH*BLOCK_HEIGHT, weight_quant) < MAXIMUM_ENCODED_WEIGHT_BITS);
 
 			size_t n = BLOCK_WIDTH;
 			size_t m = BLOCK_HEIGHT;
@@ -3886,10 +4005,8 @@ namespace impl
 			static const uint8 r_table[RANGE_32 + 1] = { 0x2, 0x3, 0x4, 0x5, 0x6, 0x7,
 														   0x2, 0x3, 0x4, 0x5, 0x6, 0x7 };
 
-			bool h = false;
-			if (weight_quant < RANGE_32 + 1) h = h_table[weight_quant];
-			size_t r = 0;
-			if (weight_quant < RANGE_32 + 1) r = r_table[weight_quant];
+			bool h = (weight_quant <= RANGE_32) ? h_table[weight_quant] : false;
+			size_t r = (weight_quant < RANGE_32) ? r_table[weight_quant] : 0;
 
 			bool d = dual_plane;
 
@@ -3902,6 +4019,7 @@ namespace impl
 			size_t cem = color_endpoint_mode;
 
 			// Block mode
+			
 			// Actually weight-grid and block sizes don't need to match. We are defining the weight-grid size here but for
 			// now we make them match.
 			if (BLOCK_WIDTH == 4 && BLOCK_HEIGHT == 4)
@@ -3923,7 +4041,7 @@ namespace impl
 			else if (BLOCK_WIDTH == 6 && BLOCK_HEIGHT == 6)
 			{
 				// This encoder doesn't support 6x6 and dual plane: that requires mismatching and interpolating the pixel grid and weight grid
-				check(dual_plane == 0);
+				miro_check(dual_plane == 0);
 
 				// Use the tenth row of Table C.2.8 in the ASTC specification.
 				size_t a = m - 6;
@@ -3941,7 +4059,7 @@ namespace impl
 			else if (BLOCK_WIDTH == 8 && BLOCK_HEIGHT == 8)
 			{
 				// This encoder doesn't support 8x8 and dual plane: that requires mismatching and interpolating the pixel grid and weight grid
-				check(dual_plane==0);
+				miro_check(dual_plane==0);
 
 				// Use the tenth row of Table C.2.8 in the ASTC specification.
 				size_t a = m - 6;
@@ -3983,13 +4101,193 @@ namespace impl
 		}
 
 
+		inline void symbolic_to_physical
+		(
+			int32 GRID_WIDTH, int32 GRID_HEIGHT,
+			color_endpoint_mode_t color_endpoint_mode,
+			range_t endpoint_quant,
+			range_t weight_quant,
+
+			size_t partition_count,
+
+			const uint8 endpoint_ise[MAXIMUM_ENCODED_COLOR_ENDPOINT_BYTES],
+
+			// FIXME: +1 needed here because orbits_8ptr breaks when the offset reaches
+			// the last byte which always happens if the weight mode is RANGE_32.
+			const uint8 weights_ise[MAXIMUM_ENCODED_WEIGHT_BYTES + 1],
+
+			PhysicalBlock* pb,
+			bool dual_plane = false
+		)
+		{
+			miro_check(weight_quant <= RANGE_32);
+			miro_check(endpoint_quant < RANGE_MAX);
+			miro_check(color_endpoint_mode < CEM_MAX);
+			miro_check(partition_count == 1);
+			miro_check(compute_ise_bitcount(GRID_WIDTH * GRID_HEIGHT, weight_quant) < MAXIMUM_ENCODED_WEIGHT_BITS);
+
+			size_t n = GRID_WIDTH;
+			size_t m = GRID_HEIGHT;
+
+			static const bool h_table[RANGE_32 + 1] = { 0, 0, 0, 0, 0, 0,
+														1, 1, 1, 1, 1, 1 };
+
+			static const uint8 r_table[RANGE_32 + 1] = { 0x2, 0x3, 0x4, 0x5, 0x6, 0x7,
+														   0x2, 0x3, 0x4, 0x5, 0x6, 0x7 };
+
+			bool h = (weight_quant <= RANGE_32) ? h_table[weight_quant] : false;
+			size_t r = (weight_quant < RANGE_32) ? r_table[weight_quant] : 0;
+
+			bool d = dual_plane;
+
+			size_t part_value = partition_count - 1;
+
+			size_t cem_offset = 13;
+			size_t ced_offset = 17;
+
+			size_t cem_bits = 4;
+			size_t cem = color_endpoint_mode;
+
+			// Block mode
+
+			// Actually weight-grid and block sizes don't need to match. We are defining the weight-grid size here but for
+			// now we make them match.
+			if (GRID_WIDTH == 4 && GRID_HEIGHT == 4)
+			{
+				// Use the first row of Table C.2.8 in the ASTC specification.
+				size_t a = m - 2;
+				size_t b = n - 4;
+
+				orbits8_ptr(pb->data, 0, getbit(r, 1), 1);
+				orbits8_ptr(pb->data, 1, getbit(r, 2), 1);
+				orbits8_ptr(pb->data, 2, 0, 1);
+				orbits8_ptr(pb->data, 3, 0, 1);
+				orbits8_ptr(pb->data, 4, getbit(r, 0), 1);
+				orbits8_ptr(pb->data, 5, a, 2);
+				orbits8_ptr(pb->data, 7, b, 2);
+				orbits8_ptr(pb->data, 9, h, 1);
+				orbits8_ptr(pb->data, 10, d, 1);
+			}
+			else if (GRID_WIDTH == 6 && GRID_HEIGHT == 6)
+			{
+				// This encoder doesn't support 6x6 and dual plane: that requires mismatching and interpolating the pixel grid and weight grid
+				miro_check(dual_plane == 0);
+
+				// Use the tenth row of Table C.2.8 in the ASTC specification.
+				size_t a = m - 6;
+				size_t b = n - 6;
+
+				orbits8_ptr(pb->data, 0, 0, 2);
+				orbits8_ptr(pb->data, 2, getbit(r, 1), 1);
+				orbits8_ptr(pb->data, 3, getbit(r, 2), 1);
+				orbits8_ptr(pb->data, 4, getbit(r, 0), 1);
+				orbits8_ptr(pb->data, 5, a, 2);
+				orbits8_ptr(pb->data, 7, 0, 1);
+				orbits8_ptr(pb->data, 8, 1, 1);
+				orbits8_ptr(pb->data, 9, b, 2);
+			}
+			else if (GRID_WIDTH == 8 && GRID_HEIGHT == 8)
+			{
+				// This encoder doesn't support 8x8 and dual plane: that requires mismatching and interpolating the pixel grid and weight grid
+				miro_check(dual_plane == 0);
+
+				// Use the tenth row of Table C.2.8 in the ASTC specification.
+				size_t a = m - 6;
+				size_t b = n - 6;
+
+				orbits8_ptr(pb->data, 0, 0, 2);
+				orbits8_ptr(pb->data, 2, getbit(r, 1), 1);
+				orbits8_ptr(pb->data, 3, getbit(r, 2), 1);
+				orbits8_ptr(pb->data, 4, getbit(r, 0), 1);
+				orbits8_ptr(pb->data, 5, a, 2);
+				orbits8_ptr(pb->data, 7, 0, 1);
+				orbits8_ptr(pb->data, 8, 1, 1);
+				orbits8_ptr(pb->data, 9, b, 2);
+			}
+			else if (GRID_WIDTH == 7 && GRID_HEIGHT == 7)
+			{
+				if (weight_quant == RANGE_3)
+				{
+					// Use the tenth row of Table C.2.8 in the ASTC specification with r=0b011, a=1, b=1
+#ifdef UE_MIRO_DEBUG
+					size_t a = m - 6;
+					size_t b = n - 6;
+					orbits8_ptr(pb->data, 0, 0, 2);
+					orbits8_ptr(pb->data, 2, getbit(r, 1), 1);
+					orbits8_ptr(pb->data, 3, getbit(r, 2), 1);
+					orbits8_ptr(pb->data, 4, getbit(r, 0), 1);
+					orbits8_ptr(pb->data, 5, a, 2);
+					orbits8_ptr(pb->data, 7, 0, 1);
+					orbits8_ptr(pb->data, 8, 1, 1);
+					orbits8_ptr(pb->data, 9, b, 2);
+					check(pb->data[0] == 0b00110100);
+					check(pb->data[1] == 0b011);
+#else
+					pb->data[0] = 0b00110100;
+					pb->data[1] = 0b011;					
+#endif
+				}
+				else
+				{
+					// Use the tenth row of Table C.2.8 in the ASTC specification with r=0b010, a=1, b=1
+					miro_check(weight_quant == RANGE_2);
+
+#ifdef UE_MIRO_DEBUG
+					size_t a = m - 6;
+					size_t b = n - 6;
+					orbits8_ptr(pb->data, 0, 0, 2);
+					orbits8_ptr(pb->data, 2, getbit(r, 1), 1);
+					orbits8_ptr(pb->data, 3, getbit(r, 2), 1);
+					orbits8_ptr(pb->data, 4, getbit(r, 0), 1);
+					orbits8_ptr(pb->data, 5, a, 2);
+					orbits8_ptr(pb->data, 7, 0, 1);
+					orbits8_ptr(pb->data, 8, 1, 1);
+					orbits8_ptr(pb->data, 9, b, 2);
+					check(pb->data[0] == 0b00100100);
+					check(pb->data[1] == 0b011);
+#else
+					pb->data[0] = 0b00100100;
+					pb->data[1] = 0b011;
+#endif
+				}
+
+				// This encoding row doesn't support dual plane.
+				miro_check(!dual_plane);
+			}
+			else
+			{
+				check(false);
+			}
+
+
+			// Partitions
+			orbits8_ptr(pb->data, 11, part_value, 2);
+
+			// CEM
+			orbits8_ptr(pb->data, cem_offset, cem, cem_bits);
+
+			copy_bytes(endpoint_ise, MAXIMUM_ENCODED_COLOR_ENDPOINT_BYTES, pb->data, ced_offset);
+
+			reverse_bytes(weights_ise, MAXIMUM_ENCODED_WEIGHT_BYTES, pb->data + 15);
+
+			if (dual_plane)
+			{
+				size_t bits_for_weights = compute_ise_bitcount(2 * GRID_WIDTH * GRID_HEIGHT, weight_quant);
+
+				// \TODO: dual plane always sets alpha as the separate channel
+				size_t secondPlaneChannel = 3;
+				orbits8_ptr(pb->data, 128 - bits_for_weights - 2, secondPlaneChannel, 2);
+			}
+		}
+
+
 		uint8 quantize_weight(range_t weight_quant, size_t weight)
 		{
-			DCHECK(weight_quant <= RANGE_32);
+			miro_check(weight_quant <= RANGE_32);
 			// anticto: this may happen because of rounding in some ranges.
 			// apparently, it is ok to clamp it, based on what the arm reference implementation does.
 			if (weight > 1024) weight = 1024;
-			//DCHECK(weight <= 1024);
+			//miro_check(weight <= 1024);
 			if ((weight_quant <= RANGE_32) && (weight <= 1024))
 			{
 				return weight_quantize_table[weight_quant][weight];
@@ -4014,7 +4312,7 @@ namespace impl
 		 */
 		size_t project(size_t k, size_t m, size_t t)
 		{
-			DCHECK(k > 0);
+			miro_check(k > 0);
 			// anticto fix: underflow is possible because we use the unquantized limit, which may be
 			// bigger than the value. so i think we need to clamp.
 			// return size_t((t - m) * 1024) / k;
@@ -4027,13 +4325,13 @@ namespace impl
 		 */
 		size_t project(vec3i_t k, int kk, vec3i_t m, vec3i_t t)
 		{
-			DCHECK(kk > 0);
+			miro_check(kk > 0);
 			return static_cast<size_t>(FMath::Clamp(dot(t - m, k) * 1024 / kk, 0, 1024));
 		}
 
 		size_t project(vec4i_t k, int kk, vec4i_t m, vec4i_t t)
 		{
-			DCHECK(kk > 0);
+			miro_check(kk > 0);
 			return static_cast<size_t>(FMath::Clamp(dot(t - m, k) * 1024 / kk, 0, 1024));
 		}
 
@@ -4044,7 +4342,7 @@ namespace impl
 			uint8 l1,
 			uint8 weights[BLOCK_TEXEL_COUNT])
 		{
-			DCHECK(l0 <= l1);
+			miro_check(l0 <= l1);
 			if (l0 < l1)
 			{
 				size_t k = l1 - l0;
@@ -4060,7 +4358,7 @@ namespace impl
 			{
 				for (size_t i = 0; i < BLOCK_TEXEL_COUNT; ++i)
 				{
-					DCHECK(static_cast<size_t>(texels[i]) == l0);
+					miro_check(static_cast<size_t>(texels[i]) == l0);
 					weights[i] = quantize_weight(quant, 0);
 				}
 			}
@@ -4084,7 +4382,7 @@ namespace impl
 			}
 			else
 			{
-				DCHECK(l0 < l1);
+				miro_check(l0 < l1);
 
 				size_t k = l1 - l0;
 				size_t m = l0;
@@ -4152,34 +4450,51 @@ namespace impl
 			}
 		}
 
-
+#ifdef UE_MIRO_DEBUG
 		// This should be compiled out, together with the actual tables.
-		template<int32 BLOCK_WIDTH, int32 BLOCK_HEIGHT>
-		constexpr range_t endpoint_quantization(size_t partitions, range_t weight_quant, color_endpoint_mode_t endpoint_mode)
+		range_t endpoint_quantization(int32 WeightX, int32 WeightY, int32 Partitions, range_t weight_quant, color_endpoint_mode_t endpoint_mode, bool bDualPlane)
 		{
-			int8 ce_range = 0;
+			// Brute force: shouldn't happen at runtime
+			int8 ce_range = RANGE_MAX;
 			
-			if (BLOCK_WIDTH == 4 && BLOCK_HEIGHT == 4)
-			{
-				ce_range = color_endpoint_range_table4x4[partitions - 1][weight_quant][endpoint_mode];
-			}
-			else if (BLOCK_WIDTH == 6 && BLOCK_HEIGHT == 6)
+			miro_check( !bDualPlane || Partitions<4);
+			int32 ConfigBits = 17;
+			if (Partitions > 1)
 			{
-				ce_range = color_endpoint_range_table6x6[partitions - 1][weight_quant][endpoint_mode];
+				check(false);
 			}
-			else if(BLOCK_WIDTH == 8 && BLOCK_HEIGHT == 8)
+
+			int32 WeightZ = 1;
+			int32 NumWeights = WeightX * WeightY * WeightZ;
+			miro_check(NumWeights <= 64);
+
+			if (bDualPlane)
 			{
-				ce_range = color_endpoint_range_table8x8[partitions - 1][weight_quant][endpoint_mode];
+				ConfigBits += 2;
+				NumWeights *= 2;
 			}
-			else
+
+			int32 WeightBits = compute_ise_bitcount(NumWeights, weight_quant);
+			miro_check(WeightBits >= 24);
+			miro_check(WeightBits <= MAXIMUM_ENCODED_WEIGHT_BITS);
+			int32 RemainingBits = 128 - ConfigBits - WeightBits;
+
+			int32 CEMValueCount = cem_value_count[endpoint_mode] * Partitions;
+			miro_check(CEMValueCount<=18);
+
+			for (ce_range = RANGE_MAX - 1; ce_range >= 0; --ce_range)
 			{
-				check(false);
+				int32 RangeBits = compute_ise_bitcount(CEMValueCount, static_cast<range_t>(ce_range) );
+				if (RangeBits <= RemainingBits)
+				{
+					break;
+				}
 			}
 
-			DCHECK(ce_range >= 0 && ce_range <= RANGE_MAX);
+			miro_check(ce_range >= 0 && ce_range <= RANGE_MAX);
 			return static_cast<range_t>(ce_range);
 		}
-
+#endif
 
 
 		/** Write void extent block bits for LDR mode and unused extent coordinates. */
@@ -4199,27 +4514,35 @@ namespace impl
 		void encode_luminance(const uint8 texels[BLOCK_WIDTH*BLOCK_HEIGHT], PhysicalBlock* physical_block)
 		{
 			size_t partition_count = 1;
+			bool bDualPlane = false;
 
 			color_endpoint_mode_t color_endpoint_mode = CEM_LDR_LUMINANCE_DIRECT;
-			range_t weight_quant = RANGE_32;
+			range_t weight_quant = RANGE_MAX;
+			range_t endpoint_quant = RANGE_MAX;
 			if (BLOCK_WIDTH == 4 && BLOCK_HEIGHT == 4)
 			{
 				weight_quant = RANGE_32;
+				endpoint_quant = RANGE_256;
 			}
 			else if (BLOCK_WIDTH == 6 && BLOCK_HEIGHT == 6)
 			{
 				weight_quant = RANGE_4;
+				endpoint_quant = RANGE_256;
 			}
 			else if (BLOCK_WIDTH == 8 && BLOCK_HEIGHT == 8)
 			{
 				weight_quant = RANGE_2;
+				endpoint_quant = RANGE_256;
 			}
 			else
 			{
 				check(false);
 			}
 
-			range_t endpoint_quant = endpoint_quantization<BLOCK_WIDTH, BLOCK_HEIGHT>(partition_count, weight_quant, color_endpoint_mode);
+#ifdef UE_MIRO_DEBUG
+			range_t ExpectedEndpointQuant = endpoint_quantization(BLOCK_WIDTH, BLOCK_HEIGHT, partition_count, weight_quant, color_endpoint_mode, bDualPlane);
+			miro_check(endpoint_quant == ExpectedEndpointQuant);
+#endif
 
 			uint8 l0 = 255;
 			uint8 l1 = 0;
@@ -4256,32 +4579,43 @@ namespace impl
 			PhysicalBlock* physical_block)
 		{
 			constexpr size_t partition_count = 1;
+			bool bDualPlane = false;
 
 			color_endpoint_mode_t color_endpoint_mode = CEM_LDR_RGB_DIRECT;
-			range_t weight_quant = RANGE_12;
+			range_t weight_quant = RANGE_MAX;
+			range_t endpoint_quant = RANGE_MAX;
 
 			if (BLOCK_WIDTH == 4 && BLOCK_HEIGHT == 4)
 			{
 				weight_quant = RANGE_12;
+				endpoint_quant = RANGE_256;
 			}
 			else if (BLOCK_WIDTH == 6 && BLOCK_HEIGHT == 6)
 			{
-				weight_quant = RANGE_5; // \TODO: enlarge later?
+				weight_quant = RANGE_3;
+				endpoint_quant = RANGE_256;
 			}
 			else if (BLOCK_WIDTH == 8 && BLOCK_HEIGHT == 8)
 			{
-				weight_quant = RANGE_2; // \TODO: enlarge later?
+				weight_quant = RANGE_2;
+				endpoint_quant = RANGE_192;
 			}
 			else
 			{
 				check(false);
 			}
 
-			range_t endpoint_quant = endpoint_quantization<BLOCK_WIDTH, BLOCK_HEIGHT>(partition_count, weight_quant, color_endpoint_mode);
+#ifdef UE_MIRO_DEBUG
+			range_t ExpectedEndpointQuant = endpoint_quantization(BLOCK_WIDTH, BLOCK_HEIGHT, partition_count, weight_quant, color_endpoint_mode, bDualPlane);
+			miro_check(endpoint_quant == ExpectedEndpointQuant);
+#endif
 
 			vec3i_t endpoint_unquantized[2];
 			uint8 endpoint_quantized[6];
-			encode_rgb_direct(endpoint_quant, round(e0), round(e1), endpoint_quantized, endpoint_unquantized);
+			vec3i_t endpoint0 = round(e0);
+			vec3i_t endpoint1 = round(e1);
+
+			encode_rgb_direct(endpoint_quant, endpoint0, endpoint1, endpoint_quantized, endpoint_unquantized);
 
 			uint8 weights_quantized[BLOCK_WIDTH * BLOCK_HEIGHT];
 			calculate_quantized_weights_rgb<BLOCK_WIDTH*BLOCK_HEIGHT>(texels, weight_quant, endpoint_unquantized[0], endpoint_unquantized[1], weights_quantized);
@@ -4304,11 +4638,12 @@ namespace impl
 			PhysicalBlock* physical_block)
 		{
 			size_t partition_count = 1;
+			bool bDualPlane = false;
 
 			color_endpoint_mode_t color_endpoint_mode = CEM_LDR_RGBA_DIRECT;
 
-			range_t weight_quant = RANGE_5;
-			range_t endpoint_quant = RANGE_16;
+			range_t weight_quant = RANGE_MAX;
+			range_t endpoint_quant = RANGE_MAX;
 			if (BLOCK_WIDTH == 4 && BLOCK_HEIGHT == 4)
 			{
 				weight_quant = RANGE_5;
@@ -4326,9 +4661,17 @@ namespace impl
 				check(false);
 			}
 
+			// There is only one valid option for endpoint_quant
+#ifdef UE_MIRO_DEBUG
+			range_t ExpectedEndpointQuant = endpoint_quantization(BLOCK_WIDTH, BLOCK_HEIGHT, partition_count, weight_quant, color_endpoint_mode, bDualPlane);
+			miro_check( endpoint_quant == ExpectedEndpointQuant );
+#endif
+
 			vec4i_t endpoint_unquantized[2];
 			uint8 endpoint_quantized[8];
-			encode_rgba_direct(endpoint_quant, round(e0), round(e1), endpoint_quantized, endpoint_unquantized);
+			vec4i_t endpoint0 = round(e0);
+			vec4i_t endpoint1 = round(e1);
+			encode_rgba_direct(endpoint_quant, endpoint0, endpoint1, endpoint_quantized, endpoint_unquantized, true);
 
 			uint8 weights_quantized[BLOCK_WIDTH * BLOCK_HEIGHT];
 			calculate_quantized_weights_rgb<BLOCK_WIDTH*BLOCK_HEIGHT>(texels, weight_quant,
@@ -4363,41 +4706,50 @@ namespace impl
 		}
 
 
-
 		template<int32 BLOCK_WIDTH, int32 BLOCK_HEIGHT>
 		void encode_rgba_single_partition(const unorm8_t texels[BLOCK_WIDTH * BLOCK_HEIGHT],
-			vec4f_t e0,
-			vec4f_t e1,
-			PhysicalBlock* physical_block)
+			vec4f_t e0, vec4f_t e1,
+			PhysicalBlock* physical_block,
+			bool bDualPlane)
 		{
 			size_t partition_count = 1;
-			bool bDualPartition = false;
 
 			color_endpoint_mode_t color_endpoint_mode = CEM_LDR_RGBA_DIRECT;
 
-			range_t weight_quant = RANGE_5;
-			range_t endpoint_quant = RANGE_16;
+			range_t weight_quant = RANGE_MAX;
+			range_t endpoint_quant = RANGE_MAX;
 			if (BLOCK_WIDTH == 4 && BLOCK_HEIGHT == 4)
 			{
+				endpoint_quant = RANGE_192;
 				weight_quant = RANGE_8;
 			}
 			else if (BLOCK_WIDTH == 6 && BLOCK_HEIGHT == 6)
 			{
+				endpoint_quant = RANGE_24;
 				weight_quant = RANGE_4;
+				//endpoint_quant = RANGE_80;
+				//weight_quant = RANGE_3;
 			}
 			else if (BLOCK_WIDTH == 8 && BLOCK_HEIGHT == 8)
 			{
+				endpoint_quant = RANGE_48;
 				weight_quant = RANGE_2;
 			}
 			else
 			{
 				check(false);
 			}
-			endpoint_quant = endpoint_quantization<BLOCK_WIDTH,BLOCK_HEIGHT>(partition_count, weight_quant, color_endpoint_mode);
+
+#ifdef UE_MIRO_DEBUG
+			range_t ExpectedEndpointQuant = endpoint_quantization(BLOCK_WIDTH, BLOCK_HEIGHT, partition_count, weight_quant, color_endpoint_mode, bDualPlane);
+			miro_check(endpoint_quant == ExpectedEndpointQuant);
+#endif
 
 			vec4i_t endpoint_unquantized[2];
 			uint8 endpoint_quantized[8];
-			encode_rgba_direct(endpoint_quant, round(e0), round(e1), endpoint_quantized, endpoint_unquantized);
+			vec4i_t endpoint0 = round(e0);
+			vec4i_t endpoint1 = round(e1);
+			encode_rgba_direct(endpoint_quant, endpoint0, endpoint1, endpoint_quantized, endpoint_unquantized, false);
 
 			uint8 weights_quantized[BLOCK_WIDTH * BLOCK_HEIGHT];
 			calculate_quantized_weights_rgba<BLOCK_WIDTH*BLOCK_HEIGHT>(texels, weight_quant,
@@ -4411,9 +4763,321 @@ namespace impl
 			uint8 weights_ise[MAXIMUM_ENCODED_WEIGHT_BYTES + 1] = { 0 };
 			integer_sequence_encode(weights_quantized, BLOCK_WIDTH * BLOCK_HEIGHT, weight_quant, weights_ise);
 
-			symbolic_to_physical<BLOCK_WIDTH, BLOCK_HEIGHT>(color_endpoint_mode, endpoint_quant, weight_quant,
-				partition_count, endpoint_ise,
-				weights_ise, physical_block, bDualPartition);
+			symbolic_to_physical<BLOCK_WIDTH, BLOCK_HEIGHT>(color_endpoint_mode, endpoint_quant, weight_quant, partition_count, 
+				endpoint_ise, weights_ise, 
+				physical_block, bDualPlane);
+		}
+
+
+		void encode_rgba_single_partition_8x8_infill_7x7(const unorm8_t texels[8*8],
+			vec4f_t e0, vec4f_t e1,
+			PhysicalBlock* physical_block,
+			bool bDualPlane)
+		{
+			size_t partition_count = 1;
+			// texel matrix
+			constexpr int32 BLOCK_X = 8;
+			constexpr int32 BLOCK_Y = 8;
+			// weight matrix
+			constexpr int32 GRID_X = 7;
+			constexpr int32 GRID_Y = 7;
+
+			FVector3f ve0(e0.r, e0.g, e0.b);
+			FVector3f ve1(e1.r, e1.g, e1.b);
+
+			// Resample texels to the weight grid.
+			// Reference implementation
+			unorm8_t grid_texels[GRID_X * GRID_Y];
+			for (int32 gy = 0; gy < GRID_Y; ++gy)
+			{
+				for (int32 gx = 0; gx < GRID_X; ++gx)
+				{
+					int32 tx_16 = ((gx << 16) / GRID_X) * BLOCK_X;
+					int32 ty_16 = ((gy << 16) / GRID_Y) * BLOCK_Y;
+					int32 wx_16 = tx_16 & 0xffff;
+					int32 wy_16 = ty_16 & 0xffff;
+
+					int32 tx0 = (tx_16 >> 16);
+					int32 tx1 = FMath::Min(tx0 + 1, BLOCK_X - 1);
+					int32 ty0 = (ty_16 >> 16);
+					int32 ty1 = FMath::Min(ty0 + 1, BLOCK_Y - 1);
+
+					int32 offset_00 = ty0 * BLOCK_X + tx0;
+					miro_check(offset_00 < BLOCK_X * BLOCK_Y);
+					int32 offset_10 = ty0 * BLOCK_X + tx1;
+					miro_check(offset_10 < BLOCK_X * BLOCK_Y);
+					int32 offset_01 = ty1 * BLOCK_X + tx0;
+					miro_check(offset_01 < BLOCK_X * BLOCK_Y);
+					int32 offset_11 = ty1 * BLOCK_X + tx1;
+					miro_check(offset_11 < BLOCK_X * BLOCK_Y);
+
+					vec4i_t e0_16 = to_vec4i(texels[offset_00]) * (0x10000 - wx_16) + to_vec4i(texels[offset_10]) * wx_16;
+					vec4i_t e1_16 = to_vec4i(texels[offset_01]) * (0x10000 - wx_16) + to_vec4i(texels[offset_11]) * wx_16;
+
+					vec4i_t v_16 = (e0_16 / 0x10000) * (0x10000 - wy_16) + (e1_16 / 0x10000) * wy_16;
+
+					unorm8_t InterpolatedTexel = to_unorm8(v_16 / 0x10000);
+
+					int32 OffsetGrid = gy * GRID_X + gx;
+					// bleh hack
+					//int32 OffsetGrid = gy * GRID_X + (GRID_X-gx-1);
+
+					grid_texels[OffsetGrid] = InterpolatedTexel;
+				}
+			}
+
+			// Look for a good endpoint encoding option.
+			struct FEncodingOption
+			{
+				range_t weight_quant = RANGE_MAX;
+				range_t endpoint_quant = RANGE_MAX;
+				color_endpoint_mode_t color_endpoint_mode = CEM_MAX;
+				int32 EndpointValueCount = 0;
+				vec4i_t endpoint0;
+				vec4i_t endpoint1;
+				vec4i_t endpoint_unquantized[2];
+				uint8 endpoint_quantized[8];
+				bool bValid = false;
+			};
+
+			constexpr int32 MaxOptionCount = 4;
+			FEncodingOption Options[MaxOptionCount];
+
+			//constexpr float GreyscaleThreshold = 0.4f;
+			//float Luminance0 = (e0.r + e0.g + e0.b) / 3.0f;
+			//float Luminance1 = (e1.r + e1.g + e1.b) / 3.0f;
+			//bool bIsGreyscale = FMath::IsNearlyEqual(e0.r, Luminance0, GreyscaleThreshold) && FMath::IsNearlyEqual(e0.g, Luminance0, 0.4f) && FMath::IsNearlyEqual(e0.b, Luminance0, GreyscaleThreshold)
+			//	&& FMath::IsNearlyEqual(e1.r, Luminance1, GreyscaleThreshold) && FMath::IsNearlyEqual(e1.g, Luminance1, GreyscaleThreshold) && FMath::IsNearlyEqual(e1.b, Luminance1, GreyscaleThreshold);
+
+			constexpr float WhiteThreshold = 2.0f;
+			bool bAlphaWhite = FMath::IsNearlyEqual(e0.a, 255.0f, WhiteThreshold) && FMath::IsNearlyEqual(e1.a, 255.0f, WhiteThreshold);
+
+			//constexpr float LinearThreshold = 0.01f;
+			//bool bRGBLinear = false;
+			//if (ve1.GetMin() > UE_SMALL_NUMBER)
+			//{
+			//	vec4f_t RatioV = e0 / e1;
+			//	float RatioAverage = (RatioV.r + RatioV.g + RatioV.b) / 3.0f;
+			//	bRGBLinear = FMath::IsNearlyEqual(RatioV.r, RatioAverage, LinearThreshold) && FMath::IsNearlyEqual(RatioV.g, RatioAverage, LinearThreshold) && FMath::IsNearlyEqual(RatioV.b, RatioAverage, LinearThreshold);
+			//}
+
+			int32 LastOption = 0;
+
+			if (bAlphaWhite)
+			{
+				// There is no need to encode the alpha
+
+				// another option
+				{
+					FEncodingOption& Option = Options[LastOption++];
+
+					Option.color_endpoint_mode = CEM_LDR_LUMINANCE_DIRECT;
+					Option.EndpointValueCount = 2;
+					Option.weight_quant = RANGE_3;
+					Option.endpoint_quant = RANGE_256;
+					Option.endpoint0 = round(e0);
+					Option.endpoint1 = round(e1);
+
+					uint8 endpoint_rgb_unquantized[2];
+					encode_luminance_direct(Option.endpoint_quant, Option.endpoint0.r, Option.endpoint1.r, Option.endpoint_quantized, endpoint_rgb_unquantized);
+
+					Option.endpoint_unquantized[0].r = endpoint_rgb_unquantized[0];
+					Option.endpoint_unquantized[0].g = endpoint_rgb_unquantized[0];
+					Option.endpoint_unquantized[0].b = endpoint_rgb_unquantized[0];
+					Option.endpoint_unquantized[0].a = 255;
+					Option.endpoint_unquantized[1].r = endpoint_rgb_unquantized[1];
+					Option.endpoint_unquantized[1].g = endpoint_rgb_unquantized[1];
+					Option.endpoint_unquantized[1].b = endpoint_rgb_unquantized[1];
+					Option.endpoint_unquantized[1].a = 255;
+				}
+
+				// another option
+				//if (bRGBLinear)
+				{
+					FEncodingOption& Option = Options[LastOption++];
+
+					Option.color_endpoint_mode = CEM_LDR_RGB_BASE_SCALE;
+					Option.EndpointValueCount = 4;
+					Option.weight_quant = RANGE_3;
+					Option.endpoint_quant = RANGE_256;
+					Option.endpoint0 = round(e0);
+					Option.endpoint1 = round(e1);
+
+					int32 MaxComponent = (ve1[0] > ve1[1]) 
+						? ((ve1[0] > ve1[2]) ? 0 : ((ve1[1] > ve1[2]) ? 1 : 2) ) 
+						: ((ve1[1] > ve1[2]) ? 1 : 2);
+					float Ratio = 1.0;
+					if (ve1[MaxComponent] > UE_SMALL_NUMBER)
+					{
+						Ratio = ve0[MaxComponent] / ve1[MaxComponent];
+					}
+					uint8 Scale = 0;
+					if (Ratio < 1.0f)
+					{
+						Scale = static_cast<int8>(Ratio * 255.0f);
+					}
+					else
+					{
+						Scale = static_cast<int8>((1.0f / Ratio) * 255.0f);
+						Swap(Option.endpoint0, Option.endpoint1);
+					}
+
+					vec3i_t endpoint_rgb_unquantized[2];
+					encode_rgb_base_scale(Option.endpoint_quant, Option.endpoint1.rgb(), Scale, Option.endpoint_quantized, endpoint_rgb_unquantized);
+
+					Option.endpoint_unquantized[0].r = endpoint_rgb_unquantized[0].r;
+					Option.endpoint_unquantized[0].g = endpoint_rgb_unquantized[0].g;
+					Option.endpoint_unquantized[0].b = endpoint_rgb_unquantized[0].b;
+					Option.endpoint_unquantized[0].a = 255;
+					Option.endpoint_unquantized[1].r = endpoint_rgb_unquantized[1].r;
+					Option.endpoint_unquantized[1].g = endpoint_rgb_unquantized[1].g;
+					Option.endpoint_unquantized[1].b = endpoint_rgb_unquantized[1].b;
+					Option.endpoint_unquantized[1].a = 255;
+				}
+
+				// another option
+				{
+					FEncodingOption& Option = Options[LastOption++];
+
+					Option.color_endpoint_mode = CEM_LDR_RGB_DIRECT;
+					Option.EndpointValueCount = 6;
+					Option.weight_quant = RANGE_2;
+					Option.endpoint_quant = RANGE_256;
+					Option.endpoint0 = round(e0);
+					Option.endpoint1 = round(e1);
+
+					vec3i_t endpoint_rgb_unquantized[2];
+					encode_rgb_direct(Option.endpoint_quant, Option.endpoint0, Option.endpoint1, Option.endpoint_quantized, endpoint_rgb_unquantized);
+
+					Option.endpoint_unquantized[0].r = endpoint_rgb_unquantized[0].r;
+					Option.endpoint_unquantized[0].g = endpoint_rgb_unquantized[0].g;
+					Option.endpoint_unquantized[0].b = endpoint_rgb_unquantized[0].b;
+					Option.endpoint_unquantized[0].a = 255;
+					Option.endpoint_unquantized[1].r = endpoint_rgb_unquantized[1].r;
+					Option.endpoint_unquantized[1].g = endpoint_rgb_unquantized[1].g;
+					Option.endpoint_unquantized[1].b = endpoint_rgb_unquantized[1].b;
+					Option.endpoint_unquantized[1].a = 255;
+				}
+
+			}
+			else
+			{
+				// most generic option.
+				FEncodingOption& Option = Options[LastOption++];
+
+				Option.color_endpoint_mode = CEM_LDR_RGBA_DIRECT;
+				Option.EndpointValueCount = 8;
+				Option.weight_quant = RANGE_2;
+				Option.endpoint_quant = RANGE_192;
+				Option.endpoint0 = round(e0);
+				Option.endpoint1 = round(e1);
+
+				encode_rgba_direct(Option.endpoint_quant, Option.endpoint0, Option.endpoint1, Option.endpoint_quantized, Option.endpoint_unquantized, false);
+			}
+			check(LastOption<=MaxOptionCount);
+
+			int32 MinError = TNumericLimits<int32>::Max();
+			int32 MinErrorOptionIndex = -1;
+
+			for (int32 o=0; o<LastOption; ++o)
+			{
+				FEncodingOption& Option = Options[o];
+
+#ifdef UE_MIRO_DEBUG
+				range_t ExpectedEndpointQuant = endpoint_quantization(GRID_X, GRID_Y, partition_count, Option.weight_quant, Option.color_endpoint_mode, bDualPlane);
+				miro_check(Option.endpoint_quant == ExpectedEndpointQuant);
+#endif
+
+				int32 EndpointError = quadrance(Option.endpoint0 - Option.endpoint_unquantized[0])
+					+ quadrance(Option.endpoint1 - Option.endpoint_unquantized[1]);
+
+				// Heuristic to compensate for the bigger precision in some options.
+				//const int32 RangeErrorFactor[RANGE_MAX] = { 128, 85, 64, 51, 42, 32, 25, 20, 16, 12, 10, 8, 6, 5, 4, 1, 1, 1, 1, 1, 1 };
+				//OptionError *= RangeErrorFactor[Option.weight_quant];
+
+				if (EndpointError < MinError)
+				{
+					MinError = EndpointError;
+					MinErrorOptionIndex = o;
+				}
+			}
+
+			if (MinErrorOptionIndex < 0 || MinErrorOptionIndex >= LastOption)
+			{
+				check(false);
+				return;
+			}
+
+			const FEncodingOption& Option = Options[MinErrorOptionIndex];
+
+			uint8 grid_weights_quantized[MAX_BLOCK_WIDTH * MAX_BLOCK_HEIGHT];
+			calculate_quantized_weights_rgba<GRID_X * GRID_Y>(grid_texels, Option.weight_quant,
+				Option.endpoint_unquantized[0], Option.endpoint_unquantized[1],
+				grid_weights_quantized);
+
+#ifdef UE_MIRO_DEBUG
+			if (s_DebugBlock == s_CurrentBlock)
+			{
+				UE_LOG(LogMutableCore, Log, TEXT("endpoints f : (%.3f,%.3f,%.3f,%.3f) (%.3f,%.3f,%.3f,%.3f)"),
+					e0.r, e0.g, e0.b, e0.a,
+					e1.r, e1.g, e1.b, e1.a);
+				UE_LOG(LogMutableCore, Log, TEXT("endpoints d : (%3d,%3d,%3d,%3d) (%3d,%3d,%3d,%3d)"),
+					Option.endpoint_unquantized[0].r, Option.endpoint_unquantized[0].g, Option.endpoint_unquantized[0].b, Option.endpoint_unquantized[0].a,
+					Option.endpoint_unquantized[1].r, Option.endpoint_unquantized[1].g, Option.endpoint_unquantized[1].b, Option.endpoint_unquantized[1].a);
+
+				UE_LOG(LogMutableCore, Log, TEXT("texel block pixels :"));
+				int T = 0;
+				for (uint64 Y = 0; Y < BLOCK_Y; ++Y)
+				{
+					FString Line;
+					for (uint64 X = 0; X < BLOCK_X; ++X)
+					{
+						vec4i_t texel = to_vec4i(texels[T++]);
+						Line += FString::Printf(TEXT("(%3d,%3d,%3d,%3d), "), texel.r, texel.g, texel.b, texel.a);
+					}
+
+					UE_LOG(LogMutableCore, Log, TEXT("%s"), *Line);
+				}
+
+				UE_LOG(LogMutableCore, Log, TEXT("texel grid pixels :"));
+				T = 0;
+				for (uint64 Y = 0; Y < GRID_Y; ++Y)
+				{
+					FString Line;
+					for (uint64 X = 0; X < GRID_X; ++X)
+					{
+						vec4i_t texel = to_vec4i(grid_texels[T++]);
+						Line += FString::Printf(TEXT("(%3d,%3d,%3d,%3d), "), texel.r, texel.g, texel.b, texel.a);
+					}
+
+					UE_LOG(LogMutableCore, Log, TEXT("%s"), *Line);
+				}
+
+				UE_LOG(LogMutableCore, Log, TEXT("texel grid encoded :"));
+				T = 0;
+				for (uint64 Y = 0; Y < GRID_Y; ++Y)
+				{
+					FString Line;
+					for (uint64 X = 0; X < GRID_X; ++X)
+					{
+						uint8 texel = grid_weights_quantized[T++];
+						Line += FString::Printf(TEXT("%3d, "), texel);
+					}
+
+					UE_LOG(LogMutableCore, Log, TEXT("%s"), *Line);
+				}
+			}
+#endif
+
+			uint8 endpoint_ise[MAXIMUM_ENCODED_COLOR_ENDPOINT_BYTES] = { 0 };
+			integer_sequence_encode(Option.endpoint_quantized, Option.EndpointValueCount, Option.endpoint_quant, endpoint_ise);
+
+			uint8 weights_ise[MAXIMUM_ENCODED_WEIGHT_BYTES + 1] = { 0 };
+			integer_sequence_encode(grid_weights_quantized, GRID_X * GRID_Y, Option.weight_quant, weights_ise);
+
+			symbolic_to_physical(GRID_X, GRID_Y, Option.color_endpoint_mode, Option.endpoint_quant, Option.weight_quant, partition_count,
+				endpoint_ise, weights_ise,
+				physical_block, bDualPlane);
 		}
 
 
@@ -4427,33 +5091,39 @@ namespace impl
 			bool bDualPlane = false;
 
 			color_endpoint_mode_t color_endpoint_mode = CEM_LDR_LUMINANCE_ALPHA_DIRECT;
-			range_t weight_quant = RANGE_6;
-			range_t endpoint_quant = RANGE_64;
+			range_t weight_quant = RANGE_MAX;
+			range_t endpoint_quant = RANGE_MAX;
 
 			if (BLOCK_WIDTH == 4 && BLOCK_HEIGHT == 4)
 			{
 				weight_quant = RANGE_12;
-				//endpoint_quant = RANGE_64;
+				endpoint_quant = RANGE_256;
 			}
 			else if (BLOCK_WIDTH == 6 && BLOCK_HEIGHT == 6)
 			{
-				weight_quant = RANGE_5;
-				//endpoint_quant = RANGE_256;
+				weight_quant = RANGE_4;
+				endpoint_quant = RANGE_256;
 			}
 			else if (BLOCK_WIDTH == 8 && BLOCK_HEIGHT == 8)
 			{
 				weight_quant = RANGE_2;
-				//endpoint_quant = RANGE_256;
+				endpoint_quant = RANGE_256;
 			}
 			else
 			{
-				check(false);
+				miro_check(false);
 			}
-			endpoint_quant = endpoint_quantization<BLOCK_WIDTH, BLOCK_HEIGHT>(partition_count, weight_quant, color_endpoint_mode);
+
+#ifdef UE_MIRO_DEBUG
+			range_t ExpectedEndpointQuant = endpoint_quantization(BLOCK_WIDTH, BLOCK_HEIGHT, partition_count, weight_quant, color_endpoint_mode, bDualPlane);
+			miro_check(endpoint_quant == ExpectedEndpointQuant);
+#endif
 
 			vec3i_t endpoint_unquantized[2];
 			uint8 endpoint_quantized[6];
-			encode_rgb_direct(endpoint_quant, round(e0), round(e1), endpoint_quantized, endpoint_unquantized);
+			vec3i_t endpoint0 = round(e0);
+			vec3i_t endpoint1 = round(e1);
+			encode_rgb_direct(endpoint_quant, endpoint0, endpoint1, endpoint_quantized, endpoint_unquantized);
 
 			uint8 weights_quantized[BLOCK_WIDTH * BLOCK_HEIGHT];
 			// TODO: optimize with rg skipping b
@@ -4481,8 +5151,8 @@ namespace impl
 			bool bDualPlane = true;
 
 			color_endpoint_mode_t color_endpoint_mode = CEM_LDR_LUMINANCE_ALPHA_DIRECT;
-			range_t weight_quant = RANGE_6;
-			range_t endpoint_quant = RANGE_64;
+			range_t weight_quant = RANGE_MAX;
+			range_t endpoint_quant = RANGE_MAX;
 
 			if (BLOCK_WIDTH == 4 && BLOCK_HEIGHT == 4)
 			{
@@ -4496,12 +5166,19 @@ namespace impl
 			}
 			else
 			{
-				check(false);
+				miro_check(false);
 			}
 
+#ifdef UE_MIRO_DEBUG
+			range_t ExpectedEndpointQuant = endpoint_quantization(BLOCK_WIDTH, BLOCK_HEIGHT, partition_count, weight_quant, color_endpoint_mode, bDualPlane);
+			miro_check(endpoint_quant == ExpectedEndpointQuant);
+#endif
+
 			vec4i_t endpoint_unquantized[2];
 			uint8 endpoint_quantized[8];
-			encode_rgba_direct(endpoint_quant, round(e0), round(e1), endpoint_quantized, endpoint_unquantized);
+			vec4i_t endpoint0 = round(e0);
+			vec4i_t endpoint1 = round(e1);
+			encode_rgba_direct(endpoint_quant, endpoint0, endpoint1, endpoint_quantized, endpoint_unquantized, true);
 
 			uint8 r_weights_quantized[BLOCK_WIDTH * BLOCK_HEIGHT];
 			calculate_quantized_weights_channel<BLOCK_WIDTH* BLOCK_HEIGHT>(texels, weight_quant,
@@ -4537,11 +5214,11 @@ namespace impl
 		}
 
 
-		bool is_solid(const unorm8_t* texels, size_t count, unorm8_t* color)
+		bool is_solid(const unorm8_t* texels, size_t count, unorm8_t* color, int32 Epsilon)
 		{
 			for (size_t i = 0; i < count; ++i)
 			{
-				if (!approx_equal(to_vec3i(texels[i]), to_vec3i(texels[0])))
+				if (!approx_equal(to_vec3i(texels[i]), to_vec3i(texels[0]), Epsilon))
 				{
 					return false;
 				}
@@ -4553,11 +5230,11 @@ namespace impl
 		}
 
 
-		bool is_solid_rgba(const unorm8_t* texels, size_t count, unorm8_t* color)
+		bool is_solid_rgba(const unorm8_t* texels, size_t count, unorm8_t* color, int32 Epsilon)
 		{
 			for (size_t i = 0; i < count; ++i)
 			{
-				if (!approx_equal(to_vec4i(texels[i]), to_vec4i(texels[0])))
+				if (!approx_equal(to_vec4i(texels[i]), to_vec4i(texels[0]), Epsilon))
 				{
 					return false;
 				}
@@ -4569,14 +5246,14 @@ namespace impl
 		}
 
 
-		bool is_greyscale(const unorm8_t* texels, size_t count, uint8* luminances)
+		bool is_greyscale(const unorm8_t* texels, size_t count, uint8* luminances, int32 Epsilon)
 		{
 			for (size_t i = 0; i < count; ++i)
 			{
 				vec3i_t color = to_vec3i(texels[i]);
 				luminances[i] = static_cast<uint8>(luminance(color));
 				vec3i_t lum(luminances[i], luminances[i], luminances[i]);
-				if (!approx_equal(color, lum))
+				if (!approx_equal(color, lum, Epsilon))
 				{
 					return false;
 				}
@@ -4586,7 +5263,6 @@ namespace impl
 		}
 
 
-
 		struct mat3x3f_t
 		{
 		public:
@@ -4610,6 +5286,32 @@ namespace impl
 			vec3f_t m[3];
 		};
 
+
+		struct mat4x4f_t
+		{
+		public:
+			mat4x4f_t() {}
+
+			mat4x4f_t(float m00, float m01, float m02, float m03,
+				float m10, float m11, float m12, float m13,
+				float m20, float m21, float m22, float m23,
+				float m30, float m31, float m32, float m33)
+			{
+				m[0] = vec4f_t(m00, m01, m02, m03);
+				m[1] = vec4f_t(m10, m11, m12, m13);
+				m[2] = vec4f_t(m20, m21, m22, m23);
+				m[3] = vec4f_t(m30, m31, m32, m33);
+			}
+
+			const vec4f_t& row(size_t i) const { return m[i]; }
+
+			float& at(size_t i, size_t j) { return m[i].components(j); }
+			const float& at(size_t i, size_t j) const { return m[i].components(j); }
+
+		private:
+			vec4f_t m[4];
+		};
+
 		inline vec3f_t operator*(const mat3x3f_t& a, vec3f_t b)
 		{
 			vec3f_t tmp;
@@ -4619,6 +5321,16 @@ namespace impl
 			return tmp;
 		}
 
+		inline vec4f_t operator*(const mat4x4f_t& a, vec4f_t b)
+		{
+			vec4f_t tmp;
+			tmp.r = dot(a.row(0), b);
+			tmp.g = dot(a.row(1), b);
+			tmp.b = dot(a.row(2), b);
+			tmp.a = dot(a.row(3), b);
+			return tmp;
+		}
+
 		vec3f_t mean(const unorm8_t* texels, size_t count)
 		{
 			vec3i_t sum(0, 0, 0);
@@ -4630,10 +5342,18 @@ namespace impl
 			return to_vec3f(sum) / static_cast<float>(count);
 		}
 
-		void subtract(const unorm8_t* texels,
-			size_t count,
-			vec3f_t v,
-			vec3f_t* output)
+		vec4f_t mean4(const unorm8_t* texels, size_t count)
+		{
+			vec4i_t sum(0, 0, 0, 0);
+			for (size_t i = 0; i < count; ++i)
+			{
+				sum = sum + to_vec4i(texels[i]);
+			}
+
+			return to_vec4f(sum) / static_cast<float>(count);
+		}
+
+		void subtract(const unorm8_t* texels, size_t count, vec3f_t v, vec3f_t* output)
 		{
 			for (size_t i = 0; i < count; ++i)
 			{
@@ -4641,6 +5361,14 @@ namespace impl
 			}
 		}
 
+		void subtract(const unorm8_t* texels, size_t count, vec4f_t v, vec4f_t* output)
+		{
+			for (size_t i = 0; i < count; ++i)
+			{
+				output[i] = to_vec4f(texels[i]) - v;
+			}
+		}
+
 		mat3x3f_t covariance(const vec3f_t* m, size_t count)
 		{
 			mat3x3f_t cov;
@@ -4660,6 +5388,25 @@ namespace impl
 			return cov;
 		}
 
+		mat4x4f_t covariance(const vec4f_t* m, size_t count)
+		{
+			mat4x4f_t cov;
+			for (size_t i = 0; i < 4; ++i)
+			{
+				for (size_t j = 0; j < 4; ++j)
+				{
+					float s = 0;
+					for (size_t k = 0; k < count; ++k)
+					{
+						s += m[k].components(i) * m[k].components(j);
+					}
+					cov.at(i, j) = s / static_cast<float>(count - 1);
+				}
+			}
+
+			return cov;
+		}
+
 		void eigen_vector(const mat3x3f_t& a, vec3f_t& eig)
 		{
 			vec3f_t b = signorm(vec3f_t(1, 3, 2));  // FIXME: Magic number
@@ -4671,6 +5418,17 @@ namespace impl
 			eig = b;
 		}
 
+		void eigen_vector(const mat4x4f_t& a, vec4f_t& eig)
+		{
+			vec4f_t b = signorm(vec4f_t(1, 3, 2, 4));  // FIXME: Magic number
+			for (size_t i = 0; i < 8; ++i)
+			{
+				b = signorm(a * b);
+			}
+
+			eig = b;
+		}
+
 
 		void find_min_max(const unorm8_t* texels,
 			size_t count,
@@ -4679,7 +5437,7 @@ namespace impl
 			vec3f_t& e0,
 			vec3f_t& e1)
 		{
-			DCHECK(FMath::IsNearlyEqual(quadrance(line_k), 1.0, 0.0001f));
+			miro_check(FMath::IsNearlyEqual(quadrance(line_k), 1.0, 0.0001f));
 
 			float a, b;
 			{
@@ -4699,11 +5457,38 @@ namespace impl
 		}
 
 
+		void find_min_max(const unorm8_t* texels,
+			size_t count,
+			vec4f_t line_k,
+			vec4f_t line_m,
+			vec4f_t& e0,
+			vec4f_t& e1)
+		{
+			miro_check(FMath::IsNearlyEqual(quadrance(line_k), 1.0, 0.0001f));
+
+			float a, b;
+			{
+				float t = dot(to_vec4f(texels[0]) - line_m, line_k);
+				a = t;
+				b = t;
+			}
+
+			for (size_t i = 1; i < count; ++i) {
+				float t = dot(to_vec4f(texels[i]) - line_m, line_k);
+				a = FMath::Min(a, t);
+				b = FMath::Max(b, t);
+			}
+
+			e0 = clamp_rgba(line_k * a + line_m);
+			e1 = clamp_rgba(line_k * b + line_m);
+		}
+
+
 		template<int32 BLOCK_TEXEL_COUNT>
 		void principal_component_analysis(const unorm8_t texels[BLOCK_TEXEL_COUNT],
 			size_t count,
 			vec3f_t& line_k,
-			vec3f_t& line_m) 
+			vec3f_t& line_m)
 		{
 			line_m = mean(texels, count);
 
@@ -4716,12 +5501,30 @@ namespace impl
 		}
 
 
+		template<int32 BLOCK_TEXEL_COUNT>
+		void principal_component_analysis(const unorm8_t texels[BLOCK_TEXEL_COUNT],
+			size_t count,
+			vec4f_t& line_k,
+			vec4f_t& line_m)
+		{
+			line_m = mean4(texels, count);
+
+			vec4f_t n[BLOCK_TEXEL_COUNT];
+			subtract(texels, count, line_m, n);
+
+			mat4x4f_t w = covariance(n, count);
+
+			eigen_vector(w, line_k);
+		}
+
+
 		template<int32 BLOCK_WIDTH, int32 BLOCK_HEIGHT>
 		void compress_block(const unorm8_t texels[BLOCK_WIDTH*BLOCK_HEIGHT], PhysicalBlock* physical_block)
 		{
 			{
 				unorm8_t color;
-				if (is_solid(texels, BLOCK_WIDTH * BLOCK_HEIGHT, &color))
+				int32 Epsilon = APPROX_COLOR_EPSILON;
+				if (is_solid(texels, BLOCK_WIDTH * BLOCK_HEIGHT, &color, Epsilon))
 				{
 					encode_void_extent(to_vec3i(color), physical_block);
 					return;
@@ -4730,7 +5533,8 @@ namespace impl
 
 			{
 				uint8 luminances[BLOCK_WIDTH * BLOCK_HEIGHT];
-				if (is_greyscale(texels, BLOCK_WIDTH * BLOCK_HEIGHT, luminances))
+				int32 Epsilon = APPROX_COLOR_EPSILON;
+				if (is_greyscale(texels, BLOCK_WIDTH * BLOCK_HEIGHT, luminances, Epsilon))
 				{
 					encode_luminance<BLOCK_WIDTH,BLOCK_HEIGHT>(luminances, physical_block);
 					return;
@@ -4746,56 +5550,109 @@ namespace impl
 
 
 		template<int32 BLOCK_WIDTH, int32 BLOCK_HEIGHT>
-		void compress_block_rgba(const unorm8_t texels[BLOCK_WIDTH * BLOCK_HEIGHT], PhysicalBlock* physical_block, bool bDualPlane)
+		void compress_block_rgba(const unorm8_t texels[BLOCK_WIDTH * BLOCK_HEIGHT], PhysicalBlock* physical_block)
 		{
 			unorm8_t color;
-			bool isSolidRGBA = is_solid_rgba(texels, BLOCK_WIDTH * BLOCK_HEIGHT, &color);
+
+			// Option 1: block mode 10111111100 : void extent block.
+			
+			// \TODO: This actually helps reduce the error due to the bad quality of the non-void extent blocks
+			// When that quality improves, reduce this epsilon.
+			int32 Epsilon = APPROX_COLOR_EPSILON; 
+
+			bool isSolidRGBA = is_solid_rgba(texels, BLOCK_WIDTH * BLOCK_HEIGHT, &color, Epsilon);
 			if (isSolidRGBA)
 			{
 				encode_void_extent(to_vec4i(color), physical_block);
 				return;
 			}
 
-			vec3f_t e0, e1;
-			bool isSolid = is_solid(texels, BLOCK_WIDTH * BLOCK_HEIGHT, &color);
-			if (isSolid)
-			{
-				e0 = e1 = to_vec3f(texels[0]);
-			}
-			else
-			{
-				vec3f_t k, m;
-				principal_component_analysis<BLOCK_WIDTH* BLOCK_HEIGHT>(texels, BLOCK_WIDTH * BLOCK_HEIGHT, k, m);
-				find_min_max(texels, BLOCK_WIDTH * BLOCK_HEIGHT, k, m, e0, e1);
-			}
+			// Option 2
 
-			// Find alpha min and max
-			uint8 l0 = 255;
-			uint8 l1 = 0;
-			for (size_t i = 0; i < BLOCK_WIDTH * BLOCK_HEIGHT; ++i)
+			bool bDualPlane = false;
+			if (bDualPlane)
 			{
-				l0 = FMath::Min(l0, texels[i].components[3]);
-				l1 = FMath::Max(l1, texels[i].components[3]);
-			}
+				// dual plane usually yields worse results!
+				vec3f_t e0, e1;
+				bool isSolid = is_solid(texels, BLOCK_WIDTH * BLOCK_HEIGHT, &color, Epsilon);
+				if (isSolid)
+				{
+					e0 = e1 = to_vec3f(texels[0]);
+				}
+				else
+				{
+					vec3f_t k, m;
+					principal_component_analysis<BLOCK_WIDTH*BLOCK_HEIGHT>(texels, BLOCK_WIDTH * BLOCK_HEIGHT, k, m);
+					find_min_max(texels, BLOCK_WIDTH * BLOCK_HEIGHT, k, m, e0, e1);
+				}
 
-			vec4f_t re0(e0.r, e0.g, e0.b, float(l0));
-			vec4f_t re1(e1.r, e1.g, e1.b, float(l1));
+				// Find alpha min and max
+				uint8 l0 = 255;
+				uint8 l1 = 0;
+				for (size_t i = 0; i < BLOCK_WIDTH * BLOCK_HEIGHT; ++i)
+				{
+					l0 = FMath::Min(l0, texels[i].components[3]);
+					l1 = FMath::Max(l1, texels[i].components[3]);
+				}
 
-			//        if (s_debuglog)
-			//        {
-			//            UE_LOG(LogMutableCore,Warning," limits: %.3f %.3f %.3f %.3f \t %.3f %.3f %.3f %.3f",
-			//                    re0.x, re0.y, re0.z, re0.w,
-			//                    re1.x, re1.y, re1.z, re1.w );
-			//        }
+				vec4f_t re0(e0.r, e0.g, e0.b, float(l0));
+				vec4f_t re1(e1.r, e1.g, e1.b, float(l1));
 
-			if (bDualPlane)
-			{
-				// dual plane usually yields worse results!
 				encode_rgba_single_partition_dual_plane<BLOCK_WIDTH, BLOCK_HEIGHT>(texels, re0, re1, physical_block);
 			}
 			else
 			{
-				encode_rgba_single_partition<BLOCK_WIDTH, BLOCK_HEIGHT>(texels, re0, re1, physical_block);
+				// Analysing as 4D
+				vec4f_t re0, re1;
+				bool isSolid = is_solid_rgba(texels, BLOCK_WIDTH * BLOCK_HEIGHT, &color, Epsilon);
+				if (isSolid)
+				{
+					re0 = re1 = to_vec4f(texels[0]);
+				}
+				else
+				{
+					vec4f_t k, m;
+					principal_component_analysis<BLOCK_WIDTH* BLOCK_HEIGHT>(texels, BLOCK_WIDTH * BLOCK_HEIGHT, k, m);
+					find_min_max(texels, BLOCK_WIDTH * BLOCK_HEIGHT, k, m, re0, re1);
+				}
+				
+				// This would be an RGB analysis with a hacky patch for alpha limits.
+				// {
+				//vec3f_t e0, e1;
+				//bool isSolid = is_solid(texels, BLOCK_WIDTH * BLOCK_HEIGHT, &color, Epsilon);
+				//if (isSolid)
+				//{
+				//	e0 = e1 = to_vec3f(texels[0]);
+				//}
+				//else
+				//{
+				//	vec3f_t k, m;
+				//	principal_component_analysis<BLOCK_WIDTH* BLOCK_HEIGHT>(texels, BLOCK_WIDTH * BLOCK_HEIGHT, k, m);
+				//	find_min_max(texels, BLOCK_WIDTH * BLOCK_HEIGHT, k, m, e0, e1);
+				//}
+
+				//// Find alpha min and max
+				//uint8 l0 = 255;
+				//uint8 l1 = 0;
+				//for (size_t i = 0; i < BLOCK_WIDTH * BLOCK_HEIGHT; ++i)
+				//{
+				//	l0 = FMath::Min(l0, texels[i].components[3]);
+				//	l1 = FMath::Max(l1, texels[i].components[3]);
+				//}
+				//vec4f_t re0(e0.r, e0.g, e0.b, float(l0));
+				//vec4f_t re1(e1.r, e1.g, e1.b, float(l1));
+				//}
+
+				if (BLOCK_WIDTH < 7 && BLOCK_HEIGHT < 7)
+				{
+					// Option 1
+					encode_rgba_single_partition<BLOCK_WIDTH, BLOCK_HEIGHT>(texels, re0, re1, physical_block, bDualPlane);
+				}
+				else
+				{
+					// Option 2
+					encode_rgba_single_partition_8x8_infill_7x7(texels, re0, re1, physical_block, bDualPlane);
+				}
 			}
 		}
 
@@ -4804,7 +5661,8 @@ namespace impl
 		void compress_block_rg(const unorm8_t texels[BLOCK_WIDTH * BLOCK_HEIGHT], PhysicalBlock* physical_block)
 		{
 			unorm8_t color;
-			bool isSolidRGBA = is_solid_rgba(texels, BLOCK_WIDTH * BLOCK_HEIGHT, &color);
+			int32 Epsilon = APPROX_COLOR_EPSILON;
+			bool isSolidRGBA = is_solid_rgba(texels, BLOCK_WIDTH * BLOCK_HEIGHT, &color, Epsilon);
 			if (isSolidRGBA)
 			{
 				encode_void_extent(vec4i_t(color.components[2], color.components[2], color.components[2], color.components[1]), physical_block);
@@ -4902,8 +5760,8 @@ namespace impl
 			uint32_t& z() { return m_c[2]; }
 			uint32_t& w() { return m_c[3]; }
 
-			uint32_t operator[] (uint32_t idx) const { check(idx < 4);  return m_c[idx]; }
-			uint32_t& operator[] (uint32_t idx) { check(idx < 4);  return m_c[idx]; }
+			uint32_t operator[] (uint32_t idx) const { miro_check(idx < 4);  return m_c[idx]; }
+			uint32_t& operator[] (uint32_t idx) { miro_check(idx < 4);  return m_c[idx]; }
 		};
 
 		struct IVec4
@@ -4941,8 +5799,8 @@ namespace impl
 				return UVec4(FMath::Max(0, m_c[0]), FMath::Max(0, m_c[1]), FMath::Max(0, m_c[2]), FMath::Max(0, m_c[3]));
 			}
 
-			int32_t operator[] (uint32_t idx) const { check(idx < 4);  return m_c[idx]; }
-			int32_t& operator[] (uint32_t idx) { check(idx < 4);  return m_c[idx]; }
+			int32_t operator[] (uint32_t idx) const { miro_check(idx < 4);  return m_c[idx]; }
+			int32_t& operator[] (uint32_t idx) { miro_check(idx < 4);  return m_c[idx]; }
 		};
 
 		struct IVec3
@@ -4971,8 +5829,8 @@ namespace impl
 			int32_t& y() { return m_c[1]; }
 			int32_t& z() { return m_c[2]; }
 
-			int32_t operator[] (uint32_t idx) const { check(idx < 3);  return m_c[idx]; }
-			int32_t& operator[] (uint32_t idx) { check(idx < 3);  return m_c[idx]; }
+			int32_t operator[] (uint32_t idx) const { miro_check(idx < 3);  return m_c[idx]; }
+			int32_t& operator[] (uint32_t idx) { miro_check(idx < 3);  return m_c[idx]; }
 		};
 
 		static uint32_t deDivRoundUp32(uint32_t a, uint32_t b)
@@ -6245,6 +7103,8 @@ namespace impl
 				}
 			return result;
 		}
+
+
 		DecompressResult decompressBlock(void* dst, const Block128& blockData, int blockWidth, int blockHeight, bool isSRGB, bool isLDR)
 		{
 			DE_ASSERT(isLDR || !isSRGB);
@@ -6309,6 +7169,156 @@ namespace impl
 		}
 
 
+#ifdef UE_MIRO_DEBUG
+		void LogBlock(const Block128& blockData, int blockWidth, int blockHeight, bool isSRGB, bool isLDR)
+		{
+			DE_ASSERT(isLDR || !isSRGB);
+
+			// Decode block mode.
+			uint32 BlockMode = blockData.getBits(0, 10);
+			const ASTCBlockMode blockMode = getASTCBlockMode(BlockMode);
+
+#define MODE_TO_BINARY(byte)  \
+  ((byte) & 0x400 ? L'1' : L'0'), \
+  ((byte) & 0x200 ? L'1' : L'0'), \
+  ((byte) & 0x100 ? L'1' : L'0'), \
+  ((byte) & 0x080 ? L'1' : L'0'), \
+  ((byte) & 0x040 ? L'1' : L'0'), \
+  ((byte) & 0x020 ? L'1' : L'0'), \
+  ((byte) & 0x010 ? L'1' : L'0'), \
+  ((byte) & 0x008 ? L'1' : L'0'), \
+  ((byte) & 0x004 ? L'1' : L'0'), \
+  ((byte) & 0x002 ? L'1' : L'0'), \
+  ((byte) & 0x001 ? L'1' : L'0') 
+
+			UE_LOG(LogMutableCore, Log, TEXT("block_mode : %c%c %c%c%c%c %c%c%c%c%c"), MODE_TO_BINARY(BlockMode) );
+
+
+			// Check for block mode errors.
+			if (blockMode.isError)
+			{
+				miro_check(false);
+			}
+			// Separate path for void-extent.
+			if (blockMode.isVoidExtent)
+			{
+				return; // decodeVoidExtentBlock(dst, blockData, blockWidth, blockHeight, isSRGB, isLDR);
+			}
+
+			// Compute weight grid values.
+			const int numWeights = computeNumWeights(blockMode);
+			const int numWeightDataBits = computeNumRequiredBits(blockMode.weightISEParams, numWeights);
+			const int numPartitions = (int)blockData.getBits(11, 12) + 1;
+
+			UE_LOG(LogMutableCore, Log, TEXT("weight_grid : %d x %d"), blockMode.weightGridWidth, blockMode.weightGridHeight);
+			UE_LOG(LogMutableCore, Log, TEXT("num_weights : %d"), numWeights);
+			UE_LOG(LogMutableCore, Log, TEXT("num_weights_data_bits : %d"), numWeightDataBits);
+			UE_LOG(LogMutableCore, Log, TEXT("num_partitions : %d"), numPartitions);
+			UE_LOG(LogMutableCore, Log, TEXT("dual_plane : %d"), blockMode.isDualPlane);
+
+			// Check for errors in weight grid, partition and dual-plane parameters.
+			if (numWeights > 64 ||
+				numWeightDataBits > 96 ||
+				numWeightDataBits < 24 ||
+				blockMode.weightGridWidth > blockWidth ||
+				blockMode.weightGridHeight > blockHeight ||
+				(numPartitions == 4 && blockMode.isDualPlane))
+			{
+				miro_check(false);
+			}
+
+			// Compute number of bits available for color endpoint data.
+			const bool	isSingleUniqueCem = numPartitions == 1 || blockData.getBits(23, 24) == 0;
+			const int	numConfigDataBits = (numPartitions == 1 ? 17 : isSingleUniqueCem ? 29 : 25 + 3 * numPartitions) +
+				(blockMode.isDualPlane ? 2 : 0);
+			const int	numBitsForColorEndpoints = 128 - numWeightDataBits - numConfigDataBits;
+			const int	extraCemBitsStart = 127 - numWeightDataBits - (isSingleUniqueCem ? -1
+				: numPartitions == 4 ? 7
+				: numPartitions == 3 ? 4
+				: numPartitions == 2 ? 1
+				: 0);
+
+			// Decode color endpoint modes.
+			deUint32 colorEndpointModes[4];
+			decodeColorEndpointModes(&colorEndpointModes[0], blockData, numPartitions, extraCemBitsStart);
+			const int numColorEndpointValues = computeNumColorEndpointValues(colorEndpointModes, numPartitions);
+
+			UE_LOG(LogMutableCore, Log, TEXT("num_endpoint_data_bits : %d"), numBitsForColorEndpoints);
+			UE_LOG(LogMutableCore, Log, TEXT("endpoint_mode : %d"), colorEndpointModes[0]);
+			UE_LOG(LogMutableCore, Log, TEXT("num_endpoint_values : %d"), numColorEndpointValues);
+
+			// Check for errors in color endpoint value count.
+			if (numColorEndpointValues > 18 || numBitsForColorEndpoints < (int)deDivRoundUp32(13 * numColorEndpointValues, 5))
+			{
+				miro_check(false);
+			}
+			// Compute color endpoints.
+			ColorEndpointPair colorEndpoints[4];
+			computeColorEndpoints(&colorEndpoints[0], blockData, &colorEndpointModes[0], numPartitions, numColorEndpointValues,
+				computeMaximumRangeISEParams(numBitsForColorEndpoints, numColorEndpointValues), numBitsForColorEndpoints);
+
+			UE_LOG(LogMutableCore, Log, TEXT("endpoints : ( %3d, %3d, %3d, %3d ), ( %3d, %3d, %3d, %3d )"), 
+				colorEndpoints[0].e0[0], colorEndpoints[0].e0[1], colorEndpoints[0].e0[2], colorEndpoints[0].e0[3],
+				colorEndpoints[0].e1[0], colorEndpoints[0].e1[1], colorEndpoints[0].e1[2], colorEndpoints[0].e1[3] );
+
+			// Compute texel weights.
+			TexelWeightPair texelWeights[MAX_BLOCK_WIDTH * MAX_BLOCK_HEIGHT];
+			//computeTexelWeights(&texelWeights[0], blockData, blockWidth, blockHeight, blockMode);
+
+			ISEDecodedResult weightGrid[64];
+			{
+				BitAccessStream dataStream(blockData, 127, numWeightDataBits, false);
+				decodeISE(&weightGrid[0], computeNumWeights(blockMode), dataStream, blockMode.weightISEParams);
+			}
+			deUint32 unquantizedWeights[64];
+			{				
+				unquantizeWeights(&unquantizedWeights[0], &weightGrid[0], blockMode);
+				interpolateWeights(&texelWeights[0], unquantizedWeights, blockWidth, blockHeight, blockMode);
+			}
+
+			int i = 0;
+			//for (int32 Row = 0; Row < blockHeight; ++Row)
+			//{
+			//	UE_LOG(LogMutableCore, Log, TEXT("texel weights : %3d %3d %3d %3d %3d %3d %3d %3d"), texelWeights[i++].w[0], texelWeights[i++].w[0], texelWeights[i++].w[0], texelWeights[i++].w[0], texelWeights[i++].w[0], texelWeights[i++].w[0], texelWeights[i++].w[0], texelWeights[i++].w[0]);
+			//}
+
+			i = 0;
+			for (int32 Row = 0; Row < blockMode.weightGridHeight; ++Row)
+			{
+				UE_LOG(LogMutableCore, Log, TEXT("texel weights : %3d %3d %3d %3d %3d %3d %3d"), unquantizedWeights[i++], unquantizedWeights[i++], unquantizedWeights[i++], unquantizedWeights[i++], unquantizedWeights[i++], unquantizedWeights[i++], unquantizedWeights[i++]);
+			}
+
+			// Set texel colors.
+			const int		ccs = blockMode.isDualPlane ? (int)blockData.getBits(extraCemBitsStart - 2, extraCemBitsStart - 1) : -1;
+			const deUint32	partitionIndexSeed = numPartitions > 1 ? blockData.getBits(13, 22) : (deUint32)-1;
+			astcrt::vec4f_t dst[MAX_BLOCK_WIDTH * MAX_BLOCK_HEIGHT];
+			DecompressResult Result = setTexelColors(dst, &colorEndpoints[0], &texelWeights[0], ccs, partitionIndexSeed, numPartitions, blockWidth, blockHeight, isSRGB, isLDR, &colorEndpointModes[0]);
+			check(Result==DecompressResult::DECOMPRESS_RESULT_VALID_BLOCK);
+
+			UE_LOG(LogMutableCore, Log, TEXT("decoded texels:"));
+			i = 0;
+			for (int32 Row = 0; Row < blockHeight; ++Row)
+			{
+				FString Line;
+				for (uint64 X = 0; X < blockWidth; ++X)
+				{
+					astcrt::vec4i_t texel;
+					texel.components(0) = FMath::Clamp<int>((int)(dst[i].components(0) * 65536.0f + .5f), 0, 65535) >> 8;
+					texel.components(1) = FMath::Clamp<int>((int)(dst[i].components(1) * 65536.0f + .5f), 0, 65535) >> 8;
+					texel.components(2) = FMath::Clamp<int>((int)(dst[i].components(2) * 65536.0f + .5f), 0, 65535) >> 8;
+					texel.components(3) = FMath::Clamp<int>((int)(dst[i].components(3) * 65536.0f + .5f), 0, 65535) >> 8;
+					Line += FString::Printf(TEXT("(%3d,%3d,%3d,%3d), "), texel.r, texel.g, texel.b, texel.a);
+					i++;
+				}
+
+				UE_LOG(LogMutableCore, Log, TEXT("%s"), *Line);
+			}
+
+#undef MODE_TO_BINARY
+
+		}
+#endif //UE_MIRO_DEBUG
+
 		bool decompress(uint8_t* pDst, const uint8_t* data, bool isSRGB, int blockWidth, int blockHeight)
 		{
 			// rg - We only support LDR here, although adding back in HDR would be easy.
@@ -6336,7 +7346,11 @@ namespace impl
 
 			if (decompressBlock(isSRGB ? (void*)pDst : (void*)&linear[0],
 				blockData, blockWidth, blockHeight, isSRGB, isLDR) != DECOMPRESS_RESULT_VALID_BLOCK)
+			{
+				// Invalid ASTC block.
+				miro_check(false);
 				return false;
+			}
 
 			if (!isSRGB)
 			{
@@ -7505,7 +8519,8 @@ namespace miro
 				bool bIsSRGB = false;
 				uint8 Block[BLOCK_SIZE * BLOCK_SIZE * 4];
 				bool bSuccess = astcdec::decompress(Block, from, bIsSRGB, BLOCK_SIZE, BLOCK_SIZE);
-				check(bSuccess);
+				miro_check(bSuccess);
+
 
 				for (uint32 py = 0; py < BLOCK_SIZE; py++)
 				{
@@ -7544,7 +8559,7 @@ namespace miro
 				bool bIsSRGB = false;
 				uint8 Block[BLOCK_SIZE * BLOCK_SIZE * 4];
 				bool bSuccess = astcdec::decompress(Block, from, bIsSRGB, BLOCK_SIZE, BLOCK_SIZE);
-				check(bSuccess);
+				miro_check(bSuccess);
 
 				for (uint32 py = 0; py < BLOCK_SIZE; py++)
 				{
@@ -7581,14 +8596,15 @@ namespace miro
 
 		uint32 bx = FMath::DivideAndRoundUp(sx, BLOCK_SIZE);
 		uint32 by = FMath::DivideAndRoundUp(sy, BLOCK_SIZE);
-
-		//for ( uint32 y=0; y<by; ++y )
-		ParallelFor(by,
-			[
-				bx, sx, sy, from, to, &physical_block_zero
-			] (uint32 y)
+		
+#ifdef UE_MIRO_DEBUG
+		s_CurrentBlock = 0;
+		for ( uint32 y=0; y<by; ++y )
+#else
+		ParallelFor(by, [ bx, sx, sy, from, to, &physical_block_zero ] (uint32 y)
+#endif
 			{
-				astcrt::PhysicalBlock* rowTo = reinterpret_cast<astcrt::PhysicalBlock*>(to + sizeof(astcrt::PhysicalBlock) * bx * y);
+				astcrt::PhysicalBlock * rowTo = reinterpret_cast<astcrt::PhysicalBlock*>(to + sizeof(astcrt::PhysicalBlock) * bx * y);
 
 				for (uint32 x = 0; x < bx; ++x)
 				{
@@ -7615,12 +8631,18 @@ namespace miro
 					}
 
 					*rowTo = physical_block_zero;
-					// \TODO: We are forcing to not use dual plane, this should be per-block.
-					astcrt::compress_block_rgba<BLOCK_SIZE, BLOCK_SIZE>((astcrt::unorm8_t*)block, rowTo, false);
+					astcrt::compress_block_rgba<BLOCK_SIZE, BLOCK_SIZE>((astcrt::unorm8_t*)block, rowTo);
 
 					++rowTo;
+
+#ifdef UE_MIRO_DEBUG
+					++s_CurrentBlock;
+				}
+			}
+#else
 				}
 			});
+#endif
 	}
 
 
@@ -7668,9 +8690,8 @@ namespace miro
 					}
 
 					*rowTo = physical_block_zero;
-					// \TODO: We are forcing to not use dual plane, this should be per-block.
-					// In this case we shoiuld weight-out A in all calculations, or even just compress as RGB.
-					astcrt::compress_block_rgba<BLOCK_SIZE, BLOCK_SIZE>((astcrt::unorm8_t*)block, rowTo, false);
+					// \TODO: In this case we shoiuld weight-out A in all calculations, or even just compress as RGB.
+					astcrt::compress_block_rgba<BLOCK_SIZE, BLOCK_SIZE>((astcrt::unorm8_t*)block, rowTo);
 
 					++rowTo;
 				}
@@ -7692,8 +8713,12 @@ namespace miro
 			return;
 		}
 
-		//for (uint32 BlockY = 0; BlockY < NumBlocksY; ++BlockY)
+#ifdef UE_MIRO_DEBUG
+		s_CurrentBlock = 0;
+		for (int32 BlockY = 0; BlockY < NumBlocksY; ++BlockY)
+#else
 		ParallelFor(NumBlocksY, [NumBlocksX, sx, sy, from, to](int32 BlockY)
+#endif
 		{
 			for (int32 BlockX = 0; BlockX < NumBlocksX; ++BlockX)
 			{
@@ -7703,8 +8728,15 @@ namespace miro
 				constexpr int32 CompressedBlockSize = 16;
 				const uint8* SrcBlockPtr = from + (BlockY * NumBlocksX + BlockX) * CompressedBlockSize;
 
+#ifdef UE_MIRO_DEBUG
+				if (s_CurrentBlock==s_DebugBlock)
+				{
+					astcdec::LogBlock(SrcBlockPtr, BLOCK_SIZE, BLOCK_SIZE, bIsSRGB, true );
+				}
+#endif
+
 				bool bSuccess = astcdec::decompress(Block, SrcBlockPtr, bIsSRGB, BLOCK_SIZE, BLOCK_SIZE);
-				check(bSuccess);
+				miro_check(bSuccess);
 
 				for (uint32 py = 0; py < BLOCK_SIZE; py++)
 				{
@@ -7720,8 +8752,15 @@ namespace miro
 						toPixel[3] = Block[py * BLOCK_SIZE * 4 + px * 4 + 3];
 					}
 				}
+
+#ifdef UE_MIRO_DEBUG
+				++s_CurrentBlock;
+			}
+		}
+#else
 			}
 		});
+#endif
 	}
 
 
@@ -7746,7 +8785,7 @@ namespace miro
 				const uint8* SrcBlockPtr = from + (BlockY * NumBlocksX + BlockX) * CompressedBlockSize;
 
 				bool bSuccess = astcdec::decompress(Block, SrcBlockPtr, bIsSRGB, BLOCK_SIZE, BLOCK_SIZE);
-				check(bSuccess);
+				miro_check(bSuccess);
 
 				for (uint32 py = 0; py < BLOCK_SIZE; py++)
 				{
@@ -7769,7 +8808,6 @@ namespace miro
 	}
 
 
-
 	//---------------------------------------------------------------------------------------------
 	template<uint32 BLOCK_SIZE>
 	void Generic_L_to_ASTCRGBL(uint32 sx, uint32 sy, const uint8_t* from, uint8_t* to, int32 Quality)
@@ -7944,7 +8982,7 @@ namespace miro
 					bool bIsSRGB = false;
 					uint8 Block[BLOCK_SIZE * BLOCK_SIZE * 4];
 					bool bSuccess = astcdec::decompress(Block, SrcBlockPtr, bIsSRGB, BLOCK_SIZE, BLOCK_SIZE);
-					check(bSuccess);
+					miro_check(bSuccess);
 
 					for (uint32 BlockY = 0; BlockY < BLOCK_SIZE; BlockY++)
 					{
@@ -7987,7 +9025,7 @@ namespace miro
 					bool bIsSRGB = false;
 					uint8 Block[BLOCK_SIZE * BLOCK_SIZE * 4];
 					bool bSuccess = astcdec::decompress(Block, SrcBlockPtr, bIsSRGB, BLOCK_SIZE, BLOCK_SIZE);
-					check(bSuccess);
+					miro_check(bSuccess);
 
 					for (uint32 BlockY = 0; BlockY < BLOCK_SIZE; BlockY++)
 					{
@@ -8291,3 +9329,7 @@ namespace miro
 
 
 }
+
+#ifdef UE_MIRO_DEBUG
+UE_ENABLE_OPTIMIZATION
+#endif
