#rnx

diff --git a/Engine/Plugins/Experimental/Mutable/Source/CustomizableObject/Private/MuCO/CustomizableInstanceLODManagement.cpp b/Engine/Plugins/Experimental/Mutable/Source/CustomizableObject/Private/MuCO/CustomizableInstanceLODManagement.cpp
--- a/Engine/Plugins/Experimental/Mutable/Source/CustomizableObject/Private/MuCO/CustomizableInstanceLODManagement.cpp
+++ b/Engine/Plugins/Experimental/Mutable/Source/CustomizableObject/Private/MuCO/CustomizableInstanceLODManagement.cpp
@@ -72,6 +72,13 @@ void UpdatePawnToInstancesDistances(const class UCustomizableObjectInstance* Onl
 {
 	for (TObjectIterator<UCustomizableSkeletalComponent> CustomizableSkeletalComponent; CustomizableSkeletalComponent; ++CustomizableSkeletalComponent)
 	{
+#if WITH_EDITOR
+		if (CustomizableSkeletalComponent && CustomizableSkeletalComponent->IsNetMode(NM_DedicatedServer))
+		{
+			continue;
+		}
+#endif
+
 		if (CustomizableSkeletalComponent->IsValidLowLevel() && (OnlyForInstance == nullptr || CustomizableSkeletalComponent->CustomizableObjectInstance == OnlyForInstance))
 		{
 			CustomizableSkeletalComponent->UpdateDistFromComponentToPlayer(ViewCenter.IsValid() ? ViewCenter.Get() : nullptr, OnlyForInstance != nullptr);
@@ -85,6 +92,12 @@ void UpdateCameraToInstancesDistance(const FVector CameraPosition)
 {
 	for (TObjectIterator<UCustomizableSkeletalComponent> CustomizableSkeletalComponent; CustomizableSkeletalComponent; ++CustomizableSkeletalComponent)
 	{
+#if WITH_EDITOR
+		if (CustomizableSkeletalComponent && CustomizableSkeletalComponent->IsNetMode(NM_DedicatedServer))
+		{
+			continue;
+		}
+#endif
 		if (CustomizableSkeletalComponent->IsValidLowLevel() && !CustomizableSkeletalComponent->IsTemplate())
 		{
 			CustomizableSkeletalComponent->UpdateDistFromComponentToLevelEditorCamera(CameraPosition);
@@ -111,6 +124,12 @@ void UCustomizableInstanceLODManagement::UpdateInstanceDistsAndLODs(FMutableInst
 
 			for (TObjectIterator<UCustomizableSkeletalComponent> CustomizableSkeletalComponent; CustomizableSkeletalComponent; ++CustomizableSkeletalComponent)
 			{
+#if WITH_EDITOR
+				if (CustomizableSkeletalComponent && CustomizableSkeletalComponent->IsNetMode(NM_DedicatedServer))
+				{
+					continue;
+				}
+#endif
 				if (CustomizableSkeletalComponent && !CustomizableSkeletalComponent->IsTemplate())
 				{
 					UWorld* LocalWorld = CustomizableSkeletalComponent->GetWorld();
@@ -275,6 +294,13 @@ void UCustomizableInstanceLODManagement::UpdateInstanceDistsAndLODs(FMutableInst
 
 		for (TObjectIterator<UCustomizableSkeletalComponent> CustomizableSkeletalComponent; CustomizableSkeletalComponent; ++CustomizableSkeletalComponent)
 		{
+#if WITH_EDITOR
+			if (CustomizableSkeletalComponent && CustomizableSkeletalComponent->IsNetMode(NM_DedicatedServer))
+			{
+				continue;
+			}
+#endif
+
 			if (CustomizableSkeletalComponent && !CustomizableSkeletalComponent->IsTemplate())
 			{
 				UCustomizableObjectInstance* COI = CustomizableSkeletalComponent->CustomizableObjectInstance;
diff --git a/Engine/Plugins/Experimental/Mutable/Source/CustomizableObject/Private/MuCO/CustomizableObjectInstance.cpp b/Engine/Plugins/Experimental/Mutable/Source/CustomizableObject/Private/MuCO/CustomizableObjectInstance.cpp
--- a/Engine/Plugins/Experimental/Mutable/Source/CustomizableObject/Private/MuCO/CustomizableObjectInstance.cpp
+++ b/Engine/Plugins/Experimental/Mutable/Source/CustomizableObject/Private/MuCO/CustomizableObjectInstance.cpp
@@ -2926,6 +2926,13 @@ void UCustomizableInstancePrivateData::DiscardResourcesAndSetReferenceSkeletalMe
 	{
 		UCustomizableSkeletalComponent* CustomizableSkeletalComponent = *It;
 
+#if WITH_EDITOR
+		if (CustomizableSkeletalComponent && CustomizableSkeletalComponent->IsNetMode(NM_DedicatedServer))
+		{
+			continue;
+		}
+#endif
+
 		if (CustomizableSkeletalComponent && CustomizableSkeletalComponent->CustomizableObjectInstance == Public)
 		{
 			UCustomizableObject* CustomizableObject = Public->GetCustomizableObject();
diff --git a/Engine/Plugins/Experimental/Mutable/Source/CustomizableObject/Private/MuCO/CustomizableObjectModule.cpp b/Engine/Plugins/Experimental/Mutable/Source/CustomizableObject/Private/MuCO/CustomizableObjectModule.cpp
--- a/Engine/Plugins/Experimental/Mutable/Source/CustomizableObject/Private/MuCO/CustomizableObjectModule.cpp
+++ b/Engine/Plugins/Experimental/Mutable/Source/CustomizableObject/Private/MuCO/CustomizableObjectModule.cpp
@@ -225,6 +225,13 @@ UCustomizableSkeletalComponent* GetPlayerCustomizableSkeletalComponent(const int
 		int32 IndexFound = INDEX_NONE;
 		for (TObjectIterator<UCustomizableSkeletalComponent> CustomizableSkeletalComponent; CustomizableSkeletalComponent; ++CustomizableSkeletalComponent)
 		{
+#if WITH_EDITOR
+			if (CustomizableSkeletalComponent && CustomizableSkeletalComponent->IsNetMode(NM_DedicatedServer))
+			{
+				continue;
+			}
+#endif
+
 			if (CustomizableSkeletalComponent->IsValidLowLevel() && !CustomizableSkeletalComponent->IsTemplate())
 			{
 				AActor* CustomizableActor = CustomizableSkeletalComponent->GetAttachmentRootActor();
@@ -249,6 +256,13 @@ UCustomizableSkeletalComponent* GetPlayerCustomizableSkeletalComponent(const int
 		int32 IndexFound = INDEX_NONE;
 		for (TObjectIterator<UCustomizableSkeletalComponent> CustomizableSkeletalComponent; CustomizableSkeletalComponent; ++CustomizableSkeletalComponent)
 		{
+#if WITH_EDITOR
+			if (CustomizableSkeletalComponent && CustomizableSkeletalComponent->IsNetMode(NM_DedicatedServer))
+			{
+				continue;
+			}
+#endif
+
 			if (CustomizableSkeletalComponent->IsValidLowLevel() && !CustomizableSkeletalComponent->IsTemplate())
 			{
 				++IndexFound;
diff --git a/Engine/Plugins/Experimental/Mutable/Source/CustomizableObject/Private/MuCO/CustomizableObjectSystem.cpp b/Engine/Plugins/Experimental/Mutable/Source/CustomizableObject/Private/MuCO/CustomizableObjectSystem.cpp
--- a/Engine/Plugins/Experimental/Mutable/Source/CustomizableObject/Private/MuCO/CustomizableObjectSystem.cpp
+++ b/Engine/Plugins/Experimental/Mutable/Source/CustomizableObject/Private/MuCO/CustomizableObjectSystem.cpp
@@ -244,6 +244,13 @@ void UCustomizableObjectSystem::LogShowData(bool bFullInfo, bool ShowMaterialInf
 	{
 		const UCustomizableSkeletalComponent* CustomizableSkeletalComponent = *It;
 
+#if WITH_EDITOR
+		if (CustomizableSkeletalComponent && CustomizableSkeletalComponent->IsNetMode(NM_DedicatedServer))
+		{
+			continue;
+		}
+#endif
+
 		if (CustomizableSkeletalComponent && CustomizableSkeletalComponent->CustomizableObjectInstance)
 		{
 			const AActor* ParentActor = CustomizableSkeletalComponent->GetAttachmentRootActor();
@@ -692,6 +699,13 @@ void FinishUpdateGlobal(UCustomizableObjectInstance* Instance, EUpdateResult Upd
 		// Call Customizable Skeletal Components updated callbacks.
 		for (TObjectIterator<UCustomizableSkeletalComponent> It; It; ++It) // Since iterating objects is expensive, for now CustomizableSkeletalComponent does not have a FinishUpdate function.
 		{
+#if WITH_EDITOR
+			if (It && It->IsNetMode(NM_DedicatedServer))
+			{
+				continue;
+			}
+#endif
+
 			if (const UCustomizableSkeletalComponent* CustomizableSkeletalComponent = *It;
 				CustomizableSkeletalComponent &&
 				CustomizableSkeletalComponent->CustomizableObjectInstance == Instance)
@@ -741,6 +755,13 @@ void UpdateSkeletalMesh(UCustomizableObjectInstance& CustomizableObjectInstance,
 	{
 		UCustomizableSkeletalComponent* CustomizableSkeletalComponent = *It;
 
+#if WITH_EDITOR
+		if (CustomizableSkeletalComponent && CustomizableSkeletalComponent->IsNetMode(NM_DedicatedServer))
+		{
+			continue;
+		}
+#endif
+
 		if (CustomizableSkeletalComponent &&
 			(CustomizableSkeletalComponent->CustomizableObjectInstance == &CustomizableObjectInstance) &&
 			CustomizableObjectInstance.SkeletalMeshes.IsValidIndex(CustomizableSkeletalComponent->ComponentIndex)
@@ -2525,6 +2546,11 @@ namespace impl
 #if WITH_EDITOR
 		for (TObjectIterator<UCustomizableSkeletalComponent> CustomizableSkeletalComponent; CustomizableSkeletalComponent && !bIsInEditorViewport; ++CustomizableSkeletalComponent)
 		{
+			if (CustomizableSkeletalComponent && CustomizableSkeletalComponent->IsNetMode(NM_DedicatedServer))
+			{
+				continue;
+			}
+
 			if (CustomizableSkeletalComponent &&
 				CustomizableSkeletalComponent->CustomizableObjectInstance == CandidateInstance)
 			{
diff --git a/Engine/Plugins/Experimental/Mutable/Source/CustomizableObject/Private/MuCO/LogBenchmarkUtil.cpp b/Engine/Plugins/Experimental/Mutable/Source/CustomizableObject/Private/MuCO/LogBenchmarkUtil.cpp
--- a/Engine/Plugins/Experimental/Mutable/Source/CustomizableObject/Private/MuCO/LogBenchmarkUtil.cpp
+++ b/Engine/Plugins/Experimental/Mutable/Source/CustomizableObject/Private/MuCO/LogBenchmarkUtil.cpp
@@ -296,6 +296,11 @@ void LogBenchmarkUtil::UpdateStats(FMutableStats& StatsToUpdate, const TArray< T
 
 				for (TObjectIterator<UCustomizableSkeletalComponent> CustomizableSkeletalComponent; CustomizableSkeletalComponent; ++CustomizableSkeletalComponent)
 				{
+					if (CustomizableSkeletalComponent && CustomizableSkeletalComponent->IsNetMode(NM_DedicatedServer))
+					{
+						continue;
+					}
+
 					AActor* ParentActor = CustomizableSkeletalComponent->GetAttachmentRootActor();
 					UCustomizableObjectInstance* Instance = CustomizableSkeletalComponent->CustomizableObjectInstance;
 
