[mutable] Static analysis fixes for recent compression changes.
#rnx

diff --git a/Engine/Plugins/Experimental/Mutable/Source/MutableRuntime/Private/MuR/BlockCompression/Miro/Miro.cpp b/Engine/Plugins/Experimental/Mutable/Source/MutableRuntime/Private/MuR/BlockCompression/Miro/Miro.cpp
--- a/Engine/Plugins/Experimental/Mutable/Source/MutableRuntime/Private/MuR/BlockCompression/Miro/Miro.cpp
+++ b/Engine/Plugins/Experimental/Mutable/Source/MutableRuntime/Private/MuR/BlockCompression/Miro/Miro.cpp
@@ -5760,8 +5760,8 @@ namespace impl
 			uint32_t& z() { return m_c[2]; }
 			uint32_t& w() { return m_c[3]; }
 
-			uint32_t operator[] (uint32_t idx) const { miro_check(idx < 4);  return m_c[idx]; }
-			uint32_t& operator[] (uint32_t idx) { miro_check(idx < 4);  return m_c[idx]; }
+			inline uint32_t operator[] (uint32_t idx) const { check(idx < 4); if (idx < 4) return m_c[idx]; return 0; }
+			inline uint32_t& operator[] (uint32_t idx) { check(idx < 4); if (idx < 4) return m_c[idx]; return m_c[0]; }
 		};
 
 		struct IVec4
@@ -5799,8 +5799,8 @@ namespace impl
 				return UVec4(FMath::Max(0, m_c[0]), FMath::Max(0, m_c[1]), FMath::Max(0, m_c[2]), FMath::Max(0, m_c[3]));
 			}
 
-			int32_t operator[] (uint32_t idx) const { miro_check(idx < 4);  return m_c[idx]; }
-			int32_t& operator[] (uint32_t idx) { miro_check(idx < 4);  return m_c[idx]; }
+			inline int32_t operator[] (uint32_t idx) const { check(idx < 4); if (idx < 4) return m_c[idx]; return 0; }
+			inline int32_t& operator[] (uint32_t idx) { check(idx < 4); if (idx < 4) return m_c[idx]; return m_c[0]; }
 		};
 
 		struct IVec3
@@ -5829,8 +5829,8 @@ namespace impl
 			int32_t& y() { return m_c[1]; }
 			int32_t& z() { return m_c[2]; }
 
-			int32_t operator[] (uint32_t idx) const { miro_check(idx < 3);  return m_c[idx]; }
-			int32_t& operator[] (uint32_t idx) { miro_check(idx < 3);  return m_c[idx]; }
+			inline int32_t operator[] (uint32_t idx) const { check(idx < 3); if (idx < 3) return m_c[idx]; return 0; }
+			inline int32_t& operator[] (uint32_t idx) { check(idx < 3); if (idx < 3) return m_c[idx]; return m_c[0]; }
 		};
 
 		static uint32_t deDivRoundUp32(uint32_t a, uint32_t b)
