[mutable] Make sure any change that requires CO resave also invalidates derived data

#rnx

diff --git a/Engine/Plugins/Experimental/Mutable/Source/CustomizableObject/Private/MuCO/CustomizableObject.cpp b/Engine/Plugins/Experimental/Mutable/Source/CustomizableObject/Private/MuCO/CustomizableObject.cpp
--- a/Engine/Plugins/Experimental/Mutable/Source/CustomizableObject/Private/MuCO/CustomizableObject.cpp
+++ b/Engine/Plugins/Experimental/Mutable/Source/CustomizableObject/Private/MuCO/CustomizableObject.cpp
@@ -139,7 +139,6 @@ void UCustomizableObject::PreSave(FObjectPreSaveContext ObjectSaveContext)
 #endif
 }
 
-
 #endif // End WITH_EDITOR
 
 
@@ -157,6 +156,12 @@ void UCustomizableObject::PostLoad()
 		ReferenceSkeletalMesh_DEPRECATED = nullptr;
 	}
 
+	// Register to dirty delegate so we update derived data version ID each time that the package is marked as dirty.
+	if (UPackage* Package = GetOutermost())
+	{
+		Package->PackageMarkedDirtyEvent.AddWeakLambda(this, [this](UPackage* Pkg, bool bWasDirty){UpdateVersionId();});
+	}
+
 	const int32 CustomizableObjectCustomVersion = GetLinkerCustomVersion(FCustomizableObjectCustomVersion::GUID);
 
 	// Convert texture compression option
@@ -1713,7 +1718,7 @@ bool UCustomizableObject::IsParameterMultidimensional(const int32& InParamIndex)
 	{
 		return PrivateData->GetModel()->IsParameterMultidimensional(InParamIndex);
 	}
-
+	
 	return false;
 }
 
