[Mutable] Fix COI reusing GeneratedMaterials between compilations.
Fix Skeletal Meshes being detected as different between consecutive updates after modifying the MaxLOD due to a change of State.
Set the console variable Mutable.ApplyFixDoComponentsNeedUpdate to true.

#rnx

diff --git a/Engine/Plugins/Experimental/Mutable/Source/CustomizableObject/Private/MuCO/CustomizableObjectInstance.cpp b/Engine/Plugins/Experimental/Mutable/Source/CustomizableObject/Private/MuCO/CustomizableObjectInstance.cpp
--- a/Engine/Plugins/Experimental/Mutable/Source/CustomizableObject/Private/MuCO/CustomizableObjectInstance.cpp
+++ b/Engine/Plugins/Experimental/Mutable/Source/CustomizableObject/Private/MuCO/CustomizableObjectInstance.cpp
@@ -132,6 +132,8 @@ void UCustomizableInstancePrivateData::InvalidateGeneratedData()
 		ComponentData.LastMeshIdPerLOD.Init(MAX_uint64, MAX_MESH_LOD_COUNT);
 	}
 
+	GeneratedMaterials.Empty();
+
 	LastUpdateData.Clear();
 }
 
@@ -1748,7 +1750,7 @@ void SetMeshUVChannelDensity(FMeshUVChannelInfo& UVChannelInfo, float Density =
 }
 
 // TODO PRP: MTBL-1653 Remove CVar
-static bool bApplyFixDoComponentsNeedUpdate = false;
+static bool bApplyFixDoComponentsNeedUpdate = true;
 FAutoConsoleVariableRef CVarApplyFixDoComponentsNeedUpdate(
 	TEXT("Mutable.ApplyFixDoComponentsNeedUpdate"),
 	bApplyFixDoComponentsNeedUpdate,
@@ -1982,6 +1984,12 @@ bool UCustomizableInstancePrivateData::UpdateSkeletalMesh_PostBeginUpdate0(UCust
 
 	TextureReuseCache.Empty(); // Sections may have changed, so invalidate the texture reuse cache because it's indexed by section
 
+	// Reset last mesh IDs.
+	for (FCustomizableInstanceComponentData& ComponentData : ComponentsData)
+	{
+		ComponentData.LastMeshIdPerLOD.Init(MAX_uint64, MAX_MESH_LOD_COUNT);
+	}
+
 	TArray<TObjectPtr<USkeletalMesh>> OldSkeletalMeshes = Public->SkeletalMeshes;
 
 	bool bSuccess = true;
