[Mutable] Merge EnqueueUpdateSkeletalMesh from Instance and System.

- Moved FMutableUpdateCandidate function bodies from header to source files.
- Moved IsUpdateRequired and GetUpdatePriority functions from the Instance to the System (who has the Update resonsability).
- Moved Hashes from the public to the private.

#rnx

diff --git a/Engine/Plugins/Experimental/Mutable/Source/CustomizableObject/Classes/MuCO/CustomizableObjectInstance.h b/Engine/Plugins/Experimental/Mutable/Source/CustomizableObject/Classes/MuCO/CustomizableObjectInstance.h
--- a/Engine/Plugins/Experimental/Mutable/Source/CustomizableObject/Classes/MuCO/CustomizableObjectInstance.h
+++ b/Engine/Plugins/Experimental/Mutable/Source/CustomizableObject/Classes/MuCO/CustomizableObjectInstance.h
@@ -278,18 +278,6 @@ public:
 	UFUNCTION(BlueprintCallable, Category = CustomizableObjectInstance)
 	void UpdateSkeletalMeshAsyncResult(FInstanceUpdateDelegate Callback, bool bIgnoreCloseDist = false, bool bForceHighPriority = false);
 
-private:
-	/** Perform all the checks required to see if the update should begin. */
-	EUpdateRequired IsUpdateRequired(bool bOnlyUpdateIfNotGenerated, bool bIgnoreCloseDist) const;
-
-public:
-	/** Private API.
-	 *
-	 * Update Skeletal Mesh asynchronously. Immersive function.
-	 * Once the update reaches this function, the update has been considered started and must complete all the update flow.
-	 * Starting at this function, all Update code paths must end up in FinishUpdateGlobal! */
-	void EnqueueUpdateSkeletalMesh(bool bOnlyUpdateIfNotGenerated, bool bIgnoreCloseDist, bool bForceHighPriority, const EUpdateRequired* OptionalUpdateRequired, FInstanceUpdateDelegate* UpdateCallback);
-
 	// Clones the instance creating a new identical transient instance.
 	UFUNCTION(BlueprintCallable, Category = CustomizableObjectInstance)
 	UCustomizableObjectInstance* Clone();
@@ -307,9 +295,6 @@ public:
 	// Releases all the mutable resources this instance holds, should only be called when it is not going to be used any more.
 	void ReleaseMutableResources(bool bCalledFromBeginDestroy);
 
-	/** Returns the priority an update issued by EnqueueUpdateSkeletalMesh would get in the current instance configuration and state */
-	EQueuePriorityType GetUpdatePriority(bool bForceHighPriority) const;
-
 	// Returns de description texture (ex: color bar) for this parameter and DescIndex
 	// This will only be valid if bBuildParameterDecorations was set to true before the last update.
 	UFUNCTION(BlueprintCallable, Category = CustomizableObjectInstance, meta = (DeprecatedFunction, DeprecationMessage = "Parameter decorations have been removed. This method will be removed in future versions."))
@@ -726,12 +711,6 @@ private:
 
 	UPROPERTY( Transient )
 	TObjectPtr<UCustomizableInstancePrivateData> PrivateData;
-	
-	/** Hash of the UCustomizableObjectInstance::Descriptor on the last update request. */
-	FDescriptorRuntimeHash UpdateDescriptorRuntimeHash;
-	
-	/** Hash of the UCustomizableObjectInstance::Descriptor on the last successful update. */
-	FDescriptorRuntimeHash DescriptorRuntimeHash;
 
 	/** LODs applied on the beginning of the last update. Represent the actual LODs the Instance is using (not strictly true since an update can fail). */
 	int32 CurrentMinLOD = -1;
@@ -800,26 +779,11 @@ public:
 		CustomizableObjectInstance(const_cast<UCustomizableObjectInstance*>(InCustomizableObjectInstance)), MinLOD(InMinLOD), MaxLOD(InMaxLOD),
 		RequestedLODLevels(InRequestedLODLevels) {}
 
-	bool HasBeenIssued()
-	{
-		return bHasBeenIssued;
-	}
-
-	void Issue()
-	{
-		bHasBeenIssued = true;
-	}
-
-	void ApplyLODUpdateParamsToInstance()
-	{
-		CustomizableObjectInstance->Descriptor.MinLOD = MinLOD;
-		CustomizableObjectInstance->Descriptor.MaxLOD = MaxLOD;
+	bool HasBeenIssued() const;
 
-		CustomizableObjectInstance->Descriptor.RequestedLODLevels = RequestedLODLevels;
+	void Issue();
 
-		CustomizableObjectInstance->UpdateDescriptorRuntimeHash.UpdateMinMaxLOD(MinLOD, MaxLOD);
-		CustomizableObjectInstance->UpdateDescriptorRuntimeHash.UpdateRequestedLODs(CustomizableObjectInstance->Descriptor.RequestedLODLevels);
-	}
+	void ApplyLODUpdateParamsToInstance();
 
 private:
 	/** If true it means that EnqueueUpdateSkeletalMesh has decided this update should be performed, if false it should be ignored. Just used for consistency checks */
diff --git a/Engine/Plugins/Experimental/Mutable/Source/CustomizableObject/Private/MuCO/CustomizableInstancePrivateData.h b/Engine/Plugins/Experimental/Mutable/Source/CustomizableObject/Private/MuCO/CustomizableInstancePrivateData.h
--- a/Engine/Plugins/Experimental/Mutable/Source/CustomizableObject/Private/MuCO/CustomizableInstancePrivateData.h
+++ b/Engine/Plugins/Experimental/Mutable/Source/CustomizableObject/Private/MuCO/CustomizableInstancePrivateData.h
@@ -365,6 +365,14 @@ private:
 	
 	ECOInstanceFlags InstanceFlagsPrivate = ECOInstanceFlags::ECONone;
 
+public:
+	/** Hash of the UCustomizableObjectInstance::Descriptor on the last update request. */
+	FDescriptorRuntimeHash UpdateDescriptorRuntimeHash;
+	
+	/** Hash of the UCustomizableObjectInstance::Descriptor on the last successful update. */
+	FDescriptorRuntimeHash DescriptorRuntimeHash;
+
+private:
 	// Struct used by BuildMaterials() to identify common materials between LODs
 	struct FMutableMaterialPlaceholder
 	{
diff --git a/Engine/Plugins/Experimental/Mutable/Source/CustomizableObject/Private/MuCO/CustomizableObjectInstance.cpp b/Engine/Plugins/Experimental/Mutable/Source/CustomizableObject/Private/MuCO/CustomizableObjectInstance.cpp
--- a/Engine/Plugins/Experimental/Mutable/Source/CustomizableObject/Private/MuCO/CustomizableObjectInstance.cpp
+++ b/Engine/Plugins/Experimental/Mutable/Source/CustomizableObject/Private/MuCO/CustomizableObjectInstance.cpp
@@ -737,171 +737,15 @@ bool UCustomizableObjectInstance::CanUpdateInstance() const
 
 void UCustomizableObjectInstance::UpdateSkeletalMeshAsync(bool bIgnoreCloseDist, bool bForceHighPriority)
 {
-	EnqueueUpdateSkeletalMesh(false, bIgnoreCloseDist, bForceHighPriority, nullptr, nullptr);
+	FCustomizableObjectSystemPrivate* SystemPrivate = UCustomizableObjectSystem::GetInstance()->GetPrivate();	
+	SystemPrivate->EnqueueUpdateSkeletalMesh(*this, false, bIgnoreCloseDist, bForceHighPriority, nullptr, nullptr);
 }
 
 
 void UCustomizableObjectInstance::UpdateSkeletalMeshAsyncResult(FInstanceUpdateDelegate Callback, bool bIgnoreCloseDist, bool bForceHighPriority)
 {
-	EnqueueUpdateSkeletalMesh(false, bIgnoreCloseDist, bForceHighPriority, nullptr, &Callback);
-}
-
-
-EUpdateRequired UCustomizableObjectInstance::IsUpdateRequired(bool bOnlyUpdateIfNotGenerated, bool bIgnoreCloseDist) const
-{
-	UCustomizableObjectSystem* System = UCustomizableObjectSystem::GetInstance();
-	const UCustomizableInstancePrivateData* const Private = GetPrivate();
-	
-	UCustomizableObject* CustomizableObject = GetCustomizableObject();
-	if (!CanUpdateInstance() || CustomizableObject->IsLocked())
-	{
-		return EUpdateRequired::NoUpdate;
-	}
-
-	const bool bIsGenerated = Private->HasCOInstanceFlags(Generated);
-	const int32 NumGeneratedInstancesLimit = System->GetInstanceLODManagement()->GetNumGeneratedInstancesLimitFullLODs();
-	const int32 NumGeneratedInstancesLimitLOD1 = System->GetInstanceLODManagement()->GetNumGeneratedInstancesLimitLOD1();
-	const int32 NumGeneratedInstancesLimitLOD2 = System->GetInstanceLODManagement()->GetNumGeneratedInstancesLimitLOD2();
-
-	if (!bIsGenerated && // Prevent generating more instances than the limit, but let updates to existing instances run normally
-		NumGeneratedInstancesLimit > 0 &&
-		System->GetPrivate()->GetCountAllocatedSkeletalMesh() > NumGeneratedInstancesLimit + NumGeneratedInstancesLimitLOD1 + NumGeneratedInstancesLimitLOD2)
-	{
-		return EUpdateRequired::NoUpdate;
-	}
-
-	const bool bShouldUpdateLODs = Private->HasCOInstanceFlags(PendingLODsUpdate);
-
-	const bool bDiscardByDistance = Private->LastMinSquareDistFromComponentToPlayer > FMath::Square(System->GetInstanceLODManagement()->GetOnlyUpdateCloseCustomizableObjectsDist());
-	const bool bLODManagementDiscard = System->GetInstanceLODManagement()->IsOnlyUpdateCloseCustomizableObjectsEnabled() &&
-			bDiscardByDistance &&
-			!bIgnoreCloseDist;
-	
-	if (Private->HasCOInstanceFlags(DiscardedByNumInstancesLimit) ||
-		bLODManagementDiscard)
-	{
-		if (bIsGenerated && !Private->HasCOInstanceFlags(Updating))
-		{
-			return EUpdateRequired::Discard;		
-		}
-		else
-		{
-			return EUpdateRequired::NoUpdate;
-		}
-	}
-
-	if (bIsGenerated &&
-		!bShouldUpdateLODs &&
-		bOnlyUpdateIfNotGenerated)
-	{
-		return EUpdateRequired::NoUpdate;
-	}
-
-	return EUpdateRequired::Update;
-}
-
-
-EQueuePriorityType UCustomizableObjectInstance::GetUpdatePriority(bool bForceHighPriority) const
-{
-	const bool bIsGenerated = GetPrivate()->HasCOInstanceFlags(Generated);
-	const bool bShouldUpdateLODs = GetPrivate()->HasCOInstanceFlags(PendingLODsUpdate);
-	const bool bIsDowngradeLODUpdate = GetPrivate()->HasCOInstanceFlags(PendingLODsDowngrade);
-	const bool bIsPlayerOrNearIt = GetPrivate()->HasCOInstanceFlags(UsedByPlayerOrNearIt);
-
-	EQueuePriorityType Priority = EQueuePriorityType::Low;
-	if (bForceHighPriority)
-	{
-		Priority = EQueuePriorityType::High;
-	}
-	else if (!bIsGenerated || !HasAnySkeletalMesh())
-	{
-		Priority = EQueuePriorityType::Med;
-	}
-	else if (bShouldUpdateLODs && bIsDowngradeLODUpdate)
-	{
-		Priority = EQueuePriorityType::Med_Low;
-	}
-	else if (bIsPlayerOrNearIt && bShouldUpdateLODs && !bIsDowngradeLODUpdate)
-	{
-		Priority = EQueuePriorityType::High;
-	}
-	else if (bShouldUpdateLODs && !bIsDowngradeLODUpdate)
-	{
-		Priority = EQueuePriorityType::Med;
-	}
-	else if (bIsPlayerOrNearIt)
-	{
-		Priority = EQueuePriorityType::High;
-	}
-
-	return Priority;
-}
-
-
-void UCustomizableObjectInstance::EnqueueUpdateSkeletalMesh(bool bOnlyUpdateIfNotGenerated, bool bIgnoreCloseDist, bool bForceHighPriority, const EUpdateRequired* OptionalUpdateRequired, FInstanceUpdateDelegate* UpdateCallback)
-{
-	MUTABLE_CPUPROFILER_SCOPE(UCustomizableInstancePrivateData::EnqueueUpdateSkeletalMesh);
-	check(IsInGameThread());
-	check(!OptionalUpdateRequired || *OptionalUpdateRequired != EUpdateRequired::NoUpdate); // If no update is required this functions must not be called.
-
-	UCustomizableObjectSystem* System = UCustomizableObjectSystem::GetInstance();
-
-	if (!CanUpdateInstance())
-	{
-		FinishUpdateGlobal(this, EUpdateResult::ErrorDiscarded, UpdateCallback);
-		return;
-	}
-	
-	const EUpdateRequired UpdateRequired = OptionalUpdateRequired ? *OptionalUpdateRequired : IsUpdateRequired(bOnlyUpdateIfNotGenerated, bIgnoreCloseDist);
-	switch (UpdateRequired)
-	{
-	case EUpdateRequired::NoUpdate:
-	{	
-		FinishUpdateGlobal(this, EUpdateResult::ErrorDiscarded, UpdateCallback);
-		break;
-	}		
-	case EUpdateRequired::Update:
-	{
-		EQueuePriorityType Priority = GetUpdatePriority(bForceHighPriority);
-
-		const uint32 InstanceId = GetUniqueID();
-		const float Distance = FMath::Sqrt(GetPrivate()->LastMinSquareDistFromComponentToPlayer);
-		const bool bIsPlayerOrNearIt = GetPrivate()->HasCOInstanceFlags(UsedByPlayerOrNearIt);
-		UE_LOG(LogMutable, Log, TEXT("Started UpdateSkeletalMesh Async. of Instance %d with priority %d at dist %f bIsPlayerOrNearIt=%d, frame=%d"), InstanceId, static_cast<int32>(Priority), Distance, bIsPlayerOrNearIt, GFrameNumber);				
-
-		UpdateDescriptorRuntimeHash = FDescriptorRuntimeHash(Descriptor);
-
-#if WITH_EDITOR
-		if (!UCustomizableObjectSystem::GetInstance()->IsCompilationDisabled())
-		{
-			if (SkeletalMeshStatus != ESkeletalMeshState::AsyncUpdatePending)
-			{
-				PreUpdateSkeletalMeshStatus = SkeletalMeshStatus;
-			}
-		}
-#endif
-
-		if (GetPrivate()->HasCOInstanceFlags(PendingLODsUpdate))
-		{
-			UE_LOG(LogMutable, Verbose, TEXT("LOD change: %d, %d -> %d, %d"), GetCurrentMinLOD(), GetCurrentMaxLOD(), GetMinLODToLoad(), GetMaxLODToLoad());
-		}
-		
-		// Do not do work after calling EnqueueUpdateSkeletalMesh. This function can optimize an update and fully complete it before even exiting its scope.
-		System->GetPrivate()->EnqueueUpdateSkeletalMesh(*this, Priority, UpdateCallback);
-		break;
-	}
-
-	case EUpdateRequired::Discard:
-	{
-		System->GetPrivate()->InitDiscardResourcesSkeletalMesh(this);
-		
-		FinishUpdateGlobal(this, EUpdateResult::ErrorDiscarded, UpdateCallback);
-		break;
-	}
-
-	default:
-		check(false); // Case not implemented.
-	}
+	FCustomizableObjectSystemPrivate* SystemPrivate = UCustomizableObjectSystem::GetInstance()->GetPrivate();	
+	SystemPrivate->EnqueueUpdateSkeletalMesh(*this, false, bIgnoreCloseDist, bForceHighPriority, nullptr, &Callback);
 }
 
 
@@ -912,7 +756,9 @@ void UCustomizableInstancePrivateData::TickUpdateCloseCustomizableObjects(UCusto
 		return;
 	}
 
-	const EUpdateRequired UpdateRequired = Public.IsUpdateRequired(true, false);
+	const FCustomizableObjectSystemPrivate* SystemPrivate = UCustomizableObjectSystem::GetInstance()->GetPrivate();	
+
+	const EUpdateRequired UpdateRequired = SystemPrivate->IsUpdateRequired(Public, true, false);
 	if (UpdateRequired != EUpdateRequired::NoUpdate) // Since this is done in the tick, avoid starting an update that we know for sure that would not be performed. Once started it has some performance implications that we want to avoid.
 	{
 		if (UpdateRequired == EUpdateRequired::Discard)
@@ -922,7 +768,7 @@ void UCustomizableInstancePrivateData::TickUpdateCloseCustomizableObjects(UCusto
 		}
 		else if (UpdateRequired == EUpdateRequired::Update)
 		{
-			EQueuePriorityType Priority = Public.GetUpdatePriority(false);
+			const EQueuePriorityType Priority = SystemPrivate->GetUpdatePriority(Public, false);
 
 			FMutableUpdateCandidate* UpdateCandidate = InOutRequestedUpdates.Find(&Public);
 
@@ -967,9 +813,10 @@ void UCustomizableInstancePrivateData::UpdateInstanceIfNotGenerated(UCustomizabl
 		return;
 	}
 
-	Public.EnqueueUpdateSkeletalMesh(true, false, false, nullptr, nullptr);
+	FCustomizableObjectSystemPrivate* SystemPrivate = UCustomizableObjectSystem::GetInstance()->GetPrivate();	
+	SystemPrivate->EnqueueUpdateSkeletalMesh(Public, true, false, false, nullptr, nullptr);
 
-	EQueuePriorityType Priority = Public.GetUpdatePriority(false);
+	EQueuePriorityType Priority = SystemPrivate->GetUpdatePriority(Public, false);
 	FMutableUpdateCandidate* UpdateCandidate = InOutRequestedUpdates.Find(&Public);
 
 	if (UpdateCandidate)
@@ -2928,7 +2775,7 @@ void UCustomizableInstancePrivateData::DiscardResourcesAndSetReferenceSkeletalMe
 	InvalidateGeneratedData();
 	
 	Public->SkeletalMeshes.Reset();
-	Public->DescriptorRuntimeHash = FDescriptorRuntimeHash();
+	DescriptorRuntimeHash = FDescriptorRuntimeHash();
 
 	for (TObjectIterator<UCustomizableSkeletalComponent> It; It; ++It)
 	{
@@ -3002,7 +2849,7 @@ void UCustomizableObjectInstance::FinishUpdate(EUpdateResult UpdateResult, const
 {
 	if (UpdateResult == EUpdateResult::Success)
 	{
-		DescriptorRuntimeHash = InUpdatedHash;
+		GetPrivate()->DescriptorRuntimeHash = InUpdatedHash;
 	}
 
 	if (UpdateResult == EUpdateResult::Success)
@@ -3016,13 +2863,13 @@ void UCustomizableObjectInstance::FinishUpdate(EUpdateResult UpdateResult, const
 
 FDescriptorRuntimeHash UCustomizableObjectInstance::GetDescriptorRuntimeHash() const
 {
-	return DescriptorRuntimeHash;
+	return GetPrivate()->DescriptorRuntimeHash;
 }
 
 
 FDescriptorRuntimeHash UCustomizableObjectInstance::GetUpdateDescriptorRuntimeHash() const
 {
-	return UpdateDescriptorRuntimeHash;
+	return GetPrivate()->UpdateDescriptorRuntimeHash;
 }
 
 #if WITH_EDITOR
@@ -3040,6 +2887,30 @@ UCustomizableInstancePrivateData* UCustomizableObjectInstance::GetPrivate() cons
 }
 
 
+bool FMutableUpdateCandidate::HasBeenIssued() const
+{
+	return bHasBeenIssued;
+}
+
+
+void FMutableUpdateCandidate::Issue()
+{
+	bHasBeenIssued = true;
+}
+
+
+void FMutableUpdateCandidate::ApplyLODUpdateParamsToInstance()
+{
+	CustomizableObjectInstance->Descriptor.MinLOD = MinLOD;
+	CustomizableObjectInstance->Descriptor.MaxLOD = MaxLOD;
+
+	CustomizableObjectInstance->Descriptor.RequestedLODLevels = RequestedLODLevels;
+
+	CustomizableObjectInstance->GetPrivate()->UpdateDescriptorRuntimeHash.UpdateMinMaxLOD(MinLOD, MaxLOD);
+	CustomizableObjectInstance->GetPrivate()->UpdateDescriptorRuntimeHash.UpdateRequestedLODs(CustomizableObjectInstance->Descriptor.RequestedLODLevels);
+}
+
+
 void UCustomizableObjectInstance::CommitMinMaxLOD()
 {
 	CurrentMinLOD = Descriptor.MinLOD;
@@ -6523,9 +6394,9 @@ void UCustomizableObjectInstance::SetRequestedLODs(int32 InMinLOD, int32 InMaxLO
 		const int32 ComponentCount = GetNumComponents();
 		MutableUpdateCandidate.RequestedLODLevels.SetNumZeroed(ComponentCount);
 
-		const TArray<uint16>& GeneratedLODsPerComponent = DescriptorRuntimeHash.GetRequestedLODs();
+		const TArray<uint16>& GeneratedLODsPerComponent = GetPrivate()->DescriptorRuntimeHash.GetRequestedLODs();
 
-		const bool bIgnoreGeneratedLODs = DescriptorRuntimeHash != UpdateDescriptorRuntimeHash || GeneratedLODsPerComponent.Num() != ComponentCount;
+		const bool bIgnoreGeneratedLODs = GetPrivate()->DescriptorRuntimeHash != GetPrivate()->UpdateDescriptorRuntimeHash || GeneratedLODsPerComponent.Num() != ComponentCount;
 
 		if (bMinMaxLODChanged || bIgnoreGeneratedLODs || Descriptor.GetRequestedLODLevels() != InRequestedLODsPerComponent)
 		{
diff --git a/Engine/Plugins/Experimental/Mutable/Source/CustomizableObject/Private/MuCO/CustomizableObjectSystem.cpp b/Engine/Plugins/Experimental/Mutable/Source/CustomizableObject/Private/MuCO/CustomizableObjectSystem.cpp
--- a/Engine/Plugins/Experimental/Mutable/Source/CustomizableObject/Private/MuCO/CustomizableObjectSystem.cpp
+++ b/Engine/Plugins/Experimental/Mutable/Source/CustomizableObject/Private/MuCO/CustomizableObjectSystem.cpp
@@ -880,79 +880,225 @@ bool FCustomizableObjectSystemPrivate::TextureHasReferences(const FMutableImageC
 }
 
 
-void FCustomizableObjectSystemPrivate::EnqueueUpdateSkeletalMesh(UCustomizableObjectInstance& Instance, const EQueuePriorityType Priority, FInstanceUpdateDelegate* UpdateCallback)
+EUpdateRequired FCustomizableObjectSystemPrivate::IsUpdateRequired(const UCustomizableObjectInstance& Instance, bool bOnlyUpdateIfNotGenerated, bool bIgnoreCloseDist) const
 {
-	MUTABLE_CPUPROFILER_SCOPE(FCustomizableObjectSystemPrivate::InitUpdateSkeletalMesh);
-
-	check(IsInGameThread());
+	UCustomizableObjectSystem* System = UCustomizableObjectSystem::GetInstance();
+	const UCustomizableInstancePrivateData* const Private = Instance.GetPrivate();
 	
-	const FDescriptorRuntimeHash UpdateDescriptorHash = Instance.GetUpdateDescriptorRuntimeHash();
+	const UCustomizableObject* CustomizableObject = Instance.GetCustomizableObject();
+	if (!Instance.CanUpdateInstance() || CustomizableObject->IsLocked())
+	{
+		return EUpdateRequired::NoUpdate;
+	}
+
+	const bool bIsGenerated = Private->HasCOInstanceFlags(Generated);
+	const int32 NumGeneratedInstancesLimit = System->GetInstanceLODManagement()->GetNumGeneratedInstancesLimitFullLODs();
+	const int32 NumGeneratedInstancesLimitLOD1 = System->GetInstanceLODManagement()->GetNumGeneratedInstancesLimitLOD1();
+	const int32 NumGeneratedInstancesLimitLOD2 = System->GetInstanceLODManagement()->GetNumGeneratedInstancesLimitLOD2();
 
-	if (const FMutablePendingInstanceUpdate* QueueElem = MutablePendingInstanceWork.GetUpdate(&Instance))
+	if (!bIsGenerated && // Prevent generating more instances than the limit, but let updates to existing instances run normally
+		NumGeneratedInstancesLimit > 0 &&
+		System->GetPrivate()->GetCountAllocatedSkeletalMesh() > NumGeneratedInstancesLimit + NumGeneratedInstancesLimitLOD1 + NumGeneratedInstancesLimitLOD2)
 	{
-		if (UpdateDescriptorHash.IsSubset(FDescriptorRuntimeHash(QueueElem->InstanceDescriptor)))
+		return EUpdateRequired::NoUpdate;
+	}
+
+	const bool bShouldUpdateLODs = Private->HasCOInstanceFlags(PendingLODsUpdate);
+
+	const bool bDiscardByDistance = Private->LastMinSquareDistFromComponentToPlayer > FMath::Square(System->GetInstanceLODManagement()->GetOnlyUpdateCloseCustomizableObjectsDist());
+	const bool bLODManagementDiscard = System->GetInstanceLODManagement()->IsOnlyUpdateCloseCustomizableObjectsEnabled() &&
+			bDiscardByDistance &&
+			!bIgnoreCloseDist;
+	
+	if (Private->HasCOInstanceFlags(DiscardedByNumInstancesLimit) ||
+		bLODManagementDiscard)
+	{
+		if (bIsGenerated && !Private->HasCOInstanceFlags(Updating))
 		{
-			FinishUpdateGlobal(&Instance, EUpdateResult::ErrorOptimized, UpdateCallback);			
-			return; // The the requested update is equal to the last enqueued update.
+			return EUpdateRequired::Discard;		
 		}
-	}	
+		else
+		{
+			return EUpdateRequired::NoUpdate;
+		}
+	}
 
-	if (CurrentMutableOperation &&
-		&Instance == CurrentMutableOperation->CustomizableObjectInstance &&
-		UpdateDescriptorHash.IsSubset(CurrentMutableOperation->InstanceDescriptorRuntimeHash))
+	if (bIsGenerated &&
+		!bShouldUpdateLODs &&
+		bOnlyUpdateIfNotGenerated)
 	{
-		FinishUpdateGlobal(&Instance, EUpdateResult::ErrorOptimized, UpdateCallback);			
-		return; // The requested update is equal to the running update.
+		return EUpdateRequired::NoUpdate;
 	}
-	
-	if (UpdateDescriptorHash.IsSubset(Instance.GetDescriptorRuntimeHash()) &&
-		!(CurrentMutableOperation &&
-		&Instance == CurrentMutableOperation->CustomizableObjectInstance)) // This condition is necessary because even if the descriptor is a subset, it will be replaced by the CurrentMutableOperation
+
+	return EUpdateRequired::Update;
+}
+
+
+EQueuePriorityType FCustomizableObjectSystemPrivate::GetUpdatePriority(const UCustomizableObjectInstance& Instance,	bool bForceHighPriority) const
+{
+	const UCustomizableInstancePrivateData* InstancePrivate = Instance.GetPrivate();
+		
+	const bool bIsGenerated = InstancePrivate->HasCOInstanceFlags(Generated);
+	const bool bShouldUpdateLODs = InstancePrivate->HasCOInstanceFlags(PendingLODsUpdate);
+	const bool bIsDowngradeLODUpdate = InstancePrivate->HasCOInstanceFlags(PendingLODsDowngrade);
+	const bool bIsPlayerOrNearIt = InstancePrivate->HasCOInstanceFlags(UsedByPlayerOrNearIt);
+
+	EQueuePriorityType Priority = EQueuePriorityType::Low;
+	if (bForceHighPriority)
 	{
-		Instance.SkeletalMeshStatus = ESkeletalMeshState::Correct; // TODO FutureGMT MTBL-1033 should not be here. Move to UCustomizableObjectInstance::Updated
-		UpdateSkeletalMesh(Instance, Instance.GetDescriptorRuntimeHash(), EUpdateResult::Success, UpdateCallback);
+		Priority = EQueuePriorityType::High;
 	}
-	else
+	else if (!bIsGenerated || !Instance.HasAnySkeletalMesh())
 	{
-		// Cache Texture Parameters being used during the update:
-		check(ImageProvider);
+		Priority = EQueuePriorityType::Med;
+	}
+	else if (bShouldUpdateLODs && bIsDowngradeLODUpdate)
+	{
+		Priority = EQueuePriorityType::Med_Low;
+	}
+	else if (bIsPlayerOrNearIt && bShouldUpdateLODs && !bIsDowngradeLODUpdate)
+	{
+		Priority = EQueuePriorityType::High;
+	}
+	else if (bShouldUpdateLODs && !bIsDowngradeLODUpdate)
+	{
+		Priority = EQueuePriorityType::Med;
+	}
+	else if (bIsPlayerOrNearIt)
+	{
+		Priority = EQueuePriorityType::High;
+	}
 
-		// Cache new Texture Parameters
-		for (const FCustomizableObjectTextureParameterValue& TextureParameters : Instance.GetDescriptor().GetTextureParameters())
-		{
-			ImageProvider->CacheImage(TextureParameters.ParameterValue, false);
+	return Priority;
+}
+
+
+void FCustomizableObjectSystemPrivate::EnqueueUpdateSkeletalMesh(UCustomizableObjectInstance& Instance, bool bOnlyUpdateIfNotGenerated, bool bIgnoreCloseDist, bool bForceHighPriority, const EUpdateRequired* OptionalUpdateRequired, FInstanceUpdateDelegate* UpdateCallback)
+{
+	MUTABLE_CPUPROFILER_SCOPE(FCustomizableObjectSystemPrivate::EnqueueUpdateSkeletalMesh);
+	check(IsInGameThread());
+
+	if (!Instance.CanUpdateInstance())
+	{
+		FinishUpdateGlobal(&Instance, EUpdateResult::ErrorDiscarded, UpdateCallback);
+		return;
+	}
+
+	UCustomizableObjectSystem* System = UCustomizableObjectSystem::GetInstance();
+
+	UCustomizableInstancePrivateData* InstancePrivate = Instance.GetPrivate();
+	
+	const EUpdateRequired UpdateRequired = IsUpdateRequired(Instance, bOnlyUpdateIfNotGenerated, bIgnoreCloseDist);
+	switch (UpdateRequired)
+	{
+	case EUpdateRequired::NoUpdate:
+	{	
+		FinishUpdateGlobal(&Instance, EUpdateResult::ErrorDiscarded, UpdateCallback);
+		break;
+	}		
+	case EUpdateRequired::Update:
+	{
+		EQueuePriorityType Priority = GetUpdatePriority(Instance, bForceHighPriority);
 
-			for (const FName& TextureParameter : TextureParameters.ParameterRangeValues)
+		const uint32 InstanceId = Instance.GetUniqueID();
+		const float Distance = FMath::Sqrt(InstancePrivate->LastMinSquareDistFromComponentToPlayer);
+		const bool bIsPlayerOrNearIt = InstancePrivate->HasCOInstanceFlags(UsedByPlayerOrNearIt);
+		UE_LOG(LogMutable, Log, TEXT("Enqueued UpdateSkeletalMesh Async. of Instance %d with priority %d at dist %f bIsPlayerOrNearIt=%d, frame=%d"), InstanceId, static_cast<int32>(Priority), Distance, bIsPlayerOrNearIt, GFrameNumber);				
+
+#if WITH_EDITOR
+		if (!UCustomizableObjectSystem::GetInstance()->IsCompilationDisabled())
+		{
+			if (Instance.SkeletalMeshStatus != ESkeletalMeshState::AsyncUpdatePending)
 			{
-				ImageProvider->CacheImage(TextureParameter, false);
+				Instance.PreUpdateSkeletalMeshStatus = Instance.SkeletalMeshStatus;
 			}
 		}
+#endif
 
-		UCustomizableInstancePrivateData* InstancePrivate = Instance.GetPrivate();
-		check(InstancePrivate);
-
-		// Uncache old Texture Parameters
-		for (const FName& TextureParameter : InstancePrivate->UpdateTextureParameters)
+		if (InstancePrivate->HasCOInstanceFlags(PendingLODsUpdate))
 		{
-			ImageProvider->UnCacheImage(TextureParameter, false);
+			UE_LOG(LogMutable, Verbose, TEXT("LOD change: %d, %d -> %d, %d"), Instance.GetCurrentMinLOD(), Instance.GetCurrentMaxLOD(), Instance.GetMinLODToLoad(), Instance.GetMaxLODToLoad());
 		}
 
-		// Update which ones are currently are being used
-		InstancePrivate->UpdateTextureParameters.Reset();
-		for (const FCustomizableObjectTextureParameterValue& TextureParameters : Instance.GetDescriptor().GetTextureParameters())
+		InstancePrivate->UpdateDescriptorRuntimeHash = FDescriptorRuntimeHash(Instance.GetDescriptor());
+
+		if (const FMutablePendingInstanceUpdate* QueueElem = MutablePendingInstanceWork.GetUpdate(&Instance))
 		{
-			InstancePrivate->UpdateTextureParameters.Add(TextureParameters.ParameterValue);
+			if (InstancePrivate->UpdateDescriptorRuntimeHash.IsSubset(FDescriptorRuntimeHash(QueueElem->InstanceDescriptor)))
+			{
+				FinishUpdateGlobal(&Instance, EUpdateResult::ErrorOptimized, UpdateCallback);			
+				return; // The the requested update is equal to the last enqueued update.
+			}
+		}	
 
-			for (const FName& TextureParameter : TextureParameters.ParameterRangeValues)
+		if (CurrentMutableOperation &&
+			&Instance == CurrentMutableOperation->CustomizableObjectInstance &&
+			InstancePrivate->UpdateDescriptorRuntimeHash.IsSubset(CurrentMutableOperation->InstanceDescriptorRuntimeHash))
+		{
+			FinishUpdateGlobal(&Instance, EUpdateResult::ErrorOptimized, UpdateCallback);			
+			return; // The requested update is equal to the running update.
+		}
+	
+		if (InstancePrivate->UpdateDescriptorRuntimeHash.IsSubset(InstancePrivate->DescriptorRuntimeHash) &&
+			!(CurrentMutableOperation &&
+			&Instance == CurrentMutableOperation->CustomizableObjectInstance)) // This condition is necessary because even if the descriptor is a subset, it will be replaced by the CurrentMutableOperation
+		{
+			Instance.SkeletalMeshStatus = ESkeletalMeshState::Correct; // TODO FutureGMT MTBL-1033 should not be here. Move to UCustomizableObjectInstance::Updated
+			UpdateSkeletalMesh(Instance, Instance.GetDescriptorRuntimeHash(), EUpdateResult::Success, UpdateCallback);
+		}
+		else
+		{
+			// Cache Texture Parameters being used during the update:
+			check(ImageProvider);
+
+			// Cache new Texture Parameters
+			for (const FCustomizableObjectTextureParameterValue& TextureParameters : Instance.GetDescriptor().GetTextureParameters())
 			{
-				InstancePrivate->UpdateTextureParameters.Add(TextureParameter);
+				ImageProvider->CacheImage(TextureParameters.ParameterValue, false);
+
+				for (const FName& TextureParameter : TextureParameters.ParameterRangeValues)
+				{
+					ImageProvider->CacheImage(TextureParameter, false);
+				}
 			}
+
+			// Uncache old Texture Parameters
+			for (const FName& TextureParameter : InstancePrivate->UpdateTextureParameters)
+			{
+				ImageProvider->UnCacheImage(TextureParameter, false);
+			}
+
+			// Update which ones are currently are being used
+			InstancePrivate->UpdateTextureParameters.Reset();
+			for (const FCustomizableObjectTextureParameterValue& TextureParameters : Instance.GetDescriptor().GetTextureParameters())
+			{
+				InstancePrivate->UpdateTextureParameters.Add(TextureParameters.ParameterValue);
+
+				for (const FName& TextureParameter : TextureParameters.ParameterRangeValues)
+				{
+					InstancePrivate->UpdateTextureParameters.Add(TextureParameter);
+				}
+			}
+
+			Instance.SkeletalMeshStatus = ESkeletalMeshState::AsyncUpdatePending;
+
+			const FMutablePendingInstanceUpdate InstanceUpdate(&Instance, Priority, UpdateCallback);
+			MutablePendingInstanceWork.AddUpdate(InstanceUpdate);
 		}
 
-		Instance.SkeletalMeshStatus = ESkeletalMeshState::AsyncUpdatePending;
+		break;
+	}
+
+	case EUpdateRequired::Discard:
+	{
+		System->GetPrivate()->InitDiscardResourcesSkeletalMesh(&Instance);
+		
+		FinishUpdateGlobal(&Instance, EUpdateResult::ErrorDiscarded, UpdateCallback);
+		break;
+	}
 
-		const FMutablePendingInstanceUpdate InstanceUpdate(&Instance, Priority, UpdateCallback);
-		MutablePendingInstanceWork.AddUpdate(InstanceUpdate);
+	default:
+		unimplemented();
 	}
 }
 
@@ -2752,7 +2898,7 @@ bool UCustomizableObjectSystem::Tick(float DeltaTime)
 
 				if (PendingUpdate.CustomizableObjectInstance.IsValid())
 				{
-					const EQueuePriorityType PriorityType = PendingUpdate.CustomizableObjectInstance->GetUpdatePriority(false);
+					const EQueuePriorityType PriorityType = Private->GetUpdatePriority(*PendingUpdate.CustomizableObjectInstance, false);
 					
 					if (PendingUpdate.PriorityType <= MaxPriorityFound)
 					{
diff --git a/Engine/Plugins/Experimental/Mutable/Source/CustomizableObject/Private/MuCO/CustomizableObjectSystemPrivate.h b/Engine/Plugins/Experimental/Mutable/Source/CustomizableObject/Private/MuCO/CustomizableObjectSystemPrivate.h
--- a/Engine/Plugins/Experimental/Mutable/Source/CustomizableObject/Private/MuCO/CustomizableObjectSystemPrivate.h
+++ b/Engine/Plugins/Experimental/Mutable/Source/CustomizableObject/Private/MuCO/CustomizableObjectSystemPrivate.h
@@ -703,7 +703,11 @@ public:
 
 	bool TextureHasReferences(const FMutableImageCacheKey& TextureId) const;
 
-	void EnqueueUpdateSkeletalMesh(UCustomizableObjectInstance& Instance, EQueuePriorityType Priority, FInstanceUpdateDelegate* UpdateCallback);
+	EUpdateRequired IsUpdateRequired(const UCustomizableObjectInstance& Instance, bool bOnlyUpdateIfNotGenerated, bool bIgnoreCloseDist) const;
+
+	EQueuePriorityType GetUpdatePriority(const UCustomizableObjectInstance& Instance, bool bForceHighPriority) const;
+
+	void EnqueueUpdateSkeletalMesh(UCustomizableObjectInstance& Instance, bool bOnlyUpdateIfNotGenerated, bool bIgnoreCloseDist, bool bForceHighPriority, const EUpdateRequired* OptionalUpdateRequired, FInstanceUpdateDelegate* UpdateCallback);
 		
 	// Init an async and safe release of the UE and Mutable resources used by the instance without actually destroying the instance, for example if it's very far away
 	void InitDiscardResourcesSkeletalMesh(UCustomizableObjectInstance* InCustomizableObjectInstance);
