[Mutable] Replace a call to GetIndexBuffer with IsIndexBufferValid to fix an assert that fires when there's an empty LOD

diff --git a/Engine/Plugins/Experimental/Mutable/Source/CustomizableObject/Private/MuCO/UnrealConversionUtils.cpp b/Engine/Plugins/Experimental/Mutable/Source/CustomizableObject/Private/MuCO/UnrealConversionUtils.cpp
--- a/Engine/Plugins/Experimental/Mutable/Source/CustomizableObject/Private/MuCO/UnrealConversionUtils.cpp
+++ b/Engine/Plugins/Experimental/Mutable/Source/CustomizableObject/Private/MuCO/UnrealConversionUtils.cpp
@@ -555,7 +555,7 @@ namespace UnrealConversionUtils
 
 		 // Copying Indices
 		 {
-			 if (SrcLODModel.MultiSizeIndexContainer.GetIndexBuffer())
+			 if (SrcLODModel.MultiSizeIndexContainer.IsIndexBufferValid())
 			 {
 				 int32 IndexCount = SrcLODModel.MultiSizeIndexContainer.GetIndexBuffer()->Num();
 				 int32 ElementSize = SrcLODModel.MultiSizeIndexContainer.GetDataTypeSize();
