[mutable] Remove most of the legacy string usage in favor of FString and FName.
Renamed some legacy structs to use the proper convention.
Removed some unnecessary boilerplate get and set methods.
Optimize mesh extract operation with switches and conditionals.
#rnx

diff --git a/Engine/Plugins/Experimental/Mutable/Source/CustomizableObject/Classes/MuCO/CustomizableObject.h b/Engine/Plugins/Experimental/Mutable/Source/CustomizableObject/Classes/MuCO/CustomizableObject.h
--- a/Engine/Plugins/Experimental/Mutable/Source/CustomizableObject/Classes/MuCO/CustomizableObject.h
+++ b/Engine/Plugins/Experimental/Mutable/Source/CustomizableObject/Classes/MuCO/CustomizableObject.h
@@ -1459,7 +1459,7 @@ private:
 	// This is a manual version number for the binary blobs in this asset.
 	// Increasing it invalidates all the previously compiled models.
 	// Warning: If while merging code both versions have changed, take the highest+1.
-	static const int32 CurrentSupportedVersion = 403;
+	static const int32 CurrentSupportedVersion = 405;
 
 public:
 
diff --git a/Engine/Plugins/Experimental/Mutable/Source/CustomizableObject/Private/MuCO/CustomizableObject.cpp b/Engine/Plugins/Experimental/Mutable/Source/CustomizableObject/Private/MuCO/CustomizableObject.cpp
--- a/Engine/Plugins/Experimental/Mutable/Source/CustomizableObject/Private/MuCO/CustomizableObject.cpp
+++ b/Engine/Plugins/Experimental/Mutable/Source/CustomizableObject/Private/MuCO/CustomizableObject.cpp
@@ -1149,7 +1149,7 @@ FString UCustomizableObject::GetStateName(int32 StateIndex) const
 
 	if (PrivateData->GetModel())
 	{
-		Result = StringCast<TCHAR>( PrivateData->GetModel()->GetStateName(StateIndex) ).Get();
+		Result = PrivateData->GetModel()->GetStateName(StateIndex);
 	}
 
 	return Result;
diff --git a/Engine/Plugins/Experimental/Mutable/Source/CustomizableObject/Private/MuCO/CustomizableObjectInstance.cpp b/Engine/Plugins/Experimental/Mutable/Source/CustomizableObject/Private/MuCO/CustomizableObjectInstance.cpp
--- a/Engine/Plugins/Experimental/Mutable/Source/CustomizableObject/Private/MuCO/CustomizableObjectInstance.cpp
+++ b/Engine/Plugins/Experimental/Mutable/Source/CustomizableObject/Private/MuCO/CustomizableObjectInstance.cpp
@@ -5636,7 +5636,8 @@ void UCustomizableInstancePrivateData::BuildMaterials(const TSharedPtr<FMutableO
 						// Decoding Material Layer from Mutable parameter name
 						FString EncodingString = "-MutableLayerParam:";
 
-						int32 EncodingPosition = Vector.Name.Find(EncodingString);
+						FString VectorName = Vector.Name.ToString();
+						int32 EncodingPosition = VectorName.Find(EncodingString);
 						int32 LayerIndex = -1;
 
 						if (EncodingPosition == INDEX_NONE)
@@ -5647,17 +5648,16 @@ void UCustomizableInstancePrivateData::BuildMaterials(const TSharedPtr<FMutableO
 						else
 						{
 							//Getting layer index
-							int32 LayerPosition = Vector.Name.Len() - (EncodingPosition + EncodingString.Len());
-							FString IndexString = Vector.Name.RightChop(Vector.Name.Len() - LayerPosition);
+							int32 LayerPosition = VectorName.Len() - (EncodingPosition + EncodingString.Len());
+							FString IndexString = VectorName.RightChop(VectorName.Len() - LayerPosition);
 							LayerIndex = FCString::Atof(*IndexString);
 
 							//Getting parameter name
 							FString Sufix = EncodingString + FString::FromInt(LayerIndex);
-							FString NameCopy = Vector.Name;
-							NameCopy.RemoveFromEnd(Sufix);
+							VectorName.RemoveFromEnd(Sufix);
 
 							MutableMaterialPlaceholder.AddParam(
-								FMutableMaterialPlaceholder::FMutableMaterialPlaceHolderParam(FName(NameCopy), LayerIndex, Vector.Vector));
+								FMutableMaterialPlaceholder::FMutableMaterialPlaceHolderParam(FName(VectorName), LayerIndex, Vector.Vector));
 						}
 					}
 
@@ -5668,7 +5668,8 @@ void UCustomizableInstancePrivateData::BuildMaterials(const TSharedPtr<FMutableO
 						// Decoding Material Layer from Mutable parameter name
 						FString EncodingString = "-MutableLayerParam:";
 
-						int32 EncodingPosition = Scalar.Name.Find(EncodingString);
+						FString ScalarName = Scalar.Name.ToString();
+						int32 EncodingPosition = ScalarName.Find(EncodingString);
 						int32 LayerIndex = -1;
 
 						if (EncodingPosition == INDEX_NONE)
@@ -5679,17 +5680,16 @@ void UCustomizableInstancePrivateData::BuildMaterials(const TSharedPtr<FMutableO
 						else
 						{
 							//Getting layer index
-							int32 LayerPosition = Scalar.Name.Len() - (EncodingPosition + EncodingString.Len());
-							FString IndexString = Scalar.Name.RightChop(Scalar.Name.Len() - LayerPosition);
+							int32 LayerPosition = ScalarName.Len() - (EncodingPosition + EncodingString.Len());
+							FString IndexString = ScalarName.RightChop(ScalarName.Len() - LayerPosition);
 							LayerIndex = FCString::Atof(*IndexString);
 
 							//Getting parameter name
 							FString Sufix = EncodingString + FString::FromInt(LayerIndex);
-							FString NameCopy = Scalar.Name;
-							NameCopy.RemoveFromEnd(Sufix);
+							ScalarName.RemoveFromEnd(Sufix);
 
 							MutableMaterialPlaceholder.AddParam(
-								FMutableMaterialPlaceholder::FMutableMaterialPlaceHolderParam(FName(NameCopy), LayerIndex, Scalar.Scalar));
+								FMutableMaterialPlaceholder::FMutableMaterialPlaceHolderParam(FName(ScalarName), LayerIndex, Scalar.Scalar));
 						}
 					}
 				}
@@ -5709,7 +5709,7 @@ void UCustomizableInstancePrivateData::BuildMaterials(const TSharedPtr<FMutableO
 					for (int32 ImageIndex = 0; ImageIndex < Surface.ImageCount; ++ImageIndex)
 					{
 						const FInstanceUpdateData::FImage& Image = OperationData->InstanceUpdateData.Images[Surface.FirstImage + ImageIndex];
-						FString KeyName = Image.Name;
+						FString KeyName = Image.Name.ToString();
 						mu::ImagePtrConst MutableImage = Image.Image;
 
 						UTexture2D* MutableTexture = nullptr; // Texture generated by mutable
@@ -5917,7 +5917,7 @@ void UCustomizableInstancePrivateData::BuildMaterials(const TSharedPtr<FMutableO
 							}
 
 							// Decoding Material Layer from Mutable parameter name
-							FString ImageName = Image.Name;
+							FString ImageName = Image.Name.ToString();
 							FString EncodingString = "-MutableLayerParam:";
 
 							int32 EncodingPosition = ImageName.Find(EncodingString);
diff --git a/Engine/Plugins/Experimental/Mutable/Source/CustomizableObject/Private/MuCO/CustomizableObjectSystem.cpp b/Engine/Plugins/Experimental/Mutable/Source/CustomizableObject/Private/MuCO/CustomizableObjectSystem.cpp
--- a/Engine/Plugins/Experimental/Mutable/Source/CustomizableObject/Private/MuCO/CustomizableObjectSystem.cpp
+++ b/Engine/Plugins/Experimental/Mutable/Source/CustomizableObject/Private/MuCO/CustomizableObjectSystem.cpp
@@ -1675,13 +1675,13 @@ namespace impl
 						{
 							MUTABLE_CPUPROFILER_SCOPE(GetScalar)
 
-							const FString ScalarName = Instance->GetScalarName(BaseLODIndex, ComponentIndex, InstanceSurfaceIndex, ScalarIndex);
+							const FName ScalarName = Instance->GetScalarName(BaseLODIndex, ComponentIndex, InstanceSurfaceIndex, ScalarIndex);
 							const float ScalarValue = Instance->GetScalar(BaseLODIndex, ComponentIndex, InstanceSurfaceIndex, ScalarIndex);
 							
 							FString EncodingMaterialIdString = "__MutableMaterialId";
 							
 							// Decoding Material Switch from Mutable parameter name
-							if (ScalarName.Equals(EncodingMaterialIdString))
+							if (ScalarName.ToString().Equals(EncodingMaterialIdString))
 							{
 								Surface.MaterialIndex = static_cast<uint32>(ScalarValue);
 							
@@ -1701,16 +1701,15 @@ namespace impl
 		// Copy ExtensionData Object node input from the Instance to the InstanceUpdateData
 		for (int32 ExtensionDataIndex = 0; ExtensionDataIndex < Instance->GetExtensionDataCount(); ExtensionDataIndex++)
 		{
-			mu::ExtensionDataPtrConst ExtensionData;
-			const char* NameAnsi = nullptr;
-			Instance->GetExtensionData(ExtensionDataIndex, ExtensionData, NameAnsi);
+			mu::Ptr<const mu::ExtensionData> ExtensionData;
+			FName Name;
+			Instance->GetExtensionData(ExtensionDataIndex, ExtensionData, Name);
 
 			check(ExtensionData);
-			check(NameAnsi);
 
 			FInstanceUpdateData::FNamedExtensionData& NewEntry = OperationData->InstanceUpdateData.ExtendedInputPins.AddDefaulted_GetRef();
 			NewEntry.Data = ExtensionData;
-			NewEntry.Name = NameAnsi;
+			NewEntry.Name = Name;
 			check(NewEntry.Name != NAME_None);
 		}
 	}
@@ -1779,7 +1778,7 @@ namespace impl
 
 				const int32 MipSizeX = FMath::Max(Image.FullImageSizeX >> MipsToSkip, 1);
 				const int32 MipSizeY = FMath::Max(Image.FullImageSizeY >> MipsToSkip, 1);
-				if (MipsToSkip > 0 && CustomizableObjectSystemPrivateData->EnableSkipGenerateResidentMips != 0 && OperationData->LowPriorityTextures.Find(Image.Name) != INDEX_NONE)
+				if (MipsToSkip > 0 && CustomizableObjectSystemPrivateData->EnableSkipGenerateResidentMips != 0 && OperationData->LowPriorityTextures.Find(Image.Name.ToString()) != INDEX_NONE)
 				{
 					Image.Image = new mu::Image(MipSizeX, MipSizeY, FullLODCount - MipsToSkip, ImageDesc.m_format, mu::EInitializationType::Black);
 				}
@@ -1848,7 +1847,7 @@ namespace impl
 			{
 				const FInstanceUpdateData::FImage& Image = OperationData->InstanceUpdateData.Images[Surface.FirstImage+ImageIndex];
 
-				const FString KeyName = Image.Name;
+				const FName KeyName = Image.Name;
 				mu::ImagePtrConst MutableImage = Image.Image;
 
 				// If the image is null, it must be in the cache (or repeated in this instance), and we don't need to do anything here.
@@ -1861,7 +1860,7 @@ namespace impl
 						{
 							FTexturePlatformData* PlatformData = UCustomizableInstancePrivateData::MutableCreateImagePlatformData(MutableImage, -1, Image.FullImageSizeX, Image.FullImageSizeY);
 							OperationData->ImageToPlatformDataMap.Add(Image.ImageID, PlatformData);
-							OperationData->PendingTextureCoverageQueries.Add({ KeyName, Surface.MaterialIndex, PlatformData });
+							OperationData->PendingTextureCoverageQueries.Add({ KeyName.ToString(), Surface.MaterialIndex, PlatformData});
 						}
 						else
 						{
diff --git a/Engine/Plugins/Experimental/Mutable/Source/CustomizableObjectEditor/Private/MuCOE/SMutableCodeViewer.cpp b/Engine/Plugins/Experimental/Mutable/Source/CustomizableObjectEditor/Private/MuCOE/SMutableCodeViewer.cpp
--- a/Engine/Plugins/Experimental/Mutable/Source/CustomizableObjectEditor/Private/MuCOE/SMutableCodeViewer.cpp
+++ b/Engine/Plugins/Experimental/Mutable/Source/CustomizableObjectEditor/Private/MuCOE/SMutableCodeViewer.cpp
@@ -401,7 +401,7 @@ void SMutableCodeViewer::Construct(const FArguments& InArgs, const TSharedPtr<mu
 					{
 						const FString SaveFileName = FString(SaveFilenames[0]);
 
-						mu::OutputFileStream Stream(StringCast<ANSICHAR>(*SaveFileName).Get());
+						mu::OutputFileStream Stream(SaveFileName);
 						Stream.Write(MUTABLE_COMPILED_MODEL_FILETAG, 4);
 						mu::OutputArchive Archive(&Stream);
 						mu::Model::Serialise(InMutableModel.Get(), Archive);
@@ -1393,7 +1393,7 @@ void SMutableCodeViewer::GenerateAllTreeElements()
 	for ( uint32 StateIndex=0; StateIndex<StateCount; ++StateIndex )
 	{
 		const mu::FProgram::FState& State = ModelPrivate->m_program.m_states[StateIndex];
-		FString Caption = FString::Printf( TEXT("state [%s]"), StringCast<TCHAR>(State.m_name.c_str()).Get() );
+		FString Caption = FString::Printf( TEXT("state [%s]"), *State.Name );
 
 		const FSlateColor LabelColor = ColorPerComputationalCost[StaticCast<uint8>(GetOperationTypeComputationalCost(
 			ModelPrivate->m_program.GetOpType(State.m_root)))];
@@ -2703,19 +2703,13 @@ void SMutableCodeViewer::PrepareProjectorViewer()
 }
 
 
-void SMutableCodeViewer::PreviewMutableString(const mu::string* InStringPtr)
+void SMutableCodeViewer::PreviewMutableString(const FString& InString)
 {
-	if (!InStringPtr)
-	{
-		UE_LOG(LogTemp,Error,TEXT("Unable to preview data on null String pointer."))
-		return;
-	}
-	
 	// Prepare the previewer object to receive data 
 	PrepareStringViewer();
 	
 	//  Provide the desired data to the previewer object
-	const FText TextToShow = FText::FromString(FString(InStringPtr->c_str()));
+	const FText TextToShow = FText::FromString(InString);
 	PreviewStringViewer->SetString(TextToShow);
 }
 
@@ -2813,7 +2807,7 @@ FReply SMutableCodeViewer::OnDragOver(const FGeometry& MyGeometry, const FDragDr
 				if (DraggedFileExtension == TEXT(".mutable_compiled"))
 				{
 					// Dump source model to a file.
-					mu::InputFileStream stream(StringCast<ANSICHAR>(*Files[0]).Get());
+					mu::InputFileStream stream(Files[0]);
 
 					char MutableSourceTag[4] = {};
 					stream.Read(MutableSourceTag, 4);
@@ -2849,7 +2843,7 @@ FReply SMutableCodeViewer::OnDrop(const FGeometry& MyGeometry, const FDragDropEv
 				if (DraggedFileExtension == TEXT(".mutable_compiled"))
 				{
 					// Read a mutable compiled model file.
-					mu::InputFileStream stream(StringCast<ANSICHAR>(*Files[0]).Get());
+					mu::InputFileStream stream(Files[0]);
 
 					char MutableSourceTag[4] = {};
 					stream.Read(MutableSourceTag, 4);
diff --git a/Engine/Plugins/Experimental/Mutable/Source/CustomizableObjectEditor/Private/MuCOE/SMutableCodeViewer.h b/Engine/Plugins/Experimental/Mutable/Source/CustomizableObjectEditor/Private/MuCOE/SMutableCodeViewer.h
--- a/Engine/Plugins/Experimental/Mutable/Source/CustomizableObjectEditor/Private/MuCOE/SMutableCodeViewer.h
+++ b/Engine/Plugins/Experimental/Mutable/Source/CustomizableObjectEditor/Private/MuCOE/SMutableCodeViewer.h
@@ -784,7 +784,7 @@ public:
 	void PreviewMutableMesh (mu::MeshPtrConst InMeshPtr);
 	void PreviewMutableLayout(mu::LayoutPtrConst Layout);
 	void PreviewMutableSkeleton(mu::SkeletonPtrConst Skeleton);
-	void PreviewMutableString(const mu::string* InStringPtr);
+	void PreviewMutableString(const FString& InString);
 	void PreviewMutableProjector(const mu::FProjector* Projector);
 	void PreviewMutableMatrix(const mu::mat4f* Mat);
 	void PreviewMutableShape(const mu::FShape* Shape);
@@ -911,7 +911,7 @@ public:
 		{
 			mu::OP::ParameterArgs Args = Program.GetOpArgs<mu::OP::ParameterArgs>(MutableOperation);
 			OpName += TEXT(" ");
-			OpName += StringCast<TCHAR>(Program.m_parameters[int32(Args.variable)].m_name.c_str()).Get();
+			OpName += Program.m_parameters[int32(Args.variable)].m_name;
 			break;
 		}
 
@@ -1008,7 +1008,7 @@ public:
 			if (Program.m_constantStrings.IsValidIndex(Args.name))
 			{
 				OpName += TEXT(" name: ");
-				OpName += FString(Program.m_constantStrings[Args.name].c_str());
+				OpName += Program.m_constantStrings[Args.name];
 			}
  			break;
 		}
diff --git a/Engine/Plugins/Experimental/Mutable/Source/CustomizableObjectEditor/Private/MuCOE/SMutableConstantsWidget.cpp b/Engine/Plugins/Experimental/Mutable/Source/CustomizableObjectEditor/Private/MuCOE/SMutableConstantsWidget.cpp
--- a/Engine/Plugins/Experimental/Mutable/Source/CustomizableObjectEditor/Private/MuCOE/SMutableConstantsWidget.cpp
+++ b/Engine/Plugins/Experimental/Mutable/Source/CustomizableObjectEditor/Private/MuCOE/SMutableConstantsWidget.cpp
@@ -118,7 +118,7 @@ public:
 		// Generate the text to be displayed taking in mind the string value held by the constant to be able to "preview"
 		// it for easier navigation
 		const FString MainString = FString::FromInt(InRowElement->IndexOnSourceVector) + FString(TEXT("_STR "));
-		const FString ConstantStringText = FString(InRowElement->MutableString->c_str());
+		const FString ConstantStringText = InRowElement->MutableString;
 		const FString GlimpseConstantText = ConstantStringText.Left(GlimpseCharacterCount);
 		
 		// Compose the FStrings to produce the UI text to be displayed
@@ -1124,12 +1124,12 @@ void SMutableConstantsWidget::LoadConstantStrings()
 	for (int32 StringAddressIndex = 0; StringAddressIndex < ConstantsCount; StringAddressIndex++)
 	{
 		TSharedPtr<FMutableConstantStringElement> ConstantStringElement = MakeShared<FMutableConstantStringElement>();
-		ConstantStringElement->MutableString = &(MutableProgramPtr->m_constantStrings[StringAddressIndex]);
+		ConstantStringElement->MutableString = MutableProgramPtr->m_constantStrings[StringAddressIndex];
 		ConstantStringElement->IndexOnSourceVector = StringAddressIndex;
 		
 		// Cache resource size
 		// in case we change the type of the contents of the mu::string we check its size as if it was a vector<>
-		ConstantStringsAccumulatedSize += ConstantStringElement->MutableString->size() * sizeof (mu::string::value_type);
+		ConstantStringsAccumulatedSize += ConstantStringElement->MutableString.GetAllocatedSize();
 		
 		ConstantStringElements.Add(ConstantStringElement);
 	}
diff --git a/Engine/Plugins/Experimental/Mutable/Source/CustomizableObjectEditor/Private/MuCOE/SMutableConstantsWidget.h b/Engine/Plugins/Experimental/Mutable/Source/CustomizableObjectEditor/Private/MuCOE/SMutableConstantsWidget.h
--- a/Engine/Plugins/Experimental/Mutable/Source/CustomizableObjectEditor/Private/MuCOE/SMutableConstantsWidget.h
+++ b/Engine/Plugins/Experimental/Mutable/Source/CustomizableObjectEditor/Private/MuCOE/SMutableConstantsWidget.h
@@ -51,7 +51,7 @@ struct FMutableConstantMeshElement : public FMutableConstantElement
 */
 struct FMutableConstantStringElement : public FMutableConstantElement
 {
-	const mu::string* MutableString = nullptr;
+	FString MutableString;
 };
 
 /**
diff --git a/Engine/Plugins/Experimental/Mutable/Source/CustomizableObjectEditor/Private/MuCOE/SMutableGraphViewer.cpp b/Engine/Plugins/Experimental/Mutable/Source/CustomizableObjectEditor/Private/MuCOE/SMutableGraphViewer.cpp
--- a/Engine/Plugins/Experimental/Mutable/Source/CustomizableObjectEditor/Private/MuCOE/SMutableGraphViewer.cpp
+++ b/Engine/Plugins/Experimental/Mutable/Source/CustomizableObjectEditor/Private/MuCOE/SMutableGraphViewer.cpp
@@ -156,7 +156,7 @@ void SMutableGraphViewer::Construct(const FArguments& InArgs, const mu::NodePtr&
 
 					// Dump source model to a file.
 					FString SaveFileName = FString(SaveFilenames[0]);
-					mu::OutputFileStream stream(StringCast<ANSICHAR>(*SaveFileName).Get());
+					mu::OutputFileStream stream(SaveFileName);
 					stream.Write(MUTABLE_SOURCE_MODEL_FILETAG, 4);
 					mu::OutputArchive arch(&stream);
 					mu::Node::Serialise(InRootNode.get(), arch);
@@ -490,7 +490,7 @@ FReply SMutableGraphViewer::OnDragOver(const FGeometry& MyGeometry, const FDragD
 				if (DraggedFileExtension == TEXT(".mutable_source"))
 				{
 					// Dump source model to a file.
-					mu::InputFileStream stream(StringCast<ANSICHAR>(*Files[0]).Get());
+					mu::InputFileStream stream(Files[0]);
 
 					char MutableSourceTag[4] = {};
 					stream.Read(MutableSourceTag, 4);
@@ -526,7 +526,7 @@ FReply SMutableGraphViewer::OnDrop(const FGeometry& MyGeometry, const FDragDropE
 				if (DraggedFileExtension == TEXT(".mutable_source"))
 				{
 					// Dump source model to a file.
-					mu::InputFileStream stream(StringCast<ANSICHAR>(*Files[0]).Get());
+					mu::InputFileStream stream(Files[0]);
 
 					char MutableSourceTag[4] = {};
 					stream.Read(MutableSourceTag, 4);
diff --git a/Engine/Plugins/Experimental/Mutable/Source/CustomizableObjectEditor/Private/MuCOE/SMutableImageViewer.cpp b/Engine/Plugins/Experimental/Mutable/Source/CustomizableObjectEditor/Private/MuCOE/SMutableImageViewer.cpp
--- a/Engine/Plugins/Experimental/Mutable/Source/CustomizableObjectEditor/Private/MuCOE/SMutableImageViewer.cpp
+++ b/Engine/Plugins/Experimental/Mutable/Source/CustomizableObjectEditor/Private/MuCOE/SMutableImageViewer.cpp
@@ -59,7 +59,8 @@ void SSimpleTextureViewer::SetTexture( UTexture* InTexture)
 	FVector2D Size(Resource->GetSizeX(), Resource->GetSizeY());
 	TextureBrush = MakeShared<FSlateImageBrush>(Texture, Size,
 		FSlateColor(FLinearColor(1, 1, 1, 1)), ESlateBrushTileType::NoTile,
-		ESlateBrushImageType::Linear);
+		ESlateBrushImageType::Linear
+);
 
 }
 
@@ -284,7 +285,7 @@ void SMutableImageViewer::Tick(const FGeometry& AllottedGeometry, const double I
 			}
 
 			FMutableModelImageProperties Props;
-			Props.Filter = TF_Bilinear;
+			Props.Filter = TF_Nearest;
 			Props.SRGB = true;
 			Props.LODBias = 0;
 			UCustomizableInstancePrivateData::ConvertImage(UnrealImage, ImageToConvert, Props, CurrentVisibleLOD, ExtractChannel);
diff --git a/Engine/Plugins/Experimental/Mutable/Source/CustomizableObjectEditor/Private/MuCOE/SMutableParametersWidget.cpp b/Engine/Plugins/Experimental/Mutable/Source/CustomizableObjectEditor/Private/MuCOE/SMutableParametersWidget.cpp
--- a/Engine/Plugins/Experimental/Mutable/Source/CustomizableObjectEditor/Private/MuCOE/SMutableParametersWidget.cpp
+++ b/Engine/Plugins/Experimental/Mutable/Source/CustomizableObjectEditor/Private/MuCOE/SMutableParametersWidget.cpp
@@ -64,7 +64,7 @@ void SMutableParametersWidget::Tick(const FGeometry& AllottedGeometry, const dou
 	
 	for ( int32 ParamIndex=0; ParamIndex< MutableParameters->GetCount(); ++ParamIndex )
 	{
-		FString ParamName = StringCast<TCHAR>(MutableParameters->GetName(ParamIndex)).Get();
+		FString ParamName = MutableParameters->GetName(ParamIndex);
 
 		TSharedPtr<SHorizontalBox> ParameterBox;
 
@@ -288,8 +288,8 @@ void SMutableParametersWidget::GenerateAndAttachParameterSlate(const int32 Param
 			const int32 ValueIndex = MutableParameters->GetIntValueIndex( ParamIndex, Value );
 			for (int32 i = 0; i < ValueCount; ++i)
 			{
-				const char* ValueText = MutableParameters->GetIntPossibleValueName( ParamIndex, i );
-				OptionNamesAttribute.Add( FString(StringCast<TCHAR>(ValueText).Get()) );				
+				const FString& ValueText = MutableParameters->GetIntPossibleValueName( ParamIndex, i );
+				OptionNamesAttribute.Add( ValueText );				
 			}
 
 			ParameterHorizontalBox->AddSlot()
@@ -530,7 +530,7 @@ void SMutableParametersWidget::OnIntParameterChanged(int32 InValue, int32 ParamI
 		// Update the text combo if any
 		if (Combo)
 		{
-			const FString Text = StringCast<TCHAR>(MutableParameters->GetIntPossibleValueName(ParamIndex, InValue)).Get();
+			const FString Text = MutableParameters->GetIntPossibleValueName(ParamIndex, InValue);
 			Combo->SetText(FText::FromString(Text));
 		}
 		
diff --git a/Engine/Plugins/Experimental/Mutable/Source/MutableRuntime/Private/MuR/CodeRunner.cpp b/Engine/Plugins/Experimental/Mutable/Source/MutableRuntime/Private/MuR/CodeRunner.cpp
--- a/Engine/Plugins/Experimental/Mutable/Source/MutableRuntime/Private/MuR/CodeRunner.cpp
+++ b/Engine/Plugins/Experimental/Mutable/Source/MutableRuntime/Private/MuR/CodeRunner.cpp
@@ -430,9 +430,9 @@ namespace mu
 
                     OP::ADDRESS nameAd = args.name;
                     check(  nameAd < (uint32)pModel->GetPrivate()->m_program.m_constantStrings.Num() );
-                    const char* strName = pModel->GetPrivate()->m_program.m_constantStrings[ nameAd ].c_str();
+                    const FString& Name = pModel->GetPrivate()->m_program.m_constantStrings[ nameAd ];
 
-                    pResult->GetPrivate()->AddVector( 0, 0, 0, value, strName );
+                    pResult->GetPrivate()->AddVector( 0, 0, 0, value, FName(Name) );
                 }
                 StoreInstance( item, pResult );
                 break;
@@ -475,9 +475,9 @@ namespace mu
 
                     OP::ADDRESS nameAd = args.name;
                     check(  nameAd < (uint32)pModel->GetPrivate()->m_program.m_constantStrings.Num() );
-                    const char* strName = pModel->GetPrivate()->m_program.m_constantStrings[ nameAd ].c_str();
+                    const FString& Name = pModel->GetPrivate()->m_program.m_constantStrings[ nameAd ];
 
-                    pResult->GetPrivate()->AddScalar( 0, 0, 0, value, strName );
+                    pResult->GetPrivate()->AddScalar( 0, 0, 0, value, FName(Name));
                 }
                 StoreInstance( item, pResult );
                 break;
@@ -521,10 +521,9 @@ namespace mu
 
                     OP::ADDRESS nameAd = args.name;
                     check( nameAd < (uint32)pModel->GetPrivate()->m_program.m_constantStrings.Num() );
-                    const char* strName =
-                        pModel->GetPrivate()->m_program.m_constantStrings[nameAd].c_str();
+                    const FString& Name = pModel->GetPrivate()->m_program.m_constantStrings[nameAd];
 
-                    pResult->GetPrivate()->AddString( 0, 0, 0, value->GetValue(), strName );
+                    pResult->GetPrivate()->AddString( 0, 0, 0, value->GetValue(), FName(Name) );
                 }
                 StoreInstance( item, pResult );
                 break;
@@ -567,22 +566,22 @@ namespace mu
 
                     int cindex = pResult->GetPrivate()->AddComponent( 0 );
 
-                    if ( !pComp->GetPrivate()->m_lods.IsEmpty()
+                    if ( !pComp->GetPrivate()->Lods.IsEmpty()
                          &&
-                         !pResult->GetPrivate()->m_lods.IsEmpty()
+                         !pResult->GetPrivate()->Lods.IsEmpty()
                          &&
-                         !pComp->GetPrivate()->m_lods[0].m_components.IsEmpty() )
+                         !pComp->GetPrivate()->Lods[0].Components.IsEmpty() )
                     {
-                        pResult->GetPrivate()->m_lods[0].m_components[cindex] =
-                                pComp->GetPrivate()->m_lods[0].m_components[0];
+                        pResult->GetPrivate()->Lods[0].Components[cindex] =
+                                pComp->GetPrivate()->Lods[0].Components[0];
 
-                    	pResult->GetPrivate()->m_lods[0].m_components[cindex].m_id = args.id;
+                    	pResult->GetPrivate()->Lods[0].Components[cindex].Id = args.id;
                     	
                         // Name
                         OP::ADDRESS nameAd = args.name;
                         check( nameAd < (uint32)pModel->GetPrivate()->m_program.m_constantStrings.Num() );
-                        const char* strName = pModel->GetPrivate()->m_program.m_constantStrings[ nameAd ].c_str();
-                        pResult->GetPrivate()->SetComponentName( 0, cindex, strName );
+                        const FString& Name = pModel->GetPrivate()->m_program.m_constantStrings[ nameAd ];
+                        pResult->GetPrivate()->SetComponentName( 0, cindex, FName(Name) );
                     }
                 }
                 StoreInstance( item, pResult );
@@ -632,29 +631,29 @@ namespace mu
                     // Surface data
                     if (pSurf
                             &&
-                            pSurf->GetPrivate()->m_lods.Num()
+                            pSurf->GetPrivate()->Lods.Num()
                             &&
-                            pSurf->GetPrivate()->m_lods[0].m_components.Num()
+                            pSurf->GetPrivate()->Lods[0].Components.Num()
                             &&
-                            pSurf->GetPrivate()->m_lods[0].m_components[0].m_surfaces.Num())
+                            pSurf->GetPrivate()->Lods[0].Components[0].Surfaces.Num())
                     {
-                        pResult->GetPrivate()->m_lods[0].m_components[0].m_surfaces[sindex] =
-                            pSurf->GetPrivate()->m_lods[0].m_components[0].m_surfaces[0];
+                        pResult->GetPrivate()->Lods[0].Components[0].Surfaces[sindex] =
+                            pSurf->GetPrivate()->Lods[0].Components[0].Surfaces[0];
 
                         // Meshes must be added later.
-                        check(!pSurf->GetPrivate()->m_lods[0].m_components[0].m_meshes.Num());
+                        check(!pSurf->GetPrivate()->Lods[0].Components[0].Meshes.Num());
                     }
 
                     // Name
                     OP::ADDRESS nameAd = args.name;
                     check( nameAd < (uint32)pModel->GetPrivate()->m_program.m_constantStrings.Num() );
-                    const char* strName = pModel->GetPrivate()->m_program.m_constantStrings[ nameAd ].c_str();
-                    pResult->GetPrivate()->SetSurfaceName( 0, 0, sindex, strName );
+                    const FString& Name = pModel->GetPrivate()->m_program.m_constantStrings[ nameAd ];
+                    pResult->GetPrivate()->SetSurfaceName( 0, 0, sindex, FName(Name) );
 
                     // IDs
-                    pResult->GetPrivate()->m_lods[0].m_components[0].m_surfaces[sindex].InternalId = args.id;
-                    pResult->GetPrivate()->m_lods[0].m_components[0].m_surfaces[sindex].ExternalId = args.ExternalId;
-                    pResult->GetPrivate()->m_lods[0].m_components[0].m_surfaces[sindex].SharedId = args.SharedSurfaceId;
+                    pResult->GetPrivate()->Lods[0].Components[0].Surfaces[sindex].InternalId = args.id;
+                    pResult->GetPrivate()->Lods[0].Components[0].Surfaces[sindex].ExternalId = args.ExternalId;
+                    pResult->GetPrivate()->Lods[0].Components[0].Surfaces[sindex].SharedId = args.SharedSurfaceId;
                 }
                 StoreInstance( item, pResult );
                 break;
@@ -711,9 +710,9 @@ namespace mu
                             int LODIndex = pResult->GetPrivate()->AddLOD();
 
                             // In a degenerated case, the returned pLOD may not have an LOD inside
-                            if ( pLOD && !pLOD->GetPrivate()->m_lods.IsEmpty() )
+                            if ( pLOD && !pLOD->GetPrivate()->Lods.IsEmpty() )
                             {
-                                pResult->GetPrivate()->m_lods[LODIndex] = pLOD->GetPrivate()->m_lods[0];
+                                pResult->GetPrivate()->Lods[LODIndex] = pLOD->GetPrivate()->Lods[0];
                             }
                         }
                         else
@@ -766,9 +765,9 @@ namespace mu
 					{
 						const OP::ADDRESS NameAddress = Args.ExtensionDataName;
 						check(NameAddress < (uint32)pModel->GetPrivate()->m_program.m_constantStrings.Num());
-						const char* NameString = pModel->GetPrivate()->m_program.m_constantStrings[NameAddress].c_str();
+						const FString& NameString = pModel->GetPrivate()->m_program.m_constantStrings[NameAddress];
 
-						Result->GetPrivate()->AddExtensionData(ExtensionData, NameString);
+						Result->GetPrivate()->AddExtensionData(ExtensionData, FName(NameString) );
 					}
 
 					StoreInstance(item, Result);
@@ -833,8 +832,8 @@ namespace mu
 					FResourceID MeshId = m_pSystem->WorkingMemoryManager.GetResourceKey(InModel,InParams,args.relevantParametersListIndex, args.value);
 					OP::ADDRESS NameAd = args.name;
 					check(NameAd < (uint32)InModel->GetPrivate()->m_program.m_constantStrings.Num());
-					const char* Name = InModel->GetPrivate()->m_program.m_constantStrings[NameAd].c_str();
-					pResult->GetPrivate()->AddMesh(0, 0, MeshId, Name);
+					const FString& Name = InModel->GetPrivate()->m_program.m_constantStrings[NameAd];
+					pResult->GetPrivate()->AddMesh(0, 0, MeshId, FName(Name));
                 }
                 StoreInstance( item, pResult );
                 break;
@@ -874,8 +873,8 @@ namespace mu
 					FResourceID ImageId = m_pSystem->WorkingMemoryManager.GetResourceKey(InModel, InParams, args.relevantParametersListIndex, args.value);
 					OP::ADDRESS NameAd = args.name;
 					check(NameAd < (uint32)InModel->GetPrivate()->m_program.m_constantStrings.Num());
-					const char* Name = InModel->GetPrivate()->m_program.m_constantStrings[NameAd].c_str();
-					pResult->GetPrivate()->AddImage(0, 0, 0, ImageId, Name);
+					const FString& Name = InModel->GetPrivate()->m_program.m_constantStrings[NameAd];
+					pResult->GetPrivate()->AddImage(0, 0, 0, ImageId, FName(Name) );
                 }
                 StoreInstance( item, pResult );
                 break;
@@ -5278,8 +5277,8 @@ namespace mu
 			OP::ResourceConstantArgs args = pModel->GetPrivate()->m_program.GetOpArgs<OP::ResourceConstantArgs>( item.At );
             check( args.value < (uint32)pModel->GetPrivate()->m_program.m_constantStrings.Num() );
 
-            const std::string& result = pModel->GetPrivate()->m_program.m_constantStrings[args.value];
-            StoreString( item, new String(result.c_str()) );
+            const FString& result = pModel->GetPrivate()->m_program.m_constantStrings[args.value];
+            StoreString( item, new String(result) );
 
             break;
         }
@@ -5288,8 +5287,9 @@ namespace mu
         {
 			OP::ParameterArgs args = pModel->GetPrivate()->m_program.GetOpArgs<OP::ParameterArgs>( item.At );
 			Ptr<RangeIndex> index = BuildCurrentOpRangeIndex( item, pParams, pModel, args.variable );
-            string result = pParams->GetStringValue( args.variable, index );
-            StoreString( item, new String( result.c_str() ) );
+			FString result;
+			pParams->GetStringValue(args.variable, result, index);
+            StoreString( item, new String(result) );
             break;
         }
 
@@ -6114,7 +6114,7 @@ namespace mu
 				}
 
 				// At least keep the levels we already have.
-				int startLevel = BaseDesc.m_lods;
+				int32 startLevel = BaseDesc.m_lods;
 				levelCount = FMath::Max(startLevel, levelCount);
 
 				// Update result.
diff --git a/Engine/Plugins/Experimental/Mutable/Source/MutableRuntime/Private/MuR/CodeRunnerTasked.cpp b/Engine/Plugins/Experimental/Mutable/Source/MutableRuntime/Private/MuR/CodeRunnerTasked.cpp
--- a/Engine/Plugins/Experimental/Mutable/Source/MutableRuntime/Private/MuR/CodeRunnerTasked.cpp
+++ b/Engine/Plugins/Experimental/Mutable/Source/MutableRuntime/Private/MuR/CodeRunnerTasked.cpp
@@ -187,7 +187,6 @@ namespace mu
 		{   0,   0, -10,   0 },	// ME_REMOVEMASK
 		{   0,   0,   0,   0 },	// ME_FORMAT
 		{   0, -10,   0,   0 },	// ME_EXTRACTLAYOUTBLOCK
-		{   0, -10,   0,   0 },	// ME_EXTRACTFACEGROUP
 		{   0,   0,   0,   0 },	// ME_TRANSFORM
 		{   0, -10,   0,   0 },	// ME_CLIPMORPHPLANE
 		{   0, -10,   0,   0 },	// ME_CLIPWITHMESH
@@ -208,8 +207,9 @@ namespace mu
 		{   0,   0,   0,   0 },	// IN_ADDSTRING
 		{   0,   0,   0,   0 },	// IN_ADDSURFACE
 		{   0,   0,   0,   0 },	// IN_ADDCOMPONENT
-		{   0,   0,   0,   0 },	// IN_ADDLOD
-		    
+		{	0,   0,   0,   0 },	// IN_ADDLOD
+		{	0,   0,   0,   0 },	// IN_ADDEXTENSIONDATA
+
 		{   0,   0,   0,   0 },	// LA_PACK
 		{   0,   0,   0,   0 },	// LA_MERGE
 		{   0,  -1,   0,   0 },	// LA_REMOVEBLOCKS
diff --git a/Engine/Plugins/Experimental/Mutable/Source/MutableRuntime/Private/MuR/Instance.cpp b/Engine/Plugins/Experimental/Mutable/Source/MutableRuntime/Private/MuR/Instance.cpp
--- a/Engine/Plugins/Experimental/Mutable/Source/MutableRuntime/Private/MuR/Instance.cpp
+++ b/Engine/Plugins/Experimental/Mutable/Source/MutableRuntime/Private/MuR/Instance.cpp
@@ -53,61 +53,44 @@ namespace mu
 	//---------------------------------------------------------------------------------------------
 	int32 Instance::GetDataSize() const
 	{
-		return 16 + sizeof(Private) + m_pD->m_lods.GetAllocatedSize() + m_pD->m_extensionData.GetAllocatedSize();
+		return 16 + sizeof(Private) + m_pD->Lods.GetAllocatedSize() + m_pD->ExtensionData.GetAllocatedSize();
 	}
 
 
     //---------------------------------------------------------------------------------------------
     Instance::ID Instance::GetId() const
     {
-        return m_pD->m_id;
+        return m_pD->Id;
     }
 
 
     //---------------------------------------------------------------------------------------------
     int32 Instance::GetLODCount() const
     {
-		return m_pD->m_lods.Num();
+		return m_pD->Lods.Num();
     }
 
 
 	//---------------------------------------------------------------------------------------------
-	int32 Instance::GetComponentCount( int32 lod ) const
+	int32 Instance::GetComponentCount( int32 LODIndex ) const
 	{
-		check(lod >= 0 && lod < m_pD->m_lods.Num());
-		if (lod >= 0 && lod < m_pD->m_lods.Num())
+		check(LODIndex >= 0 && LODIndex < m_pD->Lods.Num());
+		if (LODIndex >= 0 && LODIndex < m_pD->Lods.Num())
 		{
-			return m_pD->m_lods[lod].m_components.Num();
+			return m_pD->Lods[LODIndex].Components.Num();
 		}
 
 		return 0;
 	}
 
-
-    //---------------------------------------------------------------------------------------------
-    const char* Instance::GetComponentName( int32 lod, int32 comp ) const
-    {
-        if ( lod>=0 && lod<m_pD->m_lods.Num() &&
-             comp>=0 && comp<m_pD->m_lods[lod].m_components.Num() )
-        {
-            return m_pD->m_lods[lod].m_components[comp].m_name.c_str();
-        }
-		else
-		{
-			check(false);
-		}
-
-        return "";
-    }
-
 	
 	//---------------------------------------------------------------------------------------------
-	uint16 Instance::GetComponentId( int32 lod, int32 comp ) const
+	uint16 Instance::GetComponentId( int32 LODIndex, int32 ComponentIndex ) const
 	{
-		if ( lod>=0 && lod<m_pD->m_lods.Num() &&
-			 comp>=0 && comp<m_pD->m_lods[lod].m_components.Num() )
+		if ( LODIndex>=0 && LODIndex<m_pD->Lods.Num() &&
+			 ComponentIndex>=0 && ComponentIndex<m_pD->Lods[LODIndex].Components.Num() )
 		{
-			return m_pD->m_lods[lod].m_components[comp].m_id;
+			return m_pD->Lods[LODIndex].Components[ComponentIndex].Id;
 		}
 		else
 		{
@@ -119,12 +102,12 @@ namespace mu
 
 
     //---------------------------------------------------------------------------------------------
-    int32 Instance::GetSurfaceCount( int32 lod, int32 comp ) const
+    int32 Instance::GetSurfaceCount( int32 LODIndex, int32 ComponentIndex ) const
     {
-		if (lod >= 0 && lod < m_pD->m_lods.Num() &&
-			comp >= 0 && comp < m_pD->m_lods[lod].m_components.Num())
+		if (LODIndex >= 0 && LODIndex < m_pD->Lods.Num() &&
+			ComponentIndex >= 0 && ComponentIndex < m_pD->Lods[LODIndex].Components.Num())
 		{
-			return m_pD->m_lods[lod].m_components[comp].m_surfaces.Num();
+			return m_pD->Lods[LODIndex].Components[ComponentIndex].Surfaces.Num();
 		}
 		else
 		{
@@ -135,32 +118,14 @@ namespace mu
 	}
 
 
-	//---------------------------------------------------------------------------------------------
-    const char* Instance::GetSurfaceName( int32 lod, int32 comp, int32 surf ) const
-    {
-        if ( lod>=0 && lod<m_pD->m_lods.Num() &&
-             comp>=0 && comp<m_pD->m_lods[lod].m_components.Num() &&
-             surf>=0 && surf<m_pD->m_lods[lod].m_components[comp].m_surfaces.Num() )
-        {
-            return m_pD->m_lods[lod].m_components[comp].m_surfaces[surf].m_name.c_str();
-        }
-		else
-		{
-			check(false);
-		}
-
-        return "";
-    }
-
-
     //---------------------------------------------------------------------------------------------
-    uint32 Instance::GetSurfaceId( int32 lod, int32 comp, int32 surf ) const
+    uint32 Instance::GetSurfaceId( int32 LODIndex, int32 ComponentIndex, int32 SurfaceIndex ) const
     {
-        if ( lod>=0 && lod<m_pD->m_lods.Num() &&
-             comp>=0 && comp<m_pD->m_lods[lod].m_components.Num() &&
-             surf>=0 && surf<m_pD->m_lods[lod].m_components[comp].m_surfaces.Num() )
+        if ( LODIndex>=0 && LODIndex<m_pD->Lods.Num() &&
+             ComponentIndex>=0 && ComponentIndex<m_pD->Lods[LODIndex].Components.Num() &&
+             SurfaceIndex>=0 && SurfaceIndex<m_pD->Lods[LODIndex].Components[ComponentIndex].Surfaces.Num() )
         {
-            return m_pD->m_lods[lod].m_components[comp].m_surfaces[surf].InternalId;
+            return m_pD->Lods[LODIndex].Components[ComponentIndex].Surfaces[SurfaceIndex].InternalId;
         }
 		else
 		{
@@ -172,14 +137,14 @@ namespace mu
 
 
     //---------------------------------------------------------------------------------------------
-    int32 Instance::FindSurfaceById( int32 lod, int32 comp, uint32 id ) const
+    int32 Instance::FindSurfaceById( int32 LODIndex, int32 ComponentIndex, uint32 id ) const
     {
-		if (lod >= 0 && lod < m_pD->m_lods.Num() &&
-			comp >= 0 && comp < m_pD->m_lods[lod].m_components.Num())
+		if (LODIndex >= 0 && LODIndex < m_pD->Lods.Num() &&
+			ComponentIndex >= 0 && ComponentIndex < m_pD->Lods[LODIndex].Components.Num())
 		{
-			for (int32 i = 0; i < m_pD->m_lods[lod].m_components[comp].m_surfaces.Num(); ++i)
+			for (int32 i = 0; i < m_pD->Lods[LODIndex].Components[ComponentIndex].Surfaces.Num(); ++i)
 			{
-				if (m_pD->m_lods[lod].m_components[comp].m_surfaces[i].InternalId == id)
+				if (m_pD->Lods[LODIndex].Components[ComponentIndex].Surfaces[i].InternalId == id)
 				{
 					return i;
 				}
@@ -197,13 +162,13 @@ namespace mu
 	//---------------------------------------------------------------------------------------------
 	void Instance::FindBaseSurfaceBySharedId(int32 CompIndex, int32 SharedId, int32& OutSurfaceIndex, int32& OutLODIndex) const
 	{
-		for (int32 LodIndex = 0; LodIndex < m_pD->m_lods.Num(); LodIndex++)
+		for (int32 LodIndex = 0; LodIndex < m_pD->Lods.Num(); LodIndex++)
 		{
-			if (m_pD->m_lods[LodIndex].m_components.IsValidIndex(CompIndex))
+			if (m_pD->Lods[LodIndex].Components.IsValidIndex(CompIndex))
 			{
-				for (int32 SurfaceIndex = 0; SurfaceIndex < m_pD->m_lods[LodIndex].m_components[CompIndex].m_surfaces.Num(); ++SurfaceIndex)
+				for (int32 SurfaceIndex = 0; SurfaceIndex < m_pD->Lods[LodIndex].Components[CompIndex].Surfaces.Num(); ++SurfaceIndex)
 				{
-					if (m_pD->m_lods[LodIndex].m_components[CompIndex].m_surfaces[SurfaceIndex].SharedId == SharedId)
+					if (m_pD->Lods[LodIndex].Components[CompIndex].Surfaces[SurfaceIndex].SharedId == SharedId)
 					{
 						OutSurfaceIndex = SurfaceIndex;
 						OutLODIndex = LodIndex;
@@ -222,11 +187,11 @@ namespace mu
 	//---------------------------------------------------------------------------------------------
 	int32 Instance::GetSharedSurfaceId(int32 LodIndex, int32 CompIndex, int32 SurfaceIndex) const
 	{
-		if (LodIndex >= 0 && LodIndex < m_pD->m_lods.Num() &&
-			CompIndex >= 0 && CompIndex < m_pD->m_lods[LodIndex].m_components.Num() &&
-			SurfaceIndex >= 0 && SurfaceIndex < m_pD->m_lods[LodIndex].m_components[CompIndex].m_surfaces.Num())
+		if (LodIndex >= 0 && LodIndex < m_pD->Lods.Num() &&
+			CompIndex >= 0 && CompIndex < m_pD->Lods[LodIndex].Components.Num() &&
+			SurfaceIndex >= 0 && SurfaceIndex < m_pD->Lods[LodIndex].Components[CompIndex].Surfaces.Num())
 		{
-			return m_pD->m_lods[LodIndex].m_components[CompIndex].m_surfaces[SurfaceIndex].SharedId;
+			return m_pD->Lods[LodIndex].Components[CompIndex].Surfaces[SurfaceIndex].SharedId;
 		}
 		else
 		{
@@ -238,13 +203,13 @@ namespace mu
 	
 
     //---------------------------------------------------------------------------------------------
-    uint32 Instance::GetSurfaceCustomId( int32 lod, int32 comp, int32 surf ) const
+    uint32 Instance::GetSurfaceCustomId( int32 LODIndex, int32 ComponentIndex, int32 SurfaceIndex ) const
     {
-        if ( lod>=0 && lod<m_pD->m_lods.Num() &&
-             comp>=0 && comp<m_pD->m_lods[lod].m_components.Num() &&
-             surf>=0 && surf<m_pD->m_lods[lod].m_components[comp].m_surfaces.Num() )
+        if ( LODIndex>=0 && LODIndex<m_pD->Lods.Num() &&
+             ComponentIndex>=0 && ComponentIndex<m_pD->Lods[LODIndex].Components.Num() &&
+             SurfaceIndex>=0 && SurfaceIndex<m_pD->Lods[LODIndex].Components[ComponentIndex].Surfaces.Num() )
         {
-            return m_pD->m_lods[lod].m_components[comp].m_surfaces[surf].ExternalId;
+            return m_pD->Lods[LODIndex].Components[ComponentIndex].Surfaces[SurfaceIndex].ExternalId;
         }
 		else
 		{
@@ -256,170 +221,164 @@ namespace mu
 
 
 	//---------------------------------------------------------------------------------------------
-    int32 Instance::GetMeshCount( int32 lod, int32 comp ) const
+    int32 Instance::GetMeshCount( int32 LODIndex, int32 ComponentIndex ) const
 	{
-		check(lod >= 0 && lod < m_pD->m_lods.Num());
-		check(comp >= 0 && comp < m_pD->m_lods[lod].m_components.Num());
+		check(LODIndex >= 0 && LODIndex < m_pD->Lods.Num());
+		check(ComponentIndex >= 0 && ComponentIndex < m_pD->Lods[LODIndex].Components.Num());
 
-		return m_pD->m_lods[lod].m_components[comp].m_meshes.Num();
+		return m_pD->Lods[LODIndex].Components[ComponentIndex].Meshes.Num();
 	}
 
 
 	//---------------------------------------------------------------------------------------------
-    int32 Instance::GetImageCount( int32 lod, int32 comp, int32 surf ) const
+    int32 Instance::GetImageCount( int32 LODIndex, int32 ComponentIndex, int32 SurfaceIndex ) const
 	{
-		check(lod >= 0 && lod < m_pD->m_lods.Num());
-		check(comp >= 0 && comp < m_pD->m_lods[lod].m_components.Num());
-		check(surf >= 0 && surf < m_pD->m_lods[lod].m_components[comp].m_surfaces.Num());
+		check(LODIndex >= 0 && LODIndex < m_pD->Lods.Num());
+		check(ComponentIndex >= 0 && ComponentIndex < m_pD->Lods[LODIndex].Components.Num());
+		check(SurfaceIndex >= 0 && SurfaceIndex < m_pD->Lods[LODIndex].Components[ComponentIndex].Surfaces.Num());
 
-		return m_pD->m_lods[lod].m_components[comp].m_surfaces[surf].m_images.Num();
+		return m_pD->Lods[LODIndex].Components[ComponentIndex].Surfaces[SurfaceIndex].Images.Num();
 	}
 
 
 	//---------------------------------------------------------------------------------------------
-    int32 Instance::GetVectorCount( int32 lod, int32 comp, int32 surf ) const
+    int32 Instance::GetVectorCount( int32 LODIndex, int32 ComponentIndex, int32 SurfaceIndex ) const
 	{
-		check(lod >= 0 && lod < m_pD->m_lods.Num());
-		check(comp >= 0 && comp < m_pD->m_lods[lod].m_components.Num());
-		check(surf >= 0 && surf < m_pD->m_lods[lod].m_components[comp].m_surfaces.Num());
+		check(LODIndex >= 0 && LODIndex < m_pD->Lods.Num());
+		check(ComponentIndex >= 0 && ComponentIndex < m_pD->Lods[LODIndex].Components.Num());
+		check(SurfaceIndex >= 0 && SurfaceIndex < m_pD->Lods[LODIndex].Components[ComponentIndex].Surfaces.Num());
 
-		return m_pD->m_lods[lod].m_components[comp].m_surfaces[surf].m_vectors.Num();
+		return m_pD->Lods[LODIndex].Components[ComponentIndex].Surfaces[SurfaceIndex].Vectors.Num();
 	}
 
 
 	//---------------------------------------------------------------------------------------------
-    int32 Instance::GetScalarCount( int32 lod, int32 comp, int32 surf ) const
+    int32 Instance::GetScalarCount( int32 LODIndex, int32 ComponentIndex, int32 SurfaceIndex ) const
 	{
-		check(lod >= 0 && lod < m_pD->m_lods.Num());
-		check(comp >= 0 && comp < m_pD->m_lods[lod].m_components.Num());
-		check(surf >= 0 && surf < m_pD->m_lods[lod].m_components[comp].m_surfaces.Num());
+		check(LODIndex >= 0 && LODIndex < m_pD->Lods.Num());
+		check(ComponentIndex >= 0 && ComponentIndex < m_pD->Lods[LODIndex].Components.Num());
+		check(SurfaceIndex >= 0 && SurfaceIndex < m_pD->Lods[LODIndex].Components[ComponentIndex].Surfaces.Num());
 
-		return m_pD->m_lods[lod].m_components[comp].m_surfaces[surf].m_scalars.Num();
+		return m_pD->Lods[LODIndex].Components[ComponentIndex].Surfaces[SurfaceIndex].Scalars.Num();
 	}
 
 
     //---------------------------------------------------------------------------------------------
-    int32 Instance::GetStringCount( int32 lod, int32 comp, int32 surf ) const
+    int32 Instance::GetStringCount( int32 LODIndex, int32 ComponentIndex, int32 SurfaceIndex ) const
     {
-		check(lod >= 0 && lod < m_pD->m_lods.Num());
-		check(comp >= 0 && comp < m_pD->m_lods[lod].m_components.Num());
-		check(surf >= 0 && surf < m_pD->m_lods[lod].m_components[comp].m_surfaces.Num());
+		check(LODIndex >= 0 && LODIndex < m_pD->Lods.Num());
+		check(ComponentIndex >= 0 && ComponentIndex < m_pD->Lods[LODIndex].Components.Num());
+		check(SurfaceIndex >= 0 && SurfaceIndex < m_pD->Lods[LODIndex].Components[ComponentIndex].Surfaces.Num());
 
-		return m_pD->m_lods[lod].m_components[comp].m_surfaces[surf].m_strings.Num();
+		return m_pD->Lods[LODIndex].Components[ComponentIndex].Surfaces[SurfaceIndex].Strings.Num();
 	}
 
 
-
     //---------------------------------------------------------------------------------------------
-	FResourceID Instance::GetMeshId( int32 lod, int32 comp, int32 mesh ) const
+	FResourceID Instance::GetMeshId( int32 LODIndex, int32 ComponentIndex, int32 mesh ) const
     {
-        check( lod>=0 && lod<m_pD->m_lods.Num() );
-        check( comp>=0 && comp<m_pD->m_lods[lod].m_components.Num() );
-        check( mesh>=0 && mesh<m_pD->m_lods[lod].m_components[comp].m_meshes.Num() );
+        check( LODIndex>=0 && LODIndex<m_pD->Lods.Num() );
+        check( ComponentIndex>=0 && ComponentIndex<m_pD->Lods[LODIndex].Components.Num() );
+        check( mesh>=0 && mesh<m_pD->Lods[LODIndex].Components[ComponentIndex].Meshes.Num() );
 
-		FResourceID result = m_pD->m_lods[lod].m_components[comp].m_meshes[mesh].m_meshId;
+		FResourceID result = m_pD->Lods[LODIndex].Components[ComponentIndex].Meshes[mesh].Id;
         return result;
     }
 
 
 	//---------------------------------------------------------------------------------------------
-	FResourceID Instance::GetImageId( int32 lod, int32 comp, int32 surf, int32 img ) const
+	FResourceID Instance::GetImageId( int32 LODIndex, int32 ComponentIndex, int32 SurfaceIndex, int32 img ) const
 	{
-		check( lod>=0 && lod<m_pD->m_lods.Num() );
-		check( comp>=0 && comp<m_pD->m_lods[lod].m_components.Num() );
-        check( surf>=0 && surf<m_pD->m_lods[lod].m_components[comp].m_surfaces.Num() );
-        check( img>=0 && img<m_pD->m_lods[lod].m_components[comp].m_surfaces[surf].m_images.Num() );
+		check( LODIndex>=0 && LODIndex<m_pD->Lods.Num() );
+		check( ComponentIndex>=0 && ComponentIndex<m_pD->Lods[LODIndex].Components.Num() );
+        check( SurfaceIndex>=0 && SurfaceIndex<m_pD->Lods[LODIndex].Components[ComponentIndex].Surfaces.Num() );
+        check( img>=0 && img<m_pD->Lods[LODIndex].Components[ComponentIndex].Surfaces[SurfaceIndex].Images.Num() );
 
-		FResourceID result = m_pD->m_lods[lod].m_components[comp].m_surfaces[surf].m_images[img].m_imageId;
+		FResourceID result = m_pD->Lods[LODIndex].Components[ComponentIndex].Surfaces[SurfaceIndex].Images[img].Id;
         return result;
 	}
 
 
 	//---------------------------------------------------------------------------------------------
-    const char* Instance::GetImageName( int32 lod, int32 comp, int32 surf, int32 img ) const
+    FName Instance::GetImageName( int32 LODIndex, int32 ComponentIndex, int32 SurfaceIndex, int32 img ) const
 	{
-		check( lod>=0 && lod<m_pD->m_lods.Num() );
-		check( comp>=0 && comp<m_pD->m_lods[lod].m_components.Num() );
-        check( surf>=0 && surf<m_pD->m_lods[lod].m_components[comp].m_surfaces.Num() );
-        check( img>=0 && img<m_pD->m_lods[lod].m_components[comp].m_surfaces[surf].m_images.Num() );
+		check( LODIndex>=0 && LODIndex<m_pD->Lods.Num() );
+		check( ComponentIndex>=0 && ComponentIndex<m_pD->Lods[LODIndex].Components.Num() );
+        check( SurfaceIndex>=0 && SurfaceIndex<m_pD->Lods[LODIndex].Components[ComponentIndex].Surfaces.Num() );
+        check( img>=0 && img<m_pD->Lods[LODIndex].Components[ComponentIndex].Surfaces[SurfaceIndex].Images.Num() );
 
-        const char* strResult = m_pD->m_lods[lod].m_components[comp].m_surfaces[surf].m_images[img].m_name.c_str();
-		return strResult;
+        return m_pD->Lods[LODIndex].Components[ComponentIndex].Surfaces[SurfaceIndex].Images[img].Name;
 	}
 
 
 	//---------------------------------------------------------------------------------------------
-	FVector4f Instance::GetVector( int32 lod, int32 comp, int32 surf, int32 vec ) const
+	FVector4f Instance::GetVector( int32 LODIndex, int32 ComponentIndex, int32 SurfaceIndex, int32 vec ) const
 	{
-		check( lod>=0 && lod<m_pD->m_lods.Num() );
-		check( comp>=0 && comp<m_pD->m_lods[lod].m_components.Num() );
-        check( surf>=0 && surf<m_pD->m_lods[lod].m_components[comp].m_surfaces.Num() );
-        check( vec>=0 && vec<m_pD->m_lods[lod].m_components[comp].m_surfaces[surf].m_vectors.Num() );
+		check( LODIndex>=0 && LODIndex<m_pD->Lods.Num() );
+		check( ComponentIndex>=0 && ComponentIndex<m_pD->Lods[LODIndex].Components.Num() );
+        check( SurfaceIndex>=0 && SurfaceIndex<m_pD->Lods[LODIndex].Components[ComponentIndex].Surfaces.Num() );
+        check( vec>=0 && vec<m_pD->Lods[LODIndex].Components[ComponentIndex].Surfaces[SurfaceIndex].Vectors.Num() );
 
-        return m_pD->m_lods[lod].m_components[comp].m_surfaces[surf].m_vectors[vec].m_vec;
+        return m_pD->Lods[LODIndex].Components[ComponentIndex].Surfaces[SurfaceIndex].Vectors[vec].Value;
 	}
 
 
 	//---------------------------------------------------------------------------------------------
-    const char* Instance::GetVectorName( int32 lod, int32 comp, int32 surf, int32 vec ) const
+	FName Instance::GetVectorName( int32 LODIndex, int32 ComponentIndex, int32 SurfaceIndex, int32 vec ) const
 	{
-		check( lod>=0 && lod<m_pD->m_lods.Num() );
-		check( comp>=0 && comp<m_pD->m_lods[lod].m_components.Num() );
-        check( surf>=0 && surf<m_pD->m_lods[lod].m_components[comp].m_surfaces.Num() );
-        check( vec>=0 && vec<m_pD->m_lods[lod].m_components[comp].m_surfaces[surf].m_vectors.Num() );
+		check( LODIndex>=0 && LODIndex<m_pD->Lods.Num() );
+		check( ComponentIndex>=0 && ComponentIndex<m_pD->Lods[LODIndex].Components.Num() );
+        check( SurfaceIndex>=0 && SurfaceIndex<m_pD->Lods[LODIndex].Components[ComponentIndex].Surfaces.Num() );
+        check( vec>=0 && vec<m_pD->Lods[LODIndex].Components[ComponentIndex].Surfaces[SurfaceIndex].Vectors.Num() );
 
-        const char* strResult = m_pD->m_lods[lod].m_components[comp].m_surfaces[surf].m_vectors[vec].m_name.c_str();
-		return strResult;
+        return m_pD->Lods[LODIndex].Components[ComponentIndex].Surfaces[SurfaceIndex].Vectors[vec].Name;
 	}
 
 
 	//---------------------------------------------------------------------------------------------
-    float Instance::GetScalar( int32 lod, int32 comp, int32 surf, int32 sca ) const
+    float Instance::GetScalar( int32 LODIndex, int32 ComponentIndex, int32 SurfaceIndex, int32 sca ) const
 	{
-		check( lod>=0 && lod<m_pD->m_lods.Num() );
-		check( comp>=0 && comp<m_pD->m_lods[lod].m_components.Num() );
-        check( surf>=0 && surf<m_pD->m_lods[lod].m_components[comp].m_surfaces.Num() );
-        check( sca>=0 && sca<m_pD->m_lods[lod].m_components[comp].m_surfaces[surf].m_scalars.Num() );
+		check( LODIndex>=0 && LODIndex<m_pD->Lods.Num() );
+		check( ComponentIndex>=0 && ComponentIndex<m_pD->Lods[LODIndex].Components.Num() );
+        check( SurfaceIndex>=0 && SurfaceIndex<m_pD->Lods[LODIndex].Components[ComponentIndex].Surfaces.Num() );
+        check( sca>=0 && sca<m_pD->Lods[LODIndex].Components[ComponentIndex].Surfaces[SurfaceIndex].Scalars.Num() );
 
-        float result = m_pD->m_lods[lod].m_components[comp].m_surfaces[surf].m_scalars[sca].m_scalar;
-		return result;
+        return m_pD->Lods[LODIndex].Components[ComponentIndex].Surfaces[SurfaceIndex].Scalars[sca].Value;
 	}
 
 
 	//---------------------------------------------------------------------------------------------
-    const char* Instance::GetScalarName( int32 lod, int32 comp, int32 surf, int32 sca ) const
+	FName Instance::GetScalarName( int32 LODIndex, int32 ComponentIndex, int32 SurfaceIndex, int32 sca ) const
 	{
-		check( lod>=0 && lod<m_pD->m_lods.Num() );
-		check( comp>=0 && comp<m_pD->m_lods[lod].m_components.Num() );
-        check( surf>=0 && surf<m_pD->m_lods[lod].m_components[comp].m_surfaces.Num() );
-        check( sca>=0 && sca<m_pD->m_lods[lod].m_components[comp].m_surfaces[surf].m_scalars.Num() );
+		check( LODIndex>=0 && LODIndex<m_pD->Lods.Num() );
+		check( ComponentIndex>=0 && ComponentIndex<m_pD->Lods[LODIndex].Components.Num() );
+        check( SurfaceIndex>=0 && SurfaceIndex<m_pD->Lods[LODIndex].Components[ComponentIndex].Surfaces.Num() );
+        check( sca>=0 && sca<m_pD->Lods[LODIndex].Components[ComponentIndex].Surfaces[SurfaceIndex].Scalars.Num() );
 
-        const char* strResult = m_pD->m_lods[lod].m_components[comp].m_surfaces[surf].m_scalars[sca].m_name.c_str();
-		return strResult;
+        return m_pD->Lods[LODIndex].Components[ComponentIndex].Surfaces[SurfaceIndex].Scalars[sca].Name;
 	}
 
 
     //---------------------------------------------------------------------------------------------
-    const char* Instance::GetString( int32 lod, int32 comp, int32 surf, int32 str ) const
+    FString Instance::GetString( int32 LODIndex, int32 ComponentIndex, int32 SurfaceIndex, int32 str ) const
     {
-        check( lod >= 0 && lod < m_pD->m_lods.Num() );
-        check( comp >= 0 && comp < m_pD->m_lods[lod].m_components.Num() );
-        check( surf >= 0 &&
-                        surf < m_pD->m_lods[lod].m_components[comp].m_surfaces.Num() );
+        check( LODIndex >= 0 && LODIndex < m_pD->Lods.Num() );
+        check( ComponentIndex >= 0 && ComponentIndex < m_pD->Lods[LODIndex].Components.Num() );
+        check( SurfaceIndex >= 0 &&
+                        SurfaceIndex < m_pD->Lods[LODIndex].Components[ComponentIndex].Surfaces.Num() );
 
         bool valid =
             str >= 0 &&
-            str < m_pD->m_lods[lod].m_components[comp].m_surfaces[surf].m_strings.Num();
+            str < m_pD->Lods[LODIndex].Components[ComponentIndex].Surfaces[SurfaceIndex].Strings.Num();
         check(valid);
 
         if (valid)
         {
-            const char* result = m_pD->m_lods[lod]
-                                     .m_components[comp]
-                                     .m_surfaces[surf]
-                                     .m_strings[str]
-                                     .m_string.c_str();
-            return result;
+            return m_pD->Lods[LODIndex]
+				.Components[ComponentIndex]
+				.Surfaces[SurfaceIndex]
+				.Strings[str]
+				.Value;
         }
 
         return "";
@@ -427,274 +386,255 @@ namespace mu
 
 
     //---------------------------------------------------------------------------------------------
-    const char* Instance::GetStringName( int32 lod, int32 comp, int32 surf, int32 str ) const
+	FName Instance::GetStringName( int32 LODIndex, int32 ComponentIndex, int32 SurfaceIndex, int32 str ) const
     {
-        check( lod >= 0 && lod < m_pD->m_lods.Num() );
-        check( comp >= 0 && comp < m_pD->m_lods[lod].m_components.Num() );
-        check( surf >= 0 &&
-                        surf < m_pD->m_lods[lod].m_components[comp].m_surfaces.Num() );
+        check( LODIndex >= 0 && LODIndex < m_pD->Lods.Num() );
+        check( ComponentIndex >= 0 && ComponentIndex < m_pD->Lods[LODIndex].Components.Num() );
+        check( SurfaceIndex >= 0 &&
+                        SurfaceIndex < m_pD->Lods[LODIndex].Components[ComponentIndex].Surfaces.Num() );
         bool valid =
             str >= 0 &&
-            str < m_pD->m_lods[lod].m_components[comp].m_surfaces[surf].m_strings.Num();
+            str < m_pD->Lods[LODIndex].Components[ComponentIndex].Surfaces[SurfaceIndex].Strings.Num();
         check( valid );
 
         if (valid)
         {
-            const char* strResult =
-                m_pD->m_lods[lod].m_components[comp].m_surfaces[surf].m_strings[str].m_name.c_str();
-            return strResult;
+            return m_pD->Lods[LODIndex].Components[ComponentIndex].Surfaces[SurfaceIndex].Strings[str].Name;
         }
 
-        return "";
+        return NAME_None;
     }
 
 
     //---------------------------------------------------------------------------------------------
 	int32 Instance::GetExtensionDataCount() const
 	{
-		return m_pD->m_extensionData.Num();
+		return m_pD->ExtensionData.Num();
 	}
 
 
     //---------------------------------------------------------------------------------------------
-	void Instance::GetExtensionData(int32 Index, ExtensionDataPtrConst& OutExtensionData, const char*& OutName) const
+	void Instance::GetExtensionData(int32 Index, Ptr<const class ExtensionData>& OutExtensionData, FName& OutName) const
 	{
-		check(m_pD->m_extensionData.IsValidIndex(Index));
+		check(m_pD->ExtensionData.IsValidIndex(Index));
 
-		OutExtensionData = m_pD->m_extensionData[Index].Data;
-		OutName = m_pD->m_extensionData[Index].Name.c_str();
+		OutExtensionData = m_pD->ExtensionData[Index].Data;
+		OutName = m_pD->ExtensionData[Index].Name;
 	}
 
 
     //---------------------------------------------------------------------------------------------
 	int32 Instance::Private::AddLOD()
 	{
-		int32 result = m_lods.Num();
-		m_lods.Add( INSTANCE_LOD() );
+		int32 result = Lods.Num();
+		Lods.Add( INSTANCE_LOD() );
 		return result;
 	}
 
 
     //---------------------------------------------------------------------------------------------
-    int32 Instance::Private::AddComponent( int32 lod )
+    int32 Instance::Private::AddComponent( int32 LODIndex )
     {
         // Automatically create the necessary lods and components
-        while (lod >= m_lods.Num())
+        while (LODIndex >= Lods.Num())
         {
             AddLOD();
         }
 
-        int32 result = m_lods[lod].m_components.Num();
-        m_lods[lod].m_components.Add( INSTANCE_COMPONENT() );
+        int32 result = Lods[LODIndex].Components.Num();
+        Lods[LODIndex].Components.Add( INSTANCE_COMPONENT() );
         return result;
     }
 
 
     //---------------------------------------------------------------------------------------------
-    int32 Instance::Private::AddSurface( int32 lod, int32 comp )
+    int32 Instance::Private::AddSurface( int32 LODIndex, int32 ComponentIndex )
     {
         // Automatically create the necessary lods and components
-        while (lod >= m_lods.Num())
+        while (LODIndex >= Lods.Num())
         {
             AddLOD();
         }
-        while (comp >= m_lods[lod].m_components.Num())
+        while (ComponentIndex >= Lods[LODIndex].Components.Num())
         {
-            AddComponent(lod);
+            AddComponent(LODIndex);
         }
 
-        int32 result = m_lods[lod].m_components[comp].m_surfaces.Num();
-        m_lods[lod].m_components[comp].m_surfaces.Add( INSTANCE_SURFACE() );
+        int32 result = Lods[LODIndex].Components[ComponentIndex].Surfaces.Num();
+        Lods[LODIndex].Components[ComponentIndex].Surfaces.Add( INSTANCE_SURFACE() );
         return result;
     }
 
 
     //---------------------------------------------------------------------------------------------
-    void Instance::Private::SetComponentName( int32 lod, int32 comp, const char* strName )
+    void Instance::Private::SetComponentName( int32 LODIndex, int32 ComponentIndex, FName Name)
     {
         // Automatically create the necessary lods and components
-        while (lod >= m_lods.Num())
+        while (LODIndex >= Lods.Num())
         {
             AddLOD();
         }
-        while (comp >= m_lods[lod].m_components.Num())
+        while (ComponentIndex >= Lods[LODIndex].Components.Num())
         {
-            AddComponent( lod );
+            AddComponent( LODIndex );
         }
 
-        INSTANCE_COMPONENT& component = m_lods[lod].m_components[comp];
-        if (strName)
-        {
-            component.m_name = strName;
-        }
-        else
-        {
-            component.m_name = "";
-        }
+        INSTANCE_COMPONENT& component = Lods[LODIndex].Components[ComponentIndex];
+        component.Name = Name;
     }
 
 
     //---------------------------------------------------------------------------------------------
-    void Instance::Private::SetSurfaceName( int32 lod, int32 comp, int32 surf, const char* strName )
+    void Instance::Private::SetSurfaceName( int32 LODIndex, int32 ComponentIndex, int32 SurfaceIndex, FName Name)
     {
         // Automatically create the necessary lods and components
-        while ( lod>=m_lods.Num() )
+        while ( LODIndex>=Lods.Num() )
         {
             AddLOD();
         }
-        while (comp >= m_lods[lod].m_components.Num())
+        while (ComponentIndex >= Lods[LODIndex].Components.Num())
         {
-            AddComponent( lod );
+            AddComponent( LODIndex );
         }
-        while ( surf>=m_lods[lod].m_components[comp].m_surfaces.Num() )
+        while ( SurfaceIndex>=Lods[LODIndex].Components[ComponentIndex].Surfaces.Num() )
         {
-            AddSurface( lod, comp );
+            AddSurface( LODIndex, ComponentIndex );
         }
 
-        INSTANCE_SURFACE& surface = m_lods[lod].m_components[comp].m_surfaces[surf];
-        if (strName)
-        {
-            surface.m_name = strName;
-        }
-        else
-        {
-            surface.m_name = "";
-        }
+        INSTANCE_SURFACE& surface = Lods[LODIndex].Components[ComponentIndex].Surfaces[SurfaceIndex];
+        surface.Name = Name;
     }
 
 
 	//---------------------------------------------------------------------------------------------
-    int32 Instance::Private::AddMesh( int32 lod, int32 comp, FResourceID meshId, const char* strName )
+    int32 Instance::Private::AddMesh( int32 LODIndex, int32 ComponentIndex, FResourceID meshId, FName Name)
 	{
 		// Automatically create the necessary lods and components
-		while (lod >= m_lods.Num())
+		while (LODIndex >= Lods.Num())
 		{
 			AddLOD();
 		}
-		while (comp >= m_lods[lod].m_components.Num())
+		while (ComponentIndex >= Lods[LODIndex].Components.Num())
 		{
-			AddComponent( lod );
+			AddComponent( LODIndex );
 		}
 
-		INSTANCE_COMPONENT& component = m_lods[lod].m_components[comp];
-        int32 result = component.m_meshes.Num();
-        component.m_meshes.Emplace( meshId, strName );
+		INSTANCE_COMPONENT& component = Lods[LODIndex].Components[ComponentIndex];
+        int32 result = component.Meshes.Num();
+        component.Meshes.Emplace( meshId, Name );
 
 		return result;
 	}
 
 
 	//---------------------------------------------------------------------------------------------
-    int32 Instance::Private::AddImage( int32 lod, int32 comp, int32 surf, FResourceID imageId, const char* strName )
+    int32 Instance::Private::AddImage( int32 LODIndex, int32 ComponentIndex, int32 SurfaceIndex, FResourceID imageId, FName Name)
 	{
 		// Automatically create the necessary lods and components
-		while (lod >= m_lods.Num())
+		while (LODIndex >= Lods.Num())
 		{
 			AddLOD();
 		}
-		while (comp >= m_lods[lod].m_components.Num())
+		while (ComponentIndex >= Lods[LODIndex].Components.Num())
 		{
-			AddComponent(lod);
+			AddComponent(LODIndex);
 		}
-		while (surf >= m_lods[lod].m_components[comp].m_surfaces.Num())
+		while (SurfaceIndex >= Lods[LODIndex].Components[ComponentIndex].Surfaces.Num())
 		{
-			AddSurface(lod, comp);
+			AddSurface(LODIndex, ComponentIndex);
 		}
 
-        INSTANCE_SURFACE& surface = m_lods[lod].m_components[comp].m_surfaces[surf];
-        int32 result = surface.m_images.Num();
-        surface.m_images.Emplace( imageId, strName );
+        INSTANCE_SURFACE& surface = Lods[LODIndex].Components[ComponentIndex].Surfaces[SurfaceIndex];
+        int32 result = surface.Images.Num();
+        surface.Images.Emplace( imageId, Name);
 
         return result;
 	}
 
 
 	//---------------------------------------------------------------------------------------------
-    int32 Instance::Private::AddVector( int32 lod, int32 comp, int32 surf, const FVector4f& vec, const char* strName )
+    int32 Instance::Private::AddVector( int32 LODIndex, int32 ComponentIndex, int32 SurfaceIndex, const FVector4f& vec, FName Name)
 	{
 		// Automatically create the necessary lods and components
-		while (lod >= m_lods.Num())
+		while (LODIndex >= Lods.Num())
 		{
 			AddLOD();
 		}
-		while (comp >= m_lods[lod].m_components.Num())
+		while (ComponentIndex >= Lods[LODIndex].Components.Num())
 		{
-			AddComponent(lod);
+			AddComponent(LODIndex);
 		}
-		while (surf >= m_lods[lod].m_components[comp].m_surfaces.Num())
+		while (SurfaceIndex >= Lods[LODIndex].Components[ComponentIndex].Surfaces.Num())
 		{
-			AddSurface(lod, comp);
+			AddSurface(LODIndex, ComponentIndex);
 		}
 
-        INSTANCE_SURFACE& surface = m_lods[lod].m_components[comp].m_surfaces[surf];
-        int32 result = surface.m_vectors.Num();
-        surface.m_vectors.Emplace( vec, strName );
+        INSTANCE_SURFACE& surface = Lods[LODIndex].Components[ComponentIndex].Surfaces[SurfaceIndex];
+        int32 result = surface.Vectors.Num();
+        surface.Vectors.Emplace( vec, Name);
 
         return result;
 	}
 
 
     //---------------------------------------------------------------------------------------------
-    int32 Instance::Private::AddScalar( int32 lod, int32 comp, int32 surf, float sca, const char* strName )
+    int32 Instance::Private::AddScalar( int32 LODIndex, int32 ComponentIndex, int32 SurfaceIndex, float sca, FName Name)
     {
         // Automatically create the necessary lods and components
-		while (lod >= m_lods.Num())
+		while (LODIndex >= Lods.Num())
 		{
 			AddLOD();
 		}
-		while (comp >= m_lods[lod].m_components.Num())
+		while (ComponentIndex >= Lods[LODIndex].Components.Num())
 		{
-			AddComponent(lod);
+			AddComponent(LODIndex);
 		}
-		while (surf >= m_lods[lod].m_components[comp].m_surfaces.Num())
+		while (SurfaceIndex >= Lods[LODIndex].Components[ComponentIndex].Surfaces.Num())
 		{
-			AddSurface(lod, comp);
+			AddSurface(LODIndex, ComponentIndex);
 		}
 
-        INSTANCE_SURFACE& surface = m_lods[lod].m_components[comp].m_surfaces[surf];
-        int32 result = surface.m_scalars.Num();
-        surface.m_scalars.Emplace( sca, strName );
+        INSTANCE_SURFACE& surface = Lods[LODIndex].Components[ComponentIndex].Surfaces[SurfaceIndex];
+        int32 result = surface.Scalars.Num();
+        surface.Scalars.Emplace( sca, Name );
 
         return result;
     }
 
 
     //---------------------------------------------------------------------------------------------
-    int32 Instance::Private::AddString( int32 lod, int32 comp, int32 surf, const char* str, const char* strName )
+    int32 Instance::Private::AddString( int32 LODIndex, int32 ComponentIndex, int32 SurfaceIndex, const FString& Value, FName Name)
     {
         // Automatically create the necessary lods and components
-		while (lod >= m_lods.Num())
+		while (LODIndex >= Lods.Num())
 		{
 			AddLOD();
 		}
-		while (comp >= m_lods[lod].m_components.Num())
+		while (ComponentIndex >= Lods[LODIndex].Components.Num())
 		{
-			AddComponent(lod);
+			AddComponent(LODIndex);
 		}
-		while (surf >= m_lods[lod].m_components[comp].m_surfaces.Num())
+		while (SurfaceIndex >= Lods[LODIndex].Components[ComponentIndex].Surfaces.Num())
 		{
-			AddSurface(lod, comp);
+			AddSurface(LODIndex, ComponentIndex);
 		}
 
-        INSTANCE_SURFACE& surface = m_lods[lod].m_components[comp].m_surfaces[surf];
-        int32 result = surface.m_strings.Num();
-        surface.m_strings.Emplace( str, strName );
+        INSTANCE_SURFACE& surface = Lods[LODIndex].Components[ComponentIndex].Surfaces[SurfaceIndex];
+        int32 result = surface.Strings.Num();
+        surface.Strings.Emplace(Value, Name );
 
         return result;
     }
 
 
     //---------------------------------------------------------------------------------------------
-	void Instance::Private::AddExtensionData(ExtensionDataPtrConst Data, const char* Name)
+	void Instance::Private::AddExtensionData(const Ptr<const class ExtensionData>& Data, FName Name)
 	{
 		check(Data);
-		check(Name);
 
-		NamedExtensionData& Entry = m_extensionData.AddDefaulted_GetRef();
+		NamedExtensionData& Entry = ExtensionData.AddDefaulted_GetRef();
 		Entry.Data = Data;
 		Entry.Name = Name;
-
-		check(Entry.Name.length() > 0);
 	}
 }
 
diff --git a/Engine/Plugins/Experimental/Mutable/Source/MutableRuntime/Private/MuR/InstancePrivate.h b/Engine/Plugins/Experimental/Mutable/Source/MutableRuntime/Private/MuR/InstancePrivate.h
--- a/Engine/Plugins/Experimental/Mutable/Source/MutableRuntime/Private/MuR/InstancePrivate.h
+++ b/Engine/Plugins/Experimental/Mutable/Source/MutableRuntime/Private/MuR/InstancePrivate.h
@@ -25,117 +25,102 @@ namespace mu
 	/** */
     struct INSTANCE_SURFACE
 	{
-		string m_name;
+		FName Name;
         uint32 InternalId=0;
         uint32 ExternalId =0;
         uint32 SharedId =0;
 
 		struct IMAGE
 		{
-            IMAGE(FResourceID p, const char* strName )
+            IMAGE(FResourceID InId, FName InName)
 			{
-                m_imageId = p;
-				if (strName)
-				{
-					m_name = strName;
-				}
+				Id = InId;
+				Name = InName;
 			}
 
-			FResourceID m_imageId;
-			string m_name;
+			FResourceID Id;
+			FName Name;
 		};
 
-		TArray<IMAGE, TInlineAllocator<4>> m_images;
+		TArray<IMAGE, TInlineAllocator<4>> Images;
 
 		struct VECTOR
 		{
-			VECTOR( const FVector4f& v, const char* strName )
+			VECTOR( const FVector4f& v, FName InName)
 			{
-				m_vec = v;
-				if (strName)
-				{
-					m_name = strName;
-				}
+				Value = v;
+				Name = InName;
 			}
 
-			FVector4f m_vec;
-			string m_name;
+			FVector4f Value;
+			FName Name;
 		};
 
-		TArray<VECTOR> m_vectors;
+		TArray<VECTOR> Vectors;
 
         struct SCALAR
         {
-            SCALAR( float v, const char* strName )
+            SCALAR( float v, FName InName )
             {
-                m_scalar = v;
-                if ( strName )
-                {
-                    m_name = strName;
-                }
+                Value = v;
+                Name = InName;
             }
 
-            float m_scalar;
-            string m_name;
-        };
+            float Value;
+			FName Name;
+		};
 
-		TArray<SCALAR> m_scalars;
+		TArray<SCALAR> Scalars;
 
         struct STRING
         {
-            STRING( const char* strValue, const char* strName )
+            STRING(const FString& InValue, FName InName)
             {
-                m_string = strValue ? strValue : "";
-                if ( strName )
-                {
-                    m_name = strName;
-                }
-            }
+				Value = InValue;
+				Name = InName;
+			}
 
-            string m_string;
-            string m_name;
-        };
+			FString Value;
+			FName Name;
+		};
 
-		TArray<STRING> m_strings;
+		TArray<STRING> Strings;
     };
 
 
     struct INSTANCE_COMPONENT
     {
-        string m_name;
+        FName Name;
 
-    	uint16 m_id;
+    	uint16 Id;
 
 		struct MESH
 		{
-            MESH(FResourceID p, const char* strName)
+            MESH(FResourceID InId, FName InName)
 			{
-                m_meshId = p;
-				if (strName)
-				{
-					m_name = strName;
-				}
+				Id = InId;
+				Name = InName;
 			}
 
-			FResourceID m_meshId;
-			string m_name;
+			FResourceID Id;
+			FName Name;
 		};
-		TArray<MESH, TInlineAllocator<2>> m_meshes;
+		TArray<MESH, TInlineAllocator<2>> Meshes;
 
 		// The order must match the meshes surfaces
-		TArray<INSTANCE_SURFACE, TInlineAllocator<4>> m_surfaces;
+		TArray<INSTANCE_SURFACE, TInlineAllocator<4>> Surfaces;
 	};
 
 
     struct INSTANCE_LOD
     {
-		TArray<INSTANCE_COMPONENT, TInlineAllocator<4>> m_components;
+		TArray<INSTANCE_COMPONENT, TInlineAllocator<4>> Components;
     };
 
 	struct NamedExtensionData
 	{
-		ExtensionDataPtrConst Data;
-		string Name;
+		Ptr<const ExtensionData> Data;
+		FName Name;
 	};
 
 	class Instance::Private : public Base
@@ -143,31 +128,26 @@ namespace mu
 	public:
 
         //!
-        Instance::ID m_id;
+        Instance::ID Id = 0;
 
 		//!
-		TArray<INSTANCE_LOD,TInlineAllocator<4>> m_lods;
+		TArray<INSTANCE_LOD,TInlineAllocator<4>> Lods;
 
 		// Every entry must have a valid ExtensionData and name
-		TArray<NamedExtensionData> m_extensionData;
-
-        Private()
-        {
-            m_id = 0;
-        }
-
-        int AddLOD();
-        int AddComponent( int lod );
-        void SetComponentName( int32 lod, int32 comp, const char* strName );
-        int AddMesh(int32 lod, int32 comp, FResourceID, const char* strName);
-		int AddSurface( int lod, int comp );
-        void SetSurfaceName( int32 lod, int32 comp, int32 surf, const char* strName );
-        int AddImage( int32 lod, int32 comp, int32 surf, FResourceID, const char* strName );
-        int32 AddVector( int32 lod, int32 comp, int32 surf, const FVector4f&, const char* strName );
-        int32 AddScalar( int32 lod, int32 comp, int32 surf, float, const char* strName );
-        int32 AddString( int32 lod, int32 comp, int32 surf, const char* strValue, const char* strName );
+		TArray<NamedExtensionData> ExtensionData;
+
+		int32 AddLOD();
+		int32 AddComponent(int32 lod );
+        void SetComponentName( int32 lod, int32 comp, FName Name );
+		int32 AddMesh(int32 lod, int32 comp, FResourceID, FName Name);
+		int32 AddSurface(int32 lod, int32 comp );
+        void SetSurfaceName( int32 lod, int32 comp, int32 surf, FName Name);
+		int32 AddImage( int32 lod, int32 comp, int32 surf, FResourceID, FName Name);
+        int32 AddVector( int32 lod, int32 comp, int32 surf, const FVector4f&, FName Name);
+        int32 AddScalar( int32 lod, int32 comp, int32 surf, float, FName Name);
+        int32 AddString( int32 lod, int32 comp, int32 surf, const FString& Value, FName Name);
 		
-		// Data and Name must be non-null
-		void AddExtensionData(ExtensionDataPtrConst Data, const char* Name);
+		// Data must be non-null
+		void AddExtensionData(const Ptr<const class ExtensionData>& Data, FName Name);
     };
 }
diff --git a/Engine/Plugins/Experimental/Mutable/Source/MutableRuntime/Private/MuR/Mesh.cpp b/Engine/Plugins/Experimental/Mutable/Source/MutableRuntime/Private/MuR/Mesh.cpp
--- a/Engine/Plugins/Experimental/Mutable/Source/MutableRuntime/Private/MuR/Mesh.cpp
+++ b/Engine/Plugins/Experimental/Mutable/Source/MutableRuntime/Private/MuR/Mesh.cpp
@@ -50,7 +50,6 @@ MeshPtr Mesh::Clone() const
 	pResult->m_surfaces = m_surfaces;
 	pResult->m_pSkeleton = m_pSkeleton;
 	pResult->m_pPhysicsBody = m_pPhysicsBody;
-	pResult->m_faceGroups = m_faceGroups;
 	pResult->m_tags = m_tags;
 
     // Clone the main buffers
@@ -105,11 +104,6 @@ MeshPtr Mesh::Clone(EMeshCopyFlags Flags) const
 		pResult->m_pPhysicsBody = m_pPhysicsBody;
 	}
 
-	if (EnumHasAnyFlags(Flags, EMeshCopyFlags::WithFaceGroups))
-	{
-		pResult->m_faceGroups = m_faceGroups;
-	}
-
 	if (EnumHasAnyFlags(Flags, EMeshCopyFlags::WithTags))
 	{
 		pResult->m_tags = m_tags;
@@ -197,11 +191,6 @@ void Mesh::CopyFrom(const Mesh& From, EMeshCopyFlags Flags)
 		m_pPhysicsBody = From.m_pPhysicsBody;
 	}
 
-	if (EnumHasAnyFlags(Flags, EMeshCopyFlags::WithFaceGroups))
-	{
-		m_faceGroups = From.m_faceGroups;
-	}
-
 	if (EnumHasAnyFlags(Flags, EMeshCopyFlags::WithTags))
 	{
 		m_tags = From.m_tags;
@@ -482,70 +471,6 @@ void Mesh::SetLayout( int i, Ptr<const Layout> pLayout )
 }
 
 
-//---------------------------------------------------------------------------------------------
-void Mesh::SetFaceGroupCount( int count )
-{
-    m_faceGroups.SetNum( count );
-}
-
-
-//---------------------------------------------------------------------------------------------
-int Mesh::GetFaceGroupCount() const
-{
-    return m_faceGroups.Num();
-}
-
-
-//---------------------------------------------------------------------------------------------
-const char* Mesh::GetFaceGroupName( int group ) const
-{
-    check( group>=0 && group<m_faceGroups.Num() );
-    return m_faceGroups[group].m_name.c_str();
-}
-
-
-//---------------------------------------------------------------------------------------------
-void Mesh::SetFaceGroupName( int group, const char* strName )
-{
-    check( group>=0 && group<m_faceGroups.Num() );
-    m_faceGroups[group].m_name = strName;
-}
-
-
-//---------------------------------------------------------------------------------------------
-int Mesh::GetFaceGroupFaceCount( int group ) const
-{
-    check( group>=0 && group<m_faceGroups.Num() );
-    return m_faceGroups[group].m_faces.Num();
-}
-
-
-//---------------------------------------------------------------------------------------------
-const int32* Mesh::GetFaceGroupFaces( int group ) const
-{
-    check( group>=0 && group<m_faceGroups.Num() );
-    const int32* pResult = m_faceGroups[group].m_faces.GetData();
-    return pResult;
-}
-
-
-//---------------------------------------------------------------------------------------------
-void Mesh::SetFaceGroupFaces( int group, int count, const int32* faces )
-{
-    check( group>=0 && group<m_faceGroups.Num() );
-    check( !count || faces );
-	LLM_SCOPE_BYNAME(TEXT("MutableRuntime"));
-
-    m_faceGroups[group].m_faces.SetNum( count );
-
-    if (count)
-    {
-		FMemory::Memcpy( m_faceGroups[group].m_faces.GetData(), faces, count*sizeof(int32) );
-    }
-
-}
-
-
 //---------------------------------------------------------------------------------------------
 int Mesh::GetTagCount() const
 {
@@ -562,30 +487,31 @@ void Mesh::SetTagCount( int count )
 
 
 //---------------------------------------------------------------------------------------------
-const char* Mesh::GetTag( int tagIndex ) const
+const FString& Mesh::GetTag( int tagIndex ) const
 {
     check( tagIndex>=0 && tagIndex<GetTagCount() );
 
     if (tagIndex >= 0 && tagIndex < GetTagCount())
     {
-        return m_tags[tagIndex].c_str();
+        return m_tags[tagIndex];
     }
     else
     {
-        return nullptr;
+		static FString NullString;
+        return NullString;
     }
 }
 
 
 //---------------------------------------------------------------------------------------------
-void Mesh::SetTag( int tagIndex, const char* strName )
+void Mesh::SetTag( int tagIndex, const FString& Name )
 {
     check( tagIndex>=0 && tagIndex<GetTagCount() );
 	LLM_SCOPE_BYNAME(TEXT("MutableRuntime"));
 
     if (tagIndex >= 0 && tagIndex < GetTagCount())
     {
-        m_tags[tagIndex] = strName ? strName : "";
+        m_tags[tagIndex] = Name;
     }
 }
 
@@ -1197,29 +1123,6 @@ void MESH_SURFACE::Unserialise(InputArchive& arch)
 }
 
 
-//-------------------------------------------------------------------------------------------------
-void Mesh::FACE_GROUP::Serialise(OutputArchive& arch) const
-{
-	const int32 ver = 0;
-	arch << ver;
-
-	arch << m_name;
-	arch << m_faces;
-}
-
-
-//-------------------------------------------------------------------------------------------------
-void Mesh::FACE_GROUP::Unserialise(InputArchive& arch)
-{
-	int32 ver = 0;
-	arch >> ver;
-	check(ver == 0);
-
-	arch >> m_name;
-	arch >> m_faces;
-}
-
-
 //-------------------------------------------------------------------------------------------------
 void Mesh::FBonePose::Serialise(OutputArchive& arch) const
 {
@@ -1241,8 +1144,8 @@ void Mesh::FBonePose::Unserialise(InputArchive& arch)
 
 	if (ver <= 1)
 	{
-		string BoneName;
-		arch >> BoneName;
+		std::string DeprecatedBoneName;
+		arch >> DeprecatedBoneName;
 
 		BoneId = 0;
 	}
@@ -1269,7 +1172,7 @@ void Mesh::FBonePose::Unserialise(InputArchive& arch)
 //-------------------------------------------------------------------------------------------------
 void Mesh::Serialise(OutputArchive& arch) const
 {
-	uint32 ver = 16;
+	uint32 ver = 17;
 	arch << ver;
 
 	arch << m_IndexBuffers;
@@ -1285,7 +1188,6 @@ void Mesh::Serialise(OutputArchive& arch) const
 
 	arch << m_staticFormatFlags;
 	arch << m_surfaces;
-	arch << m_faceGroups;
 
 	arch << m_tags;
 
@@ -1301,7 +1203,7 @@ void Mesh::Unserialise(InputArchive& arch)
 {
 	uint32 ver;
 	arch >> ver;
-	check(ver <= 16);
+	check(ver <= 17);
 
 	arch >> m_IndexBuffers;
 	arch >> m_VertexBuffers;
@@ -1335,9 +1237,40 @@ void Mesh::Unserialise(InputArchive& arch)
 		// Deserialize LegacySurfaces
 		UnserialiseLegacySurfaces(arch, m_surfaces);
 	}
-	arch >> m_faceGroups;
 
-	arch >> m_tags;
+	if (ver <= 16)
+	{
+		struct FACE_GROUP_DEPRECATED
+		{
+			std::string m_name;
+			TArray<int32> m_faces;
+			inline void Unserialise(InputArchive& arch) 
+			{
+				int32 ver = 0;
+				arch >> ver;
+				arch >> m_name;
+				arch >> m_faces;
+			}
+
+		};
+		TArray<FACE_GROUP_DEPRECATED> FaceGroups;
+		arch >> FaceGroups;
+	}
+
+	if (ver <= 16)
+	{
+		TArray < std::string > Temp;
+		arch >> Temp;
+		m_tags.SetNum(Temp.Num());
+		for (int32 c = 0; c < Temp.Num(); ++c)
+		{
+			m_tags[c] = Temp[c].c_str();
+		}
+	}
+	else
+	{
+		arch >> m_tags;
+	}
 
 	if (ver >= 13)
 	{
@@ -1421,7 +1354,6 @@ bool Mesh::IsSimilar(const Mesh& o, bool bCompareLayouts) const
 	}
 
 	if (equal) equal = (m_surfaces == o.m_surfaces);
-	if (equal) equal = (m_faceGroups == o.m_faceGroups);
 	if (equal) equal = (m_tags == o.m_tags);
 
 	// Special comparison for layouts
diff --git a/Engine/Plugins/Experimental/Mutable/Source/MutableRuntime/Private/MuR/Model.cpp b/Engine/Plugins/Experimental/Mutable/Source/MutableRuntime/Private/MuR/Model.cpp
--- a/Engine/Plugins/Experimental/Mutable/Source/MutableRuntime/Private/MuR/Model.cpp
+++ b/Engine/Plugins/Experimental/Mutable/Source/MutableRuntime/Private/MuR/Model.cpp
@@ -487,27 +487,28 @@ namespace mu
 
 
     //---------------------------------------------------------------------------------------------
-    const char* Model::GetStateName( int index ) const
+    const FString& Model::GetStateName( int32 index ) const
     {
         const char* strRes = 0;
 
         if ( index>=0 && index<(int)m_pD->m_program.m_states.Num() )
         {
-            strRes = m_pD->m_program.m_states[index].m_name.c_str();
+            return m_pD->m_program.m_states[index].Name;
         }
 
-        return strRes;
+		static FString None;
+        return None;
     }
 
 
     //---------------------------------------------------------------------------------------------
-    int Model::FindState( const char* strName ) const
+    int32 Model::FindState( const FString& Name ) const
     {
         int res = -1;
 
         for ( int i=0; res<0 && i<(int)m_pD->m_program.m_states.Num(); ++i )
         {
-            if ( m_pD->m_program.m_states[i].m_name == strName )
+            if ( m_pD->m_program.m_states[i].Name == Name )
             {
                 res = i;
             }
diff --git a/Engine/Plugins/Experimental/Mutable/Source/MutableRuntime/Private/MuR/ModelPrivate.h b/Engine/Plugins/Experimental/Mutable/Source/MutableRuntime/Private/MuR/ModelPrivate.h
--- a/Engine/Plugins/Experimental/Mutable/Source/MutableRuntime/Private/MuR/ModelPrivate.h
+++ b/Engine/Plugins/Experimental/Mutable/Source/MutableRuntime/Private/MuR/ModelPrivate.h
@@ -123,7 +123,7 @@ namespace mu
         struct FState
         {
             //! Name of the state
-            string m_name;
+            FString Name;
 
             //! First instruction of the full build of an instance in this state
             OP::ADDRESS m_root = 0;
@@ -145,7 +145,7 @@ namespace mu
             //!
             inline void Serialise( OutputArchive& arch ) const
             {
-                arch << m_name;
+                arch << Name;
                 arch << m_root;
                 arch << m_runtimeParameters;
                 arch << m_updateCache;
@@ -155,7 +155,7 @@ namespace mu
             //!
             inline void Unserialise( InputArchive& arch )
             {
-                arch >> m_name;
+                arch >> Name;
                 arch >> m_root;
                 arch >> m_runtimeParameters;
                 arch >> m_updateCache;
@@ -235,7 +235,7 @@ namespace mu
 		TArray<FExtensionDataConstant> m_constantExtensionData;
 
         //! Constant string data
-		TArray<string> m_constantStrings;
+		TArray<FString> m_constantStrings;
 
         //! Constant layout data
 		TArray<Ptr<const Layout>> m_constantLayouts;
@@ -460,7 +460,7 @@ namespace mu
             return index;
         }
 
-        OP::ADDRESS AddConstant( const string& str )
+        OP::ADDRESS AddConstant( const FString& str )
         {            
             // Ensure unique
             for ( SIZE_T i=0; i<m_constantStrings.Num(); ++i)
diff --git a/Engine/Plugins/Experimental/Mutable/Source/MutableRuntime/Private/MuR/MutableString.cpp b/Engine/Plugins/Experimental/Mutable/Source/MutableRuntime/Private/MuR/MutableString.cpp
--- a/Engine/Plugins/Experimental/Mutable/Source/MutableRuntime/Private/MuR/MutableString.cpp
+++ b/Engine/Plugins/Experimental/Mutable/Source/MutableRuntime/Private/MuR/MutableString.cpp
@@ -7,36 +7,32 @@
 namespace mu {
 
 	//---------------------------------------------------------------------------------------------
-	String::String( const char* v )
+	String::String( const FString& InValue)
 	{
-		if (v)
-		{
-            m_value = v;
-        }
+        Value = InValue;
 	}
 
 
-	//---------------------------------------------------------------------------------------------
-	void String::Serialise( const String* p, OutputArchive& arch )
-	{
-		arch << *p;
-	}
+	////---------------------------------------------------------------------------------------------
+	//void String::Serialise( const String* p, OutputArchive& arch )
+	//{
+	//	arch << *p;
+	//}
 
 
-	//---------------------------------------------------------------------------------------------
-	Ptr<String> String::StaticUnserialise( InputArchive& arch )
-	{
-		Ptr<String> pResult = new String();
-		arch >> *pResult;
-		return pResult;
-	}
+	////---------------------------------------------------------------------------------------------
+	//Ptr<String> String::StaticUnserialise( InputArchive& arch )
+	//{
+	//	Ptr<String> pResult = new String();
+	//	arch >> *pResult;
+	//	return pResult;
+	//}
 
 
 	//---------------------------------------------------------------------------------------------
 	Ptr<String> String::Clone() const
 	{
-		Ptr<String> pResult = new String();
-		pResult->m_value = m_value;
+		Ptr<String> pResult = new String(Value);
 		return pResult;
 	}
 
@@ -44,14 +40,14 @@ namespace mu {
 	//---------------------------------------------------------------------------------------------
 	int32 String::GetDataSize() const
 	{
-		return sizeof(String) + m_value.capacity();
+		return sizeof(String) + Value.GetAllocatedSize();
 	}
 
 
 	//---------------------------------------------------------------------------------------------
-	const char* String::GetValue() const
+	const FString& String::GetValue() const
 	{
-        return m_value.c_str();
+        return Value;
 	}
 
 }
diff --git a/Engine/Plugins/Experimental/Mutable/Source/MutableRuntime/Private/MuR/OpLayoutPack.h b/Engine/Plugins/Experimental/Mutable/Source/MutableRuntime/Private/MuR/OpLayoutPack.h
--- a/Engine/Plugins/Experimental/Mutable/Source/MutableRuntime/Private/MuR/OpLayoutPack.h
+++ b/Engine/Plugins/Experimental/Mutable/Source/MutableRuntime/Private/MuR/OpLayoutPack.h
@@ -11,8 +11,6 @@
 
 namespace mu
 {
-	typedef std::string string;
-
 	//---------------------------------------------------------------------------------------------
     // inline void SimpleLayoutPack( Layout* pResult, const Layout* pSource )
 	// {
@@ -211,7 +209,7 @@ namespace mu
     };
 
 
-    inline string DebugGetBlockAt( const SCRATCH_LAYOUT_PACK& scratch,
+    inline char DebugGetBlockAt( const SCRATCH_LAYOUT_PACK& scratch,
                                    const TArray<uint8_t>& packedFlag,
                                    int x, int y )
     {
@@ -226,14 +224,11 @@ namespace mu
                      y<scratch.positions[i][1]+scratch.sorted[b].size[1]
                      )
                 {
-                    char temp[2];
-                    temp[1] = 0;
-                    temp[0] = char('a')+char(b);
-                    return temp;
+					return char('a')+char(b);
                 }
             }
         }
-        return ".";
+        return '.';
     }
 
 
@@ -375,12 +370,10 @@ namespace mu
 	}
 
 
-    //---------------------------------------------------------------------------------------------
-    //! Even more expensive but precise version
-    //---------------------------------------------------------------------------------------------
-	inline void ReduceBlock(int blockCount, int* area, int* reverse_it, SCRATCH_LAYOUT_PACK* scratch, EReductionMethod ReductionMethod)
+    /** Updates the area, and the iterator. */
+	inline void ReduceBlock(int blockCount, int32& InOutArea, int32& InOutBlockIt, SCRATCH_LAYOUT_PACK* scratch, EReductionMethod ReductionMethod)
 	{
-		int r_it = (*reverse_it);
+		int r_it = InOutBlockIt;
 		bool pass = false;
 
 		int oldBlockArea = scratch->sorted[r_it].size[0] * scratch->sorted[r_it].size[1];
@@ -454,15 +447,15 @@ namespace mu
 		
 		int newBlockArea = scratch->sorted[r_it].size[0] * scratch->sorted[r_it].size[1];
 
-		(*area) = (*area) - (oldBlockArea - newBlockArea);
+		InOutArea = InOutArea - (oldBlockArea - newBlockArea);
 
 		if (pass)
 		{
-			(*reverse_it) = (*reverse_it) + 1;
+			InOutBlockIt = InOutBlockIt + 1;
 
-			if ((*reverse_it) >= blockCount)
+			if (InOutBlockIt >= blockCount)
 			{
-				(*reverse_it) = 0;
+				InOutBlockIt = 0;
 			}
 		}
 	}
@@ -797,7 +790,9 @@ namespace mu
 		}
         
         int bestY = maxY;
-		int r_it = 0;
+
+		// This is used to iterate through blocks.
+		int32 BlockIterator = 0;
 
         // Sort by height, area
         check( (int)scratch->sorted.Num()==blockCount );
@@ -815,10 +810,10 @@ namespace mu
 				scratch->sorted.Sort(CompareBlocksPriority);
 			}
 
-			//Shrink blocks in case we do not have enough space to pack everything
+			// Shrink blocks in case we do not have enough space to pack everything
 			while (maxX*maxY < area)
 			{
-				ReduceBlock(blockCount, &area, &r_it, scratch, pSourceLayout->GetBlockReductionMethod());
+				ReduceBlock(blockCount, area, BlockIterator, scratch, pSourceLayout->GetBlockReductionMethod());
 			}
 			
 		}
@@ -827,24 +822,24 @@ namespace mu
 
 		while (!fits)
 		{
-			//Sort by height&area before packing
+			// Sort by height&area before packing
 			if (usePriority)
 			{
 				scratch->sorted.Sort(CompareBlocks);
 			}
 
-			//Try to pack everything
+			// Try to pack everything
 			fits = SetPositions(bestY, layoutSizeY, &maxX, &maxY, scratch, pSourceLayout->GetLayoutPackingStrategy());
 
 			if (!fits && pSourceLayout->GetLayoutPackingStrategy() == EPackStrategy::FIXED_LAYOUT)
 			{
-				//Sort by priority before shrink
+				// Sort by priority before shrink
 				if (usePriority)
 				{
 					scratch->sorted.Sort(CompareBlocksPriority);
 				}
 
-				ReduceBlock(blockCount, &area, &r_it, scratch, pSourceLayout->GetBlockReductionMethod());
+				ReduceBlock(blockCount, area, BlockIterator, scratch, pSourceLayout->GetBlockReductionMethod());
 			}
 		}
 
diff --git a/Engine/Plugins/Experimental/Mutable/Source/MutableRuntime/Private/MuR/OpMeshApplyPose.h b/Engine/Plugins/Experimental/Mutable/Source/MutableRuntime/Private/MuR/OpMeshApplyPose.h
--- a/Engine/Plugins/Experimental/Mutable/Source/MutableRuntime/Private/MuR/OpMeshApplyPose.h
+++ b/Engine/Plugins/Experimental/Mutable/Source/MutableRuntime/Private/MuR/OpMeshApplyPose.h
@@ -6,7 +6,6 @@
 #include "MuR/ConvertData.h"
 #include "MuR/MutableTrace.h"
 #include "MuR/Platform.h"
-#include "MuR/OpMeshChartDifference.h"
 
 
 namespace mu
diff --git a/Engine/Plugins/Experimental/Mutable/Source/MutableRuntime/Private/MuR/OpMeshClipWithMesh.h b/Engine/Plugins/Experimental/Mutable/Source/MutableRuntime/Private/MuR/OpMeshClipWithMesh.h
--- a/Engine/Plugins/Experimental/Mutable/Source/MutableRuntime/Private/MuR/OpMeshClipWithMesh.h
+++ b/Engine/Plugins/Experimental/Mutable/Source/MutableRuntime/Private/MuR/OpMeshClipWithMesh.h
@@ -5,7 +5,6 @@
 #include "MuR/MeshPrivate.h"
 #include "MuR/ConvertData.h"
 #include "MuR/Platform.h"
-#include "MuR/OpMeshChartDifference.h"
 #include "MuR/MutableTrace.h"
 #include "Math/Ray.h"
 #include "TriangleTypes.h"
diff --git a/Engine/Plugins/Experimental/Mutable/Source/MutableRuntime/Private/MuR/OpMeshDifference.h b/Engine/Plugins/Experimental/Mutable/Source/MutableRuntime/Private/MuR/OpMeshDifference.h
--- a/Engine/Plugins/Experimental/Mutable/Source/MutableRuntime/Private/MuR/OpMeshDifference.h
+++ b/Engine/Plugins/Experimental/Mutable/Source/MutableRuntime/Private/MuR/OpMeshDifference.h
@@ -72,7 +72,6 @@ struct MORPH_VERTEX
 					if ( sem != MBS_VERTEXINDEX &&
 					     sem != MBS_BONEINDICES &&
 					     sem != MBS_BONEWEIGHTS &&
-					     sem != MBS_CHART &&
 					     sem != MBS_LAYOUTBLOCK &&
 					     sem != MBS_OTHER && 
 					     ( !ignoreTexCoords || sem!=MBS_TEXCOORDS ) )
diff --git a/Engine/Plugins/Experimental/Mutable/Source/MutableRuntime/Private/MuR/OpMeshExtractLayoutBlock.h b/Engine/Plugins/Experimental/Mutable/Source/MutableRuntime/Private/MuR/OpMeshExtractLayoutBlock.h
--- a/Engine/Plugins/Experimental/Mutable/Source/MutableRuntime/Private/MuR/OpMeshExtractLayoutBlock.h
+++ b/Engine/Plugins/Experimental/Mutable/Source/MutableRuntime/Private/MuR/OpMeshExtractLayoutBlock.h
@@ -124,9 +124,6 @@ namespace mu
 		// TODO: Optimise
 		Result->CopyFrom(*pSource);
 
-        // TODO
-        Result->SetFaceGroupCount( 0 );
-
 		UntypedMeshBufferIteratorConst itBlocks( pSource->GetVertexBuffers(),
 												 MBS_LAYOUTBLOCK, layout );
 
@@ -174,49 +171,4 @@ namespace mu
         Result->EnsureSurfaceData();
 	}
 
-
-    //---------------------------------------------------------------------------------------------
-    inline void MeshExtractFaceGroup(Mesh* Result, const Mesh* pSource, int group, bool& bOutSuccess)
-    {
-		check(pSource);
-		bOutSuccess = true;
-
-        if (group < 0 || group >= pSource->GetFaceGroupCount())
-        {
-			bOutSuccess = true; // use empty mesh as retun value.
-            return;
-        }
-
-        // TODO: Optimise
-		Result->CopyFrom(*pSource);
-
-        int resultVertices = 0;
-		TArray<int> oldToNew;
-		oldToNew.Init(-1,pSource->GetVertexCount());
-		TArray<int> newToOld;
-        newToOld.Reserve( pSource->GetVertexCount() );
-
-        UntypedMeshBufferIteratorConst itIndex( pSource->GetIndexBuffers(), MBS_VERTEXINDEX );
-        for ( int32 f=0; f<pSource->m_faceGroups[group].m_faces.Num(); ++f )
-        {
-            int32 face = pSource->m_faceGroups[group].m_faces[f];
-
-            for (int i=0; i<3; i++)
-            {
-                uint32_t vertexIndex = (itIndex+face*3+i).GetAsUINT32();
-                if ( oldToNew[vertexIndex]<0 )
-                {
-                    oldToNew[vertexIndex] = resultVertices++;
-                    newToOld.Add( vertexIndex );
-                }
-            }
-        }
-
-        MeshExtractFromVertices(pSource, Result, oldToNew, newToOld);
-
-        Result->m_surfaces.Empty();
-        Result->EnsureSurfaceData();
-
-    }
-
 }
diff --git a/Engine/Plugins/Experimental/Mutable/Source/MutableRuntime/Private/MuR/OpMeshFormat.cpp b/Engine/Plugins/Experimental/Mutable/Source/MutableRuntime/Private/MuR/OpMeshFormat.cpp
--- a/Engine/Plugins/Experimental/Mutable/Source/MutableRuntime/Private/MuR/OpMeshFormat.cpp
+++ b/Engine/Plugins/Experimental/Mutable/Source/MutableRuntime/Private/MuR/OpMeshFormat.cpp
@@ -522,7 +522,6 @@ namespace mu
 					);
 
 					if (sourceSemantic == MBS_LAYOUTBLOCK
-						|| sourceSemantic == MBS_CHART
 						|| (isVertexBuffer && sourceSemantic == MBS_VERTEXINDEX)
 						)
 					{
diff --git a/Engine/Plugins/Experimental/Mutable/Source/MutableRuntime/Private/MuR/OpMeshMerge.h b/Engine/Plugins/Experimental/Mutable/Source/MutableRuntime/Private/MuR/OpMeshMerge.h
--- a/Engine/Plugins/Experimental/Mutable/Source/MutableRuntime/Private/MuR/OpMeshMerge.h
+++ b/Engine/Plugins/Experimental/Mutable/Source/MutableRuntime/Private/MuR/OpMeshMerge.h
@@ -879,9 +879,7 @@ namespace mu
 						// Skip system buffers
 						if (chan.m_semantic == MBS_VERTEXINDEX
 							||
-							chan.m_semantic == MBS_LAYOUTBLOCK
-							||
-							chan.m_semantic == MBS_CHART)
+							chan.m_semantic == MBS_LAYOUTBLOCK)
 						{
 							continue;
 						}
@@ -1138,7 +1136,7 @@ namespace mu
 		// Tags
 		Result->m_tags = pFirst->m_tags;
 
-		for (const string& SecondTag : pSecond->m_tags)
+		for (const FString& SecondTag : pSecond->m_tags)
 		{
 			Result->m_tags.AddUnique(SecondTag);
 		}
