[mutable] Remove most of the legacy string usage in favor of FString and FName.

diff --git a/Engine/Plugins/Experimental/Mutable/Source/MutableRuntime/Private/MuR/OpMeshRemove.cpp b/Engine/Plugins/Experimental/Mutable/Source/MutableRuntime/Private/MuR/OpMeshRemove.cpp
--- a/Engine/Plugins/Experimental/Mutable/Source/MutableRuntime/Private/MuR/OpMeshRemove.cpp
+++ b/Engine/Plugins/Experimental/Mutable/Source/MutableRuntime/Private/MuR/OpMeshRemove.cpp
@@ -253,9 +253,6 @@ namespace mu
 
 		Result->CopyFrom(*pSource);
 
-        // TODO
-        Result->SetFaceGroupCount( 0 );
-
         MeshBufferIteratorConst<MBF_UINT32,uint32,1> itMaskVI( pMask->GetVertexBuffers(), MBS_VERTEXINDEX );
 
         int32 firstFreeVertex = 0;
diff --git a/Engine/Plugins/Experimental/Mutable/Source/MutableRuntime/Private/MuR/Parameters.cpp b/Engine/Plugins/Experimental/Mutable/Source/MutableRuntime/Private/MuR/Parameters.cpp
--- a/Engine/Plugins/Experimental/Mutable/Source/MutableRuntime/Private/MuR/Parameters.cpp
+++ b/Engine/Plugins/Experimental/Mutable/Source/MutableRuntime/Private/MuR/Parameters.cpp
@@ -18,7 +18,6 @@
 namespace mu
 {
     MUTABLE_IMPLEMENT_ENUM_SERIALISABLE(PARAMETER_TYPE)              
-    MUTABLE_IMPLEMENT_ENUM_SERIALISABLE(PARAMETER_DETAILED_TYPE)     
     MUTABLE_IMPLEMENT_ENUM_SERIALISABLE(PROJECTOR_TYPE)
 
 	
@@ -57,108 +56,108 @@ namespace mu
 
 
     //---------------------------------------------------------------------------------------------
-    void Parameters::SerialisePortable( const Parameters* params, OutputArchive& arch )
-    {
-        const int32_t ver = 0;
-        arch << ver;
+    //void Parameters::SerialisePortable( const Parameters* params, OutputArchive& arch )
+    //{
+    //    const int32_t ver = 0;
+    //    arch << ver;
 
-        uint32 valueCount = params->GetPrivate()->m_values.Num();
-        arch << valueCount;
+    //    uint32 valueCount = params->GetPrivate()->m_values.Num();
+    //    arch << valueCount;
 
-        for (size_t p=0; p<valueCount; ++p)
-        {
-            const FParameterDesc& desc = params->GetPrivate()->m_pModel->GetPrivate()->m_program.m_parameters[p];
-            arch << desc.m_name;
-            arch << desc.m_uid;
-            arch << desc.m_type;
-            arch << params->GetPrivate()->m_values[p];
-        }
-    }
+    //    for (size_t p=0; p<valueCount; ++p)
+    //    {
+    //        const FParameterDesc& desc = params->GetPrivate()->m_pModel->GetPrivate()->m_program.m_parameters[p];
+    //        arch << desc.m_name;
+    //        arch << desc.m_uid;
+    //        arch << desc.m_type;
+    //        arch << params->GetPrivate()->m_values[p];
+    //    }
+    //}
 
 
     //---------------------------------------------------------------------------------------------
-    ParametersPtr Parameters::UnserialisePortable( InputArchive& arch, TSharedPtr<const Model> pModel )
-    {
-		LLM_SCOPE_BYNAME(TEXT("MutableRuntime"));
-
-        ParametersPtr pResult = Model::NewParameters(pModel);
-        size_t modelParameters = pModel->GetPrivate()->m_program.m_parameters.Num();
-
-        int32_t ver;
-        arch >> ver;
-        check( ver == 0 );
-
-        uint32_t valueCount = 0;
-        arch >> valueCount;
-
-        for (uint32_t p=0; p<valueCount; ++p)
-        {
-            // Read the parameter data.
-            mu::string name;
-            arch >> name;
-
-            mu::string uid;
-            arch >> uid;
-
-            PARAMETER_TYPE type;
-            arch >> type;
-
-            PARAMETER_VALUE value;
-            arch >> value;
-
-            int modelParam = -1;
-
-            // UIDs have higher priority
-            if ( !uid.empty() )
-            {
-                for (size_t mp=0; mp<modelParameters; ++mp)
-                {
-                    const FParameterDesc& desc = pModel->GetPrivate()->m_program.m_parameters[mp];
-                    if (desc.m_uid==uid)
-                    {
-                        modelParam = int(mp);
-                        break;
-                    }
-                }
-            }
-
-            // Try with name+type
-            if (modelParam<0)
-            {
-                for (size_t mp=0; mp<modelParameters; ++mp)
-                {
-                    const FParameterDesc& desc = pModel->GetPrivate()->m_program.m_parameters[mp];
-                    if (desc.m_name==name && desc.m_type==type)
-                    {
-                        modelParam = int(mp);
-                        break;
-                    }
-                }
-            }
-
-            // Try with name only, and maybe adapt type
-            if (modelParam<0)
-            {
-                for (size_t mp=0; mp<modelParameters; ++mp)
-                {
-                    const FParameterDesc& desc = pModel->GetPrivate()->m_program.m_parameters[mp];
-                    if (desc.m_name==name)
-                    {
-                        modelParam = int(mp);
-                        break;
-                    }
-                }
-            }
-
-            if (modelParam>=0)
-            {
-                // We found something. \todo: convert the type if necessary?
-                pResult->GetPrivate()->m_values[modelParam] = value;
-            }
-        }
-
-        return pResult;
-    }
+  //  ParametersPtr Parameters::UnserialisePortable( InputArchive& arch, TSharedPtr<const Model> pModel )
+  //  {
+		//LLM_SCOPE_BYNAME(TEXT("MutableRuntime"));
+
+  //      ParametersPtr pResult = Model::NewParameters(pModel);
+  //      size_t modelParameters = pModel->GetPrivate()->m_program.m_parameters.Num();
+
+  //      int32_t ver;
+  //      arch >> ver;
+  //      check( ver == 0 );
+
+  //      uint32_t valueCount = 0;
+  //      arch >> valueCount;
+
+  //      for (uint32_t p=0; p<valueCount; ++p)
+  //      {
+  //          // Read the parameter data.
+  //          std::string name;
+  //          arch >> name;
+
+  //          std::string uid;
+  //          arch >> uid;
+
+  //          PARAMETER_TYPE type;
+  //          arch >> type;
+
+  //          PARAMETER_VALUE value;
+  //          arch >> value;
+
+  //          int modelParam = -1;
+
+  //          // UIDs have higher priority
+  //          if ( !uid.empty() )
+  //          {
+  //              for (size_t mp=0; mp<modelParameters; ++mp)
+  //              {
+  //                  const FParameterDesc& desc = pModel->GetPrivate()->m_program.m_parameters[mp];
+  //                  if (desc.m_uid==uid)
+  //                  {
+  //                      modelParam = int(mp);
+  //                      break;
+  //                  }
+  //              }
+  //          }
+
+  //          // Try with name+type
+  //          if (modelParam<0)
+  //          {
+  //              for (size_t mp=0; mp<modelParameters; ++mp)
+  //              {
+  //                  const FParameterDesc& desc = pModel->GetPrivate()->m_program.m_parameters[mp];
+  //                  if (desc.m_name==name && desc.m_type==type)
+  //                  {
+  //                      modelParam = int(mp);
+  //                      break;
+  //                  }
+  //              }
+  //          }
+
+  //          // Try with name only, and maybe adapt type
+  //          if (modelParam<0)
+  //          {
+  //              for (size_t mp=0; mp<modelParameters; ++mp)
+  //              {
+  //                  const FParameterDesc& desc = pModel->GetPrivate()->m_program.m_parameters[mp];
+  //                  if (desc.m_name==name)
+  //                  {
+  //                      modelParam = int(mp);
+  //                      break;
+  //                  }
+  //              }
+  //          }
+
+  //          if (modelParam>=0)
+  //          {
+  //              // We found something. \todo: convert the type if necessary?
+  //              pResult->GetPrivate()->m_values[modelParam] = value;
+  //          }
+  //      }
+
+  //      return pResult;
+  //  }
 
 
 	//---------------------------------------------------------------------------------------------
@@ -191,27 +190,27 @@ namespace mu
 
 
 	//---------------------------------------------------------------------------------------------
-	const char* Parameters::GetName( int index ) const
+	const FString& Parameters::GetName( int index ) const
 	{
 		const FProgram& program = m_pD->m_pModel->GetPrivate()->m_program;
 		check( index>=0 && index<(int)program.m_parameters.Num() );
 
-		return program.m_parameters[index].m_name.c_str();
+		return program.m_parameters[index].m_name;
 	}
 
 
 	//---------------------------------------------------------------------------------------------
-	const char* Parameters::GetUid( int index ) const
+	const FString& Parameters::GetUid( int index ) const
 	{
 		const FProgram& program = m_pD->m_pModel->GetPrivate()->m_program;
 		check( index>=0 && index<(int)program.m_parameters.Num() );
 
-		return program.m_parameters[index].m_uid.c_str();
+		return program.m_parameters[index].m_uid;
 	}
 
 
 	//---------------------------------------------------------------------------------------------
-	int Parameters::Find( const char* strName ) const
+	int Parameters::Find( const FString& strName ) const
 	{
 		return m_pD->Find( strName );
 	}
@@ -227,16 +226,6 @@ namespace mu
 	}
 
 
-	//---------------------------------------------------------------------------------------------
-	PARAMETER_DETAILED_TYPE Parameters::GetDetailedType( int index ) const
-	{
-		const FProgram& program = m_pD->m_pModel->GetPrivate()->m_program;
-		check( index>=0 && index<(int)program.m_parameters.Num() );
-
-		return program.m_parameters[index].m_detailedType;
-	}
-
-
     //---------------------------------------------------------------------------------------------
     Ptr<RangeIndex> Parameters::NewRangeIndex( int paramIndex ) const
     {
@@ -423,7 +412,7 @@ namespace mu
 
 
 	//---------------------------------------------------------------------------------------------
-	int Parameters::GetIntValueIndex(int paramIndex, const char* valueName) const
+	int Parameters::GetIntValueIndex(int paramIndex, const FString& ValueName) const
 	{
 		const FProgram& program = m_pD->m_pModel->GetPrivate()->m_program;
 		check(paramIndex >= 0
@@ -432,7 +421,7 @@ namespace mu
 		int result = -1;
 		for (size_t v = 0; v < program.m_parameters[paramIndex].m_possibleValues.Num(); ++v)
 		{
-			if (program.m_parameters[paramIndex].m_possibleValues[v].m_name == valueName)
+			if (program.m_parameters[paramIndex].m_possibleValues[v].m_name == ValueName)
 			{
 				result = (int)v;
 				break;
@@ -461,7 +450,7 @@ namespace mu
 
 
 	//---------------------------------------------------------------------------------------------
-	const char* Parameters::GetIntPossibleValueName( int paramIndex, int valueIndex ) const
+	const FString& Parameters::GetIntPossibleValueName( int paramIndex, int valueIndex ) const
 	{
 		const FProgram& program = m_pD->m_pModel->GetPrivate()->m_program;
 		check( paramIndex>=0
@@ -469,13 +458,12 @@ namespace mu
 		check( valueIndex>=0
 				&& valueIndex<(int)program.m_parameters[paramIndex].m_possibleValues.Num() );
 
-		return program.m_parameters[paramIndex].m_possibleValues[valueIndex].m_name.c_str();
+		return program.m_parameters[paramIndex].m_possibleValues[valueIndex].m_name;
 	}
 
 
 	//---------------------------------------------------------------------------------------------
-    int Parameters::GetIntValue( int index,
-                                 const Ptr<const RangeIndex>& pos ) const
+    int Parameters::GetIntValue( int index, const Ptr<const RangeIndex>& pos ) const
 	{
 		check( index>=0 && index<(int)m_pD->m_values.Num() );
 		check( GetType(index)== PARAMETER_TYPE::T_INT );
@@ -828,7 +816,7 @@ namespace mu
 
 
     //---------------------------------------------------------------------------------------------
-    const char* Parameters::GetStringValue( int index, const Ptr<const RangeIndex>& pos ) const
+    void Parameters::GetStringValue( int index, FString& OutValue, const Ptr<const RangeIndex>& pos ) const
     {
         check( index >= 0 && index < (int)m_pD->m_values.Num() );
         check( GetType( index ) == PARAMETER_TYPE::T_STRING );
@@ -837,15 +825,16 @@ namespace mu
         if ( index < 0 || index >= (int)m_pD->m_values.Num() ||
              GetType( index ) != PARAMETER_TYPE::T_STRING )
         {
-            return "";
+            OutValue = TEXT("");
+			return;
         }
 
         // Single value case
         if ( !pos )
         {
             // Return the single value
-            const string& String = m_pD->m_values[index].Get<ParamStringType>();
-            return String.c_str();
+			OutValue = m_pD->m_values[index].Get<ParamStringType>();
+            return;
         }
 
         // Multivalue case
@@ -857,19 +846,19 @@ namespace mu
 			const PARAMETER_VALUE* it = m.Find( pos->m_pD->m_values );
             if ( it )
             {
-                const string& String = it->Get<ParamStringType>();
-                return String.c_str();
+				OutValue = it->Get<ParamStringType>();
+                return;
             }
         }
 
         // Multivalue parameter, but no multivalue set. Return single value.
-        const string& String = m_pD->m_values[index].Get<ParamStringType>();
-        return String.c_str();
+		OutValue = m_pD->m_values[index].Get<ParamStringType>();
+        return;
     }
 
 
     //---------------------------------------------------------------------------------------------
-    void Parameters::SetStringValue( int index, const char* value, const Ptr<const RangeIndex>& pos )
+    void Parameters::SetStringValue( int index, const FString& Value, const Ptr<const RangeIndex>& pos )
     {
 		LLM_SCOPE_BYNAME(TEXT("MutableRuntime"));
 
@@ -892,7 +881,7 @@ namespace mu
                 m_pD->m_multiValues[index].Empty();
             }
 
-        	m_pD->m_values[index].Set<ParamStringType>(value);
+        	m_pD->m_values[index].Set<ParamStringType>(Value);
         }
 
         // Multivalue case
@@ -907,14 +896,13 @@ namespace mu
 
 			TMap< TArray<int32_t>, PARAMETER_VALUE >& m = m_pD->m_multiValues[index];
 			PARAMETER_VALUE& it = m.FindOrAdd( pos->m_pD->m_values );
-        	it.Set<ParamStringType>(value);
+        	it.Set<ParamStringType>(Value);
         }
     }
 
 
     //---------------------------------------------------------------------------------------------
-    FProjector Parameters::Private::GetProjectorValue( int index,
-                                                      const Ptr<const RangeIndex>& pos ) const
+    FProjector Parameters::Private::GetProjectorValue( int index, const Ptr<const RangeIndex>& pos ) const
     {
 
         const FProgram& program = m_pModel->GetPrivate()->m_program;
@@ -1119,7 +1107,7 @@ namespace mu
 
 
 	//---------------------------------------------------------------------------------------------
-	int Parameters::Private::Find( const char* strName ) const
+	int32 Parameters::Private::Find( const FString& Name ) const
 	{
 		const FProgram& program = m_pModel->GetPrivate()->m_program;
 
@@ -1129,7 +1117,7 @@ namespace mu
 		{
 			const FParameterDesc& p = program.m_parameters[i];
 
-			if ( p.m_name == strName )
+			if (p.m_name == Name)
 			{
 				result = i;
 			}
@@ -1172,7 +1160,7 @@ namespace mu
 
 
     //---------------------------------------------------------------------------------------------
-    const char* RangeIndex::GetRangeName( int index ) const
+    const FString& RangeIndex::GetRangeName( int index ) const
     {
         check( index >= 0 && index < GetRangeCount() );
         check( index < int(m_pD->m_pParameters->GetPrivate()->m_values.Num()) );
@@ -1180,12 +1168,12 @@ namespace mu
         check(pModel);
         check( index < int(pModel->GetPrivate()->m_program.m_ranges.Num()) );
 
-        return pModel->GetPrivate()->m_program.m_ranges[index].m_name.c_str();
+        return pModel->GetPrivate()->m_program.m_ranges[index].m_name;
     }
 
 
     //---------------------------------------------------------------------------------------------
-    const char* RangeIndex::GetRangeUid( int index ) const
+    const FString& RangeIndex::GetRangeUid( int index ) const
     {
         check( index >= 0 && index < GetRangeCount() );
         check( index < int(m_pD->m_pParameters->GetPrivate()->m_values.Num()) );
@@ -1193,7 +1181,7 @@ namespace mu
 		check(pModel);
         check( index < int(pModel->GetPrivate()->m_program.m_ranges.Num()) );
 
-        return pModel->GetPrivate()->m_program.m_ranges[index].m_uid.c_str();
+        return pModel->GetPrivate()->m_program.m_ranges[index].m_uid;
     }
 
 
diff --git a/Engine/Plugins/Experimental/Mutable/Source/MutableRuntime/Private/MuR/ParametersPrivate.h b/Engine/Plugins/Experimental/Mutable/Source/MutableRuntime/Private/MuR/ParametersPrivate.h
--- a/Engine/Plugins/Experimental/Mutable/Source/MutableRuntime/Private/MuR/ParametersPrivate.h
+++ b/Engine/Plugins/Experimental/Mutable/Source/MutableRuntime/Private/MuR/ParametersPrivate.h
@@ -12,6 +12,7 @@
 
 #include "MuR/Model.h"
 
+#include <string>
 
 namespace mu
 {
@@ -134,7 +135,6 @@ namespace Private
 }
 
     MUTABLE_DEFINE_ENUM_SERIALISABLE(PARAMETER_TYPE)
-    MUTABLE_DEFINE_ENUM_SERIALISABLE(PARAMETER_DETAILED_TYPE)
     MUTABLE_DEFINE_ENUM_SERIALISABLE(PROJECTOR_TYPE)
 
 
@@ -373,7 +373,7 @@ namespace Private
 	using ParamColorType = FVector3f;
 	using ParamProjectorType = Private::TIndirectObject<FProjector>;
 	using ParamImageType = FName;
-	using ParamStringType = Private::TIndirectObject<string>;
+	using ParamStringType = Private::TIndirectObject<FString>;
 	
 	using PARAMETER_VALUE = TVariant<
             ParamBoolType, ParamIntType, ParamFloatType, ParamColorType, ParamProjectorType, ParamImageType, ParamStringType>;
@@ -402,15 +402,13 @@ namespace Private
 
     struct FParameterDesc
     {
-        string m_name;
+        FString m_name;
 
         //! Unique id (provided externally, so no actual guarantee that it is unique.)
-		string m_uid;
+		FString m_uid;
 
         PARAMETER_TYPE m_type = PARAMETER_TYPE::T_NONE;
 
-        PARAMETER_DETAILED_TYPE m_detailedType = PARAMETER_DETAILED_TYPE::UNKNOWN;
-
         PARAMETER_VALUE m_defaultValue;
 
         //! Ranges, if the parameter is multi-dimensional. The indices refer to the Model's program
@@ -418,13 +416,13 @@ namespace Private
 		TArray<uint32> m_ranges;
 
         //! Possible values of the parameter in case of being an integer, and its names
-        struct INT_VALUE_DESC
+        struct FIntValueDesc
         {
             int16 m_value;
-            string m_name;
+			FString m_name;
 
             //!
-            bool operator==( const INT_VALUE_DESC& other ) const
+            bool operator==( const FIntValueDesc& other ) const
             {
                 return m_value==other.m_value &&
                         m_name==other.m_name;
@@ -447,7 +445,7 @@ namespace Private
 
         //! For integer parameters, this contains the description of the possible values.
         //! If empty, the integer may have any value.
-		TArray<INT_VALUE_DESC> m_possibleValues;
+		TArray<FIntValueDesc> m_possibleValues;
 
         //!
         bool operator==( const FParameterDesc& other ) const
@@ -461,7 +459,7 @@ namespace Private
         //!
         void Serialise( OutputArchive& arch ) const
         {
-            const int32 ver = 7;
+            const int32 ver = 8;
             arch << ver;
 
 			arch << m_name;
@@ -477,28 +475,22 @@ namespace Private
         {
             int32 ver;
             arch >> ver;
-            check( ver >= 6 );
+            check( ver==8 );
 
 			arch >> m_name;
-            arch >> m_uid;
+			arch >> m_uid;
             arch >> m_type;
             arch >> m_defaultValue;
             arch >> m_ranges;
-
-			if (ver<=6)
-			{ 
-				TArray<OP::ADDRESS> UnusedDescImages;
-				arch >> UnusedDescImages;
-			}
-            arch >> m_possibleValues;
+			arch >> m_possibleValues;
         }
     };
 
 
     struct FRangeDesc
     {
-        string m_name;
-        string m_uid;
+		FString m_name;
+		FString m_uid;
 
 		/** Parameter that controls the size of this range, if any. */
 		int32 m_dimensionParameter = -1;
@@ -516,7 +508,7 @@ namespace Private
         //!
         void Serialise( OutputArchive& arch ) const
         {
-            const uint32 ver = 2;
+            const uint32 ver = 3;
             arch << ver;
 
             arch << m_name;
@@ -529,15 +521,11 @@ namespace Private
         {
             uint32 ver;
             arch >> ver;
-            check( ver <= 2 );
+            check( ver == 3 );
 
-            arch >> m_name;
-            arch >> m_uid;
-
-			if (ver >= 2)
-			{
-				arch >> m_dimensionParameter;
-			}
+			arch >> m_name;
+			arch >> m_uid;
+			arch >> m_dimensionParameter;
         }
     };
 
@@ -590,14 +578,14 @@ namespace Private
         {
             uint32 ver;
             arch >> ver;
-			check(ver <= 3);
+			check(ver == 3);
        
 		    arch >> m_values;
 			arch >> m_multiValues;
         }
 
         //!
-        int Find( const char* strName ) const;
+        int32 Find( const FString& Name ) const;
 
         //!
         FProjector GetProjectorValue( int index, const Ptr<const RangeIndex>& pos ) const;
diff --git a/Engine/Plugins/Experimental/Mutable/Source/MutableRuntime/Private/MuR/PhysicsBody.cpp b/Engine/Plugins/Experimental/Mutable/Source/MutableRuntime/Private/MuR/PhysicsBody.cpp
--- a/Engine/Plugins/Experimental/Mutable/Source/MutableRuntime/Private/MuR/PhysicsBody.cpp
+++ b/Engine/Plugins/Experimental/Mutable/Source/MutableRuntime/Private/MuR/PhysicsBody.cpp
@@ -27,7 +27,7 @@ namespace mu
 	//-------------------------------------------------------------------------------------------------
 	void FBodyShape::Serialise(OutputArchive& arch) const
 	{
-		const uint32 ver = 0;
+		const uint32 ver = 1;
 		arch << ver;
 
 		arch << Name;
@@ -40,9 +40,19 @@ namespace mu
 	{
 		uint32 ver;
 		arch >> ver;
-		check(ver == 0);
+		check(ver <= 1);
+
+		if (ver == 0)
+		{
+			std::string Temp;
+			arch >> Temp;
+			Name = Temp.c_str();
+		}
+		else
+		{
+			arch >> Name;
+		}
 
-		arch >> Name;
 		arch >> Flags;
 	}
 
@@ -661,48 +671,48 @@ namespace mu
 	}
 
 	//-------------------------------------------------------------------------------------------------
-	const char* PhysicsBody::GetSphereName(int32 B, int32 I) const
+	const FString& PhysicsBody::GetSphereName(int32 B, int32 I) const
 	{
 		check(B >= 0 && B < Bodies.Num());
 		check(I >= 0 && I < Bodies[B].Spheres.Num());
 
-		return Bodies[B].Spheres[I].Name.c_str();
+		return Bodies[B].Spheres[I].Name;
 	}
 
 	//-------------------------------------------------------------------------------------------------
-	const char* PhysicsBody::GetBoxName(int32 B, int32 I) const
+	const FString& PhysicsBody::GetBoxName(int32 B, int32 I) const
 	{
 		check(B >= 0 && B < Bodies.Num());
 		check(I >= 0 && I < Bodies[B].Boxes.Num());
 
-		return Bodies[B].Boxes[I].Name.c_str();
+		return Bodies[B].Boxes[I].Name;
 	}
 
 	//-------------------------------------------------------------------------------------------------
-	const char* PhysicsBody::GetConvexName(int32 B, int32 I) const
+	const FString& PhysicsBody::GetConvexName(int32 B, int32 I) const
 	{
 		check(B >= 0 && B < Bodies.Num());
 		check(I >= 0 && I < Bodies[B].Convex.Num());
 
-		return Bodies[B].Convex[I].Name.c_str();
+		return Bodies[B].Convex[I].Name;
 	}
 
 	//-------------------------------------------------------------------------------------------------
-	const char* PhysicsBody::GetSphylName(int32 B, int32 I) const
+	const FString& PhysicsBody::GetSphylName(int32 B, int32 I) const
 	{
 		check(B >= 0 && B < Bodies.Num());
 		check(I >= 0 && I < Bodies[B].Sphyls.Num());
 
-		return Bodies[B].Sphyls[I].Name.c_str();
+		return Bodies[B].Sphyls[I].Name;
 	}
 
 	//-------------------------------------------------------------------------------------------------
-	const char* PhysicsBody::GetTaperedCapsuleName(int32 B, int32 I) const
+	const FString& PhysicsBody::GetTaperedCapsuleName(int32 B, int32 I) const
 	{
 		check(B >= 0 && B < Bodies.Num());
 		check(I >= 0 && I < Bodies[B].TaperedCapsules.Num());
 
-		return Bodies[B].TaperedCapsules[I].Name.c_str();
+		return Bodies[B].TaperedCapsules[I].Name;
 	}
 
 
@@ -740,11 +750,11 @@ namespace mu
 		}
 		else
 		{
-			TArray<string> BoneNames;
+			TArray<std::string> BoneNames;
 			arch >> BoneNames;
 
 			const int32 NumBoneNames = BoneNames.Num();
-			TArray<string> UniqueBoneNames;
+			TArray<std::string> UniqueBoneNames;
 			UniqueBoneNames.Reserve(NumBoneNames);
 			BoneIds.Reserve(NumBoneNames);
 			for (int32 BoneIndex = 0; BoneIndex < NumBoneNames; ++BoneIndex)
diff --git a/Engine/Plugins/Experimental/Mutable/Source/MutableRuntime/Private/MuR/SerialisationPrivate.h b/Engine/Plugins/Experimental/Mutable/Source/MutableRuntime/Private/MuR/SerialisationPrivate.h
--- a/Engine/Plugins/Experimental/Mutable/Source/MutableRuntime/Private/MuR/SerialisationPrivate.h
+++ b/Engine/Plugins/Experimental/Mutable/Source/MutableRuntime/Private/MuR/SerialisationPrivate.h
@@ -22,9 +22,7 @@
 
 
 namespace mu
-{
-	typedef std::string string;
-	
+{	
 
     //!
     class MUTABLERUNTIME_API InputMemoryStream : public InputStream
@@ -373,15 +371,15 @@ namespace mu
 	MUTABLE_IMPLEMENT_POD_VECTOR_SERIALISABLE(FVector4f);
 
 	//---------------------------------------------------------------------------------------------
-	template<>
-	inline void operator<< <std::string>( OutputArchive& arch, const std::string& v )
-	{
-        arch << (uint32)v.size();
-		if ( v.size() )
-		{
-			arch.GetPrivate()->m_pStream->Write( &v[0], (unsigned)v.size()*sizeof(char) );
-		}
-	}
+	//template<>
+	//inline void operator<< <std::string>( OutputArchive& arch, const std::string& v )
+	//{
+	//	arch << (uint32)v.size();
+	//	if ( v.size() )
+	//	{
+	//		arch.GetPrivate()->m_pStream->Write( &v[0], (unsigned)v.size()*sizeof(char) );
+	//	}
+	//}
 
 	template<>
 	inline void operator>> <std::string>( InputArchive& arch, std::string& v )
diff --git a/Engine/Plugins/Experimental/Mutable/Source/MutableRuntime/Private/MuR/Skeleton.cpp b/Engine/Plugins/Experimental/Mutable/Source/MutableRuntime/Private/MuR/Skeleton.cpp
--- a/Engine/Plugins/Experimental/Mutable/Source/MutableRuntime/Private/MuR/Skeleton.cpp
+++ b/Engine/Plugins/Experimental/Mutable/Source/MutableRuntime/Private/MuR/Skeleton.cpp
@@ -154,7 +154,7 @@ namespace mu
 		}
 		else
 		{
-			TArray<string> OldBoneNames;
+			TArray<std::string> OldBoneNames;
 			arch >> OldBoneNames;
 
 			const int32 NumBones = OldBoneNames.Num();
diff --git a/Engine/Plugins/Experimental/Mutable/Source/MutableRuntime/Private/MuR/System.cpp b/Engine/Plugins/Experimental/Mutable/Source/MutableRuntime/Private/MuR/System.cpp
--- a/Engine/Plugins/Experimental/Mutable/Source/MutableRuntime/Private/MuR/System.cpp
+++ b/Engine/Plugins/Experimental/Mutable/Source/MutableRuntime/Private/MuR/System.cpp
@@ -303,7 +303,7 @@ namespace mu
 		pLiveInstance->Instance = Result;
 		if (Result)
 		{
-			Result->GetPrivate()->m_id = pLiveInstance->InstanceID;
+			Result->GetPrivate()->Id = pLiveInstance->InstanceID;
 		}
 
 		m_pD->WorkingMemoryManager.EndRunnerThread();
diff --git a/Engine/Plugins/Experimental/Mutable/Source/MutableRuntime/Public/MuR/Instance.h b/Engine/Plugins/Experimental/Mutable/Source/MutableRuntime/Public/MuR/Instance.h
--- a/Engine/Plugins/Experimental/Mutable/Source/MutableRuntime/Public/MuR/Instance.h
+++ b/Engine/Plugins/Experimental/Mutable/Source/MutableRuntime/Public/MuR/Instance.h
@@ -33,7 +33,7 @@ namespace mu
 		Instance();
 
         //! Clone this instance
-        InstancePtr Clone() const;
+        Ptr<Instance> Clone() const;
 
 		// Resource interface
 		int32 GetDataSize() const override;
@@ -56,11 +56,6 @@ namespace mu
         //! Get the number of components in a level-of-detail
         //! \param lod Index of the level of detail, from 0 to GetLODCount()-1
         int32 GetComponentCount( int32 lod ) const;
-
-        //! Get the name of a component
-        //! \param lod Index of the level of detail, from 0 to GetLODCount()-1
-        //! \param comp Index of the component, from 0 to GetComponentCount(lod)-1
-        const char* GetComponentName( int32 lod, int32 comp ) const;
 		
 		//! Get the Id of a component
 		//! \param lod Index of the level of detail, from 0 to GetLODCount()-1
@@ -72,12 +67,6 @@ namespace mu
         //! \param comp Index of the component, from 0 to GetComponentCount(lod)-1
         int32 GetSurfaceCount( int32 lod, int32 comp ) const;
 
-        //! Get the name of a surface
-        //! \param lod Index of the level of detail, from 0 to GetLODCount()-1
-        //! \param comp Index of the component, from 0 to GetComponentCount(lod)-1
-        //! \param surf Index of the surface, from 0 to GetSurfaceCount(lod,comp)-1
-        const char* GetSurfaceName( int32 lod, int32 comp, int32 surf ) const;
-
         //! Get an id that can be used to match the surface data with the mesh surface data.
         //! \param lod Index of the level of detail, from 0 to GetLODCount()-1
         //! \param comp Index of the component, from 0 to GetComponentCount(lod)-1
@@ -90,7 +79,6 @@ namespace mu
         //! \param id ID of the surface to look for.
         int32 FindSurfaceById( int32 lod, int32 comp, uint32 id ) const;
 
-
 		//! Find the base surface index and Lod index when reusing surfaces between LODs. Return the surface index
 		//! and the LOD it belongs to.
 		//! \param Comp - Index of the component, from 0 to GetComponentCount(lod)-1
@@ -99,14 +87,12 @@ namespace mu
 		//! \param OutLODIndex - Index of the first LOD where the surface can be found. 
 		void FindBaseSurfaceBySharedId(int32 CompIndex, int32 SharedId, int32& OutSurfaceIndex, int32& OutLODIndex) const;
 
-		
 		//! Get an id that can be used to find the same surface on other LODs
 		//! \param lod Index of the level of detail, from 0 to GetLODCount()-1
 		//! \param comp Index of the component, from 0 to GetComponentCount(lod)-1
 		//! \param surf Index of the surface, from 0 to GetSurfaceCount(lod,comp)-1
 		int32 GetSharedSurfaceId(int32 lod, int32 comp, int32 surf) const;
 
-
         //! Get an optional, opaque application-defined identifier for this surface. The meaning of
         //! this ID depends on each application, and it is specified when creating the source data
         //! that generates this surface.
@@ -151,7 +137,7 @@ namespace mu
 		//! \param comp Index of the component, from 0 to GetComponentCount(lod)-1
         //! \param surf Index of the surface, from 0 to GetSurfaceCount(lod,comp)-1
         //! \param img Index of the image, from 0 to GetImageCount(lod,comp,surf)-1
-        const char* GetImageName( int32 lod, int32 comp, int32 surf, int32 img ) const;
+		FName GetImageName( int32 lod, int32 comp, int32 surf, int32 img ) const;
 
 		//! Get the number of vectors in a component
 		//! \param lod Index of the level of detail, from 0 to GetLODCount()
@@ -173,7 +159,7 @@ namespace mu
 		//! \param comp Index of the component, from 0 to GetComponentCount(lod)-1
         //! \param surf Index of the surface, from 0 to GetSurfaceCount(lod,comp)-1
         //! \param vec Index of the vector, from 0 to GetVectorCount(lod,comp)-1
-        const char* GetVectorName( int32 lod, int32 comp, int32 surf, int32 vec ) const;
+		FName GetVectorName( int32 lod, int32 comp, int32 surf, int32 vec ) const;
 
         //! Get the number of scalar values in a component
         //! \param lod Index of the level of detail, from 0 to GetLODCount()
@@ -193,7 +179,7 @@ namespace mu
         //! \param comp Index of the component, from 0 to GetComponentCount(lod)-1
         //! \param surf Index of the surface, from 0 to GetSurfaceCount(lod,comp)-1
         //! \param sca Index of the scalar, from 0 to GetScalarCount(lod,comp)-1
-        const char* GetScalarName( int32 lod, int32 comp, int32 surf, int32 sca ) const;
+		FName GetScalarName( int32 lod, int32 comp, int32 surf, int32 sca ) const;
 
         //! Get the number of string values in a component
         //! \param lod Index of the level of detail, from 0 to GetLODCount()
@@ -206,14 +192,14 @@ namespace mu
         //! \param comp Index of the component, from 0 to GetComponentCount(lod)-1
         //! \param surf Index of the surface, from 0 to GetSurfaceCount(lod,comp)-1
         //! \param sca Index of the string, from 0 to GetStringCount(lod,comp)-1
-        const char* GetString( int32 lod, int32 comp, int32 surf, int32 str ) const;
+        FString GetString( int32 lod, int32 comp, int32 surf, int32 str ) const;
 
         //! Get the name of a string from a component
         //! \param lod Index of the level of detail, from 0 to GetLODCount()-1
         //! \param comp Index of the component, from 0 to GetComponentCount(lod)-1
         //! \param surf Index of the surface, from 0 to GetSurfaceCount(lod,comp)-1
         //! \param sca Index of the string, from 0 to GetStringCount(lod,comp)-1
-        const char* GetStringName( int32 lod, int32 comp, int32 surf, int32 str ) const;
+        FName GetStringName( int32 lod, int32 comp, int32 surf, int32 str ) const;
 
 		//! Get the number of ExtensionData values in a component
 		int32 GetExtensionDataCount() const;
@@ -222,7 +208,7 @@ namespace mu
 		//! \param Index Index of the ExtensionData to fetch
 		//! \param OutExtensionData Receives the ExtensionData
 		//! \param OutName Receives the name associated with the ExtensionData. Guaranteed to be a valid string of non-zero length.
-		void GetExtensionData(int32 Index, ExtensionDataPtrConst& OutExtensionData, const char*& OutName) const;
+		void GetExtensionData(int32 Index, Ptr<const class ExtensionData>& OutExtensionData, FName& OutName) const;
 
         //-----------------------------------------------------------------------------------------
 		// Interface pattern
diff --git a/Engine/Plugins/Experimental/Mutable/Source/MutableRuntime/Public/MuR/Mesh.h b/Engine/Plugins/Experimental/Mutable/Source/MutableRuntime/Public/MuR/Mesh.h
--- a/Engine/Plugins/Experimental/Mutable/Source/MutableRuntime/Public/MuR/Mesh.h
+++ b/Engine/Plugins/Experimental/Mutable/Source/MutableRuntime/Public/MuR/Mesh.h
@@ -262,33 +262,6 @@ namespace mu
 
         //! \}
 
-        //! \name Face groups
-        //! \{
-
-        //!
-        void SetFaceGroupCount( int count );
-
-        //!
-        int GetFaceGroupCount() const;
-
-        //!
-        const char* GetFaceGroupName( int group ) const;
-
-        //!
-        void SetFaceGroupName( int group, const char* strName );
-
-        //!
-        int GetFaceGroupFaceCount( int group ) const;
-
-        //!
-        const int32* GetFaceGroupFaces( int group ) const;
-
-        //!
-        void SetFaceGroupFaces( int group, int count, const int32* faces );
-
-        //! \}
-
-
         //! \name Tags
         //! \{
 
@@ -299,10 +272,10 @@ namespace mu
         int GetTagCount() const;
 
         //!
-        const char* GetTag( int tagIndex ) const;
+        const FString& GetTag( int tagIndex ) const;
 
         //!
-        void SetTag( int tagIndex, const char* strName );
+        void SetTag( int tagIndex, const FString& Name );
 
 		//!
 		int32 FindBonePose(uint16 BoneId) const;
@@ -395,30 +368,10 @@ namespace mu
 
 		//! Texture Layout blocks attached to this mesh. They are const because they could be shared with
 		//! other meshes, so they need to be cloned and replaced if a modification is needed.
-		TArray<Ptr<const Layout>> m_layouts;
-
-		struct FACE_GROUP
-		{
-			string m_name;
-			TArray<int32> m_faces;
-
-			inline void Serialise(OutputArchive& arch) const;
-
-
-			inline void Unserialise(InputArchive& arch);
-
-			//!
-			inline bool operator==(const FACE_GROUP& o) const
-			{
-				return m_name == o.m_name
-					&& m_faces == o.m_faces;
-			}
-
-		};
-		TArray<FACE_GROUP> m_faceGroups;
+		TArray<Ptr<const Layout>> m_layouts;		
 
 		//!
-		TArray<string> m_tags;
+		TArray<FString> m_tags;
 
 		struct FBonePose
 		{
@@ -476,7 +429,6 @@ namespace mu
 				}
 			}
 			if (equal) equal = (m_surfaces == o.m_surfaces);
-			if (equal) equal = (m_faceGroups == o.m_faceGroups);
 			if (equal) equal = (m_tags == o.m_tags);
 			if (equal) equal = (SkeletonIDs == o.SkeletonIDs);
 
diff --git a/Engine/Plugins/Experimental/Mutable/Source/MutableRuntime/Public/MuR/MeshBufferSet.h b/Engine/Plugins/Experimental/Mutable/Source/MutableRuntime/Public/MuR/MeshBufferSet.h
--- a/Engine/Plugins/Experimental/Mutable/Source/MutableRuntime/Public/MuR/MeshBufferSet.h
+++ b/Engine/Plugins/Experimental/Mutable/Source/MutableRuntime/Public/MuR/MeshBufferSet.h
@@ -105,9 +105,7 @@ namespace mu
 		//! It will never be in the same buffer that other vertex semantics.
 		MBS_LAYOUTBLOCK,
 
-		//! Internal semantic indicating what chart does the given vertex or face belong to.
-		//! It can be safely ignored if present in meshes returned by the system.
-		MBS_CHART,
+		MBS_CHART_DEPRECATED,
 
 		//! To let users define channels with semantics unknown to the system.
 		//! These channels will never be transformed, and the per-vertex or per-index data will be
diff --git a/Engine/Plugins/Experimental/Mutable/Source/MutableRuntime/Public/MuR/Model.h b/Engine/Plugins/Experimental/Mutable/Source/MutableRuntime/Public/MuR/Model.h
--- a/Engine/Plugins/Experimental/Mutable/Source/MutableRuntime/Public/MuR/Model.h
+++ b/Engine/Plugins/Experimental/Mutable/Source/MutableRuntime/Public/MuR/Model.h
@@ -13,7 +13,7 @@
 //! data structures. Compiled models are not necessarily compatible when the runtime is updated,
 //! so this version number can be used externally to verify this. It is not used internally, and
 //! serializing models from different versions than this runtime will probably result in a crash.
-#define MUTABLE_COMPILED_MODEL_CODE_VERSION		uint32( 77 )
+#define MUTABLE_COMPILED_MODEL_CODE_VERSION		uint32( 79 )
 #define MUTABLE_PARAMETERS_VERSION              uint32( 2 )
 
 
@@ -82,10 +82,10 @@ namespace mu
 		int GetStateCount() const;
 
 		//! Get a state name by state index from 0 to GetStateCount-1
-		const char* GetStateName( int stateIndex ) const;
+		const FString& GetStateName( int32 StateIndex ) const;
 
 		//! Find a state index by state name
-		int FindState( const char* strName ) const;
+		int32 FindState( const FString& Name ) const;
 
 		//! Get the number of parameters available in a particular state.
 		int GetStateParameterCount( int stateIndex ) const;
diff --git a/Engine/Plugins/Experimental/Mutable/Source/MutableRuntime/Public/MuR/MutableString.h b/Engine/Plugins/Experimental/Mutable/Source/MutableRuntime/Public/MuR/MutableString.h
--- a/Engine/Plugins/Experimental/Mutable/Source/MutableRuntime/Public/MuR/MutableString.h
+++ b/Engine/Plugins/Experimental/Mutable/Source/MutableRuntime/Public/MuR/MutableString.h
@@ -22,14 +22,14 @@ namespace mu
 		// Life cycle
 		//-----------------------------------------------------------------------------------------
 
-		String(const char* = nullptr);
+		String(const FString&);
 
 		//! Deep clone this string.
 		Ptr<String> Clone() const;
 
 		//! Serialisation
-		static void Serialise( const String* p, OutputArchive& arch );
-		static Ptr<String> StaticUnserialise( InputArchive& arch );
+		//static void Serialise( const String* p, OutputArchive& arch );
+		//static Ptr<String> StaticUnserialise( InputArchive& arch );
 
 		// Resource interface
 		int32 GetDataSize() const override;
@@ -39,7 +39,7 @@ namespace mu
 		//-----------------------------------------------------------------------------------------
 
 		//! Get the string data.
-		const char* GetValue() const;
+		const FString& GetValue() const;
 
 	protected:
 
@@ -49,31 +49,40 @@ namespace mu
 	public:
 
 		//!
-		string m_value;
+		FString Value;
 
 		//!
-		void Serialise(OutputArchive& arch) const
-		{
-			uint32 ver = 0;
-			arch << ver;
-
-			arch << m_value;
-		}
-
-		//!
-		void Unserialise(InputArchive& arch)
-		{
-			uint32 ver;
-			arch >> ver;
-			check(ver <= 0);
-
-			arch >> m_value;
-		}
+		//void Serialise(OutputArchive& arch) const
+		//{
+		//	uint32 ver = 1;
+		//	arch << ver;
+
+		//	arch << Value;
+		//}
+
+		////!
+		//void Unserialise(InputArchive& arch)
+		//{
+		//	uint32 ver;
+		//	arch >> ver;
+		//	check(ver <= 1);
+
+		//	if (ver == 0)
+		//	{
+		//		std::string Temp;
+		//		arch >> Temp;
+		//		Value = Temp.c_str();
+		//	}
+		//	else
+		//	{
+		//		arch >> Value;
+		//	}
+		//}
 
 		//!
 		inline bool operator==(const String& o) const
 		{
-			return (m_value == o.m_value);
+			return (Value == o.Value);
 		}
 
 	};
diff --git a/Engine/Plugins/Experimental/Mutable/Source/MutableRuntime/Public/MuR/Parameters.h b/Engine/Plugins/Experimental/Mutable/Source/MutableRuntime/Public/MuR/Parameters.h
--- a/Engine/Plugins/Experimental/Mutable/Source/MutableRuntime/Public/MuR/Parameters.h
+++ b/Engine/Plugins/Experimental/Mutable/Source/MutableRuntime/Public/MuR/Parameters.h
@@ -64,22 +64,6 @@ namespace mu
     };
 
 
-    //! \brief Additional information about a particular parameter type.
-    //!
-	//! This information is not strictly relevant, but it can provide hints for dynamic user
-	//! interfaces.
-	enum class PARAMETER_DETAILED_TYPE : uint32
-	{
-		//!
-		UNKNOWN,
-
-		//! The parameter belongs to a triangular barycentric set of parameters
-		//! The parameter name can be used to identify which coordinate it is (ending number)
-		TRIANGLE
-
-	};
-
-
     //! \brief Types of projectors.
     //! \ingroup runtime
     enum class PROJECTOR_TYPE : uint32
@@ -122,11 +106,11 @@ namespace mu
 
         //! Return the name of a range.
         //! \param index Index of the range from 0 to GetRangeCount()-1
-        const char* GetRangeName( int index ) const;
+        const FString& GetRangeName( int index ) const;
 
         //! Return the Guid of the parameter, resistant to parameter name changes.
         //! \param index Index of the parameter from 0 to GetRangeCount()-1
-        const char* GetRangeUid( int index ) const;
+        const FString& GetRangeUid( int index ) const;
 
         //! Set the position in one of the dimensions of the index.
         //! \param index Index of the parameter from 0 to GetRangeCount()-1
@@ -185,8 +169,8 @@ namespace mu
         //! - parameters have been removed
         //! - parameter types have changed
         //! - parameter names have changed, and the application is using UIDs
-        static void SerialisePortable( const Parameters*, OutputArchive& );
-        static ParametersPtr UnserialisePortable( InputArchive&, TSharedPtr<const Model> );
+        //static void SerialisePortable( const Parameters*, OutputArchive& );
+        //static ParametersPtr UnserialisePortable( InputArchive&, TSharedPtr<const Model> );
 
 		//! Deep clone this object.
 		ParametersPtr Clone() const;
@@ -200,24 +184,20 @@ namespace mu
 
 		//! Return the name of the parameter.
 		//! \param index Index of the parameter from 0 to GetCount()-1
-		const char* GetName( int index ) const;
+		const FString& GetName( int index ) const;
 
 		//! Return the Guid of the parameter, resistant to parameter name changes.
 		//! \param index Index of the parameter from 0 to GetCount()-1
-		const char* GetUid( int index ) const;
+		const FString& GetUid( int index ) const;
 
 		//! Find the parameter index by name.
 		//! It returns -1 if the parameter doesn't exist.
-		int Find( const char* strName ) const;
+		int Find( const FString& Name ) const;
 
 		//! Return the type of the parameter.
 		//! \param index Index of the parameter from 0 to GetCount()-1
 		PARAMETER_TYPE GetType( int index ) const;
 
-		//! Return the additional type data of the parameter.
-		//! \param index Index of the parameter from 0 to GetCount()-1
-		PARAMETER_DETAILED_TYPE GetDetailedType( int index ) const;
-
         //! Create a new RangeIndex object to use to access a multi-dimensional parameter.
         //! It will return nullptr if the parameter is not multidimensional.
         Ptr<RangeIndex> NewRangeIndex( int paramIndex ) const;
@@ -259,11 +239,11 @@ namespace mu
 
         //! Get the name of one of the possible values for this integer.
         //! The paramIndex is in the range of 0 to GetIntPossibleValueCount()-1
-        const char* GetIntPossibleValueName( int paramIndex, int valueIndex ) const;
+        const FString& GetIntPossibleValueName( int paramIndex, int valueIndex ) const;
 
         //! Get the index of the value of one of the possible values for this integer.
         //! The paramIndex is in the range of 0 to GetIntPossibleValueCount()-1
-		int GetIntValueIndex(int paramIndex, const char* valueName) const;
+		int GetIntValueIndex(int paramIndex, const FString& valueName) const;
 
 		//! Get the index of the value of one of the possible values for this integer.
 		//! The paramIndex is in the range of 0 to GetIntPossibleValueCount()-1
@@ -363,13 +343,13 @@ namespace mu
         //! \pre The parameter specified by index is a T_FLOAT.
         //! \param index Index of the parameter from 0 to GetCount()-1
         //! \param pos Only for multidimensional parameters: relevant position to get in the ranges
-        const char* GetStringValue( int index, const Ptr<const RangeIndex>& pos = nullptr ) const;
+        void GetStringValue( int index, FString& OutValue, const Ptr<const RangeIndex>& pos = nullptr) const;
 
         //! If the parameter is of the float type, set its value.
         //! \param index Index of the parameter from 0 to GetCount()-1
         //! \param value new value of the parameter
         //! \param pos Only for multidimensional parameters: relevant position to set in the ranges
-        void SetStringValue( int index, const char* value, const Ptr<const RangeIndex>& pos = nullptr );
+        void SetStringValue( int index, const FString& Value, const Ptr<const RangeIndex>& pos = nullptr );
 
         //! Utility method to compare the values of a specific parameter with the values of another
         //! Parameters object. It returns false if type or values are different.
diff --git a/Engine/Plugins/Experimental/Mutable/Source/MutableRuntime/Public/MuR/PhysicsBody.h b/Engine/Plugins/Experimental/Mutable/Source/MutableRuntime/Public/MuR/PhysicsBody.h
--- a/Engine/Plugins/Experimental/Mutable/Source/MutableRuntime/Public/MuR/PhysicsBody.h
+++ b/Engine/Plugins/Experimental/Mutable/Source/MutableRuntime/Public/MuR/PhysicsBody.h
@@ -19,11 +19,9 @@
 
 namespace mu
 {
-	typedef std::string string;
-
 	struct FBodyShape
 	{
-		string Name;
+		FString Name;
 		uint32 Flags = 0;
 
 		bool operator==(const FBodyShape& Other) const 
@@ -281,11 +279,11 @@ namespace mu
 		uint32 GetSphylFlags(int32 B, int32 I) const;
 		uint32 GetTaperedCapsuleFlags(int32 B, int32 I) const;
 
-		const char* GetSphereName(int32 B, int32 I) const;
-		const char* GetBoxName(int32 B, int32 I) const;
-		const char* GetConvexName(int32 B, int32 I) const;
-		const char* GetSphylName(int32 B, int32 I) const;
-		const char* GetTaperedCapsuleName(int32 B, int32 I) const;
+		const FString& GetSphereName(int32 B, int32 I) const;
+		const FString& GetBoxName(int32 B, int32 I) const;
+		const FString& GetConvexName(int32 B, int32 I) const;
+		const FString& GetSphylName(int32 B, int32 I) const;
+		const FString& GetTaperedCapsuleName(int32 B, int32 I) const;
 
 	protected:
 		//! Forbidden. Manage with the Ptr<> template.
diff --git a/Engine/Plugins/Experimental/Mutable/Source/MutableRuntime/Public/MuR/Skeleton.h b/Engine/Plugins/Experimental/Mutable/Source/MutableRuntime/Public/MuR/Skeleton.h
--- a/Engine/Plugins/Experimental/Mutable/Source/MutableRuntime/Public/MuR/Skeleton.h
+++ b/Engine/Plugins/Experimental/Mutable/Source/MutableRuntime/Public/MuR/Skeleton.h
@@ -11,13 +11,12 @@
 #include "MuR/Ptr.h"
 #include "MuR/RefCounted.h"
 
+// Remove when removin deprecated data
 #include <string>
 
 
 namespace mu
-{
-	typedef std::string string;
-	
+{	
 	// Forward references
     class Skeleton;
 
@@ -77,7 +76,7 @@ namespace mu
 	public:
 
 		//! Deprecated
-		TArray<string> m_bones_DEPRECATED;
+		TArray<std::string> m_bones_DEPRECATED;
 		TArray<FTransform3f> m_boneTransforms_DEPRECATED;
 
 
diff --git a/Engine/Plugins/Experimental/Mutable/Source/MutableTools/Private/MuT/AST.h b/Engine/Plugins/Experimental/Mutable/Source/MutableTools/Private/MuT/AST.h
--- a/Engine/Plugins/Experimental/Mutable/Source/MutableTools/Private/MuT/AST.h
+++ b/Engine/Plugins/Experimental/Mutable/Source/MutableTools/Private/MuT/AST.h
@@ -515,13 +515,13 @@ namespace mu
             ASTChild rangeSize;
 
             //!
-            string rangeName;
+            FString rangeName;
 
             //!
-            string rangeUID;
+			FString rangeUID;
 
             //!
-            FRangeData( ASTOp* parentOp, Ptr<ASTOp> childOp, const string& name, const string& uid )
+            FRangeData( ASTOp* parentOp, Ptr<ASTOp> childOp, const FString& name, const FString& uid )
                 : rangeSize( parentOp, childOp )
                 , rangeName(name)
                 , rangeUID(uid)
diff --git a/Engine/Plugins/Experimental/Mutable/Source/MutableTools/Private/MuT/ASTOpAddExtensionData.cpp b/Engine/Plugins/Experimental/Mutable/Source/MutableTools/Private/MuT/ASTOpAddExtensionData.cpp
--- a/Engine/Plugins/Experimental/Mutable/Source/MutableTools/Private/MuT/ASTOpAddExtensionData.cpp
+++ b/Engine/Plugins/Experimental/Mutable/Source/MutableTools/Private/MuT/ASTOpAddExtensionData.cpp
@@ -47,7 +47,6 @@ uint64 ASTOpAddExtensionData::Hash() const
 	
 	hash_combine(Result, Instance.child().get());
 	hash_combine(Result, ExtensionData.child().get());
-	hash_combine(Result, ExtensionDataName);
 	
 	return Result;
 }
@@ -86,7 +85,7 @@ void ASTOpAddExtensionData::Link(FProgram& Program, FLinkerOptions*)
 	check(ExtensionData->linkedAddress);
 	Args.ExtensionData = ExtensionData->linkedAddress;
 
-	check(ExtensionDataName.length() > 0);
+	check(ExtensionDataName.Len() > 0);
 	Args.ExtensionDataName = Program.AddConstant(ExtensionDataName);
 
 	linkedAddress = (OP::ADDRESS)Program.m_opAddress.Num();
diff --git a/Engine/Plugins/Experimental/Mutable/Source/MutableTools/Private/MuT/ASTOpAddExtensionData.h b/Engine/Plugins/Experimental/Mutable/Source/MutableTools/Private/MuT/ASTOpAddExtensionData.h
--- a/Engine/Plugins/Experimental/Mutable/Source/MutableTools/Private/MuT/ASTOpAddExtensionData.h
+++ b/Engine/Plugins/Experimental/Mutable/Source/MutableTools/Private/MuT/ASTOpAddExtensionData.h
@@ -23,7 +23,7 @@ public:
 
 	ASTChild Instance;
 	ASTChild ExtensionData;
-	string ExtensionDataName;
+	FString ExtensionDataName;
 
 public:
 
diff --git a/Engine/Plugins/Experimental/Mutable/Source/MutableTools/Private/MuT/ASTOpConstantString.cpp b/Engine/Plugins/Experimental/Mutable/Source/MutableTools/Private/MuT/ASTOpConstantString.cpp
--- a/Engine/Plugins/Experimental/Mutable/Source/MutableTools/Private/MuT/ASTOpConstantString.cpp
+++ b/Engine/Plugins/Experimental/Mutable/Source/MutableTools/Private/MuT/ASTOpConstantString.cpp
@@ -24,7 +24,7 @@ namespace mu
 
 	uint64 ASTOpConstantString::Hash() const
 	{
-		uint64 res = std::hash<std::string>()(value.c_str());
+		uint64 res = std::hash<int32>()(value.Len());
 		return res;
 	}
 
@@ -52,7 +52,7 @@ namespace mu
 		if (!linkedAddress)
 		{
 			OP::ResourceConstantArgs args;
-			memset(&args, 0, sizeof(args));
+			FMemory::Memset(&args, 0, sizeof(args));
 			args.value = program.AddConstant(value);
 
 			linkedAddress = (OP::ADDRESS)program.m_opAddress.Num();
diff --git a/Engine/Plugins/Experimental/Mutable/Source/MutableTools/Private/MuT/ASTOpConstantString.h b/Engine/Plugins/Experimental/Mutable/Source/MutableTools/Private/MuT/ASTOpConstantString.h
--- a/Engine/Plugins/Experimental/Mutable/Source/MutableTools/Private/MuT/ASTOpConstantString.h
+++ b/Engine/Plugins/Experimental/Mutable/Source/MutableTools/Private/MuT/ASTOpConstantString.h
@@ -20,7 +20,7 @@ struct FProgram;
 	{
 	public:
 		//!
-		string value;
+		FString value;
 
 	public:
 		OP_TYPE GetOpType() const override { return OP_TYPE::ST_CONSTANT; }
diff --git a/Engine/Plugins/Experimental/Mutable/Source/MutableTools/Private/MuT/ASTOpImageMultiLayer.cpp b/Engine/Plugins/Experimental/Mutable/Source/MutableTools/Private/MuT/ASTOpImageMultiLayer.cpp
--- a/Engine/Plugins/Experimental/Mutable/Source/MutableTools/Private/MuT/ASTOpImageMultiLayer.cpp
+++ b/Engine/Plugins/Experimental/Mutable/Source/MutableTools/Private/MuT/ASTOpImageMultiLayer.cpp
@@ -21,7 +21,7 @@ namespace mu
 		: base(this)
 		, blend(this)
 		, mask(this)
-		, range(this, nullptr, string(), string())
+		, range(this, nullptr, FString(), FString())
 	{
 	}
 
diff --git a/Engine/Plugins/Experimental/Mutable/Source/MutableTools/Private/MuT/ASTOpInstanceAdd.cpp b/Engine/Plugins/Experimental/Mutable/Source/MutableTools/Private/MuT/ASTOpInstanceAdd.cpp
--- a/Engine/Plugins/Experimental/Mutable/Source/MutableTools/Private/MuT/ASTOpInstanceAdd.cpp
+++ b/Engine/Plugins/Experimental/Mutable/Source/MutableTools/Private/MuT/ASTOpInstanceAdd.cpp
@@ -116,7 +116,7 @@ namespace mu
 		if (!linkedAddress)
 		{
 			OP::InstanceAddArgs args;
-			memset(&args, 0, sizeof(args));
+			FMemory::Memzero(&args, sizeof(args));
 			args.id = id;
 			args.ExternalId = ExternalId;
 			args.SharedSurfaceId = SharedSurfaceId;
@@ -134,7 +134,7 @@ namespace mu
 				visitor.Run(value.child());
 
 				TArray<uint16> params;
-				for (const string& paramName : visitor.m_params)
+				for (const FString& paramName : visitor.m_params)
 				{
 					for (int32 i = 0; i < program.m_parameters.Num(); ++i)
 					{
diff --git a/Engine/Plugins/Experimental/Mutable/Source/MutableTools/Private/MuT/ASTOpInstanceAdd.h b/Engine/Plugins/Experimental/Mutable/Source/MutableTools/Private/MuT/ASTOpInstanceAdd.h
--- a/Engine/Plugins/Experimental/Mutable/Source/MutableTools/Private/MuT/ASTOpInstanceAdd.h
+++ b/Engine/Plugins/Experimental/Mutable/Source/MutableTools/Private/MuT/ASTOpInstanceAdd.h
@@ -28,7 +28,7 @@ struct FProgram;
 		uint32_t id = 0;
 		uint32_t ExternalId = 0;
 		int32 SharedSurfaceId = INDEX_NONE;
-		string name;
+		FString name;
 
 	public:
 
diff --git a/Engine/Plugins/Experimental/Mutable/Source/MutableTools/Private/MuT/ASTOpMeshExtractLayoutBlocks.cpp b/Engine/Plugins/Experimental/Mutable/Source/MutableTools/Private/MuT/ASTOpMeshExtractLayoutBlocks.cpp
--- a/Engine/Plugins/Experimental/Mutable/Source/MutableTools/Private/MuT/ASTOpMeshExtractLayoutBlocks.cpp
+++ b/Engine/Plugins/Experimental/Mutable/Source/MutableTools/Private/MuT/ASTOpMeshExtractLayoutBlocks.cpp
@@ -2,6 +2,8 @@
 
 #include "MuT/ASTOpMeshExtractLayoutBlocks.h"
 
+#include "MuT/ASTOpSwitch.h"
+#include "MuT/ASTOpConditional.h"
 #include "Misc/AssertionMacros.h"
 #include "MuR/ModelPrivate.h"
 #include "MuR/RefCounted.h"
@@ -17,7 +19,7 @@ namespace mu
 	//---------------------------------------------------------------------------------------------
 	//---------------------------------------------------------------------------------------------
 	ASTOpMeshExtractLayoutBlocks::ASTOpMeshExtractLayoutBlocks()
-		: source(this)
+		: Source(this)
 	{
 	}
 
@@ -33,7 +35,7 @@ namespace mu
 	{
 		if (auto other = dynamic_cast<const ASTOpMeshExtractLayoutBlocks*>(&otherUntyped))
 		{
-			return source == other->source && layout == other->layout && blocks == other->blocks;
+			return Source == other->Source && Layout == other->Layout && Blocks == other->Blocks;
 		}
 		return false;
 	}
@@ -42,29 +44,29 @@ namespace mu
 	mu::Ptr<ASTOp> ASTOpMeshExtractLayoutBlocks::Clone(MapChildFuncRef mapChild) const
 	{
 		Ptr<ASTOpMeshExtractLayoutBlocks> n = new ASTOpMeshExtractLayoutBlocks();
-		n->source = mapChild(source.child());
-		n->layout = layout;
-		n->blocks = blocks;
+		n->Source = mapChild(Source.child());
+		n->Layout = Layout;
+		n->Blocks = Blocks;
 		return n;
 	}
 
 
 	void ASTOpMeshExtractLayoutBlocks::Assert()
 	{
-		check(blocks.Num() < std::numeric_limits<uint16>::max());
+		check(Blocks.Num() < std::numeric_limits<uint16>::max());
 		ASTOp::Assert();
 	}
 
 
 	void ASTOpMeshExtractLayoutBlocks::ForEachChild(const TFunctionRef<void(ASTChild&)> f)
 	{
-		f(source);
+		f(Source);
 	}
 
 
 	uint64 ASTOpMeshExtractLayoutBlocks::Hash() const
 	{
-		uint64 res = std::hash<size_t>()(size_t(source.child().get()));
+		uint64 res = std::hash<size_t>()(size_t(Source.child().get()));
 		return res;
 	}
 
@@ -78,16 +80,90 @@ namespace mu
 
 			program.m_opAddress.Add((uint32)program.m_byteCode.Num());
 			AppendCode(program.m_byteCode, OP_TYPE::ME_EXTRACTLAYOUTBLOCK);
-			OP::ADDRESS sourceAt = source ? source->linkedAddress : 0;
+			OP::ADDRESS sourceAt = Source ? Source->linkedAddress : 0;
 			AppendCode(program.m_byteCode, sourceAt);
-			AppendCode(program.m_byteCode, (uint16)layout);
-			AppendCode(program.m_byteCode, (uint16)blocks.Num());
+			AppendCode(program.m_byteCode, (uint16)Layout);
+			AppendCode(program.m_byteCode, (uint16)Blocks.Num());
 
-			for (auto b : blocks)
+			for (auto b : Blocks)
 			{
 				AppendCode(program.m_byteCode, (uint32)b);
 			}
 		}
 	}
 
+
+	mu::Ptr<ASTOp> ASTOpMeshExtractLayoutBlocks::OptimiseSink(const FModelOptimizationOptions&, FOptimizeSinkContext&) const
+	{
+		Ptr<ASTOp> NewOp;
+
+		if (!Source.child())
+		{
+			return nullptr;
+		}
+
+		OP_TYPE SourceType = Source.child()->GetOpType();
+
+		// Optimize only the mesh parameter
+		switch (SourceType)
+		{
+
+		case OP_TYPE::ME_SWITCH:
+		{
+			// Move the operation down all the paths
+			Ptr<ASTOpSwitch> NewSwitch = mu::Clone<ASTOpSwitch>(Source.child());
+
+			if (NewSwitch->def)
+			{
+				Ptr<ASTOpMeshExtractLayoutBlocks> NewBind = mu::Clone<ASTOpMeshExtractLayoutBlocks>(this);
+				NewBind->Source = NewSwitch->def.child();
+				NewSwitch->def = NewBind;
+			}
+
+			for (int32 v = 0; v < NewSwitch->cases.Num(); ++v)
+			{
+				if (NewSwitch->cases[v].branch)
+				{
+					Ptr<ASTOpMeshExtractLayoutBlocks> NewBind = mu::Clone<ASTOpMeshExtractLayoutBlocks>(this);
+					NewBind->Source = NewSwitch->cases[v].branch.child();
+					NewSwitch->cases[v].branch = NewBind;
+				}
+			}
+
+			NewOp = NewSwitch;
+			break;
+		}
+
+		case OP_TYPE::ME_CONDITIONAL:
+		{
+			// Move the operation down all the paths
+			Ptr<ASTOpConditional> NewConditional = mu::Clone<ASTOpConditional>(Source.child());
+
+			if (NewConditional->yes)
+			{
+				Ptr<ASTOpMeshExtractLayoutBlocks> NewBind = mu::Clone<ASTOpMeshExtractLayoutBlocks>(this);
+				NewBind->Source = NewConditional->yes.child();
+				NewConditional->yes = NewBind;
+			}
+
+			if (NewConditional->no)
+			{
+				Ptr<ASTOpMeshExtractLayoutBlocks> NewBind = mu::Clone<ASTOpMeshExtractLayoutBlocks>(this);
+				NewBind->Source = NewConditional->no.child();
+				NewConditional->no = NewBind;
+			}
+
+			NewOp = NewConditional;
+			break;
+		}
+
+		default:
+			break;
+
+		}
+
+		return NewOp;
+	}
+
+
 }
diff --git a/Engine/Plugins/Experimental/Mutable/Source/MutableTools/Private/MuT/ASTOpMeshExtractLayoutBlocks.h b/Engine/Plugins/Experimental/Mutable/Source/MutableTools/Private/MuT/ASTOpMeshExtractLayoutBlocks.h
--- a/Engine/Plugins/Experimental/Mutable/Source/MutableTools/Private/MuT/ASTOpMeshExtractLayoutBlocks.h
+++ b/Engine/Plugins/Experimental/Mutable/Source/MutableTools/Private/MuT/ASTOpMeshExtractLayoutBlocks.h
@@ -10,7 +10,7 @@
 
 namespace mu
 {
-struct FProgram;
+	struct FProgram;
 
 
 	//---------------------------------------------------------------------------------------------
@@ -20,14 +20,14 @@ struct FProgram;
 	{
 	public:
 
-		//! Source mesh to extract block from.
-		ASTChild source;
+		/** Source mesh to extract block from. */
+		ASTChild Source;
 
-		//! Layout to use to select the blocks.
-		uint16 layout = 0;
+		/** Layout to use to select the blocks. */
+		uint16 Layout = 0;
 
-		//! Blocks
-		TArray<uint32> blocks;
+		/** Blocks to include in the resulting mesh. */
+		TArray<uint32> Blocks;
 
 	public:
 
@@ -43,6 +43,7 @@ struct FProgram;
 		Ptr<ASTOp> Clone(MapChildFuncRef mapChild) const override;
 		void Assert() override;
 		void Link(FProgram& program, FLinkerOptions* Options) override;
+		mu::Ptr<ASTOp> OptimiseSink(const FModelOptimizationOptions&, FOptimizeSinkContext&) const;
 
 	};
 
diff --git a/Engine/Plugins/Experimental/Mutable/Source/MutableTools/Private/MuT/ASTOpParameter.cpp b/Engine/Plugins/Experimental/Mutable/Source/MutableTools/Private/MuT/ASTOpParameter.cpp
--- a/Engine/Plugins/Experimental/Mutable/Source/MutableTools/Private/MuT/ASTOpParameter.cpp
+++ b/Engine/Plugins/Experimental/Mutable/Source/MutableTools/Private/MuT/ASTOpParameter.cpp
@@ -63,7 +63,7 @@ namespace mu
 	{
 		uint64 res = std::hash<uint64>()(uint64(type));
 		hash_combine(res, parameter.m_type);
-		hash_combine(res, parameter.m_name);
+		hash_combine(res, parameter.m_name.Len());
 		return res;
 	}
 
diff --git a/Engine/Plugins/Experimental/Mutable/Source/MutableTools/Private/MuT/ASTOpSwitch.cpp b/Engine/Plugins/Experimental/Mutable/Source/MutableTools/Private/MuT/ASTOpSwitch.cpp
--- a/Engine/Plugins/Experimental/Mutable/Source/MutableTools/Private/MuT/ASTOpSwitch.cpp
+++ b/Engine/Plugins/Experimental/Mutable/Source/MutableTools/Private/MuT/ASTOpSwitch.cpp
@@ -445,7 +445,7 @@ namespace mu
 			bool bFirstValue = true;
 			bool bAllSame = true;
 			Ptr<ASTOp> SameBranch = nullptr;
-			for (const FParameterDesc::INT_VALUE_DESC& Value : ParamOp->parameter.m_possibleValues)
+			for (const FParameterDesc::FIntValueDesc& Value : ParamOp->parameter.m_possibleValues)
 			{
 				// Look for the switch branch it would take
 				Ptr<ASTOp> Branch = def.child();
diff --git a/Engine/Plugins/Experimental/Mutable/Source/MutableTools/Private/MuT/CodeGenerator.cpp b/Engine/Plugins/Experimental/Mutable/Source/MutableTools/Private/MuT/CodeGenerator.cpp
--- a/Engine/Plugins/Experimental/Mutable/Source/MutableTools/Private/MuT/CodeGenerator.cpp
+++ b/Engine/Plugins/Experimental/Mutable/Source/MutableTools/Private/MuT/CodeGenerator.cpp
@@ -148,7 +148,7 @@ namespace mu
 		m_generatedRanges.Reset();
 		m_generatedTables.clear();
 		m_firstPass = FirstPassGenerator();
-		m_currentBottomUpState = BOTTOM_UP_STATE();
+		m_currentBottomUpState = FBottomUpState();
 		m_currentParents.Empty();
 		m_currentObject.Empty();
 		m_additionalComponents.clear();
@@ -204,7 +204,7 @@ namespace mu
 		{
 			// This happens only if we generate a node graph that has a NodeSurfaceNew at the root.
 			FSurfaceGenerationResult surfResult;
-			const TArray<FirstPassGenerator::SURFACE::EDIT> edits;
+			const TArray<FirstPassGenerator::FSurface::FEdit> edits;
 			GenerateSurface(surfResult, surfNode, edits);
 			return surfResult.surfaceOp;
 		}
@@ -296,13 +296,13 @@ namespace mu
 
 
 	//---------------------------------------------------------------------------------------------
-	Ptr<ASTOp> CodeGenerator::GenerateTableVariable(TablePtr pTable, const string& strName)
+	Ptr<ASTOp> CodeGenerator::GenerateTableVariable(TablePtr pTable, const FString& strName)
 	{
 		Ptr<ASTOp> result;
 
         FParameterDesc param;
         param.m_name = strName;
-        if ( param.m_name.size()==0 )
+        if ( param.m_name.Len()==0 )
         {
             param.m_name = TCHAR_TO_ANSI(*pTable->GetName());
         }
@@ -325,7 +325,7 @@ namespace mu
 
 			if (pTable->GetPrivate()->bNoneOption)
 			{
-				FParameterDesc::INT_VALUE_DESC nullValue;
+				FParameterDesc::FIntValueDesc nullValue;
 				nullValue.m_value = -1;
 				nullValue.m_name = "None";
 				param.m_possibleValues.Add(nullValue);
@@ -336,7 +336,7 @@ namespace mu
 			int32 rows = pTable->GetPrivate()->Rows.Num();
 			for (size_t i = 0; i < rows; ++i)
 			{
-				FParameterDesc::INT_VALUE_DESC value;
+				FParameterDesc::FIntValueDesc value;
 				value.m_value = (int16_t)pTable->GetPrivate()->Rows[i].Id;
 
 				if (nameCol > -1)
@@ -584,7 +584,7 @@ namespace mu
 		if (!pNodeLayout)
 		{
 			FString Msg = FString::Printf(TEXT("In object [%s] NodePatchImage couldn't find the layout in parent."),
-				m_currentParents.Last().m_pObject->m_name.c_str()
+				*m_currentParents.Last().m_pObject->m_name
 			);
 
 			m_pErrorLog->GetPrivate()->Add(Msg, ELMT_ERROR, node.m_errorContext);
@@ -625,7 +625,7 @@ namespace mu
 		}
 
 		// Add components from child objects
-		ADDITIONAL_COMPONENT_KEY thisKey;
+		FAdditionalComponentKey thisKey;
 		thisKey.m_lod = m_currentParents.Last().m_lod;
 		thisKey.m_pObject = m_currentParents.Last().m_pObject;
 		auto addIt = m_additionalComponents.find(thisKey);
@@ -658,7 +658,7 @@ namespace mu
 		if (lastCompOp && m_currentParents.Num() > 2)
 		{
 			const auto& parentObjectKey = m_currentParents[m_currentParents.Num() - 2];
-			ADDITIONAL_COMPONENT_KEY parentKey;
+			FAdditionalComponentKey parentKey;
 			parentKey.m_lod = m_currentParents.Last().m_lod;
 			parentKey.m_pObject = parentObjectKey.m_pObject;
 			m_additionalComponents[parentKey].Add(lastCompOp);
@@ -729,7 +729,7 @@ namespace mu
     //---------------------------------------------------------------------------------------------
     void CodeGenerator::GenerateSurface( FSurfaceGenerationResult& result,
                                          NodeSurfaceNewPtrConst surfaceNode,
-                                         const TArray<FirstPassGenerator::SURFACE::EDIT>& edits )
+                                         const TArray<FirstPassGenerator::FSurface::FEdit>& edits )
     {
         //MUTABLE_CPUPROFILER_SCOPE(GenerateSurface);
 
@@ -767,8 +767,8 @@ namespace mu
 
 		// Do we need to generate the mesh? Or was it already generated for state conditions 
 		// accepting the current state?
-		FirstPassGenerator::SURFACE* targetSurface = nullptr;
-		for (FirstPassGenerator::SURFACE& its : m_firstPass.surfaces)
+		FirstPassGenerator::FSurface* targetSurface = nullptr;
+		for (FirstPassGenerator::FSurface& its : m_firstPass.surfaces)
 		{
             if (its.node->GetPrivate() == &node)
             {
@@ -856,7 +856,7 @@ namespace mu
             // Apply mesh merges from child objects "edit surface" nodes
             for ( int32 editIndex=0; editIndex<edits.Num(); ++editIndex )
             {
-                const FirstPassGenerator::SURFACE::EDIT& e = edits[editIndex];
+                const FirstPassGenerator::FSurface::FEdit& e = edits[editIndex];
 
                 if ( e.node->m_pMesh )
                 {
@@ -928,7 +928,7 @@ namespace mu
             // and not the base.
 			// \TODO: Apply base removes first, and then "added meshes" removes here. It may have lower memory footprint during generation.
             Ptr<ASTOpMeshRemoveMask> rop;
-            for ( const FirstPassGenerator::SURFACE::EDIT& e: edits )
+            for ( const FirstPassGenerator::FSurface::FEdit& e: edits )
             {
                 if ( e.node->m_pMesh )
                 {
@@ -978,7 +978,7 @@ namespace mu
             }
 
             // Apply mesh morphs from child objects "edit surface" nodes
-            for ( const FirstPassGenerator::SURFACE::EDIT& e: edits )
+            for ( const FirstPassGenerator::FSurface::FEdit& e: edits )
             {
                 if ( NodeMeshPtr pMorph = e.node->m_pMorph )
                 {
@@ -1240,9 +1240,9 @@ namespace mu
 					{
 						int currentLOD = m_currentParents.Last().m_lod;
 						FString Msg = FString::Printf( TEXT("An image block for [%s] [%s] [%s] at lod [%d] has zero size and will not be generated. "),
-							node.m_images[t].m_name.c_str(),
-							node.m_images[t].m_materialName.c_str(),
-							node.m_images[t].m_materialParameterName.c_str(),
+							*node.m_images[t].m_name,
+							*node.m_images[t].m_materialName,
+							*node.m_images[t].m_materialParameterName,
 							currentLOD
 						);
 						m_pErrorLog->GetPrivate()->Add(Msg, ELMT_INFO, node.m_errorContext);
@@ -1250,7 +1250,7 @@ namespace mu
 
 					else if (desc.m_format == EImageFormat::IF_NONE)
 					{
-						FString Msg = FString::Printf(TEXT("An image [%s] has an unidentified pixel format. "), node.m_images[t].m_name.c_str());
+						FString Msg = FString::Printf(TEXT("An image [%s] has an unidentified pixel format. "), *node.m_images[t].m_name);
 						m_pErrorLog->GetPrivate()->Add(Msg, ELMT_ERROR, node.m_errorContext);
 					}
 
@@ -1271,7 +1271,7 @@ namespace mu
 						// Look for patches to this block
 						for (int32 editIndex = 0; editIndex < edits.Num(); ++editIndex)
 						{
-							const FirstPassGenerator::SURFACE::EDIT& e = edits[editIndex];
+							const FirstPassGenerator::FSurface::FEdit& e = edits[editIndex];
 							if (t < e.node->m_textures.Num())
 							{
 								if (const NodePatchImage* pPatch = e.node->m_textures[t].m_pPatch.get())
@@ -1351,9 +1351,9 @@ namespace mu
 
 								int currentLOD = m_currentParents.Last().m_lod;
 								FString Msg = FString::Printf( TEXT("A texture [%s] for material [%s] parameter [%s] in LOD [%d] has been resized from [%d x %d] to [%d x %d] because it didn't fit the layout [%d x %d]. "),
-									node.m_images[t].m_name.c_str(),
-									node.m_images[t].m_materialName.c_str(),
-									node.m_images[t].m_materialParameterName.c_str(),
+									*node.m_images[t].m_name,
+									*node.m_images[t].m_materialName,
+									*node.m_images[t].m_materialParameterName,
 									currentLOD,
 									oldSize[0], oldSize[1], desc.m_size[0], desc.m_size[1], grid[0], grid[1]);
 								m_pErrorLog->GetPrivate()->Add(Msg, ELMT_INFO, node.m_errorContext);
@@ -1421,7 +1421,7 @@ namespace mu
 								// Look for patches to this block
 								for (int32 editIndex = 0; editIndex < edits.Num(); ++editIndex)
 								{
-									const FirstPassGenerator::SURFACE::EDIT& e = edits[editIndex];
+									const FirstPassGenerator::FSurface::FEdit& e = edits[editIndex];
 									if (t < e.node->m_textures.Num())
 									{
 										if (const NodePatchImage* pPatch = e.node->m_textures[t].m_pPatch.get())
@@ -1458,7 +1458,7 @@ namespace mu
 							// Apply composition of blocks coming from child objects
 							for (int32 editIndex = 0; editIndex < edits.Num(); ++editIndex)
 							{
-								const FirstPassGenerator::SURFACE::EDIT& e = edits[editIndex];
+								const FirstPassGenerator::FSurface::FEdit& e = edits[editIndex];
 								if (t < e.node->m_textures.Num())
 								{
 									Ptr<NodeImage> pExtend = e.node->m_textures[t].m_pExtend;
@@ -1469,7 +1469,7 @@ namespace mu
 											!meshResults.extraMeshLayouts[editIndex].GeneratedLayouts[LayoutIndex])
 										{
 											FString Msg = FString::Printf(TEXT("Trying to extend a layout that doesn't exist in object [%s]."),
-												m_currentParents.Last().m_pObject->m_name.c_str()
+												*m_currentParents.Last().m_pObject->m_name
 											);
 
 											m_pErrorLog->GetPrivate()->Add(Msg, ELMT_ERROR, node.m_errorContext);
@@ -1633,7 +1633,7 @@ namespace mu
 							op->value = imageAd;
 
 							// Name
-							op->name = node.m_images[t].m_name.c_str();
+							op->name = node.m_images[t].m_name;
 
 							lastSurfOp = op;
 						}
@@ -1734,7 +1734,7 @@ namespace mu
 		// Build a series of operations to assemble the component
         Ptr<ASTOp> lastCompOp;
         Ptr<ASTOp> lastMeshOp;
-        string lastMeshName;
+        FString lastMeshName;
 
         // This generates a different ID for each surface. It can be used to match it to the
         // mesh surface, or for debugging. It cannot be 0 because it is a special case for the
@@ -1744,7 +1744,7 @@ namespace mu
         // Look for all surfaces that belong to this component
 		for (int32 i = 0; i<m_firstPass.surfaces.Num(); ++i, ++surfaceID)
 		{
-			const FirstPassGenerator::SURFACE& its = m_firstPass.surfaces[i];
+			const FirstPassGenerator::FSurface& its = m_firstPass.surfaces[i];
 			if (its.component==&node)
 			{
                 // Apply state conditions: only generate it if it enabled in this state
@@ -1947,7 +1947,7 @@ namespace mu
 			}
 
 			// Name must be valid
-			check(NamedNode.Name.length() > 0);
+			check(NamedNode.Name.Len() > 0);
 
 			FExtensionDataGenerationResult Result;
 			GenerateExtensionData(Result, NamedNode.Node);
@@ -2002,7 +2002,7 @@ namespace mu
     //---------------------------------------------------------------------------------------------
     Ptr<ASTOp> CodeGenerator::Visit( const NodeObjectGroup::Private& node )
     {
-		vector<string> usedNames;
+		TArray<FString> usedNames;
 
         // Parse the child objects first, which will accumulate operations in the patching lists
         for ( int32 t=0; t<node.m_children.Num(); ++t )
@@ -2014,7 +2014,7 @@ namespace mu
 				bool found = false;
                 for( int32 i = 0; !found && i != m_firstPass.objects.Num(); i++ )
 				{
-					FirstPassGenerator::OBJECT& it = m_firstPass.objects[i];
+					FirstPassGenerator::FObject& it = m_firstPass.objects[i];
 					if (it.node == pChildNode->GetBasePrivate())
 					{
 						found = true;
@@ -2038,19 +2038,15 @@ namespace mu
                 m_currentObject.Pop();
 
 				// Check for duplicated child names
-				const char* strChildName = pChildNode->GetName();
-				if (std::find(usedNames.begin(), usedNames.end(), strChildName)
-					!=
-					usedNames.end() )
+				FString strChildName = pChildNode->GetName();
+				if (usedNames.Contains(strChildName))
 				{
-					FString Msg = FString::Printf(TEXT("Object group has more than one children with the same name [%s]."),
-						strChildName
-					);
+					FString Msg = FString::Printf(TEXT("Object group has more than one children with the same name [%s]."), *strChildName );
 					m_pErrorLog->GetPrivate()->Add(Msg, ELMT_WARNING, node.m_errorContext);
 				}
 				else
 				{
-					usedNames.push_back(strChildName);
+					usedNames.Add(strChildName);
 				}
             }
         }
@@ -2079,15 +2075,15 @@ namespace mu
 
 	//---------------------------------------------------------------------------------------------
 	void CodeGenerator::GetModifiersFor(
-		const TArray<string>& tags,
+		const TArray<FString>& tags,
 		int LOD, bool bModifiersForBeforeOperations,
-		TArray<FirstPassGenerator::MODIFIER>& modifiers)
+		TArray<FirstPassGenerator::FModifier>& modifiers)
 	{
         MUTABLE_CPUPROFILER_SCOPE(GetModifiersFor);
 
 		if (tags.Num())
 		{
-			for (const FirstPassGenerator::MODIFIER& m: m_firstPass.modifiers)
+			for (const FirstPassGenerator::FModifier& m: m_firstPass.modifiers)
 			{
 				// Correct LOD?
 				if (m.lod != LOD)
@@ -2103,7 +2099,7 @@ namespace mu
 
 				// Already there?
 				bool alreadyAdded = 
-					modifiers.FindByPredicate( [&m](const FirstPassGenerator::MODIFIER& c) {return c.node == m.node; })
+					modifiers.FindByPredicate( [&m](const FirstPassGenerator::FModifier& c) {return c.node == m.node; })
 					!= 
 					nullptr;
 
@@ -2130,14 +2126,14 @@ namespace mu
 	//---------------------------------------------------------------------------------------------
 	Ptr<ASTOp> CodeGenerator::ApplyMeshModifiers(
 		const Ptr<ASTOp>& sourceOp,
-		const TArray<string>& tags,
+		const TArray<FString>& tags,
 		bool bModifiersForBeforeOperations,
 		const void* errorContext )
 	{
 		Ptr<ASTOp> lastMeshOp = sourceOp;
 
 		// Apply mesh modifiers
-		TArray<FirstPassGenerator::MODIFIER> modifiers;
+		TArray<FirstPassGenerator::FModifier> modifiers;
 
 		int currentLOD = m_currentParents.Last().m_lod;
 		GetModifiersFor(tags, currentLOD, bModifiersForBeforeOperations, modifiers);
@@ -2148,7 +2144,7 @@ namespace mu
 
 		// Process clip-with-mesh modifiers
 		Ptr<ASTOpMeshRemoveMask> removeOp;
-		for (const FirstPassGenerator::MODIFIER& m : modifiers)
+		for (const FirstPassGenerator::FModifier& m : modifiers)
 		{
 			if (const NodeModifierMeshClipWithMesh::Private* TypedClipNode = dynamic_cast<const NodeModifierMeshClipWithMesh::Private*>(m.node))
 			{
@@ -2193,7 +2189,7 @@ namespace mu
 
 
 		// Process clip-morph-plane modifiers
-		for (const FirstPassGenerator::MODIFIER& m : modifiers)
+		for (const FirstPassGenerator::FModifier& m : modifiers)
 		{
 			Ptr<ASTOp> modifiedMeshOp;
 
@@ -2273,7 +2269,7 @@ namespace mu
 		}
 
     	// Process clip deform modifiers.
-		for (const FirstPassGenerator::MODIFIER& M : modifiers)
+		for (const FirstPassGenerator::FModifier& M : modifiers)
 		{
 			Ptr<ASTOp> ModifiedMeshOp;
 
diff --git a/Engine/Plugins/Experimental/Mutable/Source/MutableTools/Private/MuT/CodeGenerator.h b/Engine/Plugins/Experimental/Mutable/Source/MutableTools/Private/MuT/CodeGenerator.h
--- a/Engine/Plugins/Experimental/Mutable/Source/MutableTools/Private/MuT/CodeGenerator.h
+++ b/Engine/Plugins/Experimental/Mutable/Source/MutableTools/Private/MuT/CodeGenerator.h
@@ -175,21 +175,21 @@ namespace mu
 
             // This reference has to be the smart pointer to avoid memory aliasing, keeping
             // processed nodes alive.
-            NodePtrConst pNode;
+            Ptr<const Node> pNode;
             int state = -1;
-			TArray<mu::string> activeTags;
+			TArray<FString> activeTags;
         };
 
         //! This struct contains additional state propagated from bottom to top of the object node graph.
         //! It is stored for every visited node, and restored when the cache is used.
-        struct BOTTOM_UP_STATE
+        struct FBottomUpState
         {
             //! Generated root address for the node.
             Ptr<ASTOp> m_address;
         };
-        BOTTOM_UP_STATE m_currentBottomUpState;
+		FBottomUpState m_currentBottomUpState;
 
-        typedef TMap<FVisitedKeyMap,BOTTOM_UP_STATE> VisitedMap;
+        typedef TMap<FVisitedKeyMap, FBottomUpState> VisitedMap;
         VisitedMap m_compiled;
 
         //!
@@ -219,7 +219,7 @@ namespace mu
         uint32 m_freeVertexIndex = 0;
 
 		// (top-down) Tags that are active when generating nodes.
-		TArray< TArray<mu::string> > m_activeTags;
+		TArray< TArray<FString> > m_activeTags;
 
         struct FParentKey
         {
@@ -247,9 +247,9 @@ namespace mu
 
         // List of additional components to add to an object that come from child objects.
         // The index is the object and lod that should receive the components.
-        struct ADDITIONAL_COMPONENT_KEY
+        struct FAdditionalComponentKey
         {
-            ADDITIONAL_COMPONENT_KEY()
+			FAdditionalComponentKey()
             {
                 m_pObject = nullptr;
                 m_lod = -1;
@@ -258,14 +258,14 @@ namespace mu
             const NodeObjectNew::Private* m_pObject;
             int m_lod;
 
-            inline bool operator<(const ADDITIONAL_COMPONENT_KEY& o) const
+            inline bool operator<(const FAdditionalComponentKey& o) const
             {
                 if (m_pObject < o.m_pObject) return true;
                 if (m_pObject > o.m_pObject) return false;
                 return m_lod < o.m_lod;
             }
         };
-        std::map< ADDITIONAL_COMPONENT_KEY, TArray<Ptr<ASTOp>> > m_additionalComponents;
+        std::map< FAdditionalComponentKey, TArray<Ptr<ASTOp>> > m_additionalComponents;
 
 
         struct OBJECT_GENERATION_DATA
@@ -275,7 +275,7 @@ namespace mu
         };
 		TArray< OBJECT_GENERATION_DATA > m_currentObject;
 
-        map< std::pair<TablePtr,string>, std::pair<TablePtr,Ptr<ASTOp>> > m_generatedTables;
+        map< std::pair<TablePtr, FString>, std::pair<TablePtr,Ptr<ASTOp>> > m_generatedTables;
 
         //! Variables added for every node
         map< Ptr<const Node>, Ptr<ASTOpParameter> > m_nodeVariables;
@@ -284,7 +284,7 @@ namespace mu
 		{
 			Ptr<ASTOp> Condition;
 			Ptr<ASTOp> ExtensionDataOp;
-			string ExtensionDataName;
+			FString ExtensionDataName;
 		};
 
 		TArray<FConditionalExtensionDataOp> m_conditionalExtensionDataOps;
@@ -292,11 +292,11 @@ namespace mu
 		//-----------------------------------------------------------------------------------------
 
 		// Get the modifiers that have to be applied to elements with a specific tag.
-		void GetModifiersFor(const TArray<string>& tags, int LOD,
-			bool bModifiersForBeforeOperations, TArray<FirstPassGenerator::MODIFIER>& modifiers);
+		void GetModifiersFor(const TArray<FString>& tags, int LOD,
+			bool bModifiersForBeforeOperations, TArray<FirstPassGenerator::FModifier>& modifiers);
 
 		// Apply the required mesh modifiers to the given operation.
-		Ptr<ASTOp> ApplyMeshModifiers( const Ptr<ASTOp>& sourceOp, const TArray<string>& tags,
+		Ptr<ASTOp> ApplyMeshModifiers( const Ptr<ASTOp>& sourceOp, const TArray<FString>& tags,
 			bool bModifiersForBeforeOperations, const void* errorContext);
 
 		// Get the modifiers that have to be applied to elements with a specific tag.
@@ -304,7 +304,7 @@ namespace mu
 
         //-----------------------------------------------------------------------------------------
         //!
-        Ptr<ASTOp> GenerateTableVariable( TablePtr pTable, const string& strName );
+        Ptr<ASTOp> GenerateTableVariable( TablePtr pTable, const FString& strName );
 
         //!
         Ptr<ASTOp> GenerateMissingBoolCode(const TCHAR* strWhere, bool value, const void* errorContext );
@@ -350,7 +350,7 @@ namespace mu
 			Ptr<const Layout> LayoutToApply;
 
 			/** Tags that are active at this point of the generation. */
-			TArray<mu::string> ActiveTags;
+			TArray<FString> ActiveTags;
 
 			friend FORCEINLINE uint32 GetTypeHash(const FImageGenerationOptions& InKey)
 			{
@@ -475,7 +475,7 @@ namespace mu
 			int32 State = 0;
 
 			/** Tags that are active at this point of the generation. */
-			TArray<mu::string> ActiveTags;
+			TArray<FString> ActiveTags;
 
 			/** Whatever mesh we reach at the leaves of the graph will need to have unique ids for its vertices.
 			* This is used to track mesh removal indices, morph data in other nodes, clothing data, etc.
@@ -716,10 +716,10 @@ namespace mu
             Ptr<ASTOp> sizeOp;
 
             //
-            string rangeName;
+            FString rangeName;
 
             //
-            string rangeUID;
+			FString rangeUID;
         };
 
         typedef TMap<FVisitedKeyMap,FRangeGenerationResult> GeneratedRangeMap;
@@ -736,7 +736,7 @@ namespace mu
 
         void GenerateSurface( FSurfaceGenerationResult& result,
                               NodeSurfaceNewPtrConst node,
-                              const TArray<FirstPassGenerator::SURFACE::EDIT>& edits );
+                              const TArray<FirstPassGenerator::FSurface::FEdit>& edits );
     };
 
 
@@ -764,9 +764,9 @@ namespace mu
         TablePtr pTable;
         Ptr<ASTOp> variable;
 
-        map< std::pair<TablePtr,string>, std::pair<TablePtr,Ptr<ASTOp>> >::iterator it
+        map< std::pair<TablePtr,FString>, std::pair<TablePtr,Ptr<ASTOp>> >::iterator it
                 = m_generatedTables.find
-                ( std::pair<TablePtr,string>(node.m_pTable,node.m_parameterName) );
+                ( std::pair<TablePtr,FString>(node.m_pTable,node.m_parameterName) );
         if ( it!=m_generatedTables.end() )
         {
             pTable = it->second.first;
@@ -779,12 +779,12 @@ namespace mu
             pTable = node.m_pTable;
             variable = GenerateTableVariable( pTable, node.m_parameterName );
 
-            m_generatedTables[ std::pair<TablePtr,string>(node.m_pTable,node.m_parameterName) ] =
+            m_generatedTables[ std::pair<TablePtr, FString>(node.m_pTable,node.m_parameterName) ] =
                     std::pair<TablePtr,Ptr<ASTOp>>( pTable, variable );
         }
 
         // Verify that the table column is the right type
-        int colIndex = pTable->FindColumn( node.m_columnName.c_str() );
+        int colIndex = pTable->FindColumn( node.m_columnName );
         if ( colIndex<0 )
         {
             m_pErrorLog->GetPrivate()->Add("Table column not found.", ELMT_ERROR, node.m_errorContext);
diff --git a/Engine/Plugins/Experimental/Mutable/Source/MutableTools/Private/MuT/CodeGenerator_Color.cpp b/Engine/Plugins/Experimental/Mutable/Source/MutableTools/Private/MuT/CodeGenerator_Color.cpp
--- a/Engine/Plugins/Experimental/Mutable/Source/MutableTools/Private/MuT/CodeGenerator_Color.cpp
+++ b/Engine/Plugins/Experimental/Mutable/Source/MutableTools/Private/MuT/CodeGenerator_Color.cpp
@@ -251,7 +251,7 @@ namespace mu
 		for (int t = node.m_variations.Num() - 1; t >= 0; --t)
 		{
 			int tagIndex = -1;
-			const string& tag = node.m_variations[t].m_tag;
+			const FString& tag = node.m_variations[t].m_tag;
 			for (int i = 0; i < m_firstPass.m_tags.Num(); ++i)
 			{
 				if (m_firstPass.m_tags[i].tag == tag)
@@ -262,7 +262,7 @@ namespace mu
 
 			if (tagIndex < 0)
 			{
-				FString Msg = FString::Printf(TEXT("Unknown tag found in color variation [%s]."), *FString(tag.c_str()));
+				FString Msg = FString::Printf(TEXT("Unknown tag found in color variation [%s]."), *tag);
 				m_pErrorLog->GetPrivate()->Add(Msg, ELMT_WARNING, node.m_errorContext);
 				continue;
 			}
diff --git a/Engine/Plugins/Experimental/Mutable/Source/MutableTools/Private/MuT/CodeGenerator_ExtensionData.cpp b/Engine/Plugins/Experimental/Mutable/Source/MutableTools/Private/MuT/CodeGenerator_ExtensionData.cpp
--- a/Engine/Plugins/Experimental/Mutable/Source/MutableTools/Private/MuT/CodeGenerator_ExtensionData.cpp
+++ b/Engine/Plugins/Experimental/Mutable/Source/MutableTools/Private/MuT/CodeGenerator_ExtensionData.cpp
@@ -165,15 +165,15 @@ class Node;
 		// Process variations in reverse order, since conditionals are built bottom-up.
 		for (int32 VariationIndex = Node.Variations.Num() - 1; VariationIndex >= 0; --VariationIndex)
 		{
-			const string& Tag = Node.Variations[VariationIndex].Tag;
-			const int32 TagIndex = m_firstPass.m_tags.IndexOfByPredicate([Tag](const FirstPassGenerator::TAG& CandidateTag)
+			const FString& Tag = Node.Variations[VariationIndex].Tag;
+			const int32 TagIndex = m_firstPass.m_tags.IndexOfByPredicate([Tag](const FirstPassGenerator::FTag& CandidateTag)
 			{
 				return CandidateTag.tag == Tag;
 			});
 
 			if (TagIndex == INDEX_NONE)
 			{
-				const FString Msg = FString::Printf(TEXT("Unknown tag found in Extension Data variation [%s]"), *FString{ Tag.c_str() });
+				const FString Msg = FString::Printf(TEXT("Unknown tag found in Extension Data variation [%s]"), *Tag);
 				m_pErrorLog->GetPrivate()->Add(Msg, ELMT_WARNING, Node.m_errorContext);
 				continue;
 			}
diff --git a/Engine/Plugins/Experimental/Mutable/Source/MutableTools/Private/MuT/CodeGenerator_FirstPass.cpp b/Engine/Plugins/Experimental/Mutable/Source/MutableTools/Private/MuT/CodeGenerator_FirstPass.cpp
--- a/Engine/Plugins/Experimental/Mutable/Source/MutableTools/Private/MuT/CodeGenerator_FirstPass.cpp
+++ b/Engine/Plugins/Experimental/Mutable/Source/MutableTools/Private/MuT/CodeGenerator_FirstPass.cpp
@@ -38,7 +38,7 @@ namespace mu
 	FirstPassGenerator::FirstPassGenerator()
 	{
 		// Default conditions when there is no restriction accumulated.
-		CONDITION_CONTEXT noCondition;
+		FConditionContext noCondition;
         m_currentCondition.Add(noCondition);
         m_currentStateCondition.Add(StateCondition());
 	}
@@ -83,7 +83,7 @@ namespace mu
             for (int32 t=0; t<surfaces[s].node->GetPrivate()->m_tags.Num(); ++t)
 			{
 				int tag = -1;
-                auto tagStr = surfaces[s].node->GetPrivate()->m_tags[t];
+                const FString& tagStr = surfaces[s].node->GetPrivate()->m_tags[t];
                 for (std::size_t i = 0; i<m_tags.Num() && tag<0; ++i)
 				{
                     if (m_tags[i].tag == tagStr)
@@ -96,7 +96,7 @@ namespace mu
 				if (tag < 0)
 				{
                     tag = (int)m_tags.Num();
-					TAG newTag;
+					FTag newTag;
                     newTag.tag = tagStr;
                     m_tags.Add(newTag);
 				}
@@ -128,7 +128,7 @@ namespace mu
                     if (tag < 0)
                     {
                         tag = (int)m_tags.Num();
-                        TAG newTag;
+						FTag newTag;
                         newTag.tag = tagStr;
                         m_tags.Add(newTag);
                     }
@@ -162,7 +162,7 @@ namespace mu
     Ptr<ASTOp> FirstPassGenerator::Visit(const NodeModifierMeshClipMorphPlane::Private& node)
 	{
         // Add the data about this modifier
-        MODIFIER thisData;
+        FModifier thisData;
         thisData.node = &node;
         thisData.objectCondition = m_currentCondition.Last().objectCondition;
         thisData.stateCondition = m_currentStateCondition.Last();
@@ -179,7 +179,7 @@ namespace mu
     Ptr<ASTOp> FirstPassGenerator::Visit(const NodeModifierMeshClipWithMesh::Private& node)
 	{
         // Add the data about this modifier
-		MODIFIER thisData;
+		FModifier thisData;
 		thisData.node = &node;
         thisData.objectCondition = m_currentCondition.Last().objectCondition;
         thisData.stateCondition = m_currentStateCondition.Last();
@@ -196,7 +196,7 @@ namespace mu
     Ptr<ASTOp> FirstPassGenerator::Visit(const NodeModifierMeshClipDeform::Private& node)
 	{
         // Add the data about this modifier
-		MODIFIER thisData;
+		FModifier thisData;
 		thisData.node = &node;
         thisData.objectCondition = m_currentCondition.Last().objectCondition;
         thisData.stateCondition = m_currentStateCondition.Last();
@@ -212,7 +212,7 @@ namespace mu
     Ptr<ASTOp> FirstPassGenerator::Visit(const NodeSurfaceNew::Private& node)
 	{
 		// Add the data about this surface
-		SURFACE thisData;
+		FSurface thisData;
 		thisData.node = dynamic_cast<const NodeSurfaceNew*>(node.m_pNode);
 		thisData.component = m_currentComponent;
         thisData.objectCondition = m_currentCondition.Last().objectCondition;
@@ -230,7 +230,7 @@ namespace mu
 	{
 		// Store a reference to this node in the surface data for the surface that this node is
 		// editing.
-		auto its = surfaces.FindByPredicate([&node](const FirstPassGenerator::SURFACE& s)
+		auto its = surfaces.FindByPredicate([&node](const FirstPassGenerator::FSurface& s)
         {
             // Are we editing the main surface node of this surface?
             if (s.node.get() == node.m_pParent.get()) return true;
@@ -251,9 +251,9 @@ namespace mu
 		// with wrong input or in case of partial models for preview.
 		if (its)
 		{
-			SURFACE& surface = *its;
+			FSurface& surface = *its;
 
-            SURFACE::EDIT edit;
+			FSurface::FEdit edit;
             edit.node = &node;
             edit.condition = m_currentCondition.Last().objectCondition;
             surface.edits.Add(edit);
@@ -470,7 +470,7 @@ namespace mu
     Ptr<ASTOp> FirstPassGenerator::Visit(const NodeObjectNew::Private& node)
 	{
 		// Add the data about this object
-		OBJECT thisData;
+		FObject thisData;
 		thisData.node = &node;
         thisData.condition = m_currentCondition.Last().objectCondition;
 		objects.Add(thisData);
@@ -482,8 +482,8 @@ namespace mu
 
             if ( s.m_runtimeParams.Num() > MUTABLE_MAX_RUNTIME_PARAMETERS_PER_STATE )
             {
-                FString Msg = FString::Printf( TEXT("State [%s] has more than %d runtime parameters. Their update may fail."),
-                	StringCast<TCHAR>(s.m_name.c_str()).Get(),
+                FString Msg = FString::Printf( TEXT("State [%s] has more than %d runtime parameters. Their update may fail."), 
+					*s.m_name,
                     MUTABLE_MAX_RUNTIME_PARAMETERS_PER_STATE);
                 m_pErrorLog->GetPrivate()->Add(Msg, ELMT_ERROR, node.m_errorContext );
             }
@@ -526,14 +526,14 @@ namespace mu
             Ptr<ASTOpParameter> op = new ASTOpParameter();
             op->type = OP_TYPE::NU_PARAMETER;
 
-            op->parameter.m_name = node.m_name;
-            op->parameter.m_uid = node.m_uid;
+            op->parameter.m_name = node.Name;
+            op->parameter.m_uid = node.Uid;
             op->parameter.m_type = PARAMETER_TYPE::T_INT;
             op->parameter.m_defaultValue.Set<ParamIntType>(-1);
 
             if ( node.m_type==NodeObjectGroup::CS_ONE_OR_NONE )
             {
-                FParameterDesc::INT_VALUE_DESC nullValue;
+                FParameterDesc::FIntValueDesc nullValue;
                 nullValue.m_value = -1;
                 nullValue.m_name = "None";
                 op->parameter.m_possibleValues.Add( nullValue );
@@ -582,8 +582,8 @@ namespace mu
                     case NodeObjectGroup::CS_ALWAYS_ONE:
                     {
                         // Add the option to the enumeration parameter
-                        FParameterDesc::INT_VALUE_DESC value;
-                        value.m_value = (int16_t)t;
+                        FParameterDesc::FIntValueDesc value;
+                        value.m_value = (int16)t;
                         value.m_name = pChildNode->GetName();
                         enumOp->parameter.m_possibleValues.Add( value );
 
@@ -619,7 +619,7 @@ namespace mu
                     paramOp = op;
                 }
 
-				CONDITION_CONTEXT data;
+				FConditionContext data;
                 data.objectCondition = paramOp;
                 m_currentCondition.Add( data );
 
diff --git a/Engine/Plugins/Experimental/Mutable/Source/MutableTools/Private/MuT/CodeGenerator_FirstPass.h b/Engine/Plugins/Experimental/Mutable/Source/MutableTools/Private/MuT/CodeGenerator_FirstPass.h
--- a/Engine/Plugins/Experimental/Mutable/Source/MutableTools/Private/MuT/CodeGenerator_FirstPass.h
+++ b/Engine/Plugins/Experimental/Mutable/Source/MutableTools/Private/MuT/CodeGenerator_FirstPass.h
@@ -69,12 +69,12 @@ namespace mu
 		//-------------------------
 
 		//! Store the conditions that will enable or disable every object
-		struct OBJECT
+		struct FObject
 		{
 			const NodeObjectNew::Private* node;
             Ptr<ASTOp> condition;
 		};
-		TArray<OBJECT> objects;
+		TArray<FObject> objects;
 
         //! Type used to represent the activation conditions regarding states
         //! This is the state mask for the states in which this surface must be added. If it
@@ -88,7 +88,7 @@ namespace mu
 		//! - all edit operators
         //! A surface may have different versions depending on the different parents and conditions
         //! it is reached with.
-		struct SURFACE
+		struct FSurface
 		{
             NodeSurfaceNewPtrConst node;
 
@@ -97,10 +97,10 @@ namespace mu
             const NodeComponentNew::Private* component = nullptr;
 
             // List of tags that are required for the presence of this surface
-			TArray<string> positiveTags;
+			TArray<FString> positiveTags;
 
             // List of tags that block the presence of this surface
-			TArray<string> negativeTags;
+			TArray<FString> negativeTags;
 
 			// This conditions is the condition of the object defining this surface which may not
 			// be the parent object where this surface will be added.
@@ -114,7 +114,7 @@ namespace mu
             Ptr<ASTOp> surfaceCondition;
 
 			// All surface editing nodes that edit this surface
-            struct EDIT
+            struct FEdit
             {
                 //! Condition that enables the effects of this edit node on the surface
                 Ptr<ASTOp> condition;
@@ -122,24 +122,24 @@ namespace mu
                 //! Weak reference to the edit node, used during compilation.
                 const NodeSurfaceEdit::Private* node = nullptr;
             };
-			TArray<EDIT> edits;
+			TArray<FEdit> edits;
 
             // This is filled in the final code generation pass
             Ptr<ASTOp> resultSurfaceOp;
             Ptr<ASTOp> resultMeshOp;
         };
-		TArray<SURFACE> surfaces;
+		TArray<FSurface> surfaces;
 
 		//! Store the conditions that enable every modifier.
-		struct MODIFIER
+		struct FModifier
 		{
             const NodeModifier::Private* node = nullptr;
 
             // List of tags that are required for the presence of this surface
-			TArray<string> positiveTags;
+			TArray<FString> positiveTags;
 
             // List of tags that block the presence of this surface
-			TArray<string> negativeTags;
+			TArray<FString> negativeTags;
 
             // This conditions is the condition of the object defining this modifier which may not
             // be the parent object where this surface will be added.
@@ -154,14 +154,14 @@ namespace mu
             StateCondition stateCondition;
 
             //
-            int lod = 0;
+            int32 lod = 0;
         };
-		TArray<MODIFIER> modifiers;
+		TArray<FModifier> modifiers;
 
 		//! Info about all found tags.
-		struct TAG
+		struct FTag
 		{
-			string tag;
+			FString tag;
 
             // Surfaces that activate the tag. These are indices to the FirstPassGenerator::surfaces
             // vector.
@@ -176,7 +176,7 @@ namespace mu
             // condition. This is filled in CodeGenerator_SecondPass.
             Ptr<ASTOp> genericCondition;
         };
-        TArray<TAG> m_tags;
+        TArray<FTag> m_tags;
 
         //! Accumulate the model states found while generating code, with their generated root
         //! nodes.
@@ -200,11 +200,11 @@ namespace mu
 
 	private:
 
-        struct CONDITION_CONTEXT
+        struct FConditionContext
         {
             Ptr<ASTOp> objectCondition;
         };
-		TArray< CONDITION_CONTEXT > m_currentCondition;
+		TArray< FConditionContext > m_currentCondition;
 
         //!
 		TArray< StateCondition > m_currentStateCondition;
@@ -213,8 +213,8 @@ namespace mu
         const NodeComponentNew::Private* m_currentComponent = nullptr;
 
         //! Current relevant tags so far. Used during traversal.
-		TArray<string> m_currentPositiveTags;
-		TArray<string> m_currentNegativeTags;
+		TArray<FString> m_currentPositiveTags;
+		TArray<FString> m_currentNegativeTags;
 
 		//! Index of the LOD we are processing
         int m_currentLOD = -1;
diff --git a/Engine/Plugins/Experimental/Mutable/Source/MutableTools/Private/MuT/CodeGenerator_Image.cpp b/Engine/Plugins/Experimental/Mutable/Source/MutableTools/Private/MuT/CodeGenerator_Image.cpp
--- a/Engine/Plugins/Experimental/Mutable/Source/MutableTools/Private/MuT/CodeGenerator_Image.cpp
+++ b/Engine/Plugins/Experimental/Mutable/Source/MutableTools/Private/MuT/CodeGenerator_Image.cpp
@@ -1349,10 +1349,10 @@ namespace mu
                 if ( GeneratedLayoutBlockId>=0 )
                 {
                     Ptr<ASTOpMeshExtractLayoutBlocks> eop = new ASTOpMeshExtractLayoutBlocks();
-                    eop->source = CurrentMeshToProjectOp;
-                    eop->layout = node.m_layout;
+                    eop->Source = CurrentMeshToProjectOp;
+                    eop->Layout = node.m_layout;
 
-                    eop->blocks.Add(GeneratedLayoutBlockId);
+                    eop->Blocks.Add(GeneratedLayoutBlockId);
 
 					CurrentMeshToProjectOp = eop;
                 }
@@ -1625,7 +1625,7 @@ namespace mu
         for ( int t = int( node.m_variations.Num() ) - 1; t >= 0; --t )
         {
             int tagIndex = -1;
-            const string& tag = node.m_variations[t].m_tag;
+            const FString& tag = node.m_variations[t].m_tag;
             for ( int i = 0; i < int( m_firstPass.m_tags.Num() ); ++i )
             {
                 if ( m_firstPass.m_tags[i].tag == tag )
@@ -1636,7 +1636,7 @@ namespace mu
 
             if ( tagIndex < 0 )
             {
-				FString Msg = FString::Printf(TEXT("Unknown tag found in image variation [%s]."), tag.c_str() );
+				FString Msg = FString::Printf(TEXT("Unknown tag found in image variation [%s]."), *tag );
 
                 m_pErrorLog->GetPrivate()->Add( Msg, ELMT_WARNING, node.m_errorContext );
                 continue;
