[mutable] Remove most of the legacy string usage in favor of FString and FName.

diff --git a/Engine/Plugins/Experimental/Mutable/Source/MutableTools/Private/MuT/CodeGenerator_Mesh.cpp b/Engine/Plugins/Experimental/Mutable/Source/MutableTools/Private/MuT/CodeGenerator_Mesh.cpp
--- a/Engine/Plugins/Experimental/Mutable/Source/MutableTools/Private/MuT/CodeGenerator_Mesh.cpp
+++ b/Engine/Plugins/Experimental/Mutable/Source/MutableTools/Private/MuT/CodeGenerator_Mesh.cpp
@@ -920,19 +920,19 @@ class Node;
                 Ptr<ASTOpMeshExtractLayoutBlocks> op = new ASTOpMeshExtractLayoutBlocks();
                 OutResult.meshOp = op;
 
-                op->source = BaseResult.meshOp;
+                op->Source = BaseResult.meshOp;
 
                 if (BaseResult.GeneratedLayouts.Num()>node.m_layoutOrGroup )
                 {
                     const Layout* pLayout = BaseResult.GeneratedLayouts[node.m_layoutOrGroup].get();
-                    op->layout = (uint16)node.m_layoutOrGroup;
+                    op->Layout = (uint16)node.m_layoutOrGroup;
 
                     for ( int32 i=0; i<node.m_blocks.Num(); ++i )
                     {
                         if (node.m_blocks[i]>=0 && node.m_blocks[i]<pLayout->m_blocks.Num() )
                         {
                             int bid = pLayout->m_blocks[ node.m_blocks[i] ].m_id;
-                            op->blocks.Add(bid);
+                            op->Blocks.Add(bid);
                         }
                         else
                         {
@@ -1189,7 +1189,7 @@ class Node;
         for ( int32 t = node.m_variations.Num()-1; t >= 0; --t )
         {
             int tagIndex = -1;
-            const string& tag = node.m_variations[t].m_tag;
+            const FString& tag = node.m_variations[t].m_tag;
             for ( int i = 0; i < m_firstPass.m_tags.Num(); ++i )
             {
                 if ( m_firstPass.m_tags[i].tag==tag)
@@ -1200,10 +1200,8 @@ class Node;
 
             if ( tagIndex < 0 )
             {
-				const char* Aux = nullptr;
-				Aux = tag.c_str();
                 m_pErrorLog->GetPrivate()->Add( 
-					FString::Printf(TEXT("Unknown tag found in mesh variation [%s]."), StringCast<TCHAR>(Aux).Get()),
+					FString::Printf(TEXT("Unknown tag found in mesh variation [%s]."), *tag),
 					ELMT_WARNING,
 					node.m_errorContext,
 					ELMSB_UNKNOWN_TAG
@@ -1434,7 +1432,7 @@ class Node;
 
  
 		// Apply the modifier for the pre-normal operations stage.
-		BOTTOM_UP_STATE temp = m_currentBottomUpState;
+		FBottomUpState temp = m_currentBottomUpState;
 
 		bool bModifiersForBeforeOperations = true;
 		OutResult.meshOp = ApplyMeshModifiers(op, InOptions.ActiveTags, bModifiersForBeforeOperations, node.m_errorContext);
diff --git a/Engine/Plugins/Experimental/Mutable/Source/MutableTools/Private/MuT/CodeGenerator_Scalar.cpp b/Engine/Plugins/Experimental/Mutable/Source/MutableTools/Private/MuT/CodeGenerator_Scalar.cpp
--- a/Engine/Plugins/Experimental/Mutable/Source/MutableTools/Private/MuT/CodeGenerator_Scalar.cpp
+++ b/Engine/Plugins/Experimental/Mutable/Source/MutableTools/Private/MuT/CodeGenerator_Scalar.cpp
@@ -144,7 +144,6 @@ namespace mu
 			param.m_uid = node.m_uid;
 			param.m_type = PARAMETER_TYPE::T_FLOAT;
 			param.m_defaultValue.Set<ParamFloatType>(node.m_defaultValue);
-			param.m_detailedType = node.m_detailedType;
 
 			op = new ASTOpParameter();
 			op->type = OP_TYPE::SC_PARAMETER;
@@ -184,7 +183,6 @@ namespace mu
 			param.m_uid = node.m_uid;
 			param.m_type = PARAMETER_TYPE::T_INT;
 			param.m_defaultValue.Set<ParamIntType>(node.m_defaultValue);
-			param.m_detailedType = node.m_detailedType;
 
 			param.m_possibleValues.SetNum(node.m_options.Num());
 			for (int32 i = 0; i < node.m_options.Num(); ++i)
@@ -296,7 +294,7 @@ namespace mu
 		for (int t = node.m_variations.Num() - 1; t >= 0; --t)
 		{
 			int tagIndex = -1;
-			const string& tag = node.m_variations[t].m_tag;
+			const FString& tag = node.m_variations[t].m_tag;
 			for (int i = 0; i < m_firstPass.m_tags.Num(); ++i)
 			{
 				if (m_firstPass.m_tags[i].tag == tag)
@@ -307,7 +305,7 @@ namespace mu
 
 			if (tagIndex < 0)
 			{
-				FString Msg = FString::Printf(TEXT("Unknown tag found in image variation [%s]."), *FString(tag.c_str()));
+				FString Msg = FString::Printf(TEXT("Unknown tag found in image variation [%s]."), *tag);
 
 				m_pErrorLog->GetPrivate()->Add(Msg, ELMT_WARNING, node.m_errorContext);
 				continue;
@@ -323,8 +321,7 @@ namespace mu
 			else
 			{
 				// This argument is required
-				variationOp = GenerateMissingScalarCode(TEXT("Variation option"), 0.0f,
-					node.m_errorContext);
+				variationOp = GenerateMissingScalarCode(TEXT("Variation option"), 0.0f, node.m_errorContext);
 			}
 
 
diff --git a/Engine/Plugins/Experimental/Mutable/Source/MutableTools/Private/MuT/CodeGenerator_SecondPass.cpp b/Engine/Plugins/Experimental/Mutable/Source/MutableTools/Private/MuT/CodeGenerator_SecondPass.cpp
--- a/Engine/Plugins/Experimental/Mutable/Source/MutableTools/Private/MuT/CodeGenerator_SecondPass.cpp
+++ b/Engine/Plugins/Experimental/Mutable/Source/MutableTools/Private/MuT/CodeGenerator_SecondPass.cpp
@@ -249,7 +249,7 @@ namespace mu
 
 		for (const auto& t : surf.positiveTags)
 		{
-			auto it = m_pFirstPass->m_tags.FindByPredicate([&](const FirstPassGenerator::TAG& e) { return e.tag == t; });
+			auto it = m_pFirstPass->m_tags.FindByPredicate([&](const FirstPassGenerator::FTag& e) { return e.tag == t; });
 			if (!it)
 			{
 				// This could happen if a tag is in a variation but noone defines it.
@@ -318,7 +318,7 @@ namespace mu
 
 		for (const auto& t : surf.negativeTags)
 		{
-			auto it = m_pFirstPass->m_tags.FindByPredicate([&](const FirstPassGenerator::TAG& e) { return e.tag == t; });
+			auto it = m_pFirstPass->m_tags.FindByPredicate([&](const FirstPassGenerator::FTag& e) { return e.tag == t; });
 			if (!it)
 			{
 				// This could happen if a tag is in a variation but noone defines it.
@@ -393,7 +393,7 @@ namespace mu
 		bool done = false;
 		for (const auto& t : mod.positiveTags)
 		{
-			auto it = m_pFirstPass->m_tags.FindByPredicate([&](const FirstPassGenerator::TAG& e) { return e.tag == t; });
+			auto it = m_pFirstPass->m_tags.FindByPredicate([&](const FirstPassGenerator::FTag& e) { return e.tag == t; });
 			if (!it)
 			{
 				// This could happen if a tag is in a variation but noone defines it.
@@ -457,7 +457,7 @@ namespace mu
 		{
 			for (const auto& t : mod.negativeTags)
 			{
-				auto it = m_pFirstPass->m_tags.FindByPredicate([&](const FirstPassGenerator::TAG& e) { return e.tag == t; });
+				auto it = m_pFirstPass->m_tags.FindByPredicate([&](const FirstPassGenerator::FTag& e) { return e.tag == t; });
 				if (!it)
 				{
 					// This could happen if a tag is in a variation but noone defines it.
@@ -561,7 +561,7 @@ namespace mu
 				auto& csurf = m_pFirstPass->surfaces[cs];
 				for (auto sct : csurf.positiveTags)
 				{
-					auto it = m_pFirstPass->m_tags.FindByPredicate([&](const FirstPassGenerator::TAG& e) { return e.tag == sct; });
+					auto it = m_pFirstPass->m_tags.FindByPredicate([&](const FirstPassGenerator::FTag& e) { return e.tag == sct; });
 					if (!it)
 					{
 						// This could happen if a tag is in a variation but noone defines it.
@@ -582,7 +582,7 @@ namespace mu
 				}
 				for (auto sct : csurf.negativeTags)
 				{
-					auto it = m_pFirstPass->m_tags.FindByPredicate([&](const FirstPassGenerator::TAG& e) { return e.tag == sct; });
+					auto it = m_pFirstPass->m_tags.FindByPredicate([&](const FirstPassGenerator::FTag& e) { return e.tag == sct; });
 					if (!it)
 					{
 						// This could happen if a tag is in a variation but noone defines it.
diff --git a/Engine/Plugins/Experimental/Mutable/Source/MutableTools/Private/MuT/CodeGenerator_String.cpp b/Engine/Plugins/Experimental/Mutable/Source/MutableTools/Private/MuT/CodeGenerator_String.cpp
--- a/Engine/Plugins/Experimental/Mutable/Source/MutableTools/Private/MuT/CodeGenerator_String.cpp
+++ b/Engine/Plugins/Experimental/Mutable/Source/MutableTools/Private/MuT/CodeGenerator_String.cpp
@@ -92,7 +92,6 @@ namespace mu
 			param.m_name = node.m_name;
 			param.m_uid = node.m_uid;
 			param.m_type = PARAMETER_TYPE::T_FLOAT;
-			param.m_detailedType = node.m_detailedType;
 			param.m_defaultValue.Set<ParamStringType>(node.m_defaultValue);
 
 			op = new ASTOpParameter();
diff --git a/Engine/Plugins/Experimental/Mutable/Source/MutableTools/Private/MuT/CodeOptimiser.h b/Engine/Plugins/Experimental/Mutable/Source/MutableTools/Private/MuT/CodeOptimiser.h
--- a/Engine/Plugins/Experimental/Mutable/Source/MutableTools/Private/MuT/CodeOptimiser.h
+++ b/Engine/Plugins/Experimental/Mutable/Source/MutableTools/Private/MuT/CodeOptimiser.h
@@ -285,7 +285,7 @@ namespace mu
         void Run( Ptr<ASTOp> root );
 
         //! After Run, list of relevant parameters.
-        std::unordered_set< string > m_params;
+        TSet< FString > m_params;
 
     private:
 
@@ -313,7 +313,7 @@ namespace mu
 
         // Result cache
         // \todo optimise by storing unique lists separately and an index here.
-        std::unordered_map< STATE, std::unordered_set< string >, state_hash > m_resultCache;
+        std::unordered_map< STATE, TSet<FString>, state_hash > m_resultCache;
     };
 
 }
diff --git a/Engine/Plugins/Experimental/Mutable/Source/MutableTools/Private/MuT/CodeOptimiser_Parameters.cpp b/Engine/Plugins/Experimental/Mutable/Source/MutableTools/Private/MuT/CodeOptimiser_Parameters.cpp
--- a/Engine/Plugins/Experimental/Mutable/Source/MutableTools/Private/MuT/CodeOptimiser_Parameters.cpp
+++ b/Engine/Plugins/Experimental/Mutable/Source/MutableTools/Private/MuT/CodeOptimiser_Parameters.cpp
@@ -142,7 +142,7 @@ namespace mu
                     case OP_TYPE::IM_PARAMETER:
                     {
 						const ASTOpParameter* typed = dynamic_cast<const ASTOpParameter*>(at.get());
-                        const TArray<string>& params = m_pState->nodeState.m_runtimeParams;
+                        const TArray<FString>& params = m_pState->nodeState.m_runtimeParams;
                         if ( params.Find( typed->parameter.m_name)
                              !=
                              INDEX_NONE )
@@ -1619,7 +1619,7 @@ namespace mu
     void SubtreeRelevantParametersVisitorAST::Run( Ptr<ASTOp> root )
     {
         // Cached?
-		std::unordered_map< STATE, std::unordered_set< string >, state_hash >::iterator it = m_resultCache.find( STATE(root,false) );
+		std::unordered_map< STATE, TSet<FString>, state_hash >::iterator it = m_resultCache.find( STATE(root,false) );
         if (it!=m_resultCache.end())
         {
             m_params = it->second;
@@ -1630,7 +1630,7 @@ namespace mu
         {
             MUTABLE_CPUPROFILER_SCOPE(SubtreeRelevantParametersVisitorAST);
 
-            m_params.clear();
+            m_params.Empty();
 
             // The state is the onlyLayoutRelevant flag
             ASTOp::Traverse_TopDown_Unique_Imprecise_WithState<bool>( root, false,
@@ -1648,7 +1648,7 @@ namespace mu
                 case OP_TYPE::IM_PARAMETER:
                 {
 					const ASTOpParameter* typedAt = dynamic_cast<const ASTOpParameter*>(at.get());
-                    m_params.insert(typedAt->parameter.m_name);
+                    m_params.Add(typedAt->parameter.m_name);
 
                     // Not interested in the parameters from the parameters decorators.
                     return false;
@@ -1748,8 +1748,8 @@ namespace mu
                     subtreeParams.Run( i.Key );
 
 					// Temp copy
-					TArray<mu::string> ParamCopy;
-					for ( const string& e: subtreeParams.m_params )
+					TArray<FString> ParamCopy;
+					for ( const FString& e: subtreeParams.m_params )
 					{
 						ParamCopy.Add(e);
 					}
diff --git a/Engine/Plugins/Experimental/Mutable/Source/MutableTools/Private/MuT/Compiler.cpp b/Engine/Plugins/Experimental/Mutable/Source/MutableTools/Private/MuT/Compiler.cpp
--- a/Engine/Plugins/Experimental/Mutable/Source/MutableTools/Private/MuT/Compiler.cpp
+++ b/Engine/Plugins/Experimental/Mutable/Source/MutableTools/Private/MuT/Compiler.cpp
@@ -202,7 +202,7 @@ namespace mu
                 STATE_COMPILATION_DATA data;
                 data.nodeState = s.first;
                 data.root = s.second;
-                data.state.m_name = s.first.m_name;
+                data.state.Name = s.first.m_name;
                 states.push_back( data );
             }
 
@@ -279,11 +279,10 @@ namespace mu
                 else
                 {
 					FString Temp = FString::Printf(TEXT(
-						"The state [%s] refers to a parameter [%s] "
-						"that has not been found in the model. This error can be "
+						"The state [%s] refers to a parameter [%s]  that has not been found in the model. This error can be "
 						"safely dismissed in case of partial compilation."), 
-						StringCast<TCHAR>(s.nodeState.m_name.c_str()).Get(),
-						StringCast<TCHAR>(s.nodeState.m_runtimeParams[p].c_str()).Get());
+						*s.nodeState.m_name,
+						*s.nodeState.m_runtimeParams[p]);
                     m_pD->m_pErrorLog->GetPrivate()->Add(Temp, ELMT_WARNING, pNode->GetBasePrivate()->m_errorContext );
                 }
             }
diff --git a/Engine/Plugins/Experimental/Mutable/Source/MutableTools/Private/MuT/CompilerPrivate.h b/Engine/Plugins/Experimental/Mutable/Source/MutableTools/Private/MuT/CompilerPrivate.h
--- a/Engine/Plugins/Experimental/Mutable/Source/MutableTools/Private/MuT/CompilerPrivate.h
+++ b/Engine/Plugins/Experimental/Mutable/Source/MutableTools/Private/MuT/CompilerPrivate.h
@@ -106,17 +106,17 @@ namespace mu
     struct FObjectState
     {
         //! Name used to identify the state from the code and user interface.
-        string m_name;
+        FString m_name;
 
         //! GPU Optimisation options
 		FStateOptimizationOptions m_optimisation;
 
         //! List of names of the runtime parameters in this state
-        TArray<string> m_runtimeParams;
+        TArray<FString> m_runtimeParams;
 
         void Serialise( OutputArchive& arch ) const
         {
-            const int32_t ver = 5;
+            const int32 ver = 6;
             arch << ver;
 
             arch << m_name;
@@ -127,13 +127,36 @@ namespace mu
 
         void Unserialise( InputArchive& arch )
         {
-            int32_t ver = 0;
+            int32 ver = 0;
             arch >> ver;
-            check( ver==5 );
+            check( ver>=5 && ver<=6 );
 
-            arch >> m_name;
+			if (ver <= 5)
+			{
+				std::string Temp;
+				arch >> Temp;
+				m_name = Temp.c_str();
+			}
+			else
+			{
+				arch >> m_name;
+			}
             arch >> m_optimisation;
-            arch >> m_runtimeParams;
+
+			if (ver <= 5)
+			{
+				TArray<std::string> Temp;
+				arch >> Temp;
+				m_runtimeParams.SetNum(Temp.Num());
+				for ( int32 i=0; i<Temp.Num(); ++i)
+				{
+					m_runtimeParams[i] = Temp[i].c_str();
+				}
+			}
+			else
+			{
+				arch >> m_runtimeParams;
+			}
         }
     };
 
@@ -151,7 +174,7 @@ namespace mu
 
         //! List of root instructions for the dynamic resources that depend on the runtime
         //! parameters of this state.
-		TArray< TPair<Ptr<ASTOp>, TArray<string> > > m_dynamicResources;
+		TArray< TPair<Ptr<ASTOp>, TArray<FString> > > m_dynamicResources;
     };
 
 
diff --git a/Engine/Plugins/Experimental/Mutable/Source/MutableTools/Private/MuT/DataPacker.cpp b/Engine/Plugins/Experimental/Mutable/Source/MutableTools/Private/MuT/DataPacker.cpp
--- a/Engine/Plugins/Experimental/Mutable/Source/MutableTools/Private/MuT/DataPacker.cpp
+++ b/Engine/Plugins/Experimental/Mutable/Source/MutableTools/Private/MuT/DataPacker.cpp
@@ -507,7 +507,6 @@ namespace mu
             // Default state: we need everything except internal semantics
             uint64_t defaultState = 0xffffffffffffffff;
             defaultState ^= (UINT64_C(1)<<MBS_LAYOUTBLOCK);
-            defaultState ^= (UINT64_C(1)<<MBS_CHART);
             defaultState ^= (UINT64_C(1)<<MBS_VERTEXINDEX);
 
             Traverse(roots,defaultState);
@@ -613,10 +612,9 @@ namespace mu
                 // todo: check if we really need all of them
                 uint64_t newState = currentSemantics;
                 newState |= (UINT64_C(1)<<MBS_LAYOUTBLOCK);
-                newState |= (UINT64_C(1)<<MBS_CHART);
                 newState |= (UINT64_C(1)<<MBS_VERTEXINDEX);
 
-                RecurseWithState( op->source.child(), newState );
+                RecurseWithState( op->Source.child(), newState );
                 break;
             }
 
@@ -700,7 +698,6 @@ namespace mu
         {
             uint64_t layoutSemantics = 0;
             layoutSemantics |= (UINT64_C(1)<<MBS_LAYOUTBLOCK);
-            layoutSemantics |= (UINT64_C(1)<<MBS_CHART);
 
             if ( (usedSemantics&layoutSemantics) == 0)
             {
diff --git a/Engine/Plugins/Experimental/Mutable/Source/MutableTools/Private/MuT/ImageDescGenerator.cpp b/Engine/Plugins/Experimental/Mutable/Source/MutableTools/Private/MuT/ImageDescGenerator.cpp
--- a/Engine/Plugins/Experimental/Mutable/Source/MutableTools/Private/MuT/ImageDescGenerator.cpp
+++ b/Engine/Plugins/Experimental/Mutable/Source/MutableTools/Private/MuT/ImageDescGenerator.cpp
@@ -120,7 +120,7 @@ namespace mu
         m_desc = MUTABLE_MISSING_IMAGE_DESC;
 
         // Verify that the table column is the right type
-        int colIndex = node.m_pTable->FindColumn( node.m_columnName.c_str() );
+        int colIndex = node.m_pTable->FindColumn( node.m_columnName );
         if ( colIndex<0 )
         {
             check( false );
diff --git a/Engine/Plugins/Experimental/Mutable/Source/MutableTools/Private/MuT/NodeBoolParameter.cpp b/Engine/Plugins/Experimental/Mutable/Source/MutableTools/Private/MuT/NodeBoolParameter.cpp
--- a/Engine/Plugins/Experimental/Mutable/Source/MutableTools/Private/MuT/NodeBoolParameter.cpp
+++ b/Engine/Plugins/Experimental/Mutable/Source/MutableTools/Private/MuT/NodeBoolParameter.cpp
@@ -66,30 +66,9 @@ namespace mu
 	//---------------------------------------------------------------------------------------------
 	// Own Interface
 	//---------------------------------------------------------------------------------------------
-	const char* NodeBoolParameter::GetName() const
+	void NodeBoolParameter::SetName( const FString& strName )
 	{
-		return m_pD->m_name.c_str();
-	}
-
-
-	//---------------------------------------------------------------------------------------------
-	void NodeBoolParameter::SetName( const char* strName )
-	{
-		if ( strName )
-		{
-			m_pD->m_name = strName;
-		}
-		else
-		{
-			m_pD->m_name = "";
-		}
-	}
-
-
-	//---------------------------------------------------------------------------------------------
-	bool NodeBoolParameter::GetDefaultValue() const
-	{
-		return m_pD->m_defaultValue;
+		m_pD->m_name = strName;
 	}
 
 
diff --git a/Engine/Plugins/Experimental/Mutable/Source/MutableTools/Private/MuT/NodeBoolPrivate.h b/Engine/Plugins/Experimental/Mutable/Source/MutableTools/Private/MuT/NodeBoolPrivate.h
--- a/Engine/Plugins/Experimental/Mutable/Source/MutableTools/Private/MuT/NodeBoolPrivate.h
+++ b/Engine/Plugins/Experimental/Mutable/Source/MutableTools/Private/MuT/NodeBoolPrivate.h
@@ -63,15 +63,15 @@ namespace mu
 		static NODE_TYPE s_type;
 
 		bool m_defaultValue;
-		string m_name;
-		string m_uid;
+		FString m_name;
+		FString m_uid;
 
         TArray<Ptr<NodeRange>> m_ranges;
 
 		//!
 		void Serialise( OutputArchive& arch ) const
 		{
-            uint32_t ver = 2;
+            uint32_t ver = 3;
 			arch << ver;
 
 			arch << m_defaultValue;
@@ -85,12 +85,24 @@ namespace mu
 		{
             uint32_t ver;
 			arch >> ver;
-            check(ver==2);
+            check(ver>=2 && ver<=3);
 
 			arch >> m_defaultValue;
-			arch >> m_name;
-            arch >> m_uid;
-            arch >> m_ranges;
+			if (ver <= 2)
+			{
+				std::string Temp;
+				arch >> Temp;
+				m_name = Temp.c_str();
+				arch >> Temp;
+				m_uid = Temp.c_str();
+			}
+			else
+			{
+				arch >> m_name;
+				arch >> m_uid;
+			}
+			
+			arch >> m_ranges;
         }
 	};
 
diff --git a/Engine/Plugins/Experimental/Mutable/Source/MutableTools/Private/MuT/NodeColourParameter.cpp b/Engine/Plugins/Experimental/Mutable/Source/MutableTools/Private/MuT/NodeColourParameter.cpp
--- a/Engine/Plugins/Experimental/Mutable/Source/MutableTools/Private/MuT/NodeColourParameter.cpp
+++ b/Engine/Plugins/Experimental/Mutable/Source/MutableTools/Private/MuT/NodeColourParameter.cpp
@@ -62,52 +62,16 @@ namespace mu
 	//---------------------------------------------------------------------------------------------
 	// Own Interface
 	//---------------------------------------------------------------------------------------------
-	const char* NodeColourParameter::GetName() const
+	void NodeColourParameter::SetName( const FString& strName )
 	{
-		return m_pD->m_name.c_str();
+		m_pD->m_name = strName;
 	}
 
 
 	//---------------------------------------------------------------------------------------------
-	void NodeColourParameter::SetName( const char* strName )
+	void NodeColourParameter::SetUid( const FString& strUid )
 	{
-		if ( strName )
-		{
-			m_pD->m_name = strName;
-		}
-		else
-		{
-			m_pD->m_name = "";
-		}
-	}
-
-
-	const char* NodeColourParameter::GetUid() const
-	{
-		return m_pD->m_uid.c_str();
-	}
-
-
-	//---------------------------------------------------------------------------------------------
-	void NodeColourParameter::SetUid( const char* strUid )
-	{
-		if ( strUid )
-		{
-			m_pD->m_uid = strUid;
-		}
-		else
-		{
-			m_pD->m_uid = "";
-		}
-	}
-
-
-	//---------------------------------------------------------------------------------------------
-	void NodeColourParameter::GetDefaultValue( float* pR, float* pG, float* pB ) const
-	{
-		if (*pR) *pR = m_pD->m_defaultValue[0];
-		if (*pG) *pG = m_pD->m_defaultValue[1];
-		if (*pB) *pB = m_pD->m_defaultValue[2];
+		m_pD->m_uid = strUid;
 	}
 
 
diff --git a/Engine/Plugins/Experimental/Mutable/Source/MutableTools/Private/MuT/NodeColourParameterPrivate.h b/Engine/Plugins/Experimental/Mutable/Source/MutableTools/Private/MuT/NodeColourParameterPrivate.h
--- a/Engine/Plugins/Experimental/Mutable/Source/MutableTools/Private/MuT/NodeColourParameterPrivate.h
+++ b/Engine/Plugins/Experimental/Mutable/Source/MutableTools/Private/MuT/NodeColourParameterPrivate.h
@@ -23,15 +23,15 @@ namespace mu
 		static NODE_TYPE s_type;
 
 		vec3<float> m_defaultValue;
-		string m_name;
-		string m_uid;
+		FString m_name;
+		FString m_uid;
 
         TArray<Ptr<NodeRange>> m_ranges;
 
 		//!
 		void Serialise( OutputArchive& arch ) const
 		{
-            uint32_t ver = 2;
+            uint32_t ver = 3;
 			arch << ver;
 
 			arch << m_defaultValue;
@@ -45,12 +45,25 @@ namespace mu
 		{
             uint32_t ver;
 			arch >> ver;
-            check(ver==2);
+            check(ver>=2&&ver<=3);
 
 			arch >> m_defaultValue;
-			arch >> m_name;
-            arch >> m_uid;
-            arch >> m_ranges;
+			
+			if (ver <= 2)
+			{
+				std::string Temp;
+				arch >> Temp;
+				m_name = Temp.c_str();
+				arch >> Temp;
+				m_uid = Temp.c_str();
+			}
+			else
+			{
+				arch >> m_name;
+				arch >> m_uid;
+			}
+			
+			arch >> m_ranges;
         }
 	};
 
diff --git a/Engine/Plugins/Experimental/Mutable/Source/MutableTools/Private/MuT/NodeColourTable.cpp b/Engine/Plugins/Experimental/Mutable/Source/MutableTools/Private/MuT/NodeColourTable.cpp
--- a/Engine/Plugins/Experimental/Mutable/Source/MutableTools/Private/MuT/NodeColourTable.cpp
+++ b/Engine/Plugins/Experimental/Mutable/Source/MutableTools/Private/MuT/NodeColourTable.cpp
@@ -58,28 +58,14 @@ namespace mu
 	//---------------------------------------------------------------------------------------------
 	// Own Interface
 	//---------------------------------------------------------------------------------------------
-	const char* NodeColourTable::GetColumn() const
+	void NodeColourTable::SetColumn( const FString& strName )
 	{
-		return m_pD->m_columnName.c_str();
+		m_pD->m_columnName = strName;
 	}
 
 
 	//---------------------------------------------------------------------------------------------
-	void NodeColourTable::SetColumn( const char* strName )
-	{
-		if (strName)
-		{
-			m_pD->m_columnName = strName;
-		}
-		else
-		{
-			m_pD->m_columnName = "";
-		}
-	}
-
-
-	//---------------------------------------------------------------------------------------------
-	void NodeColourTable::SetParameterName( const char* strName )
+	void NodeColourTable::SetParameterName( const FString& strName )
 	{
 		m_pD->m_parameterName = strName;
 	}
diff --git a/Engine/Plugins/Experimental/Mutable/Source/MutableTools/Private/MuT/NodeColourTablePrivate.h b/Engine/Plugins/Experimental/Mutable/Source/MutableTools/Private/MuT/NodeColourTablePrivate.h
--- a/Engine/Plugins/Experimental/Mutable/Source/MutableTools/Private/MuT/NodeColourTablePrivate.h
+++ b/Engine/Plugins/Experimental/Mutable/Source/MutableTools/Private/MuT/NodeColourTablePrivate.h
@@ -23,14 +23,14 @@ namespace mu
 
 		static NODE_TYPE s_type;
 
-		string m_parameterName;
+		FString m_parameterName;
 		TablePtr m_pTable;
-		string m_columnName;
+		FString m_columnName;
 
 		//!
 		void Serialise( OutputArchive& arch ) const
 		{
-            uint32_t ver = 0;
+            uint32 ver = 1;
 			arch << ver;
 
 			arch << m_parameterName;
@@ -41,13 +41,33 @@ namespace mu
 		//!
 		void Unserialise( InputArchive& arch )
 		{
-            uint32_t ver;
+            uint32 ver;
 			arch >> ver;
-            check(ver==0);
+			check(ver <= 1);
+
+			if (ver == 0)
+			{
+				std::string Temp;
+				arch >> Temp;
+				m_parameterName = Temp.c_str();
+			}
+			else
+			{
+				arch >> m_parameterName;
+			}
 
-			arch >> m_parameterName;
 			arch >> m_pTable;
-			arch >> m_columnName;
+
+			if (ver == 0)
+			{
+				std::string Temp;
+				arch >> Temp;
+				m_columnName = Temp.c_str();
+			}
+			else
+			{
+				arch >> m_columnName;
+			}
 		}
 
 	};
diff --git a/Engine/Plugins/Experimental/Mutable/Source/MutableTools/Private/MuT/NodeColourVariationPrivate.h b/Engine/Plugins/Experimental/Mutable/Source/MutableTools/Private/MuT/NodeColourVariationPrivate.h
--- a/Engine/Plugins/Experimental/Mutable/Source/MutableTools/Private/MuT/NodeColourVariationPrivate.h
+++ b/Engine/Plugins/Experimental/Mutable/Source/MutableTools/Private/MuT/NodeColourVariationPrivate.h
@@ -24,15 +24,15 @@ namespace mu
 
         NodeColourPtr m_defaultColour;
 
-        struct VARIATION
+        struct FVariation
         {
             NodeColourPtr m_colour;
-            string m_tag;
+			FString m_tag;
 
             //!
             void Serialise( OutputArchive& arch ) const
             {
-                uint32_t ver = 1;
+                uint32 ver = 2;
                 arch << ver;
 
                 arch << m_tag;
@@ -41,16 +41,25 @@ namespace mu
 
             void Unserialise( InputArchive& arch )
             {
-                uint32_t ver = 0;
+                uint32 ver = 0;
                 arch >> ver;
-                check( ver == 1 );
-
-                arch >> m_tag;
-                arch >> m_colour;
+				check(ver >= 1 && ver <= 2);
+
+				if (ver <= 1)
+				{
+					std::string Temp;
+					arch >> Temp;
+					m_tag = Temp.c_str();
+				}
+				else
+				{
+					arch >> m_tag;
+				}
+				arch >> m_colour;
             }
         };
 
-        TArray<VARIATION> m_variations;
+        TArray<FVariation> m_variations;
 
         //!
         void Serialise( OutputArchive& arch ) const
diff --git a/Engine/Plugins/Experimental/Mutable/Source/MutableTools/Private/MuT/NodeComponentNew.cpp b/Engine/Plugins/Experimental/Mutable/Source/MutableTools/Private/MuT/NodeComponentNew.cpp
--- a/Engine/Plugins/Experimental/Mutable/Source/MutableTools/Private/MuT/NodeComponentNew.cpp
+++ b/Engine/Plugins/Experimental/Mutable/Source/MutableTools/Private/MuT/NodeComponentNew.cpp
@@ -67,25 +67,16 @@ namespace mu
 	//---------------------------------------------------------------------------------------------
 	// Own Interface
 	//---------------------------------------------------------------------------------------------
-	const char* NodeComponentNew::GetName() const
+	const FString& NodeComponentNew::GetName() const
 	{
-		const char* strResult = m_pD->m_name.c_str();
-
-		return strResult;
+		return m_pD->m_name;
 	}
 
 
 	//---------------------------------------------------------------------------------------------
-	void NodeComponentNew::SetName( const char* strName )
+	void NodeComponentNew::SetName( const FString& strName )
 	{
-		if (strName)
-		{
-			m_pD->m_name = strName;
-		}
-		else
-		{
-			m_pD->m_name = "";
-		}
+		m_pD->m_name = strName;
 	}
 
 
diff --git a/Engine/Plugins/Experimental/Mutable/Source/MutableTools/Private/MuT/NodeComponentNewPrivate.h b/Engine/Plugins/Experimental/Mutable/Source/MutableTools/Private/MuT/NodeComponentNewPrivate.h
--- a/Engine/Plugins/Experimental/Mutable/Source/MutableTools/Private/MuT/NodeComponentNewPrivate.h
+++ b/Engine/Plugins/Experimental/Mutable/Source/MutableTools/Private/MuT/NodeComponentNewPrivate.h
@@ -30,7 +30,7 @@ namespace mu
 
 		static NODE_TYPE s_type;
 
-		string m_name;
+		FString m_name;
 
 		uint16 m_id = 0;
 
@@ -39,7 +39,7 @@ namespace mu
 		//!
 		void Serialise( OutputArchive& arch ) const
 		{
-            uint32_t ver = 3;
+            uint32_t ver = 4;
 			
 			arch << ver;
 
@@ -53,9 +53,18 @@ namespace mu
 		{
             uint32_t ver;
 			arch >> ver;
-			check(ver >= 2);
+			check(ver>=2 && ver<=4);
 
-			arch >> m_name;
+			if (ver <= 3)
+			{
+				std::string Temp;
+				arch>>Temp;
+				m_name = Temp.c_str();
+			}
+			else
+			{
+				arch >> m_name;
+			}
 
 			if (ver >= 3)
 			{
diff --git a/Engine/Plugins/Experimental/Mutable/Source/MutableTools/Private/MuT/NodeExtensionDataVariationPrivate.h b/Engine/Plugins/Experimental/Mutable/Source/MutableTools/Private/MuT/NodeExtensionDataVariationPrivate.h
--- a/Engine/Plugins/Experimental/Mutable/Source/MutableTools/Private/MuT/NodeExtensionDataVariationPrivate.h
+++ b/Engine/Plugins/Experimental/Mutable/Source/MutableTools/Private/MuT/NodeExtensionDataVariationPrivate.h
@@ -19,15 +19,15 @@ namespace mu
 
 		NodeExtensionDataPtr DefaultValue;
 
-		struct VARIATION
+		struct FVariation
 		{
 			NodeExtensionDataPtr Value;
-			string Tag;
+			FString Tag;
 
 			//!
 			void Serialise(OutputArchive& InArch) const
 			{
-				uint32_t Ver = 1;
+				uint32 Ver = 2;
 				InArch << Ver;
 
 				InArch << Tag;
@@ -36,16 +36,25 @@ namespace mu
 
 			void Unserialise(InputArchive& InArch)
 			{
-				uint32_t Ver = 0;
+				uint32 Ver = 0;
 				InArch >> Ver;
-				check(Ver == 1);
-
-				InArch >> Tag;
+				check(Ver >= 1 && Ver <= 2);
+
+				if (Ver <= 1)
+				{
+					std::string Temp;
+					InArch >> Temp;
+					Tag = Temp.c_str();
+				}
+				else
+				{
+					InArch >> Tag;
+				}
 				InArch >> Value;
 			}
 		};
 
-		TArray<VARIATION> Variations;
+		TArray<FVariation> Variations;
 
 		//!
 		void Serialise(OutputArchive& InArch) const
diff --git a/Engine/Plugins/Experimental/Mutable/Source/MutableTools/Private/MuT/NodeImageParameter.cpp b/Engine/Plugins/Experimental/Mutable/Source/MutableTools/Private/MuT/NodeImageParameter.cpp
--- a/Engine/Plugins/Experimental/Mutable/Source/MutableTools/Private/MuT/NodeImageParameter.cpp
+++ b/Engine/Plugins/Experimental/Mutable/Source/MutableTools/Private/MuT/NodeImageParameter.cpp
@@ -51,13 +51,6 @@ namespace mu
 
 	//---------------------------------------------------------------------------------------------
 	// Own Interface
-	//---------------------------------------------------------------------------------------------
-    const char* NodeImageParameter::GetName() const
-	{
-		return m_pD->m_name.c_str();
-	}
-
-
 	//---------------------------------------------------------------------------------------------
     void NodeImageParameter::SetName( const char* strName )
 	{
@@ -72,12 +65,6 @@ namespace mu
 	}
 
 
-	const char* NodeImageParameter::GetUid() const
-	{
-		return m_pD->m_uid.c_str();
-	}
-
-
 	//---------------------------------------------------------------------------------------------
 	void NodeImageParameter::SetUid( const char* strUid )
 	{
@@ -92,13 +79,6 @@ namespace mu
 	}
 
 	
-	//---------------------------------------------------------------------------------------------
-	FName NodeImageParameter::GetDefaultValue() const
-	{
-		return m_pD->m_defaultValue;
-	}
-
-	
 	//---------------------------------------------------------------------------------------------
 	void NodeImageParameter::SetDefaultValue(FName Value)
 	{
diff --git a/Engine/Plugins/Experimental/Mutable/Source/MutableTools/Private/MuT/NodeImageParameterPrivate.h b/Engine/Plugins/Experimental/Mutable/Source/MutableTools/Private/MuT/NodeImageParameterPrivate.h
--- a/Engine/Plugins/Experimental/Mutable/Source/MutableTools/Private/MuT/NodeImageParameterPrivate.h
+++ b/Engine/Plugins/Experimental/Mutable/Source/MutableTools/Private/MuT/NodeImageParameterPrivate.h
@@ -21,15 +21,15 @@ namespace mu
 		static NODE_TYPE s_type;
 
     	FName m_defaultValue;
-		string m_name;
-		string m_uid;
+		FString m_name;
+		FString m_uid;
 
 		TArray<Ptr<NodeRange>> m_ranges;
 
 		//!
 		void Serialise( OutputArchive& arch ) const
 		{
-            uint32 ver = 3;
+            uint32 ver = 4;
 			arch << ver;
 
 			arch << m_defaultValue;
@@ -43,15 +43,26 @@ namespace mu
 		{
             uint32 ver;
 			arch >> ver;
-            check(ver <= 3);
+            check(ver <= 4);
 
 			if (ver >= 3)
 			{
 				arch >> m_defaultValue;
 			}
 			
-			arch >> m_name;
-            arch >> m_uid;
+			if (ver <= 3)
+			{
+				std::string Temp;
+				arch >> Temp;
+				m_name = Temp.c_str();
+				arch >> Temp;
+				m_uid = Temp.c_str();
+			}
+			else
+			{
+				arch >> m_name;
+				arch >> m_uid;
+			}
 
 			if (ver >= 2)
 			{
diff --git a/Engine/Plugins/Experimental/Mutable/Source/MutableTools/Private/MuT/NodeImageTable.cpp b/Engine/Plugins/Experimental/Mutable/Source/MutableTools/Private/MuT/NodeImageTable.cpp
--- a/Engine/Plugins/Experimental/Mutable/Source/MutableTools/Private/MuT/NodeImageTable.cpp
+++ b/Engine/Plugins/Experimental/Mutable/Source/MutableTools/Private/MuT/NodeImageTable.cpp
@@ -57,28 +57,14 @@ namespace mu
 	//---------------------------------------------------------------------------------------------
 	// Own Interface
 	//---------------------------------------------------------------------------------------------
-	const char* NodeImageTable::GetColumn() const
+	void NodeImageTable::SetColumn( const FString& strName )
 	{
-		return m_pD->m_columnName.c_str();
+		m_pD->m_columnName = strName;
 	}
 
 
 	//---------------------------------------------------------------------------------------------
-	void NodeImageTable::SetColumn( const char* strName )
-	{
-		if (strName)
-		{
-			m_pD->m_columnName = strName;
-		}
-		else
-		{
-			m_pD->m_columnName = "";
-		}
-	}
-
-
-	//---------------------------------------------------------------------------------------------
-	void NodeImageTable::SetParameterName( const char* strName )
+	void NodeImageTable::SetParameterName( const FString& strName )
 	{
 		m_pD->m_parameterName = strName;
 	}
diff --git a/Engine/Plugins/Experimental/Mutable/Source/MutableTools/Private/MuT/NodeImageTablePrivate.h b/Engine/Plugins/Experimental/Mutable/Source/MutableTools/Private/MuT/NodeImageTablePrivate.h
--- a/Engine/Plugins/Experimental/Mutable/Source/MutableTools/Private/MuT/NodeImageTablePrivate.h
+++ b/Engine/Plugins/Experimental/Mutable/Source/MutableTools/Private/MuT/NodeImageTablePrivate.h
@@ -22,14 +22,14 @@ namespace mu
 
 		static NODE_TYPE s_type;
 
-		string m_parameterName;
+		FString m_parameterName;
 		TablePtr m_pTable;
-		string m_columnName;
+		FString m_columnName;
 
 		//!
 		void Serialise( OutputArchive& arch ) const
 		{
-            uint32_t ver = 1;
+            uint32_t ver = 2;
 			arch << ver;
 
 			arch << m_parameterName;
@@ -42,11 +42,31 @@ namespace mu
 		{
             uint32_t ver;
 			arch >> ver;
-            check(ver==1);
+			check(ver>=1 && ver<= 2);
+
+			if (ver == 1)
+			{
+				std::string Temp;
+				arch >> Temp;
+				m_parameterName = Temp.c_str();
+			}
+			else
+			{
+				arch >> m_parameterName;
+			}
 
-			arch >> m_parameterName;
 			arch >> m_pTable;
-			arch >> m_columnName;
+
+			if (ver == 1)
+			{
+				std::string Temp;
+				arch >> Temp;
+				m_columnName = Temp.c_str();
+			}
+			else
+			{
+				arch >> m_columnName;
+			}
 		}
 
 	};
diff --git a/Engine/Plugins/Experimental/Mutable/Source/MutableTools/Private/MuT/NodeImageVariationPrivate.h b/Engine/Plugins/Experimental/Mutable/Source/MutableTools/Private/MuT/NodeImageVariationPrivate.h
--- a/Engine/Plugins/Experimental/Mutable/Source/MutableTools/Private/MuT/NodeImageVariationPrivate.h
+++ b/Engine/Plugins/Experimental/Mutable/Source/MutableTools/Private/MuT/NodeImageVariationPrivate.h
@@ -24,15 +24,15 @@ namespace mu
 
         NodeImagePtr m_defaultImage;
 
-        struct VARIATION
+        struct FVariation
         {
             NodeImagePtr m_image;
-            string m_tag;
+			FString m_tag;
 
             //!
             void Serialise( OutputArchive& arch ) const
             {
-                uint32_t ver = 1;
+                uint32_t ver = 2;
                 arch << ver;
 
                 arch << m_tag;
@@ -43,14 +43,23 @@ namespace mu
             {
                 uint32_t ver = 0;
                 arch >> ver;
-                check( ver == 1 );
-
-                arch >> m_tag;
-                arch >> m_image;
+				check(ver >= 1 && ver <= 2);
+
+				if (ver <= 1)
+				{
+					std::string Temp;
+					arch >> Temp;
+					m_tag = Temp.c_str();
+				}
+				else
+				{
+					arch >> m_tag;
+				}
+				arch >> m_image;
             }
         };
 
-		TArray<VARIATION> m_variations;
+		TArray<FVariation> m_variations;
 
         //!
         void Serialise( OutputArchive& arch ) const
diff --git a/Engine/Plugins/Experimental/Mutable/Source/MutableTools/Private/MuT/NodeMeshClipMorphPlanePrivate.h b/Engine/Plugins/Experimental/Mutable/Source/MutableTools/Private/MuT/NodeMeshClipMorphPlanePrivate.h
--- a/Engine/Plugins/Experimental/Mutable/Source/MutableTools/Private/MuT/NodeMeshClipMorphPlanePrivate.h
+++ b/Engine/Plugins/Experimental/Mutable/Source/MutableTools/Private/MuT/NodeMeshClipMorphPlanePrivate.h
@@ -66,7 +66,7 @@ namespace mu
 		vec3f m_selectionBoxRadius;
 		uint16 m_vertexSelectionBone;
 
-		TArray<mu::string> m_tags;
+		TArray<FString> m_tags;
 
 		// Max distance a vertex can have to the bone in order to be affected. A negative value
 		// means no limit.
@@ -75,7 +75,7 @@ namespace mu
 		//!
 		void Serialise( OutputArchive& arch ) const
 		{
-            uint32_t ver = 4;
+            uint32_t ver = 5;
 			arch << ver;
 
 			arch << m_pSource;
@@ -99,7 +99,7 @@ namespace mu
 		{
             uint32_t ver;
 			arch >> ver;
-            check(ver==3);
+            check(ver>=3&&ver<=5);
 
 			arch >> m_pSource;
 			arch >> m_origin;
@@ -112,7 +112,21 @@ namespace mu
 			arch >> m_vertexSelectionType;
 			arch >> m_selectionBoxOrigin;
 			arch >> m_selectionBoxRadius;
-			arch >> m_tags;
+
+			if (ver >= 5)
+			{
+				arch >> m_tags;
+			}
+			else
+			{
+				TArray<std::string> Temp;
+				arch >> Temp;
+				m_tags.SetNum(Temp.Num());
+				for (int32 i=0;i<Temp.Num();++i)
+				{
+					m_tags[i] = Temp[i].c_str();
+				}
+			}
 
 			if (ver >= 4)
 			{
diff --git a/Engine/Plugins/Experimental/Mutable/Source/MutableTools/Private/MuT/NodeMeshClipWithMeshPrivate.h b/Engine/Plugins/Experimental/Mutable/Source/MutableTools/Private/MuT/NodeMeshClipWithMeshPrivate.h
--- a/Engine/Plugins/Experimental/Mutable/Source/MutableTools/Private/MuT/NodeMeshClipWithMeshPrivate.h
+++ b/Engine/Plugins/Experimental/Mutable/Source/MutableTools/Private/MuT/NodeMeshClipWithMeshPrivate.h
@@ -26,12 +26,12 @@ namespace mu
 		NodeMeshPtr m_pSource;
 		NodeMeshPtr m_pClipMesh;
 
-		TArray<mu::string> m_tags;
+		TArray<FString> m_tags;
 
 		//!
 		void Serialise( OutputArchive& arch ) const
 		{
-            uint32 ver = 1;
+            uint32 ver = 2;
 			arch << ver;
 
 			arch << m_pSource;
@@ -44,10 +44,25 @@ namespace mu
 		{
             uint32 ver;
 			arch >> ver;
+			check(ver<=2)
 
 			arch >> m_pSource;
 			arch >> m_pClipMesh;
-			arch >> m_tags;
+
+			if (ver <= 1)
+			{
+				TArray<std::string> Temp;
+				arch >> Temp;
+				m_tags.SetNum(Temp.Num());
+				for (int32 i = 0; i < Temp.Num(); ++i)
+				{
+					m_tags[i] = Temp[i].c_str();
+				}
+			}
+			else
+			{
+				arch >> m_tags;
+			}
 		}
 
 		// NodeMesh::Private interface
diff --git a/Engine/Plugins/Experimental/Mutable/Source/MutableTools/Private/MuT/NodeMeshTable.cpp b/Engine/Plugins/Experimental/Mutable/Source/MutableTools/Private/MuT/NodeMeshTable.cpp
--- a/Engine/Plugins/Experimental/Mutable/Source/MutableTools/Private/MuT/NodeMeshTable.cpp
+++ b/Engine/Plugins/Experimental/Mutable/Source/MutableTools/Private/MuT/NodeMeshTable.cpp
@@ -58,28 +58,14 @@ namespace mu
 	//---------------------------------------------------------------------------------------------
 	// Own Interface
 	//---------------------------------------------------------------------------------------------
-	const char* NodeMeshTable::GetColumn() const
+	void NodeMeshTable::SetColumn( const FString& strName )
 	{
-		return m_pD->m_columnName.c_str();
+		m_pD->m_columnName = strName;
 	}
 
 
 	//---------------------------------------------------------------------------------------------
-	void NodeMeshTable::SetColumn( const char* strName )
-	{
-		if (strName)
-		{
-			m_pD->m_columnName = strName;
-		}
-		else
-		{
-			m_pD->m_columnName = "";
-		}
-	}
-
-
-	//---------------------------------------------------------------------------------------------
-	void NodeMeshTable::SetParameterName( const char* strName )
+	void NodeMeshTable::SetParameterName( const FString& strName )
 	{
 		m_pD->m_parameterName = strName;
 	}
diff --git a/Engine/Plugins/Experimental/Mutable/Source/MutableTools/Private/MuT/NodeMeshVariationPrivate.h b/Engine/Plugins/Experimental/Mutable/Source/MutableTools/Private/MuT/NodeMeshVariationPrivate.h
--- a/Engine/Plugins/Experimental/Mutable/Source/MutableTools/Private/MuT/NodeMeshVariationPrivate.h
+++ b/Engine/Plugins/Experimental/Mutable/Source/MutableTools/Private/MuT/NodeMeshVariationPrivate.h
@@ -24,15 +24,15 @@ namespace mu
 
         NodeMeshPtr m_defaultMesh;
 
-        struct VARIATION
+        struct FVariation
         {
             NodeMeshPtr m_mesh;
-            string m_tag;
+            FString m_tag;
 
             //!
             void Serialise( OutputArchive& arch ) const
             {
-                uint32_t ver = 1;
+                uint32_t ver = 2;
                 arch << ver;
 
                 arch << m_tag;
@@ -43,14 +43,23 @@ namespace mu
             {
                 uint32_t ver = 0;
                 arch >> ver;
-                check( ver == 1 );
-
-                arch >> m_tag;
-                arch >> m_mesh;
+				check(ver >= 1 && ver <= 2);
+
+				if (ver <= 1)
+				{
+					std::string Temp;
+					arch >> Temp;
+					m_tag = Temp.c_str();
+				}
+				else
+				{
+					arch >> m_tag;
+				}
+				arch >> m_mesh;
             }
         };
 
-        TArray<VARIATION> m_variations;
+        TArray<FVariation> m_variations;
 
         //!
         void Serialise( OutputArchive& arch ) const
diff --git a/Engine/Plugins/Experimental/Mutable/Source/MutableTools/Private/MuT/NodeModifier.cpp b/Engine/Plugins/Experimental/Mutable/Source/MutableTools/Private/MuT/NodeModifier.cpp
--- a/Engine/Plugins/Experimental/Mutable/Source/MutableTools/Private/MuT/NodeModifier.cpp
+++ b/Engine/Plugins/Experimental/Mutable/Source/MutableTools/Private/MuT/NodeModifier.cpp
@@ -84,7 +84,7 @@ namespace mu
 	}
 
 	//---------------------------------------------------------------------------------------------
-	void NodeModifier::AddTag(const char* tagName)
+	void NodeModifier::AddTag(const FString& tagName)
 	{
 		NodeModifier::Private* pD = dynamic_cast<NodeModifier::Private*>(GetBasePrivate());
 		check(pD);
@@ -93,30 +93,6 @@ namespace mu
 	}
 
 
-    //---------------------------------------------------------------------------------------------
-    int NodeModifier::GetTagCount() const
-    {
-        NodeModifier::Private* pD = dynamic_cast<NodeModifier::Private*>(GetBasePrivate());
-        check(pD);
-
-        return pD->m_tags.Num();
-    }
-
-
-    //---------------------------------------------------------------------------------------------
-    const char* NodeModifier::GetTag( int i ) const
-    {
-        NodeModifier::Private* pD = dynamic_cast<NodeModifier::Private*>(GetBasePrivate());
-        check(pD);
-
-        if (i>=0 && i<GetTagCount())
-        {
-            return pD->m_tags[i].c_str();
-        }
-        return nullptr;
-    }
-
-
 	//---------------------------------------------------------------------------------------------
 	void NodeModifier::SetStage(bool bBeforeNormalOperation)
 	{
diff --git a/Engine/Plugins/Experimental/Mutable/Source/MutableTools/Private/MuT/NodeModifierPrivate.h b/Engine/Plugins/Experimental/Mutable/Source/MutableTools/Private/MuT/NodeModifierPrivate.h
--- a/Engine/Plugins/Experimental/Mutable/Source/MutableTools/Private/MuT/NodeModifierPrivate.h
+++ b/Engine/Plugins/Experimental/Mutable/Source/MutableTools/Private/MuT/NodeModifierPrivate.h
@@ -14,7 +14,7 @@ namespace mu
 	public:
 
 		// Tags that target surface need to have enabled to receive this modifier.
-		TArray<mu::string> m_tags;
+		TArray<FString> m_tags;
 
 		// Wether the modifier has to be applied after the normal node operations or before
 		bool m_applyBeforeNormalOperations = true;
@@ -22,7 +22,7 @@ namespace mu
 		//!
 		void Serialise(OutputArchive& arch) const
 		{
-            uint32_t ver = 2;
+            uint32_t ver = 3;
 			arch << ver;
 
 			arch << m_tags;
@@ -34,9 +34,22 @@ namespace mu
 		{
             uint32_t ver;
 			arch >> ver;
-			check(ver == 2);
-
-			arch >> m_tags;
+			check(ver>=2 && ver<=3);
+
+			if (ver <= 2)
+			{
+				TArray<std::string> Temp;
+				arch >> Temp;
+				m_tags.SetNum(Temp.Num());
+				for( int32 i=0; i<Temp.Num(); ++i)
+				{
+					m_tags[i] = Temp[i].c_str();
+				}
+			}
+			else
+			{
+				arch >> m_tags;
+			}
 			arch >> m_applyBeforeNormalOperations;
 		}
 
diff --git a/Engine/Plugins/Experimental/Mutable/Source/MutableTools/Private/MuT/NodeObjectGroup.cpp b/Engine/Plugins/Experimental/Mutable/Source/MutableTools/Private/MuT/NodeObjectGroup.cpp
--- a/Engine/Plugins/Experimental/Mutable/Source/MutableTools/Private/MuT/NodeObjectGroup.cpp
+++ b/Engine/Plugins/Experimental/Mutable/Source/MutableTools/Private/MuT/NodeObjectGroup.cpp
@@ -73,43 +73,29 @@ namespace mu
 	//---------------------------------------------------------------------------------------------
 	// Own Interface
 	//---------------------------------------------------------------------------------------------
-	const char* NodeObjectGroup::GetName() const
+	const FString& NodeObjectGroup::GetName() const
 	{
-		return m_pD->m_name.c_str();
+		return m_pD->Name;
 	}
 
 
 	//---------------------------------------------------------------------------------------------
-	void NodeObjectGroup::SetName( const char* strName )
+	void NodeObjectGroup::SetName( const FString& Name )
 	{
-		if( strName )
-		{
-			m_pD->m_name = strName;
-		}
-		else
-		{
-			m_pD->m_name = "";
-		}
+		m_pD->Name = Name;
 	}
 
 
-	const char* NodeObjectGroup::GetUid() const
+	const FString& NodeObjectGroup::GetUid() const
 	{
-		return m_pD->m_uid.c_str();
+		return m_pD->Uid;
 	}
 
 
 	//---------------------------------------------------------------------------------------------
-	void NodeObjectGroup::SetUid( const char* strUid )
+	void NodeObjectGroup::SetUid( const FString& Uid )
 	{
-		if( strUid )
-		{
-			m_pD->m_uid = strUid;
-		}
-		else
-		{
-			m_pD->m_uid = "";
-		}
+		m_pD->Uid = Uid;
 	}
 
 
diff --git a/Engine/Plugins/Experimental/Mutable/Source/MutableTools/Private/MuT/NodeObjectGroupPrivate.h b/Engine/Plugins/Experimental/Mutable/Source/MutableTools/Private/MuT/NodeObjectGroupPrivate.h
--- a/Engine/Plugins/Experimental/Mutable/Source/MutableTools/Private/MuT/NodeObjectGroupPrivate.h
+++ b/Engine/Plugins/Experimental/Mutable/Source/MutableTools/Private/MuT/NodeObjectGroupPrivate.h
@@ -23,8 +23,8 @@ namespace mu
 
 		static NODE_TYPE s_type;
 
-		string m_name;
-		string m_uid;
+		FString Name;
+		FString Uid;
 
 		CHILD_SELECTION m_type;
 
@@ -36,12 +36,12 @@ namespace mu
 		//!
 		void Serialise( OutputArchive& arch ) const
 		{
-            uint32_t ver = 1;
+            uint32_t ver = 2;
 			arch << ver;
 
 			arch << m_type;
-			arch << m_name;
-			arch << m_uid;
+			arch << Name;
+			arch << Uid;
 			arch << m_children;
 		}
 
@@ -50,11 +50,23 @@ namespace mu
 		{
             uint32_t ver;
 			arch >> ver;
-            check(ver==1);
+            check(ver>=1 && ver<=2);
 
 			arch >> m_type;
-			arch >> m_name;
-            arch >> m_uid;
+			if (ver <= 1)
+			{
+				std::string Temp;
+				arch >> Temp;
+				Name = Temp.c_str();
+				arch >> Temp;
+				Uid = Temp.c_str();
+			}
+			else
+			{
+				arch >> Name;
+				arch >> Uid;
+			}
+
 			arch >> m_children;
 		}
 
diff --git a/Engine/Plugins/Experimental/Mutable/Source/MutableTools/Private/MuT/NodeObjectNew.cpp b/Engine/Plugins/Experimental/Mutable/Source/MutableTools/Private/MuT/NodeObjectNew.cpp
--- a/Engine/Plugins/Experimental/Mutable/Source/MutableTools/Private/MuT/NodeObjectNew.cpp
+++ b/Engine/Plugins/Experimental/Mutable/Source/MutableTools/Private/MuT/NodeObjectNew.cpp
@@ -108,43 +108,29 @@ namespace mu
 	//---------------------------------------------------------------------------------------------
 	// Own Interface
 	//---------------------------------------------------------------------------------------------
-	const char* NodeObjectNew::GetName() const
+	const FString& NodeObjectNew::GetName() const
 	{
-		return m_pD->m_name.c_str();
+		return m_pD->m_name;
 	}
 
 
 	//---------------------------------------------------------------------------------------------
-	void NodeObjectNew::SetName( const char* strName )
+	void NodeObjectNew::SetName( const FString& Name )
 	{
-		if( strName )
-		{
-			m_pD->m_name = strName;
-		}
-		else
-		{
-			m_pD->m_name = "";
-		}
+		m_pD->m_name = Name;
 	}
 
 
-	const char* NodeObjectNew::GetUid() const
+	const FString& NodeObjectNew::GetUid() const
 	{
-		return m_pD->m_uid.c_str();
+		return m_pD->m_uid;
 	}
 
 
 	//---------------------------------------------------------------------------------------------
-	void NodeObjectNew::SetUid( const char* strUid )
+	void NodeObjectNew::SetUid( const FString& Uid )
 	{
-		if( strUid )
-		{
-			m_pD->m_uid = strUid;
-		}
-		else
-		{
-			m_pD->m_uid = "";
-		}
+		m_pD->m_uid = Uid;
 	}
 
 
@@ -272,29 +258,21 @@ namespace mu
 
 
 	//---------------------------------------------------------------------------------------------
-	int NodeObjectNew::GetStateCount() const
+	int32 NodeObjectNew::GetStateCount() const
 	{
 		return m_pD->m_states.Num();
 	}
 
 
 	//---------------------------------------------------------------------------------------------
-	void NodeObjectNew::SetStateCount( int c )
+	void NodeObjectNew::SetStateCount( int32 c )
 	{
 		m_pD->m_states.SetNum( c );
 	}
 
 
 	//---------------------------------------------------------------------------------------------
-	const char* NodeObjectNew::GetStateName( int s ) const
-	{
-		check( s>=0 && s<GetStateCount() );
-		return m_pD->m_states[s].m_name.c_str();
-	}
-
-
-	//---------------------------------------------------------------------------------------------
-	void NodeObjectNew::SetStateName( int s, const char* n )
+	void NodeObjectNew::SetStateName( int32 s, const FString& n )
 	{
 		check( s>=0 && s<GetStateCount() );
 		m_pD->m_states[s].m_name = n;
@@ -302,7 +280,7 @@ namespace mu
 
 
 	//---------------------------------------------------------------------------------------------
-	bool NodeObjectNew::HasStateParam( int s, const char* param ) const
+	bool NodeObjectNew::HasStateParam( int32 s, const FString& param ) const
 	{
 		check( s>=0 && s<GetStateCount() );
 		return m_pD->m_states[s].m_runtimeParams.Contains( param );
@@ -310,7 +288,7 @@ namespace mu
 
 
 	//---------------------------------------------------------------------------------------------
-	void NodeObjectNew::AddStateParam( int s, const char* param )
+	void NodeObjectNew::AddStateParam( int32 s, const FString& param )
 	{
 		check( s>=0 && s<GetStateCount() );
 
@@ -322,7 +300,7 @@ namespace mu
 
 
 	//---------------------------------------------------------------------------------------------
-	void NodeObjectNew::RemoveStateParam( int s, const char* param )
+	void NodeObjectNew::RemoveStateParam( int32 s, const FString& param )
 	{
 		check( s>=0 && s<GetStateCount() );
 
@@ -351,7 +329,7 @@ namespace mu
     }
 
     //---------------------------------------------------------------------------------------------
-	void NodeObjectNew::AddExtensionDataNode(NodeExtensionDataPtr Node, const char* Name)
+	void NodeObjectNew::AddExtensionDataNode(NodeExtensionDataPtr Node, const FString& Name)
 	{
 		NodeObjectNew::Private::NamedExtensionDataNode& Entry = m_pD->m_extensionDataNodes.AddDefaulted_GetRef();
 		Entry.Node = Node;
diff --git a/Engine/Plugins/Experimental/Mutable/Source/MutableTools/Private/MuT/NodeObjectNewPrivate.h b/Engine/Plugins/Experimental/Mutable/Source/MutableTools/Private/MuT/NodeObjectNewPrivate.h
--- a/Engine/Plugins/Experimental/Mutable/Source/MutableTools/Private/MuT/NodeObjectNewPrivate.h
+++ b/Engine/Plugins/Experimental/Mutable/Source/MutableTools/Private/MuT/NodeObjectNewPrivate.h
@@ -25,8 +25,8 @@ namespace mu
 
 		static NODE_TYPE s_type;
 
-		string m_name;
-		string m_uid;
+		FString m_name;
+		FString m_uid;
 
 		TArray<NodeLODPtr> m_lods;
 
@@ -34,8 +34,8 @@ namespace mu
 
 		struct NamedExtensionDataNode
 		{
-			NodeExtensionDataPtr Node;
-			string Name;
+			Ptr<NodeExtensionData> Node;
+			FString Name;
 
 			void Serialise(OutputArchive& arch) const
 			{
@@ -60,7 +60,7 @@ namespace mu
 		//!
 		void Serialise( OutputArchive& arch ) const
 		{
-            uint32_t ver = 3;
+            uint32_t ver = 4;
 			arch << ver;
 
 			arch << m_name;
@@ -76,18 +76,43 @@ namespace mu
 		{
             uint32_t ver;
 			arch >> ver;
-            check(ver >= 2);
+            check(ver>=2 && ver<=4);
+
+			if (ver <= 3)
+			{
+				std::string Temp;
+				arch >> Temp;
+				m_name = Temp.c_str();
+				arch >> Temp;
+				m_uid = Temp.c_str();
+			}
+			else
+			{
+				arch >> m_name;
+				arch >> m_uid;
+			}
 
-			arch >> m_name;
-            arch >> m_uid;
 			arch >> m_lods;
 			arch >> m_children;
 			arch >> m_states;
 
-			if (ver >= 3)
+			if (ver >= 4)
 			{
 				arch >> m_extensionDataNodes;
 			}
+			else if (ver >= 3)
+			{
+				int32 Num = 0;
+				arch >> Num;
+				m_extensionDataNodes.SetNum(Num);
+				for (int i=0;i<Num; ++i)
+				{
+					std::string Temp;
+					arch >> Temp;
+					m_extensionDataNodes[i].Name = Temp.c_str();;
+					arch >> m_extensionDataNodes[i].Node;
+				}
+			}
 		}
 
         //! Return true if the given component is set in any lod of this object.
diff --git a/Engine/Plugins/Experimental/Mutable/Source/MutableTools/Private/MuT/NodeProjectorParameter.cpp b/Engine/Plugins/Experimental/Mutable/Source/MutableTools/Private/MuT/NodeProjectorParameter.cpp
--- a/Engine/Plugins/Experimental/Mutable/Source/MutableTools/Private/MuT/NodeProjectorParameter.cpp
+++ b/Engine/Plugins/Experimental/Mutable/Source/MutableTools/Private/MuT/NodeProjectorParameter.cpp
@@ -65,76 +65,24 @@ namespace mu
 	//---------------------------------------------------------------------------------------------
 	// Own Interface
 	//---------------------------------------------------------------------------------------------
-	const char* NodeProjectorParameter::GetName() const
+	void NodeProjectorParameter::SetName( const FString& strName )
 	{
-		return m_pD->m_name.c_str();
+		m_pD->m_name = strName;
 	}
 
 
-	//---------------------------------------------------------------------------------------------
-	void NodeProjectorParameter::SetName( const char* strName )
-	{
-		if ( strName )
-		{
-			m_pD->m_name = strName;
-		}
-		else
-		{
-			m_pD->m_name = "";
-		}
-	}
-
-
-	const char* NodeProjectorParameter::GetUid() const
+	const FString& NodeProjectorParameter::GetUid() const
 	{
-		return m_pD->m_uid.c_str();
+		return m_pD->m_uid;
 	}
 
 
-	//---------------------------------------------------------------------------------------------
-	void NodeProjectorParameter::SetUid( const char* strUid )
+	void NodeProjectorParameter::SetUid( const FString& strUid )
 	{
-		if ( strUid )
-		{
-			m_pD->m_uid = strUid;
-		}
-		else
-		{
-			m_pD->m_uid = "";
-		}
+		m_pD->m_uid = strUid;
 	}
 
 
-	//---------------------------------------------------------------------------------------------
-    void NodeProjectorParameter::GetDefaultValue( PROJECTOR_TYPE* pType,
-                                                  float* pPosX, float* pPosY, float* pPosZ,
-												  float* pDirX, float* pDirY, float* pDirZ,
-												  float* pUpX, float* pUpY, float* pUpZ,
-                                                  float* pScaleU, float* pScaleV, float* pScaleW,
-                                                  float* pProjectionAngle ) const
-	{
-        if (pType) *pType = m_pD->m_type;
-
-        if (pPosX) *pPosX = m_pD->m_position[0];
-        if (pPosY) *pPosY = m_pD->m_position[1];
-        if (pPosZ) *pPosZ = m_pD->m_position[2];
-
-        if (pDirX) *pDirX = m_pD->m_direction[0];
-        if (pDirY) *pDirY = m_pD->m_direction[1];
-        if (pDirZ) *pDirZ = m_pD->m_direction[2];
-
-        if (pUpX) *pUpX = m_pD->m_up[0];
-        if (pUpY) *pUpY = m_pD->m_up[1];
-        if (pUpZ) *pUpZ = m_pD->m_up[2];
-
-        if (pScaleU) *pScaleU = m_pD->m_scale[0];
-        if (pScaleV) *pScaleV = m_pD->m_scale[1];
-        if (pScaleW) *pScaleW = m_pD->m_scale[2];
-
-        if (pProjectionAngle) *pProjectionAngle = m_pD->m_projectionAngle;
-    }
-
-
 	//---------------------------------------------------------------------------------------------
     void NodeProjectorParameter::SetDefaultValue( PROJECTOR_TYPE type,
                                                   float posX, float posY, float posZ,
diff --git a/Engine/Plugins/Experimental/Mutable/Source/MutableTools/Private/MuT/NodeProjectorPrivate.h b/Engine/Plugins/Experimental/Mutable/Source/MutableTools/Private/MuT/NodeProjectorPrivate.h
--- a/Engine/Plugins/Experimental/Mutable/Source/MutableTools/Private/MuT/NodeProjectorPrivate.h
+++ b/Engine/Plugins/Experimental/Mutable/Source/MutableTools/Private/MuT/NodeProjectorPrivate.h
@@ -77,8 +77,8 @@ namespace mu
 
 		static NODE_TYPE s_type;
 
-		string m_name;
-		string m_uid;
+		FString m_name;
+		FString m_uid;
 
 		TArray<Ptr<NodeRange>> m_ranges;
 
@@ -87,7 +87,7 @@ namespace mu
 		{
 			NodeProjectorConstant::Private::Serialise( arch );
 
-            uint32_t ver = 2;
+            uint32_t ver = 3;
 			arch << ver;
 
 			arch << m_name;
@@ -102,13 +102,24 @@ namespace mu
 
             uint32_t ver;
 			arch >> ver;
-            check(ver<=2);
-
-			arch >> m_name;
-            if (ver>=1)
-            {
-                arch >> m_uid;
-            }
+            check(ver<=3);
+
+			if (ver <= 2)
+			{
+				std::string Temp;
+				arch >> Temp;
+				m_name = Temp.c_str();
+				if (ver >= 1)
+				{
+					arch >> Temp;
+					m_uid = Temp.c_str();
+				}
+			}
+			else
+			{
+				arch >> m_name;
+				arch >> m_uid;
+			}
 
             if (ver>=2)
             {
diff --git a/Engine/Plugins/Experimental/Mutable/Source/MutableTools/Private/MuT/NodeRangeFromScalar.cpp b/Engine/Plugins/Experimental/Mutable/Source/MutableTools/Private/MuT/NodeRangeFromScalar.cpp
--- a/Engine/Plugins/Experimental/Mutable/Source/MutableTools/Private/MuT/NodeRangeFromScalar.cpp
+++ b/Engine/Plugins/Experimental/Mutable/Source/MutableTools/Private/MuT/NodeRangeFromScalar.cpp
@@ -81,23 +81,16 @@ namespace mu
 
 
 	//---------------------------------------------------------------------------------------------
-    const char* NodeRangeFromScalar::GetName() const
+    const FString& NodeRangeFromScalar::GetName() const
 	{
-        return m_pD->m_name.c_str();
+        return m_pD->m_name;
 	}
 
 
 	//---------------------------------------------------------------------------------------------
-    void NodeRangeFromScalar::SetName( const char* strName )
+    void NodeRangeFromScalar::SetName( const FString& strName )
 	{
-        if (!strName)
-        {
-            m_pD->m_name.clear();
-        }
-        else
-        {
-            m_pD->m_name = strName;
-        }
+        m_pD->m_name = strName;
 	}
 
 }
diff --git a/Engine/Plugins/Experimental/Mutable/Source/MutableTools/Private/MuT/NodeRangeFromScalarPrivate.h b/Engine/Plugins/Experimental/Mutable/Source/MutableTools/Private/MuT/NodeRangeFromScalarPrivate.h
--- a/Engine/Plugins/Experimental/Mutable/Source/MutableTools/Private/MuT/NodeRangeFromScalarPrivate.h
+++ b/Engine/Plugins/Experimental/Mutable/Source/MutableTools/Private/MuT/NodeRangeFromScalarPrivate.h
@@ -23,12 +23,12 @@ namespace mu
 		static NODE_TYPE s_type;
 
         NodeScalarPtr m_pSize;
-        string m_name;
+        FString m_name;
 
 		//!
 		void Serialise( OutputArchive& arch ) const
 		{
-            uint32_t ver = 0;
+            uint32_t ver = 1;
 			arch << ver;
 
             arch << m_name;
@@ -38,11 +38,21 @@ namespace mu
 		//!
 		void Unserialise( InputArchive& arch )
 		{
-            uint32_t ver;
+            uint32 ver;
 			arch >> ver;
-			check(ver==0);
+			check(ver<=1);
+
+			if (ver == 0)
+			{
+				std::string Temp;
+				arch >> Temp;
+				m_name = Temp.c_str();
+			}
+			else
+			{
+				arch >> m_name;
+			}
 
-            arch >> m_name;
             arch >> m_pSize;
 		}
 	};
diff --git a/Engine/Plugins/Experimental/Mutable/Source/MutableTools/Private/MuT/NodeScalarEnumParameter.cpp b/Engine/Plugins/Experimental/Mutable/Source/MutableTools/Private/MuT/NodeScalarEnumParameter.cpp
--- a/Engine/Plugins/Experimental/Mutable/Source/MutableTools/Private/MuT/NodeScalarEnumParameter.cpp
+++ b/Engine/Plugins/Experimental/Mutable/Source/MutableTools/Private/MuT/NodeScalarEnumParameter.cpp
@@ -56,43 +56,29 @@ namespace mu
 
 
 	//---------------------------------------------------------------------------------------------
-	const char* NodeScalarEnumParameter::GetName() const
+	const FString& NodeScalarEnumParameter::GetName() const
 	{
-		return m_pD->m_name.c_str();
+		return m_pD->m_name;
 	}
 
 
 	//---------------------------------------------------------------------------------------------
-	void NodeScalarEnumParameter::SetName( const char* strName )
+	void NodeScalarEnumParameter::SetName( const FString& strName )
 	{
-		if ( strName )
-		{
-			m_pD->m_name = strName;
-		}
-		else
-		{
-			m_pD->m_name = "";
-		}
+		m_pD->m_name = strName;
 	}
 
 
-	const char* NodeScalarEnumParameter::GetUid() const
+	const FString& NodeScalarEnumParameter::GetUid() const
 	{
-		return m_pD->m_uid.c_str();
+		return m_pD->m_uid;
 	}
 
 
 	//---------------------------------------------------------------------------------------------
-	void NodeScalarEnumParameter::SetUid( const char* strUid )
+	void NodeScalarEnumParameter::SetUid( const FString& strUid )
 	{
-		if ( strUid )
-		{
-			m_pD->m_uid = strUid;
-		}
-		else
-		{
-			m_pD->m_uid = "";
-		}
+		m_pD->m_uid = strUid;
 	}
 
 
@@ -125,7 +111,7 @@ namespace mu
 
 
 	//---------------------------------------------------------------------------------------------
-	void NodeScalarEnumParameter::SetValue( int i, float value, const char* strName )
+	void NodeScalarEnumParameter::SetValue( int i, float value, const FString& strName )
 	{
 		check( i>=0 && i<(int)m_pD->m_options.Num() );
 		m_pD->m_options[i].name = strName;
diff --git a/Engine/Plugins/Experimental/Mutable/Source/MutableTools/Private/MuT/NodeScalarEnumParameterPrivate.h b/Engine/Plugins/Experimental/Mutable/Source/MutableTools/Private/MuT/NodeScalarEnumParameterPrivate.h
--- a/Engine/Plugins/Experimental/Mutable/Source/MutableTools/Private/MuT/NodeScalarEnumParameterPrivate.h
+++ b/Engine/Plugins/Experimental/Mutable/Source/MutableTools/Private/MuT/NodeScalarEnumParameterPrivate.h
@@ -23,20 +23,19 @@ namespace mu
 
 		static NODE_TYPE s_type;
 
-		string m_name;
-		string m_uid;
-		int m_defaultValue = 0;
-		PARAMETER_DETAILED_TYPE m_detailedType = PARAMETER_DETAILED_TYPE::UNKNOWN;
+		FString m_name;
+		FString m_uid;
+		int32 m_defaultValue = 0;
 
 		struct OPTION
 		{
-			string name;
+			FString name;
 			float value;
 
 			//!
 			void Serialise( OutputArchive& arch ) const
 			{
-                uint32_t ver = 0;
+                uint32 ver = 1;
 				arch << ver;
 
 				arch << name;
@@ -46,11 +45,20 @@ namespace mu
 			//!
 			void Unserialise( InputArchive& arch )
 			{
-                uint32_t ver;
+                uint32 ver;
 				arch >> ver;
-                check(ver<=0);
-
-				arch >> name;
+                check(ver<=1);
+
+				if (ver == 0)
+				{
+					std::string Temp;
+					arch >> Temp;
+					name = Temp.c_str();
+				}
+				else
+				{
+					arch >> name;
+				}
 				arch >> value;
 			}
 		};
@@ -62,13 +70,12 @@ namespace mu
 		//!
 		void Serialise( OutputArchive& arch ) const
 		{
-            uint32_t ver = 2;
+            uint32_t ver = 4;
 			arch << ver;
 
 			arch << m_name;
 			arch << m_uid;
 			arch << m_defaultValue;
-			arch << m_detailedType;
             arch << m_options;
             arch << m_ranges;
         }
@@ -78,12 +85,28 @@ namespace mu
 		{
             uint32_t ver;
 			arch >> ver;
-            check(ver==2);
+            check(ver>=2 && ver<=4);
+
+			if (ver <= 3)
+			{
+				std::string Temp;
+				arch >> Temp;
+				m_name = Temp.c_str();
+				arch >> Temp;
+				m_uid = Temp.c_str();
+			}
+			else
+			{
+				arch >> m_name;
+				arch >> m_uid;
+			}
 
-			arch >> m_name;
-            arch >> m_uid;
 			arch >> m_defaultValue;
-			arch >> m_detailedType;
+			if (ver <= 2)
+			{
+				int32 Dummy = 0;
+				arch >> Dummy;
+			}
 			arch >> m_options;
             arch >> m_ranges;
         }
diff --git a/Engine/Plugins/Experimental/Mutable/Source/MutableTools/Private/MuT/NodeScalarParameter.cpp b/Engine/Plugins/Experimental/Mutable/Source/MutableTools/Private/MuT/NodeScalarParameter.cpp
--- a/Engine/Plugins/Experimental/Mutable/Source/MutableTools/Private/MuT/NodeScalarParameter.cpp
+++ b/Engine/Plugins/Experimental/Mutable/Source/MutableTools/Private/MuT/NodeScalarParameter.cpp
@@ -60,13 +60,6 @@ namespace mu
     }
 
 
-	//---------------------------------------------------------------------------------------------
-	const char* NodeScalarParameter::GetName() const
-	{
-		return m_pD->m_name.c_str();
-	}
-
-
 	//---------------------------------------------------------------------------------------------
 	void NodeScalarParameter::SetName( const char* strName )
 	{
@@ -81,12 +74,6 @@ namespace mu
 	}
 
 
-	const char* NodeScalarParameter::GetUid() const
-	{
-		return m_pD->m_uid.c_str();
-	}
-
-
 	//---------------------------------------------------------------------------------------------
 	void NodeScalarParameter::SetUid( const char* strUid )
 	{
@@ -101,13 +88,6 @@ namespace mu
 	}
 
 
-	//---------------------------------------------------------------------------------------------
-	float NodeScalarParameter::GetDefaultValue() const
-	{
-		return m_pD->m_defaultValue;
-	}
-
-
 	//---------------------------------------------------------------------------------------------
 	void NodeScalarParameter::SetDefaultValue( float v )
 	{
@@ -115,20 +95,6 @@ namespace mu
 	}
 
 
-	//---------------------------------------------------------------------------------------------
-	PARAMETER_DETAILED_TYPE NodeScalarParameter::GetDetailedType() const
-	{
-		return m_pD->m_detailedType;
-	}
-
-
-	//---------------------------------------------------------------------------------------------
-	void NodeScalarParameter::SetDetailedType( PARAMETER_DETAILED_TYPE t )
-	{
-		m_pD->m_detailedType = t;
-	}
-
-
     //---------------------------------------------------------------------------------------------
     void NodeScalarParameter::SetRangeCount( int i )
     {
diff --git a/Engine/Plugins/Experimental/Mutable/Source/MutableTools/Private/MuT/NodeScalarParameterPrivate.h b/Engine/Plugins/Experimental/Mutable/Source/MutableTools/Private/MuT/NodeScalarParameterPrivate.h
--- a/Engine/Plugins/Experimental/Mutable/Source/MutableTools/Private/MuT/NodeScalarParameterPrivate.h
+++ b/Engine/Plugins/Experimental/Mutable/Source/MutableTools/Private/MuT/NodeScalarParameterPrivate.h
@@ -25,22 +25,20 @@ namespace mu
 		static NODE_TYPE s_type;
 
 		float m_defaultValue = 0.0f;
-		string m_name;
-		string m_uid;
-		PARAMETER_DETAILED_TYPE m_detailedType = PARAMETER_DETAILED_TYPE::UNKNOWN;
+		FString m_name;
+		FString m_uid;
 
 		TArray<Ptr<NodeRange>> m_ranges;
 
 		//!
 		void Serialise( OutputArchive& arch ) const
 		{
-            uint32_t ver = 5;
+            uint32_t ver = 7;
 			arch << ver;
 
 			arch << m_defaultValue;
 			arch << m_name;
 			arch << m_uid;
-			arch << m_detailedType;
             arch << m_ranges;
         }
 
@@ -49,17 +47,35 @@ namespace mu
 		{
             uint32_t ver;
 			arch >> ver;
-            check(ver>=4 && ver<=5);
+            check(ver>=4 && ver<=7);
 
 			arch >> m_defaultValue;
-			arch >> m_name;
-            arch >> m_uid;
+			if (ver <= 6)
+			{
+				std::string Temp;
+				arch >> Temp;
+				m_name = Temp.c_str();
+				arch >> Temp;
+				m_uid = Temp.c_str();
+			}
+			else
+			{
+				arch >> m_name;
+				arch >> m_uid;
+			}
+
 			if (ver < 5)
 			{
 				TArray<Ptr<NodeImage>> TempAdditionalImages;
 				arch >> TempAdditionalImages;
 			}
-            arch >> m_detailedType;
+
+			if (ver <= 5)
+			{
+				int32 Dummy=0;
+				arch >> Dummy;
+			}
+
             arch >> m_ranges;
         }
 	};
diff --git a/Engine/Plugins/Experimental/Mutable/Source/MutableTools/Private/MuT/NodeScalarTable.cpp b/Engine/Plugins/Experimental/Mutable/Source/MutableTools/Private/MuT/NodeScalarTable.cpp
--- a/Engine/Plugins/Experimental/Mutable/Source/MutableTools/Private/MuT/NodeScalarTable.cpp
+++ b/Engine/Plugins/Experimental/Mutable/Source/MutableTools/Private/MuT/NodeScalarTable.cpp
@@ -37,48 +37,31 @@ namespace mu
 
 
 	//---------------------------------------------------------------------------------------------
-	Node* NodeScalarTable::GetInputNode( int i ) const
+	Node* NodeScalarTable::GetInputNode( int ) const
 	{
-		check( i >=0 && i < 0 );
-        (void)i;
-        return 0;
+		check( false);
+        return nullptr;
 	}
 
 
 	//---------------------------------------------------------------------------------------------
-    void NodeScalarTable::SetInputNode( int i, NodePtr )
+    void NodeScalarTable::SetInputNode( int, NodePtr )
 	{
-		check( i >=0 && i < 0 );
-        (void)i;
-        //m_pD->m_pObject = dynamic_cast<NodeObject*>( pNode.get() );
+		check( false);
 	}
 
 
 	//---------------------------------------------------------------------------------------------
 	// Own Interface
 	//---------------------------------------------------------------------------------------------
-	const char* NodeScalarTable::GetColumn() const
+	void NodeScalarTable::SetColumn( const FString& strName )
 	{
-		return m_pD->m_columnName.c_str();
+		m_pD->m_columnName = strName;
 	}
 
 
 	//---------------------------------------------------------------------------------------------
-	void NodeScalarTable::SetColumn( const char* strName )
-	{
-		if (strName)
-		{
-			m_pD->m_columnName = strName;
-		}
-		else
-		{
-			m_pD->m_columnName = "";
-		}
-	}
-
-
-	//---------------------------------------------------------------------------------------------
-	void NodeScalarTable::SetParameterName( const char* strName )
+	void NodeScalarTable::SetParameterName( const FString& strName )
 	{
 		m_pD->m_parameterName = strName;
 	}
diff --git a/Engine/Plugins/Experimental/Mutable/Source/MutableTools/Private/MuT/NodeScalarTablePrivate.h b/Engine/Plugins/Experimental/Mutable/Source/MutableTools/Private/MuT/NodeScalarTablePrivate.h
--- a/Engine/Plugins/Experimental/Mutable/Source/MutableTools/Private/MuT/NodeScalarTablePrivate.h
+++ b/Engine/Plugins/Experimental/Mutable/Source/MutableTools/Private/MuT/NodeScalarTablePrivate.h
@@ -22,14 +22,14 @@ namespace mu
 
 		static NODE_TYPE s_type;
 
-		string m_parameterName;
-		TablePtr m_pTable;
-		string m_columnName;
+		FString m_parameterName;
+		Ptr<Table> m_pTable;
+		FString m_columnName;
 
 		//!
 		void Serialise( OutputArchive& arch ) const
 		{
-            uint32_t ver = 0;
+            uint32 ver = 1;
 			arch << ver;
 
 			arch << m_parameterName;
@@ -40,13 +40,33 @@ namespace mu
 		//!
 		void Unserialise( InputArchive& arch )
 		{
-            uint32_t ver;
+            uint32 ver;
 			arch >> ver;
-            check(ver==0);
+            check(ver<=1);
+
+			if (ver==0)
+			{
+				std::string Temp;
+				arch >> Temp;
+				m_parameterName = Temp.c_str();
+			}
+			else
+			{
+				arch >> m_parameterName;
+			}
 
-			arch >> m_parameterName;
 			arch >> m_pTable;
-			arch >> m_columnName;
+
+			if (ver == 0)
+			{
+				std::string Temp;
+				arch >> Temp;
+				m_columnName = Temp.c_str();
+			}
+			else
+			{
+				arch >> m_columnName;
+			}
 		}
 
 	};
diff --git a/Engine/Plugins/Experimental/Mutable/Source/MutableTools/Private/MuT/NodeScalarVariationPrivate.h b/Engine/Plugins/Experimental/Mutable/Source/MutableTools/Private/MuT/NodeScalarVariationPrivate.h
--- a/Engine/Plugins/Experimental/Mutable/Source/MutableTools/Private/MuT/NodeScalarVariationPrivate.h
+++ b/Engine/Plugins/Experimental/Mutable/Source/MutableTools/Private/MuT/NodeScalarVariationPrivate.h
@@ -23,15 +23,15 @@ namespace mu
 
         NodeScalarPtr m_defaultScalar;
 
-        struct VARIATION
+        struct FVariation
         {
             NodeScalarPtr m_scalar;
-            string m_tag;
+            FString m_tag;
 
             //!
             void Serialise( OutputArchive& arch ) const
             {
-                uint32_t ver = 1;
+                uint32 ver = 2;
                 arch << ver;
 
                 arch << m_tag;
@@ -40,21 +40,31 @@ namespace mu
 
             void Unserialise( InputArchive& arch )
             {
-                uint32_t ver = 0;
+                uint32 ver = 0;
                 arch >> ver;
-                check( ver == 1 );
-
-                arch >> m_tag;
-                arch >> m_scalar;
+				check(ver >= 1 && ver <= 2);
+
+				if (ver <= 1)
+				{
+					std::string Temp;
+					arch >> Temp;
+					m_tag = Temp.c_str();
+				}
+				else
+				{
+					arch >> m_tag;
+				}
+				
+				arch >> m_scalar;
             }
         };
 
-        TArray<VARIATION> m_variations;
+        TArray<FVariation> m_variations;
 
         //!
         void Serialise( OutputArchive& arch ) const
         {
-            uint32_t ver = 1;
+            uint32 ver = 1;
             arch << ver;
 
             arch << m_defaultScalar;
@@ -64,7 +74,7 @@ namespace mu
         //!
         void Unserialise( InputArchive& arch )
         {
-            uint32_t ver;
+            uint32 ver;
             arch >> ver;
 			check(ver == 1);
 
diff --git a/Engine/Plugins/Experimental/Mutable/Source/MutableTools/Private/MuT/NodeStringConstant.cpp b/Engine/Plugins/Experimental/Mutable/Source/MutableTools/Private/MuT/NodeStringConstant.cpp
--- a/Engine/Plugins/Experimental/Mutable/Source/MutableTools/Private/MuT/NodeStringConstant.cpp
+++ b/Engine/Plugins/Experimental/Mutable/Source/MutableTools/Private/MuT/NodeStringConstant.cpp
@@ -52,15 +52,9 @@ namespace mu
 	//---------------------------------------------------------------------------------------------
 	// Own Interface
 	//---------------------------------------------------------------------------------------------
-	const char* NodeStringConstant::GetValue() const
+	void NodeStringConstant::SetValue( const FString& v )
 	{
-		return m_pD->m_value.c_str();
-	}
-
-	//---------------------------------------------------------------------------------------------
-	void NodeStringConstant::SetValue( const char* v )
-	{
-		m_pD->m_value = v?v:"";
+		m_pD->m_value = v;
 	}
 
 
diff --git a/Engine/Plugins/Experimental/Mutable/Source/MutableTools/Private/MuT/NodeStringConstantPrivate.h b/Engine/Plugins/Experimental/Mutable/Source/MutableTools/Private/MuT/NodeStringConstantPrivate.h
--- a/Engine/Plugins/Experimental/Mutable/Source/MutableTools/Private/MuT/NodeStringConstantPrivate.h
+++ b/Engine/Plugins/Experimental/Mutable/Source/MutableTools/Private/MuT/NodeStringConstantPrivate.h
@@ -24,12 +24,12 @@ namespace mu
 
 		static NODE_TYPE s_type;
 
-		string m_value;
+		FString m_value;
 
 		//!
 		void Serialise( OutputArchive& arch ) const
 		{
-            uint32_t ver = 0;
+            uint32 ver = 1;
 			arch << ver;
 
 			arch << m_value;
@@ -38,11 +38,20 @@ namespace mu
 		//!
 		void Unserialise( InputArchive& arch )
 		{
-            uint32_t ver;
+            uint32 ver;
 			arch >> ver;
-			check(ver==0);
-
-			arch >> m_value;
+			check(ver<=1);
+
+			if (ver <= 0)
+			{
+				std::string Temp;
+				arch >> Temp;
+				m_value = Temp.c_str();
+			}
+			else
+			{
+				arch >> m_value;
+			}
 		}
 	};
 
diff --git a/Engine/Plugins/Experimental/Mutable/Source/MutableTools/Private/MuT/NodeStringParameter.cpp b/Engine/Plugins/Experimental/Mutable/Source/MutableTools/Private/MuT/NodeStringParameter.cpp
--- a/Engine/Plugins/Experimental/Mutable/Source/MutableTools/Private/MuT/NodeStringParameter.cpp
+++ b/Engine/Plugins/Experimental/Mutable/Source/MutableTools/Private/MuT/NodeStringParameter.cpp
@@ -74,13 +74,6 @@ namespace mu
     }
 
 
-	//---------------------------------------------------------------------------------------------
-	const char* NodeStringParameter::GetName() const
-	{
-		return m_pD->m_name.c_str();
-	}
-
-
 	//---------------------------------------------------------------------------------------------
 	void NodeStringParameter::SetName( const char* strName )
 	{
@@ -95,12 +88,6 @@ namespace mu
 	}
 
 
-	const char* NodeStringParameter::GetUid() const
-	{
-		return m_pD->m_uid.c_str();
-	}
-
-
 	//---------------------------------------------------------------------------------------------
 	void NodeStringParameter::SetUid( const char* strUid )
 	{
@@ -115,13 +102,6 @@ namespace mu
 	}
 
 
-	//---------------------------------------------------------------------------------------------
-	const char* NodeStringParameter::GetDefaultValue() const
-	{
-		return m_pD->m_defaultValue.c_str();
-	}
-
-
 	//---------------------------------------------------------------------------------------------
 	void NodeStringParameter::SetDefaultValue( const char* v )
 	{
@@ -129,20 +109,6 @@ namespace mu
 	}
 
 
-	//---------------------------------------------------------------------------------------------
-	PARAMETER_DETAILED_TYPE NodeStringParameter::GetDetailedType() const
-	{
-		return m_pD->m_detailedType;
-	}
-
-
-	//---------------------------------------------------------------------------------------------
-	void NodeStringParameter::SetDetailedType( PARAMETER_DETAILED_TYPE t )
-	{
-		m_pD->m_detailedType = t;
-	}
-
-
     //---------------------------------------------------------------------------------------------
     void NodeStringParameter::SetRangeCount( int i )
     {
diff --git a/Engine/Plugins/Experimental/Mutable/Source/MutableTools/Private/MuT/NodeStringParameterPrivate.h b/Engine/Plugins/Experimental/Mutable/Source/MutableTools/Private/MuT/NodeStringParameterPrivate.h
--- a/Engine/Plugins/Experimental/Mutable/Source/MutableTools/Private/MuT/NodeStringParameterPrivate.h
+++ b/Engine/Plugins/Experimental/Mutable/Source/MutableTools/Private/MuT/NodeStringParameterPrivate.h
@@ -24,10 +24,9 @@ namespace mu
 
 		static NODE_TYPE s_type;
 
-		string m_defaultValue;
-		string m_name;
-		string m_uid;
-		PARAMETER_DETAILED_TYPE m_detailedType = PARAMETER_DETAILED_TYPE::UNKNOWN;
+		FString m_defaultValue;
+		FString m_name;
+		FString m_uid;
 
         TArray<Ptr<NodeImage>> m_additionalImages;
 
@@ -36,14 +35,13 @@ namespace mu
 		//!
 		void Serialise( OutputArchive& arch ) const
 		{
-            uint32_t ver = 4;
+            uint32_t ver = 6;
 			arch << ver;
 
 			arch << m_defaultValue;
 			arch << m_name;
 			arch << m_uid;
 			arch << m_additionalImages;
-			arch << m_detailedType;
             arch << m_ranges;
         }
 
@@ -52,13 +50,31 @@ namespace mu
 		{
             uint32_t ver;
 			arch >> ver;
-            check(ver==4);
+            check(ver>=4 && ver<=6);
+
+			if (ver <= 5)
+			{
+				std::string Temp;
+				arch >> Temp;
+				m_defaultValue = Temp.c_str();
+				arch >> Temp;
+				m_name = Temp.c_str();
+				arch >> Temp;
+				m_uid = Temp.c_str();
+			}
+			else
+			{
+				arch >> m_defaultValue;
+				arch >> m_name;
+				arch >> m_uid;
+			}
 
-			arch >> m_defaultValue;
-			arch >> m_name;
-            arch >> m_uid;
 			arch >> m_additionalImages;
-            arch >> m_detailedType;
+			if (ver <= 4)
+			{
+				int32 Dummy=0;
+				arch >> Dummy;
+			}
             arch >> m_ranges;
         }
 	};
diff --git a/Engine/Plugins/Experimental/Mutable/Source/MutableTools/Private/MuT/NodeSurfaceEdit.cpp b/Engine/Plugins/Experimental/Mutable/Source/MutableTools/Private/MuT/NodeSurfaceEdit.cpp
--- a/Engine/Plugins/Experimental/Mutable/Source/MutableTools/Private/MuT/NodeSurfaceEdit.cpp
+++ b/Engine/Plugins/Experimental/Mutable/Source/MutableTools/Private/MuT/NodeSurfaceEdit.cpp
@@ -130,30 +130,12 @@ namespace mu
 
 
     //---------------------------------------------------------------------------------------------
-    void NodeSurfaceEdit::AddTag(const char* tagName)
+    void NodeSurfaceEdit::AddTag(const FString& tagName)
     {
         m_pD->m_tags.Add(tagName);
     }
 
 
-    //---------------------------------------------------------------------------------------------
-    int NodeSurfaceEdit::GetTagCount() const
-    {
-        return m_pD->m_tags.Num();
-    }
-
-
-    //---------------------------------------------------------------------------------------------
-    const char* NodeSurfaceEdit::GetTag( int i ) const
-    {
-        if (i>=0 && i<GetTagCount())
-        {
-            return m_pD->m_tags[i].c_str();
-        }
-        return nullptr;
-    }
-
-
 	//---------------------------------------------------------------------------------------------
     void NodeSurfaceEdit::SetMesh(NodePatchMesh* pNode )
 	{
diff --git a/Engine/Plugins/Experimental/Mutable/Source/MutableTools/Private/MuT/NodeSurfaceEditPrivate.h b/Engine/Plugins/Experimental/Mutable/Source/MutableTools/Private/MuT/NodeSurfaceEditPrivate.h
--- a/Engine/Plugins/Experimental/Mutable/Source/MutableTools/Private/MuT/NodeSurfaceEditPrivate.h
+++ b/Engine/Plugins/Experimental/Mutable/Source/MutableTools/Private/MuT/NodeSurfaceEditPrivate.h
@@ -37,7 +37,7 @@ namespace mu
         //! morph information. This means that the morph information needs to be generated.
         bool m_morphIsTarget = true;
 
-		struct TEXTURE
+		struct FTexture
 		{
 			NodeImagePtr m_pExtend;
 			NodePatchImagePtr m_pPatch;
@@ -56,10 +56,10 @@ namespace mu
 			}
 		};
 
-		TArray<TEXTURE> m_textures;
+		TArray<FTexture> m_textures;
 
         //! Tags in this surface edit
-		TArray<mu::string> m_tags;
+		TArray<FString> m_tags;
 
 		//! Factor of the morph
 		NodeScalarPtr m_pFactor;
@@ -67,7 +67,7 @@ namespace mu
 		//!
 		void Serialise( OutputArchive& arch ) const
 		{
-            uint32_t ver = 4;
+            uint32_t ver = 5;
 			arch << ver;
 
 			arch << m_pParent;
@@ -84,14 +84,27 @@ namespace mu
 		{
             uint32_t ver;
 			arch >> ver;
-            check( ver==4 );
+            check( ver>=4 && ver<=5 );
 
 			arch >> m_pParent;
 			arch >> m_pMesh;
 			arch >> m_pMorph;
             arch >> m_morphIsTarget;
 			arch >> m_textures;
-            arch >> m_tags;
+			if (ver <= 4)
+			{
+				TArray <std::string> Temp;
+				arch >> Temp;
+				m_tags.SetNum(Temp.Num());
+				for (int32 i=0; i<Temp.Num(); ++i)
+				{
+					m_tags[i] = Temp[i].c_str();
+				}
+			}
+			else
+			{
+				arch >> m_tags;
+			}
 			arch >> m_pFactor;
         }
 	};
diff --git a/Engine/Plugins/Experimental/Mutable/Source/MutableTools/Private/MuT/NodeSurfaceNew.cpp b/Engine/Plugins/Experimental/Mutable/Source/MutableTools/Private/MuT/NodeSurfaceNew.cpp
--- a/Engine/Plugins/Experimental/Mutable/Source/MutableTools/Private/MuT/NodeSurfaceNew.cpp
+++ b/Engine/Plugins/Experimental/Mutable/Source/MutableTools/Private/MuT/NodeSurfaceNew.cpp
@@ -162,25 +162,9 @@ namespace mu
 	//---------------------------------------------------------------------------------------------
 	// Own Interface
 	//---------------------------------------------------------------------------------------------
-    const char* NodeSurfaceNew::GetName() const
-	{
-		const char* strResult = m_pD->m_name.c_str();
-
-		return strResult;
-	}
-
-
-    //---------------------------------------------------------------------------------------------
-    void NodeSurfaceNew::SetName( const char* strName )
+    void NodeSurfaceNew::SetName( const FString& strName )
     {
-        if (strName)
-        {
-            m_pD->m_name = strName;
-        }
-        else
-        {
-            m_pD->m_name = "";
-        }
+        m_pD->m_name = strName;
     }
 
 
@@ -191,13 +175,6 @@ namespace mu
     }
 
 
-	//---------------------------------------------------------------------------------------------
-	int32 NodeSurfaceNew::GetSharedSurfaceId() const
-    {
-		return m_pD->SharedSurfaceId;
-    }
-
-
     //---------------------------------------------------------------------------------------------
     void NodeSurfaceNew::SetCustomID( uint32_t id )
     {
@@ -239,29 +216,20 @@ namespace mu
 
 
 	//---------------------------------------------------------------------------------------------
-    const char* NodeSurfaceNew::GetMeshName( int index ) const
+    const FString& NodeSurfaceNew::GetMeshName( int index ) const
 	{
 		check( index >=0 && index < (int)m_pD->m_meshes.Num() );
 
-		const char* strResult = m_pD->m_meshes[ index ].m_name.c_str();
-
-		return strResult;
+		return m_pD->m_meshes[ index ].m_name;
 	}
 
 
 	//---------------------------------------------------------------------------------------------
-    void NodeSurfaceNew::SetMeshName( int index, const char* strName )
+    void NodeSurfaceNew::SetMeshName( int index, const FString& strName )
 	{
 		check( index >=0 && index < (int)m_pD->m_meshes.Num() );
 
-		if (strName)
-		{
-			m_pD->m_meshes[ index ].m_name = strName;
-		}
-		else
-		{
-			m_pD->m_meshes[ index ].m_name = "";
-		}
+		m_pD->m_meshes[ index ].m_name = strName;
 	}
 
 
@@ -299,29 +267,20 @@ namespace mu
 
 
 	//---------------------------------------------------------------------------------------------
-    const char* NodeSurfaceNew::GetImageName( int index ) const
+    const FString& NodeSurfaceNew::GetImageName( int index ) const
 	{
 		check( index >=0 && index < (int)m_pD->m_images.Num() );
 
-		const char* strResult = m_pD->m_images[ index ].m_name.c_str();
-
-		return strResult;
+		return m_pD->m_images[ index ].m_name;
 	}
 
 
 	//---------------------------------------------------------------------------------------------
-    void NodeSurfaceNew::SetImageName( int index, const char* strName )
+    void NodeSurfaceNew::SetImageName( int index, const FString& strName )
 	{
 		check( index >=0 && index < (int)m_pD->m_images.Num() );
 
-		if (strName)
-		{
-			m_pD->m_images[ index ].m_name = strName;
-		}
-		else
-		{
-			m_pD->m_images[ index ].m_name = "";
-		}
+		m_pD->m_images[ index ].m_name = strName;
 	}
 
 
@@ -344,7 +303,7 @@ namespace mu
 
 
 	//---------------------------------------------------------------------------------------------
-	void NodeSurfaceNew::SetImageAdditionalNames(int index, const char* strMaterialName, const char* strMaterialParameterName )
+	void NodeSurfaceNew::SetImageAdditionalNames(int index, const FString& strMaterialName, const FString& strMaterialParameterName )
 	{
 		if (index >= 0 && index < (int)m_pD->m_images.Num())
 		{
@@ -392,29 +351,20 @@ namespace mu
 
 
 	//---------------------------------------------------------------------------------------------
-	const char* NodeSurfaceNew::GetVectorName(int index) const
+	const FString& NodeSurfaceNew::GetVectorName(int index) const
 	{
         check(index >= 0 && index < (int)m_pD->m_vectors.Num());
 
-		const char* strResult = m_pD->m_vectors[index].m_name.c_str();
-
-		return strResult;
+		return m_pD->m_vectors[index].m_name;
 	}
 
 
 	//---------------------------------------------------------------------------------------------
-	void NodeSurfaceNew::SetVectorName(int index, const char* strName)
+	void NodeSurfaceNew::SetVectorName(int index, const FString& strName)
 	{
         check(index >= 0 && index < (int)m_pD->m_vectors.Num());
 
-		if (strName)
-		{
-			m_pD->m_vectors[index].m_name = strName;
-		}
-		else
-		{
-			m_pD->m_vectors[index].m_name = "";
-		}
+		m_pD->m_vectors[index].m_name = strName;
 	}
 
 
@@ -452,34 +402,28 @@ namespace mu
 
 
     //---------------------------------------------------------------------------------------------
-    const char* NodeSurfaceNew::GetScalarName(int index) const
+    const FString& NodeSurfaceNew::GetScalarName(int index) const
     {
         check(index >= 0 && index < (int)m_pD->m_scalars.Num());
 
-        const char* strResult = m_pD->m_scalars[index].m_name.c_str();
-
-        return strResult;
+        return m_pD->m_scalars[index].m_name;
     }
 
 
     //---------------------------------------------------------------------------------------------
-    void NodeSurfaceNew::SetScalarName(int index, const char* strName)
+    void NodeSurfaceNew::SetScalarName(int index, const FString& strName)
     {
         check(index >= 0 && index < (int)m_pD->m_scalars.Num());
 
-        if (strName)
-        {
-            m_pD->m_scalars[index].m_name = strName;
-        }
-        else
-        {
-            m_pD->m_scalars[index].m_name = "";
-        }
+        m_pD->m_scalars[index].m_name = strName;
     }
 
 
     //---------------------------------------------------------------------------------------------
-    int NodeSurfaceNew::GetStringCount() const { return (int)m_pD->m_strings.Num(); }
+    int NodeSurfaceNew::GetStringCount() const 
+	{ 
+		return (int)m_pD->m_strings.Num(); 
+	}
 
 
     //---------------------------------------------------------------------------------------------
@@ -509,34 +453,25 @@ namespace mu
 
 
     //---------------------------------------------------------------------------------------------
-    const char* NodeSurfaceNew::GetStringName( int index ) const
+    const FString& NodeSurfaceNew::GetStringName( int index ) const
     {
         check( index >= 0 && index < (int)m_pD->m_strings.Num() );
 
-        const char* strResult = m_pD->m_strings[index].m_name.c_str();
-
-        return strResult;
+        return m_pD->m_strings[index].m_name;
     }
 
 
     //---------------------------------------------------------------------------------------------
-    void NodeSurfaceNew::SetStringName( int index, const char* strName )
+    void NodeSurfaceNew::SetStringName( int index, const FString& strName )
     {
         check( index >= 0 && index < (int)m_pD->m_strings.Num() );
 
-        if ( strName )
-        {
-            m_pD->m_strings[index].m_name = strName;
-        }
-        else
-        {
-            m_pD->m_strings[index].m_name = "";
-        }
+        m_pD->m_strings[index].m_name = strName;
     }
 
 
     //---------------------------------------------------------------------------------------------
-    int NodeSurfaceNew::Private::FindImage( const char* strName ) const
+    int NodeSurfaceNew::Private::FindImage( const FString& strName ) const
 	{		
 		for ( int32 i = 0; i<m_images.Num(); ++i )
 		{
@@ -551,7 +486,7 @@ namespace mu
 
 
 	//---------------------------------------------------------------------------------------------
-	int NodeSurfaceNew::Private::FindMesh(const char* strName) const
+	int NodeSurfaceNew::Private::FindMesh(const FString& strName) const
 	{
 		for ( int32 i = 0; i < m_meshes.Num(); ++i)
 		{
@@ -566,7 +501,7 @@ namespace mu
 
 
     //---------------------------------------------------------------------------------------------
-    int NodeSurfaceNew::Private::FindVector(const char* strName) const
+    int NodeSurfaceNew::Private::FindVector(const FString& strName) const
     {
         for ( int32 i = 0; i < m_vectors.Num(); ++i)
         {
@@ -581,7 +516,7 @@ namespace mu
 
 
     //---------------------------------------------------------------------------------------------
-    int NodeSurfaceNew::Private::FindScalar( const char* strName ) const
+    int NodeSurfaceNew::Private::FindScalar( const FString& strName ) const
     {
         for (int32 i = 0; i < m_scalars.Num(); ++i)
         {
@@ -596,7 +531,7 @@ namespace mu
 
 
     //---------------------------------------------------------------------------------------------
-    int NodeSurfaceNew::Private::FindString( const char* strName ) const
+    int NodeSurfaceNew::Private::FindString( const FString& strName ) const
     {
         for (int32 i = 0; i < m_strings.Num(); ++i)
         {
@@ -611,29 +546,11 @@ namespace mu
 
 
     //---------------------------------------------------------------------------------------------
-	void NodeSurfaceNew::AddTag(const char* tagName)
+	void NodeSurfaceNew::AddTag(const FString& tagName)
 	{
 		m_pD->m_tags.Add(tagName);
 	}
 
-
-    //---------------------------------------------------------------------------------------------
-    int NodeSurfaceNew::GetTagCount() const
-    {
-        return m_pD->m_tags.Num();
-    }
-
-
-    //---------------------------------------------------------------------------------------------
-    const char* NodeSurfaceNew::GetTag( int i ) const
-    {
-        if (i>=0 && i<GetTagCount())
-        {
-            return m_pD->m_tags[i].c_str();
-        }
-        return nullptr;
-    }
-
 }
 
 
diff --git a/Engine/Plugins/Experimental/Mutable/Source/MutableTools/Private/MuT/NodeSurfaceNewPrivate.h b/Engine/Plugins/Experimental/Mutable/Source/MutableTools/Private/MuT/NodeSurfaceNewPrivate.h
--- a/Engine/Plugins/Experimental/Mutable/Source/MutableTools/Private/MuT/NodeSurfaceNewPrivate.h
+++ b/Engine/Plugins/Experimental/Mutable/Source/MutableTools/Private/MuT/NodeSurfaceNewPrivate.h
@@ -30,13 +30,13 @@ namespace mu
 
 		static NODE_TYPE s_type;
 
-		string m_name;
-        uint32_t ExternalId =0;
+		FString m_name;
+        uint32 ExternalId =0;
         int32 SharedSurfaceId =INDEX_NONE;
 
 		struct MESH
 		{
-			string m_name;
+			FString m_name;
 			NodeMeshPtr m_pMesh;
 
 			//!
@@ -57,9 +57,9 @@ namespace mu
 
 		struct IMAGE
 		{
-			string m_name;
-			string m_materialName;
-			string m_materialParameterName;
+			FString m_name;
+			FString m_materialName;
+			FString m_materialParameterName;
 			NodeImagePtr m_pImage;
 
 			// It could be negative, to indicate no layout.
@@ -89,7 +89,7 @@ namespace mu
 
 		struct VECTOR
 		{
-			string m_name;
+			FString m_name;
 			NodeColourPtr m_pVector;
 
 			//!
@@ -110,7 +110,7 @@ namespace mu
 
         struct SCALAR
         {
-            string m_name;
+			FString m_name;
             NodeScalarPtr m_pScalar;
 
             //!
@@ -131,7 +131,7 @@ namespace mu
 
         struct STRING
         {
-            string m_name;
+			FString m_name;
             NodeStringPtr m_pString;
 
             //!
@@ -151,27 +151,27 @@ namespace mu
 		TArray<STRING> m_strings;
 
         //! Tags in this surface
-		TArray<mu::string> m_tags;
+		TArray<FString> m_tags;
 
 		//! Find an image node index by name or return -1
-		int FindImage( const char* strName ) const;
+		int FindImage( const FString& strName ) const;
 
 		//! Find a mesh node index by name or return -1
-		int FindMesh(const char* strName) const;
+		int FindMesh(const FString& strName) const;
 
         //! Find a vector node index by name or return -1
-        int FindVector(const char* strName) const;
+        int FindVector(const FString& strName) const;
 
         //! Find a scalar node index by name or return -1
-        int FindScalar( const char* strName ) const;
+        int FindScalar( const FString& strName ) const;
 
         //! Find a string node index by name or return -1
-        int FindString( const char* strName ) const;
+        int FindString( const FString& strName ) const;
 
         //!
 		void Serialise( OutputArchive& arch ) const
 		{
-            uint32_t ver = 8;
+            uint32 ver = 9;
 			arch << ver;
 
 			arch << m_name;
@@ -189,11 +189,20 @@ namespace mu
         //!
 		void Unserialise( InputArchive& arch )
 		{
-            uint32_t ver;
+            uint32 ver;
 			arch >> ver;
-            check(ver<=8);
+            check(ver<=9);
 
-			arch >> m_name;
+			std::string Temp;
+			if (ver <= 8)
+			{
+				arch >> Temp;
+				m_name = Temp.c_str();
+			}
+			else
+			{
+				arch >> m_name;
+			}
             arch >> ExternalId;
 
 			if (ver >= 8)
@@ -201,12 +210,80 @@ namespace mu
 				arch >> SharedSurfaceId;
 			}
 
-			arch >> m_meshes;
-			arch >> m_images;
-            arch >> m_tags;
-            arch >> m_vectors;
-            arch >> m_scalars;
-            arch >> m_strings;
+			if (ver <= 8)
+			{
+				int32 Num = 0;
+
+				arch >> Num;
+				m_meshes.SetNum(Num);
+				for (int32 i=0; i<Num; ++i)
+				{
+					arch >> Temp;
+					m_meshes[i].m_name = Temp.c_str();
+					arch >> m_meshes[i].m_pMesh;
+				}
+
+				arch >> Num;
+				m_images.SetNum(Num);
+				for (int32 i = 0; i < Num; ++i)
+				{
+					arch >> Temp;
+					m_images[i].m_name = Temp.c_str();
+
+					arch >> Temp;
+					m_images[i].m_materialName = Temp.c_str();
+
+					arch >> Temp;
+					m_images[i].m_materialParameterName = Temp.c_str();
+
+					arch >> m_images[i].m_pImage;
+					arch >> m_images[i].m_layoutIndex;
+				}
+
+				arch >> Num;
+				m_tags.SetNum(Num);
+				for (int32 i = 0; i < Num; ++i)
+				{
+					arch >> Temp;
+					m_tags[i] = Temp.c_str();
+				}
+
+				arch >> Num;
+				m_vectors.SetNum(Num);
+				for (int32 i = 0; i < Num; ++i)
+				{
+					arch >> Temp;
+					m_vectors[i].m_name = Temp.c_str();
+					arch >> m_vectors[i].m_pVector;
+				}
+
+				arch >> Num;
+				m_scalars.SetNum(Num);
+				for (int32 i = 0; i < Num; ++i)
+				{
+					arch >> Temp;
+					m_scalars[i].m_name = Temp.c_str();
+					arch >> m_scalars[i].m_pScalar;
+				}
+
+				arch >> Num;
+				m_strings.SetNum(Num);
+				for (int32 i = 0; i < Num; ++i)
+				{
+					arch >> Temp;
+					m_strings[i].m_name = Temp.c_str();
+					arch >> m_strings[i].m_pString;
+				}
+			}
+			else
+			{
+				arch >> m_meshes;
+				arch >> m_images;
+				arch >> m_tags;
+				arch >> m_vectors;
+				arch >> m_scalars;
+				arch >> m_strings;
+			}
         }
     };
 
