[mutable] Remove most of the legacy string usage in favor of FString and FName.

diff --git a/Engine/Plugins/Experimental/Mutable/Source/MutableTools/Private/MuT/NodeSurfaceVariationPrivate.h b/Engine/Plugins/Experimental/Mutable/Source/MutableTools/Private/MuT/NodeSurfaceVariationPrivate.h
--- a/Engine/Plugins/Experimental/Mutable/Source/MutableTools/Private/MuT/NodeSurfaceVariationPrivate.h
+++ b/Engine/Plugins/Experimental/Mutable/Source/MutableTools/Private/MuT/NodeSurfaceVariationPrivate.h
@@ -28,16 +28,16 @@ namespace mu
         TArray<NodeSurfacePtr> m_defaultSurfaces;
 		TArray<NodeModifierPtr> m_defaultModifiers;
 
-		struct VARIATION
+		struct FVariation
 		{
 			TArray<NodeSurfacePtr> m_surfaces;
 			TArray<NodeModifierPtr> m_modifiers;
-            string m_tag;
+            FString m_tag;
 
 			//!
 			void Serialise( OutputArchive& arch ) const
 			{
-                uint32_t ver = 1;
+                uint32 ver = 2;
 				arch << ver;
 
 				arch << m_tag;
@@ -47,11 +47,20 @@ namespace mu
 
 			void Unserialise( InputArchive& arch )
 			{
-                uint32_t ver = 0;
+                uint32 ver = 0;
 				arch >> ver;
-                check(ver==1);
-
-				arch >> m_tag;
+                check(ver>=1&&ver<=2);
+
+				if (ver <= 1)
+				{
+					std::string Temp;
+					arch >> Temp;
+					m_tag = Temp.c_str();
+				}
+				else
+				{
+					arch >> m_tag;
+				}
 				arch >> m_surfaces;
                 arch >> m_modifiers;
 			}
@@ -59,12 +68,12 @@ namespace mu
 
         NodeSurfaceVariation::VariationType m_type = NodeSurfaceVariation::VariationType::Tag;
 
-		TArray<VARIATION> m_variations;
+		TArray<FVariation> m_variations;
 
 		//!
 		void Serialise( OutputArchive& arch ) const
 		{
-            uint32_t ver = 3;
+            uint32 ver = 3;
 			arch << ver;
 
             arch << uint32_t(m_type);
@@ -76,11 +85,11 @@ namespace mu
 		//!
 		void Unserialise( InputArchive& arch )
 		{
-            uint32_t ver;
+            uint32 ver;
 			arch >> ver;
             check( ver==3 );
 
-            uint32_t temp;
+            uint32 temp;
             arch >> temp;
             m_type = NodeSurfaceVariation::VariationType(temp);
 
diff --git a/Engine/Plugins/Experimental/Mutable/Source/MutableTools/Private/MuT/Platform.h b/Engine/Plugins/Experimental/Mutable/Source/MutableTools/Private/MuT/Platform.h
--- a/Engine/Plugins/Experimental/Mutable/Source/MutableTools/Private/MuT/Platform.h
+++ b/Engine/Plugins/Experimental/Mutable/Source/MutableTools/Private/MuT/Platform.h
@@ -19,7 +19,6 @@
 namespace mu
 {
 
-
 	//! STL-like containers using this allocator
 	template< typename T >
 	using vector = std::vector<T>;
@@ -41,54 +40,5 @@ namespace mu
 	template< typename K, typename T >
 	using pair = std::pair<K, T>;
 
-    //-------------------------------------------------------------------------------------------------
-    inline FILE* mutable_fopen
-		(
-			const char *filename,
-			const char *mode
-		)
-	{
-		FILE* pFile = 0;
-
-#ifdef _MSC_VER
-        errno_t error = fopen_s( &pFile, filename, mode );
-        if (error!=0)
-        {
-            // It's valid to fail, just return null.
-            //check(false);
-            pFile = nullptr;
-        }
-#else
-		pFile = fopen( filename, mode );
-#endif
-
-		return pFile;
-	}
-
-
-    //-------------------------------------------------------------------------------------------------
-    inline int64_t mutable_ftell( FILE* f )
-    {
-#ifdef _MSC_VER
-        // Windows normal ftell works only with 32 bit int sizes
-        return _ftelli64(f);
-#else
-        return ftell(f);
-#endif
-    }
-
-
-    //-------------------------------------------------------------------------------------------------
-    inline int64_t mutable_fseek( FILE* f, int64_t pos, int origin )
-    {
-#ifdef _MSC_VER
-        // Windows normal fseek works only with 32 bit int sizes
-        return _fseeki64(f,pos,origin);
-#else
-        return fseek(f,pos,origin);
-#endif
-    }
-
-
 }
 
diff --git a/Engine/Plugins/Experimental/Mutable/Source/MutableTools/Private/MuT/Streams.cpp b/Engine/Plugins/Experimental/Mutable/Source/MutableTools/Private/MuT/Streams.cpp
--- a/Engine/Plugins/Experimental/Mutable/Source/MutableTools/Private/MuT/Streams.cpp
+++ b/Engine/Plugins/Experimental/Mutable/Source/MutableTools/Private/MuT/Streams.cpp
@@ -19,30 +19,28 @@
 
 namespace mu
 {
-	InputFileStream::InputFileStream( const char* strFile )
+	InputFileStream::InputFileStream( const FString& File )
 	{
 		m_pD = new Private();
 
-		m_pD->m_pFile = mutable_fopen( strFile, "rb" );
-		m_pD->m_size = 0;
-		m_pD->m_pos = 0;
-        m_pD->m_readBytes = 0;
+		TSharedPtr<IFileHandle> FilePtr(FPlatformFileManager::Get().GetPlatformFile().OpenRead(*File));
+		m_pD->File = FilePtr;
 
-        m_pD->m_buffer = reinterpret_cast<uint8*>(FMemory::Malloc(MUTABLE_STREAM_BUFFER_SIZE,16));
+		m_pD->File->SeekFromEnd();
+		m_pD->FileSize = m_pD->File->Tell();
+		m_pD->File->Seek(0);
+
+		m_pD->BytesInBuffer = 0;
+		m_pD->BufferPosition = 0;
+        m_pD->FilePosition = 0;
+
+        m_pD->Buffer.SetNum(MUTABLE_STREAM_BUFFER_SIZE);
 	}
 
 
 	//---------------------------------------------------------------------------------------------
 	InputFileStream::~InputFileStream()
 	{
-        if ( m_pD->m_pFile )
-        {
-            fclose( m_pD->m_pFile );
-        }
-
-		FMemory::Free( m_pD->m_buffer );
-        m_pD->m_buffer = nullptr;
-
         delete m_pD;
 		m_pD = nullptr;
 	}
@@ -51,22 +49,15 @@ namespace mu
 	//---------------------------------------------------------------------------------------------
 	bool InputFileStream::IsOpen() const
     {
-        return m_pD->m_pFile!=0;
-    }
-    
-
-    //---------------------------------------------------------------------------------------------
-	uint64 InputFileStream::GetReadBytes() const
-    {
-        return m_pD->m_readBytes;
+        return m_pD->File!=0;
     }
     
 
     //---------------------------------------------------------------------------------------------
 	uint64 InputFileStream::Tell() const
     {
-        size_t bufferOffset = m_pD->m_size - m_pD->m_pos;
-        int64_t filePos = mutable_ftell( m_pD->m_pFile );
+        size_t bufferOffset = m_pD->BytesInBuffer - m_pD->BufferPosition;
+        int64 filePos = m_pD->File->Tell();
         check(filePos>0);
         size_t pos = size_t(filePos) - bufferOffset;
         return pos;
@@ -74,42 +65,44 @@ namespace mu
 
 
     //---------------------------------------------------------------------------------------------
-    void InputFileStream::Seek( uint64 position )
+    void InputFileStream::Seek( uint64 NewPosition )
     {
-        int64_t result = mutable_fseek( m_pD->m_pFile, int64_t(position), SEEK_SET );
-        check(result==0);
-        (void)result;
+        bool bSuccess = m_pD->File->Seek( int64(NewPosition) );
+        check(bSuccess);
 
-        m_pD->m_size = 0;
-        m_pD->m_pos = 0;
+        m_pD->BytesInBuffer = 0;
+        m_pD->BufferPosition = 0;
+		m_pD->FilePosition = NewPosition;
     }
 
 
 	//---------------------------------------------------------------------------------------------
     void InputFileStream::Read( void* pData, uint64 size )
 	{
-        uint8_t* pDest = (uint8_t*) pData;
-        while (size && m_pD->m_pFile )
+        uint8* pDest = (uint8*) pData;
+        while (size && m_pD->File )
 		{			
-			int available = FMath::Min( int(m_pD->m_size) - int(m_pD->m_pos), int(size) );
+			int64 Available = FMath::Min( int64(m_pD->BytesInBuffer) - int64(m_pD->BufferPosition), int64(size) );
 
 			// Copy from what is already loaded
-			if ( available>0 )
+			if (Available >0 )
 			{
-                FMemory::Memcpy( pDest, m_pD->m_buffer+m_pD->m_pos, available );
-				m_pD->m_pos += available;
-				size -= available;
-				pDest += available;
+                FMemory::Memcpy( pDest, m_pD->Buffer.GetData()+m_pD->BufferPosition, Available);
+				m_pD->BufferPosition += Available;
+				size -= Available;
+				pDest += Available;
 			}
 
 			// Load more data if necessary
-			if ( m_pD->m_pos == m_pD->m_size && size )
+			if ( m_pD->BufferPosition == m_pD->BytesInBuffer && size )
 			{
-                m_pD->m_size = fread( m_pD->m_buffer, 1, MUTABLE_STREAM_BUFFER_SIZE, m_pD->m_pFile );
-                check( !ferror( m_pD->m_pFile ) );
-                check( m_pD->m_size );
-				m_pD->m_pos = 0;
-                m_pD->m_readBytes += m_pD->m_size;
+				int64 BytesLeftInFile = m_pD->FileSize - m_pD->FilePosition;
+				int64 BytesToRead = FMath::Min(MUTABLE_STREAM_BUFFER_SIZE, BytesLeftInFile);
+                bool bSuccess = m_pD->File->Read( m_pD->Buffer.GetData(), BytesToRead );
+                check( bSuccess );
+				m_pD->BytesInBuffer = BytesToRead;
+				m_pD->BufferPosition = 0;
+                m_pD->FilePosition += BytesToRead;
             }
 		}
 
@@ -119,14 +112,15 @@ namespace mu
 	//---------------------------------------------------------------------------------------------
 	//---------------------------------------------------------------------------------------------
 	//---------------------------------------------------------------------------------------------
-	OutputFileStream::OutputFileStream( const char* strFile )
+	OutputFileStream::OutputFileStream( const FString& File )
 	{
 		m_pD = new Private();
-		m_pD->m_pFile = mutable_fopen( strFile, "wb" );
-		m_pD->m_pos = 0;
-		m_pD->m_buffer = reinterpret_cast<uint8*>(FMemory::Malloc(MUTABLE_STREAM_BUFFER_SIZE, 16));
+		TSharedPtr<IFileHandle> FilePtr(FPlatformFileManager::Get().GetPlatformFile().OpenWrite(*File));
+		m_pD->File = FilePtr;
+		m_pD->BufferPosition = 0;
+		m_pD->Buffer.SetNum(MUTABLE_STREAM_BUFFER_SIZE);
 
-		check( m_pD->m_pFile );
+		check( m_pD->File );
 	}
 
 
@@ -137,19 +131,19 @@ namespace mu
 
 		while (size)
 		{
-            int available = FMath::Min( int(MUTABLE_STREAM_BUFFER_SIZE) - int(m_pD->m_pos), int(size) );
+            int64 Available = FMath::Min(int64(MUTABLE_STREAM_BUFFER_SIZE) - int64(m_pD->BufferPosition), int64(size) );
 
 			// Copy from what fits in the current buffer
-			if ( available>0 )
+			if (Available >0 )
 			{
-                memcpy( m_pD->m_buffer+m_pD->m_pos, pSrc, available );
-				m_pD->m_pos += available;
-				size -= available;
-				pSrc += available;
+                FMemory::Memcpy( m_pD->Buffer.GetData() + m_pD->BufferPosition, pSrc, Available);
+				m_pD->BufferPosition += Available;
+				size -= Available;
+				pSrc += Available;
 			}
 
 			// Flush the data if necessary
-            if ( m_pD->m_pos == MUTABLE_STREAM_BUFFER_SIZE )
+            if ( m_pD->BufferPosition == MUTABLE_STREAM_BUFFER_SIZE )
 			{
 				Flush();
 			}
@@ -160,18 +154,13 @@ namespace mu
 	//---------------------------------------------------------------------------------------------
 	void OutputFileStream::Flush()
 	{
-		if ( m_pD->m_pos )
+		if ( m_pD->BufferPosition)
 		{
-            std::size_t count = fwrite( m_pD->m_buffer,
-										1,
-										m_pD->m_pos,
-										m_pD->m_pFile );
-
-			check( count == m_pD->m_pos );
-            (void)count;
+			bool bSuccess = m_pD->File->Write( m_pD->Buffer.GetData(), m_pD->BufferPosition);
+			check(bSuccess);
         }
 
-		m_pD->m_pos = 0;
+		m_pD->BufferPosition = 0;
 	}
 
 
@@ -180,11 +169,6 @@ namespace mu
 	{
 		Flush();
 
-		fclose( m_pD->m_pFile );
-
-        FMemory::Free( m_pD->m_buffer );
-        m_pD->m_buffer = nullptr;
-
         delete m_pD;
 		m_pD = 0;
 	}
@@ -279,19 +263,19 @@ namespace mu
     class ProxyFactoryMutableSourceFile::Private : public Base
     {
     public:
-        InputFileStream* m_pStream = nullptr;
-        std::string m_filename;
+        InputFileStream* Stream = nullptr;
+        FString FileName;
     };
 
 
     //---------------------------------------------------------------------------------------------
-    ProxyFactoryMutableSourceFile::ProxyFactoryMutableSourceFile( const char* strFileName, 
+    ProxyFactoryMutableSourceFile::ProxyFactoryMutableSourceFile( const FString& InFileName, 
                                                                   InputFileStream* pStream )
     {
         m_pD = new Private();
 
-        m_pD->m_pStream = pStream;
-        m_pD->m_filename = strFileName;
+        m_pD->Stream = pStream;
+        m_pD->FileName = InFileName;
     }
 
 
@@ -307,13 +291,13 @@ namespace mu
     Ptr<ResourceProxy<Image>> ProxyFactoryMutableSourceFile::NewImageProxy(InputArchive& arch)
     {
         // assert: the archive should beusing m_pStream
-        size_t position = m_pD->m_pStream->Tell();
+        size_t position = m_pD->Stream->Tell();
 
         // we are unserialising just to skip the image. It would be much better if we had the size
         // and we could skip it.
         Ptr<Image> t = Image::StaticUnserialise( arch );
 
-        return new ResourceProxyFile<Image>( m_pD->m_filename, position );
+        return new ResourceProxyFile<Image>( m_pD->FileName, position );
     }
 
 
diff --git a/Engine/Plugins/Experimental/Mutable/Source/MutableTools/Private/MuT/StreamsPrivate.h b/Engine/Plugins/Experimental/Mutable/Source/MutableTools/Private/MuT/StreamsPrivate.h
--- a/Engine/Plugins/Experimental/Mutable/Source/MutableTools/Private/MuT/StreamsPrivate.h
+++ b/Engine/Plugins/Experimental/Mutable/Source/MutableTools/Private/MuT/StreamsPrivate.h
@@ -21,15 +21,15 @@ namespace mu
 	{
 	public:
 
-		FILE* m_pFile;
+		TSharedPtr<IFileHandle> File;
 
-		uint64 m_size;
-		uint64 m_pos;
-
-		uint64 m_readBytes;
+		uint64 FileSize;
+		uint64 BytesInBuffer;
+		uint64 BufferPosition;
+		uint64 FilePosition;
 
     private:
-        uint8* m_buffer;
+        TArray<uint8> Buffer;
 
         friend class InputFileStream;
     };
@@ -40,12 +40,12 @@ namespace mu
 	{
 	public:
 
-		FILE* m_pFile;
+		TSharedPtr<IFileHandle> File;
 
-		uint64 m_pos;
+		uint64 BufferPosition;
 
     private:
-        uint8* m_buffer;
+		TArray<uint8> Buffer;
 
         friend class OutputFileStream;
     };
@@ -161,16 +161,15 @@ namespace mu
     class MUTABLETOOLS_API ResourceProxyFile : public ResourceProxy<R>
     {
     private:
-        std::string m_fileName;
-        uint64 m_filePos = 0;
+		FString FileName;
+        uint64 FilePos = 0;
 		FCriticalSection Mutex;
 
     public:
-        ResourceProxyFile( const std::string& fileName, uint64 filePos )
+        ResourceProxyFile( const FString& InFileName, uint64 InFilePos )
         {
-			FScopeLock Lock(&Mutex);
-			m_fileName = fileName;
-            m_filePos = filePos;
+			FileName = InFileName;
+            FilePos = InFilePos;
         }
 
         Ptr<const R> Get() override
@@ -178,10 +177,10 @@ namespace mu
 			FScopeLock Lock(&Mutex);
 
             Ptr<const R> r;
-            if (!m_fileName.empty())
+            if (!FileName.IsEmpty())
             {
-                InputFileStream stream( m_fileName.c_str() );
-                stream.Seek( m_filePos );
+                InputFileStream stream(FileName);
+                stream.Seek( FilePos );
                 InputArchive arch(&stream);
                 r =  R::StaticUnserialise( arch );
             }
diff --git a/Engine/Plugins/Experimental/Mutable/Source/MutableTools/Public/MuT/NodeBool.h b/Engine/Plugins/Experimental/Mutable/Source/MutableTools/Public/MuT/NodeBool.h
--- a/Engine/Plugins/Experimental/Mutable/Source/MutableTools/Public/MuT/NodeBool.h
+++ b/Engine/Plugins/Experimental/Mutable/Source/MutableTools/Public/MuT/NodeBool.h
@@ -186,16 +186,10 @@ namespace mu
 		// Own Interface
 		//-----------------------------------------------------------------------------------------
 
-		//! Get the name of the parameter. It will be exposed in the final compiled data.
-		const char* GetName() const;
-		void SetName( const char* );
-
-		//! Get the uid of the parameter. It will be exposed in the final compiled data.
-		const char* GetUid() const;
-		void SetUid( const char* );
+		//! Set the name of the parameter. It will be exposed in the final compiled data.
+		void SetName( const FString& );
 
 		//! Get the default value of the parameter.
-		bool GetDefaultValue() const;
 		void SetDefaultValue( bool v );
 
         //! Set the number of ranges (dimensions) for this parameter.
diff --git a/Engine/Plugins/Experimental/Mutable/Source/MutableTools/Public/MuT/NodeColourParameter.h b/Engine/Plugins/Experimental/Mutable/Source/MutableTools/Public/MuT/NodeColourParameter.h
--- a/Engine/Plugins/Experimental/Mutable/Source/MutableTools/Public/MuT/NodeColourParameter.h
+++ b/Engine/Plugins/Experimental/Mutable/Source/MutableTools/Public/MuT/NodeColourParameter.h
@@ -40,9 +40,6 @@ namespace mu
 		//-----------------------------------------------------------------------------------------
 		// Node Interface
 		//-----------------------------------------------------------------------------------------
-
-		
-
 		const NODE_TYPE* GetType() const override;
 		static const NODE_TYPE* GetStaticType();
 
@@ -54,20 +51,11 @@ namespace mu
 		// Own Interface
 		//-----------------------------------------------------------------------------------------
 
-		//! Get the name of the parameter. It will be exposed in the final compiled data.
-		const char* GetName() const;
-
 		//! Set the name of the parameter.
-		void SetName( const char* );
-
-		//! Get the uid of the parameter. It will be exposed in the final compiled data.
-		const char* GetUid() const;
+		void SetName( const FString& );
 
 		//! Set the uid of the parameter.
-		void SetUid( const char* );
-
-		//! Get the default value of the parameter.
-		void GetDefaultValue( float* pR, float* pG, float* pB ) const;
+		void SetUid( const FString& );
 
 		//! Set the default value of the parameter.
 		void SetDefaultValue( float r, float g, float b );
diff --git a/Engine/Plugins/Experimental/Mutable/Source/MutableTools/Public/MuT/NodeColourTable.h b/Engine/Plugins/Experimental/Mutable/Source/MutableTools/Public/MuT/NodeColourTable.h
--- a/Engine/Plugins/Experimental/Mutable/Source/MutableTools/Public/MuT/NodeColourTable.h
+++ b/Engine/Plugins/Experimental/Mutable/Source/MutableTools/Public/MuT/NodeColourTable.h
@@ -41,9 +41,6 @@ namespace mu
 		//-----------------------------------------------------------------------------------------
 		// Node Interface
 		//-----------------------------------------------------------------------------------------
-
-		
-
 		const NODE_TYPE* GetType() const override;
 		static const NODE_TYPE* GetStaticType();
 
@@ -56,15 +53,14 @@ namespace mu
 		//-----------------------------------------------------------------------------------------
 
 		//! Set the name of the implicit table parameter.
-		void SetParameterName( const char* strName );
+		void SetParameterName( const FString& strName );
 
 		//!
 		TablePtr GetTable() const;
 		void SetTable( TablePtr );
 
 		//!
-		const char* GetColumn() const;
-		void SetColumn( const char* strName );
+		void SetColumn( const FString& strName );
 
 		//-----------------------------------------------------------------------------------------
 		// Interface pattern
diff --git a/Engine/Plugins/Experimental/Mutable/Source/MutableTools/Public/MuT/NodeComponentNew.h b/Engine/Plugins/Experimental/Mutable/Source/MutableTools/Public/MuT/NodeComponentNew.h
--- a/Engine/Plugins/Experimental/Mutable/Source/MutableTools/Public/MuT/NodeComponentNew.h
+++ b/Engine/Plugins/Experimental/Mutable/Source/MutableTools/Public/MuT/NodeComponentNew.h
@@ -67,10 +67,10 @@ namespace mu
 		//-----------------------------------------------------------------------------------------
 
 		//! Get the name of the component.
-		const char* GetName() const;
+		const FString& GetName() const;
 
 		//! Set the name of the component.
-		void SetName( const char* );
+		void SetName( const FString& );
 
 		//! Get the id of the component.
 		uint16 GetId() const;
diff --git a/Engine/Plugins/Experimental/Mutable/Source/MutableTools/Public/MuT/NodeImageParameter.h b/Engine/Plugins/Experimental/Mutable/Source/MutableTools/Public/MuT/NodeImageParameter.h
--- a/Engine/Plugins/Experimental/Mutable/Source/MutableTools/Public/MuT/NodeImageParameter.h
+++ b/Engine/Plugins/Experimental/Mutable/Source/MutableTools/Public/MuT/NodeImageParameter.h
@@ -46,20 +46,13 @@ namespace mu
 		// Own Interface
 		//-----------------------------------------------------------------------------------------
 
-		//! Get the name of the parameter. It will be exposed in the final compiled data.
-		const char* GetName() const;
-
 		//! Set the name of the parameter.
 		void SetName( const char* );
 
-		//! Get the uid of the parameter. It will be exposed in the final compiled data.
-		const char* GetUid() const;
-
 		//! Set the uid of the parameter.
 		void SetUid( const char* );
 
     	//! Get the default value of the parameter.
-    	FName GetDefaultValue() const;
     	void SetDefaultValue(FName Value);
     	
 		//! Set the number of ranges (dimensions) for this parameter.
diff --git a/Engine/Plugins/Experimental/Mutable/Source/MutableTools/Public/MuT/NodeImageTable.h b/Engine/Plugins/Experimental/Mutable/Source/MutableTools/Public/MuT/NodeImageTable.h
--- a/Engine/Plugins/Experimental/Mutable/Source/MutableTools/Public/MuT/NodeImageTable.h
+++ b/Engine/Plugins/Experimental/Mutable/Source/MutableTools/Public/MuT/NodeImageTable.h
@@ -59,15 +59,14 @@ namespace mu
 		//-----------------------------------------------------------------------------------------
 
 		//! Set the name of the implicit table parameter.
-		void SetParameterName( const char* strName );
+		void SetParameterName( const FString& strName );
 
 		//!
 		TablePtr GetTable() const;
 		void SetTable( TablePtr );
 
 		//!
-		const char* GetColumn() const;
-		void SetColumn( const char* strName );
+		void SetColumn( const FString& strName );
 
 		//-----------------------------------------------------------------------------------------
 		// Interface pattern
diff --git a/Engine/Plugins/Experimental/Mutable/Source/MutableTools/Public/MuT/NodeMeshTable.h b/Engine/Plugins/Experimental/Mutable/Source/MutableTools/Public/MuT/NodeMeshTable.h
--- a/Engine/Plugins/Experimental/Mutable/Source/MutableTools/Public/MuT/NodeMeshTable.h
+++ b/Engine/Plugins/Experimental/Mutable/Source/MutableTools/Public/MuT/NodeMeshTable.h
@@ -60,15 +60,14 @@ namespace mu
 		//-----------------------------------------------------------------------------------------
 
 		//! Set the name of the implicit table parameter.
-		void SetParameterName( const char* strName );
+		void SetParameterName( const FString& strName );
 
 		//!
 		TablePtr GetTable() const;
 		void SetTable( TablePtr );
 
 		//!
-		const char* GetColumn() const;
-		void SetColumn( const char* strName );
+		void SetColumn( const FString& strName );
 
 		//! Get the number of layouts defined in the meshes on this column.
 		int GetLayoutCount() const;
diff --git a/Engine/Plugins/Experimental/Mutable/Source/MutableTools/Public/MuT/NodeModifier.h b/Engine/Plugins/Experimental/Mutable/Source/MutableTools/Public/MuT/NodeModifier.h
--- a/Engine/Plugins/Experimental/Mutable/Source/MutableTools/Public/MuT/NodeModifier.h
+++ b/Engine/Plugins/Experimental/Mutable/Source/MutableTools/Public/MuT/NodeModifier.h
@@ -38,7 +38,7 @@ namespace mu
 		//-----------------------------------------------------------------------------------------
 
 		static void Serialise( const NodeModifier* pNode, OutputArchive& arch );
-		static NodeModifierPtr StaticUnserialise( InputArchive& arch );
+		static Ptr<NodeModifier> StaticUnserialise( InputArchive& arch );
 
 		//-----------------------------------------------------------------------------------------
         // Node interface
@@ -55,13 +55,7 @@ namespace mu
         //! \{
 
         //! Add a tag to the surface, which will be affected by modifier nodes with the same tag
-        void AddTag(const char* tagName);
-
-        //! Get the number of tags added to the Surface.
-        int GetTagCount() const;
-
-        //! Get a tag string from an index (0 to GetTagCount-1)
-        const char* GetTag( int ) const;
+        void AddTag(const FString& tagName);
 
         //! \}
 
diff --git a/Engine/Plugins/Experimental/Mutable/Source/MutableTools/Public/MuT/NodeObject.h b/Engine/Plugins/Experimental/Mutable/Source/MutableTools/Public/MuT/NodeObject.h
--- a/Engine/Plugins/Experimental/Mutable/Source/MutableTools/Public/MuT/NodeObject.h
+++ b/Engine/Plugins/Experimental/Mutable/Source/MutableTools/Public/MuT/NodeObject.h
@@ -52,16 +52,16 @@ namespace mu
 		//-----------------------------------------------------------------------------------------
 
 		//! Get the name of the object.
-		virtual const char* GetName() const = 0;
+		virtual const FString& GetName() const = 0;
 
 		//! Set the name of the object.
-		virtual void SetName( const char* strName ) = 0;
+		virtual void SetName( const FString& ) = 0;
 
 		//! Get the uid of the object.
-		virtual const char* GetUid() const = 0;
+		virtual const FString& GetUid() const = 0;
 
 		//! Set the uid of the object.
-		virtual void SetUid( const char* strUid ) = 0;
+		virtual void SetUid( const FString& ) = 0;
 
 		//-----------------------------------------------------------------------------------------
 		// Interface pattern
diff --git a/Engine/Plugins/Experimental/Mutable/Source/MutableTools/Public/MuT/NodeObjectGroup.h b/Engine/Plugins/Experimental/Mutable/Source/MutableTools/Public/MuT/NodeObjectGroup.h
--- a/Engine/Plugins/Experimental/Mutable/Source/MutableTools/Public/MuT/NodeObjectGroup.h
+++ b/Engine/Plugins/Experimental/Mutable/Source/MutableTools/Public/MuT/NodeObjectGroup.h
@@ -37,9 +37,6 @@ namespace mu
 		//-----------------------------------------------------------------------------------------
 		// Node Interface
 		//-----------------------------------------------------------------------------------------
-
-        
-
         const NODE_TYPE* GetType() const override;
 		static const NODE_TYPE* GetStaticType();
 
@@ -50,10 +47,10 @@ namespace mu
 		//-----------------------------------------------------------------------------------------
 		// NodeObject Interface
 		//-----------------------------------------------------------------------------------------
-        const char* GetName() const override;
-        void SetName( const char* strName ) override;
-        const char* GetUid() const override;
-        void SetUid( const char* strUid ) override;
+        virtual const FString& GetName() const override;
+		virtual void SetName( const FString&) override;
+		virtual const FString& GetUid() const override;
+		virtual void SetUid( const FString&) override;
 
 		//-----------------------------------------------------------------------------------------
 		// Own Interface
diff --git a/Engine/Plugins/Experimental/Mutable/Source/MutableTools/Public/MuT/NodeObjectNew.h b/Engine/Plugins/Experimental/Mutable/Source/MutableTools/Public/MuT/NodeObjectNew.h
--- a/Engine/Plugins/Experimental/Mutable/Source/MutableTools/Public/MuT/NodeObjectNew.h
+++ b/Engine/Plugins/Experimental/Mutable/Source/MutableTools/Public/MuT/NodeObjectNew.h
@@ -33,9 +33,6 @@ namespace mu
 		//-----------------------------------------------------------------------------------------
 		// Node Interface
 		//-----------------------------------------------------------------------------------------
-
-        
-
         const NODE_TYPE* GetType() const override;
 		static const NODE_TYPE* GetStaticType();
 
@@ -46,10 +43,10 @@ namespace mu
 		//-----------------------------------------------------------------------------------------
 		// NodeObject Interface
 		//-----------------------------------------------------------------------------------------
-        const char* GetName() const override;
-        void SetName( const char* strName ) override;
-        const char* GetUid() const override;
-        void SetUid( const char* strUid ) override;
+        virtual const FString& GetName() const override;
+		virtual void SetName( const FString& ) override;
+		virtual const FString& GetUid() const override;
+		virtual void SetUid( const FString& ) override;
 
 		//-----------------------------------------------------------------------------------------
 		// Own Interface
@@ -72,27 +69,20 @@ namespace mu
 		void SetChild( int index, NodeObjectPtr );
 
 		//! Set the number of states that the model can be in.
-		int GetStateCount() const;
-		void SetStateCount( int c );
+		int32 GetStateCount() const;
+		void SetStateCount( int32 c );
 
 		//! Set the name of a state
-		const char* GetStateName( int s ) const;
-		void SetStateName( int s, const char* n );
-
-		//! State optimization options
-		void SetStateGPUOptimisation( int s, bool internal, bool external, bool returnGPU = false );
-		bool GetStateGPUOptimisationInternal( int s ) const;
-		bool GetStateGPUOptimisationExternal( int s ) const;
-		bool GetStateGPUOptimisationReturnImages( int s ) const;
+		void SetStateName( int32 s, const FString& n );
 
 		//! See if a state has a parameter as runtime.
-		bool HasStateParam( int s, const char* param ) const;
+		bool HasStateParam( int32 s, const FString& param ) const;
 
 		//! Add a runtime parameter to the state.
-		void AddStateParam( int s, const char* param );
+		void AddStateParam( int32 s, const FString& param );
 
 		//! Remove a runtime parameter from a state.
-		void RemoveStateParam( int s, const char* param );
+		void RemoveStateParam( int32 s, const FString& param );
 
         //! Set the optimisation properties of a state
 		void SetStateProperties(int32 StateIndex,
@@ -102,7 +92,7 @@ namespace mu
 			uint8 NumExtraLODsToBuildAfterFirstLOD);
 
 		//! Connect a node that produces ExtensionData to be added to the final Instance, and provide a name to associate with the data
-		void AddExtensionDataNode(NodeExtensionDataPtr Node, const char* Name);
+		void AddExtensionDataNode(NodeExtensionDataPtr Node, const FString& Name);
 
 		//-----------------------------------------------------------------------------------------
 		// Interface pattern
diff --git a/Engine/Plugins/Experimental/Mutable/Source/MutableTools/Public/MuT/NodeProjector.h b/Engine/Plugins/Experimental/Mutable/Source/MutableTools/Public/MuT/NodeProjector.h
--- a/Engine/Plugins/Experimental/Mutable/Source/MutableTools/Public/MuT/NodeProjector.h
+++ b/Engine/Plugins/Experimental/Mutable/Source/MutableTools/Public/MuT/NodeProjector.h
@@ -178,25 +178,13 @@ namespace mu
 		// Own Interface
 		//-----------------------------------------------------------------------------------------
 
-		//! Get the name of the parameter. It will be exposed in the final compiled data.
-		const char* GetName() const;
-
 		//! Set the name of the parameter.
-		void SetName( const char* );
+		void SetName( const FString& );
+
 
 		//! Get the uid of the parameter. It will be exposed in the final compiled data.
-		const char* GetUid() const;
-
-		//! Set the uid of the parameter.
-		void SetUid( const char* );
-
-		//! Get the default value of the parameter.
-        void GetDefaultValue( PROJECTOR_TYPE* pType,
-                              float* posX, float* posY, float* posZ,
-							  float* pDirX, float* pDirY, float* pDirZ,
-							  float* pUpX, float* pUpY, float *pUpZ,
-                              float *pScaleU, float *pScaleV, float* pScaleW,
-                              float* pProjectionAngle ) const;
+		const FString& GetUid() const;
+		void SetUid( const FString& );
 
 		//! Set the default value of the parameter.
         void SetDefaultValue( PROJECTOR_TYPE type,
diff --git a/Engine/Plugins/Experimental/Mutable/Source/MutableTools/Public/MuT/NodeRangeFromScalar.h b/Engine/Plugins/Experimental/Mutable/Source/MutableTools/Public/MuT/NodeRangeFromScalar.h
--- a/Engine/Plugins/Experimental/Mutable/Source/MutableTools/Public/MuT/NodeRangeFromScalar.h
+++ b/Engine/Plugins/Experimental/Mutable/Source/MutableTools/Public/MuT/NodeRangeFromScalar.h
@@ -51,8 +51,8 @@ namespace mu
         void SetSize( const Ptr<NodeScalar>& );
 
         //!
-        const char* GetName() const;
-        void SetName( const char* strName );
+		const FString& GetName() const;
+		void SetName(const FString& Name);
 
 		//-----------------------------------------------------------------------------------------
 		// Interface pattern
diff --git a/Engine/Plugins/Experimental/Mutable/Source/MutableTools/Public/MuT/NodeScalarEnumParameter.h b/Engine/Plugins/Experimental/Mutable/Source/MutableTools/Public/MuT/NodeScalarEnumParameter.h
--- a/Engine/Plugins/Experimental/Mutable/Source/MutableTools/Public/MuT/NodeScalarEnumParameter.h
+++ b/Engine/Plugins/Experimental/Mutable/Source/MutableTools/Public/MuT/NodeScalarEnumParameter.h
@@ -46,9 +46,6 @@ namespace mu
 		//-----------------------------------------------------------------------------------------
 		// Node Interface
 		//-----------------------------------------------------------------------------------------
-
-        
-
         const NODE_TYPE* GetType() const override;
 		static const NODE_TYPE* GetStaticType();
 
@@ -61,12 +58,12 @@ namespace mu
 		//-----------------------------------------------------------------------------------------
 
 		//! Get the name of the parameter. It will be exposed in the final compiled data.
-		const char* GetName() const;
-		void SetName( const char* );
+		const FString& GetName() const;
+		void SetName( const FString&);
 
 		//! Get the uid of the parameter. It will be exposed in the final compiled data.
-		const char* GetUid() const;
-		void SetUid( const char* );
+		const FString& GetUid() const;
+		void SetUid( const FString&);
 
 		//! Get the index of the default value of the parameter.
 		int GetDefaultValueIndex() const;
@@ -79,7 +76,7 @@ namespace mu
 		int GetValueCount() const;
 
 		//! Set the data of one of the possible values of the parameter.
-		void SetValue( int i, float value, const char* strName );
+		void SetValue( int i, float value, const FString& strName );
 
         //! Set the number of ranges (dimensions) for this parameter.
         //! By default a parameter has 0 ranges, meaning it only has one value.
diff --git a/Engine/Plugins/Experimental/Mutable/Source/MutableTools/Public/MuT/NodeScalarParameter.h b/Engine/Plugins/Experimental/Mutable/Source/MutableTools/Public/MuT/NodeScalarParameter.h
--- a/Engine/Plugins/Experimental/Mutable/Source/MutableTools/Public/MuT/NodeScalarParameter.h
+++ b/Engine/Plugins/Experimental/Mutable/Source/MutableTools/Public/MuT/NodeScalarParameter.h
@@ -59,21 +59,14 @@ namespace mu
 		//-----------------------------------------------------------------------------------------
 
 		//! Get the name of the parameter. It will be exposed in the final compiled data.
-		const char* GetName() const;
 		void SetName( const char* );
 
 		//! Get the uid of the parameter. It will be exposed in the final compiled data.
-		const char* GetUid() const;
 		void SetUid( const char* );
 
 		//! Get the default value of the parameter.
-		float GetDefaultValue() const;
 		void SetDefaultValue( float v );
 
-		//! Get the additional information about the type
-		PARAMETER_DETAILED_TYPE GetDetailedType() const ;
-		void SetDetailedType( PARAMETER_DETAILED_TYPE type );
-
         //! Set the number of ranges (dimensions) for this parameter.
         //! By default a parameter has 0 ranges, meaning it only has one value.
         void SetRangeCount( int i );
diff --git a/Engine/Plugins/Experimental/Mutable/Source/MutableTools/Public/MuT/NodeStringConstant.h b/Engine/Plugins/Experimental/Mutable/Source/MutableTools/Public/MuT/NodeStringConstant.h
--- a/Engine/Plugins/Experimental/Mutable/Source/MutableTools/Public/MuT/NodeStringConstant.h
+++ b/Engine/Plugins/Experimental/Mutable/Source/MutableTools/Public/MuT/NodeStringConstant.h
@@ -37,9 +37,6 @@ namespace mu
 		//-----------------------------------------------------------------------------------------
 		// Node Interface
 		//-----------------------------------------------------------------------------------------
-
-        
-
         const NODE_TYPE* GetType() const override;
 		static const NODE_TYPE* GetStaticType();
 
@@ -51,11 +48,8 @@ namespace mu
 		// Own Interface
 		//-----------------------------------------------------------------------------------------
 
-		//! Get the value to be returned by the node.
-		const char* GetValue() const;
-
 		//! Set the value to be returned by the node.
-		void SetValue( const char* v );
+		void SetValue( const FString& v );
 
 		//-----------------------------------------------------------------------------------------
 		// Interface pattern
diff --git a/Engine/Plugins/Experimental/Mutable/Source/MutableTools/Public/MuT/NodeStringParameter.h b/Engine/Plugins/Experimental/Mutable/Source/MutableTools/Public/MuT/NodeStringParameter.h
--- a/Engine/Plugins/Experimental/Mutable/Source/MutableTools/Public/MuT/NodeStringParameter.h
+++ b/Engine/Plugins/Experimental/Mutable/Source/MutableTools/Public/MuT/NodeStringParameter.h
@@ -60,22 +60,15 @@ namespace mu
 		// Own Interface
 		//-----------------------------------------------------------------------------------------
 
-		//! Get the name of the parameter. It will be exposed in the final compiled data.
-		const char* GetName() const;
+		//! Set the name of the parameter. It will be exposed in the final compiled data.
 		void SetName( const char* );
 
-		//! Get the uid of the parameter. It will be exposed in the final compiled data.
-		const char* GetUid() const;
+		//! Set the uid of the parameter. It will be exposed in the final compiled data.
 		void SetUid( const char* );
 
-		//! Get the default value of the parameter.
-		const char* GetDefaultValue() const;
+		//! Set the default value of the parameter.
 		void SetDefaultValue( const char* v );
 
-		//! Get the additional information about the type
-		PARAMETER_DETAILED_TYPE GetDetailedType() const ;
-		void SetDetailedType( PARAMETER_DETAILED_TYPE type );
-
         //! Set the number of ranges (dimensions) for this parameter.
         //! By default a parameter has 0 ranges, meaning it only has one value.
         void SetRangeCount( int i );
diff --git a/Engine/Plugins/Experimental/Mutable/Source/MutableTools/Public/MuT/NodeSurfaceEdit.h b/Engine/Plugins/Experimental/Mutable/Source/MutableTools/Public/MuT/NodeSurfaceEdit.h
--- a/Engine/Plugins/Experimental/Mutable/Source/MutableTools/Public/MuT/NodeSurfaceEdit.h
+++ b/Engine/Plugins/Experimental/Mutable/Source/MutableTools/Public/MuT/NodeSurfaceEdit.h
@@ -59,9 +59,6 @@ namespace mu
 		//-----------------------------------------------------------------------------------------
 		// Node Interface
 		//-----------------------------------------------------------------------------------------
-
-        
-
         const NODE_TYPE* GetType() const override;
 		static const NODE_TYPE* GetStaticType();
 
@@ -85,13 +82,7 @@ namespace mu
         //! - the surface will be affected by modifier nodes with the same tag
         //! - the tag will be enabled when the surface is added to an object, and it can activate
         //! variations for any surface.
-        void AddTag(const char* tagName);
-
-        //! Get the number of tags added to the Surface.
-        int GetTagCount() const;
-
-        //! Get a tag string from an index (0 to GetTagCount-1)
-        const char* GetTag( int ) const;
+        void AddTag(const FString& tagName);
 
         //! \}
 
diff --git a/Engine/Plugins/Experimental/Mutable/Source/MutableTools/Public/MuT/NodeSurfaceNew.h b/Engine/Plugins/Experimental/Mutable/Source/MutableTools/Public/MuT/NodeSurfaceNew.h
--- a/Engine/Plugins/Experimental/Mutable/Source/MutableTools/Public/MuT/NodeSurfaceNew.h
+++ b/Engine/Plugins/Experimental/Mutable/Source/MutableTools/Public/MuT/NodeSurfaceNew.h
@@ -74,11 +74,8 @@ namespace mu
 		// Own Interface
 		//-----------------------------------------------------------------------------------------
 
-        //! Get the name of the Surface.
-		const char* GetName() const;
-
         //! Set the name of the Surface.
-        void SetName( const char* );
+        void SetName( const FString&);
 
         //! Add an optional, opaque id that will be returned in the surfaces of the created
         //! instances. Can be useful to identify surfaces on the application side.
@@ -88,9 +85,6 @@ namespace mu
 		//! Add an id that will be used to identify the same surface in other LODs 
 		void SetSharedSurfaceId( int32 SharedSurfaceId );
 
-		//! Return the surface ID shared between LODs 
-		int32 GetSharedSurfaceId() const;
-
         //! \name Tags
         //! \{
 
@@ -98,13 +92,7 @@ namespace mu
         //! - the surface will be affected by modifier nodes with the same tag
         //! - the tag will be enabled when the surface is added to an object, and it can activate
         //! variations for any surface.
-		void AddTag(const char* tagName);
-
-        //! Get the number of tags added to the Surface.
-        int GetTagCount() const;
-
-        //! Get a tag string from an index (0 to GetTagCount-1)
-        const char* GetTag( int ) const;
+		void AddTag(const FString& tagName);
 
         //! \}
 
@@ -125,14 +113,14 @@ namespace mu
 		//! \param index index of the mesh, from 0 to GetMeshCount()-1
 		void SetMesh( int index, NodeMeshPtr );
 
-        //! Get the name of a mesh in the Surface.
-		//! \param index index of the mesh, from 0 to GetMeshCount()-1
-		const char* GetMeshName( int index ) const;
+		//! Get the name of a mesh in the Surface.
+		//! \param index index of the image, from 0 to GetMeshCount()-1
+		const FString& GetMeshName(int index) const;
 
         //! Set the name of a mesh in the Surface.
 		//! \param index index of the mesh, from 0 to GetMeshCount()-1
 		//! \param strName name of the mesh
-		void SetMeshName( int index, const char* strName );
+		void SetMeshName( int index, const FString& strName );
 
 		//! \}
 
@@ -156,12 +144,12 @@ namespace mu
 
 		//! Get the name of a image in the Surface.
 		//! \param index index of the image, from 0 to GetImageCount()-1
-		const char* GetImageName(int index) const;
+		const FString& GetImageName(int index) const;
 
         //! Set the name of a image in the Surface.
         //! \param index index of the image, from 0 to GetImageCount()-1
         //! \param strName name of the image
-        void SetImageName( int index, const char* strName );
+        void SetImageName( int index, const FString& strName );
 
         //! Get the name of a image in the Surface.
         //! \param index index of the image, from 0 to GetImageCount()-1
@@ -174,7 +162,7 @@ namespace mu
 
 		//! This can be used to ser additional information for error reporting.
 		//! It is not propagated to compiled objects in any way.
-		void SetImageAdditionalNames(int index, const char* strMaterialName, const char* strMaterialParameterName);
+		void SetImageAdditionalNames(int index, const FString& strMaterialName, const FString& strMaterialParameterName);
 
         //! \}
 
@@ -198,12 +186,12 @@ namespace mu
 
         //! Get the name of a vectors in the Surface.
         //! \param index index of the vector, from 0 to GetVectorCount()-1
-        const char* GetVectorName(int index) const;
+        const FString& GetVectorName(int index) const;
 
         //! Set the name of a image in the Surface.
         //! \param index index of the vector, from 0 to GetVectorCount()-1
         //! \param strName name of the vector
-        void SetVectorName(int index, const char* strName);
+        void SetVectorName(int index, const FString& strName);
 
         //! \}
 
@@ -227,12 +215,12 @@ namespace mu
 
         //! Get the name of a Scalars in the Surface.
         //! \param index index of the Scalar, from 0 to GetScalarCount()-1
-        const char* GetScalarName( int index ) const;
+        const FString& GetScalarName( int index ) const;
 
         //! Set the name of a image in the Surface.
         //! \param index index of the Scalar, from 0 to GetScalarCount()-1
         //! \param strName name of the Scalar
-        void SetScalarName( int index, const char* strName );
+        void SetScalarName( int index, const FString& strName );
 
         //! \}
 
@@ -256,12 +244,12 @@ namespace mu
 
         //! Get the name of a String in the Surface.
         //! \param index index of the String, from 0 to GetStringCount()-1
-        const char* GetStringName( int index ) const;
+        const FString& GetStringName( int index ) const;
 
         //! Set the name of a image in the Surface.
         //! \param index index of the String, from 0 to GetStringCount()-1
         //! \param strName name of the String
-        void SetStringName( int index, const char* strName );
+        void SetStringName( int index, const FString& strName );
 
         //! \}
 
diff --git a/Engine/Plugins/Experimental/Mutable/Source/MutableTools/Public/MuT/Streams.h b/Engine/Plugins/Experimental/Mutable/Source/MutableTools/Public/MuT/Streams.h
--- a/Engine/Plugins/Experimental/Mutable/Source/MutableTools/Public/MuT/Streams.h
+++ b/Engine/Plugins/Experimental/Mutable/Source/MutableTools/Public/MuT/Streams.h
@@ -32,7 +32,7 @@ namespace mu
 		//-----------------------------------------------------------------------------------------
 
 		//! Create the stream by reading a file.
-		InputFileStream( const char* strFile );
+		InputFileStream( const FString& File );
 
 		~InputFileStream();
 
@@ -50,11 +50,6 @@ namespace mu
         //! Return true if the file was found and it could be opened in the constructor.
         bool IsOpen() const;
 
-        //! Return the number of bytes read from the stream. This should not be used as file 
-        //! position since other bytes might have been read outside this stream, or even Seek might 
-        //! have been called to repeat reads.
-		uint64 GetReadBytes() const;
-
         //! Return the current position in bytes in the underlying file.
 		uint64 Tell() const;
 
@@ -81,7 +76,7 @@ namespace mu
 
 		//! Create a stream that will write to a file.
 		//! The file will be created overwriting any other file with the same name.
-		OutputFileStream( const char* strFile );
+		OutputFileStream(const FString& File);
 
 		~OutputFileStream();
 
@@ -175,7 +170,7 @@ namespace mu
     class MUTABLETOOLS_API ProxyFactoryMutableSourceFile : public InputArchiveWithProxies::ProxyFactory
     {
     public:
-        ProxyFactoryMutableSourceFile( const char* strFileName, InputFileStream* );
+        ProxyFactoryMutableSourceFile( const FString& FileName, InputFileStream* );
         ~ProxyFactoryMutableSourceFile() override;
 
         Ptr<ResourceProxy<mu::Image>> NewImageProxy(InputArchive& arch) override;
