[Mutable] Projector Gizmo code refactor.

ViewPortClient has completly rewritten to handle the following gizmos types:
- Hidden
- Projector
- ClipMorph
- ClipMesh
- Light

Also this class no longer has any responsability on how to save the values of the gizmo, this is now handled through delegates. This allows all projectors types (Projector Parameter, Projector Parameter Default and Projector Constant) to have the same gizmo code.

Finaly the controls of the gizmo have been moved to the top level Editor, accessible through ICustomizableObjectEditorInstance. This way, when selecting a gizmo we can sincronize all the aspects that this imply (udate panels, deselect nodes, hide other gizmos...).

Changes:
- Rewrited Projector Gizmos.
- UCustomizableObjectInstance is now Transactionable.
- Created new COI and Descriptor projector API.
- Removed Projector Temp variables from COI.
- Removed ProjectorStates from COI Private. Replaced with UProjectorParameter located in the editors.
- Removed FCustomizableObjectEditor projector flags.
- Removed GizmoRTSProxy.
- Replaced UCustomizableObjectEmptyClassForSettings with UPreviewSettings.
- CustomizableObjectCustomSettings no longer has state. All state is being moved to UPreviewSettings. Required to refresh the widget.
- Removed lots of unused commented code.
- Removed many inline forward declarations.
- Removed unused FCustomizableObjectWidget.
- Simplified SetProjectorPosition.
- Added getters and setters to NodeProjectorParameter, NodeProjectorConstant.
- Fixed typos.

Also fixes:
#tests COEditor, COIEditor
#rnx

diff --git a/Engine/Plugins/Experimental/Mutable/Source/CustomizableObject/Classes/MuCO/CustomizableObjectInstance.h b/Engine/Plugins/Experimental/Mutable/Source/CustomizableObject/Classes/MuCO/CustomizableObjectInstance.h
--- a/Engine/Plugins/Experimental/Mutable/Source/CustomizableObject/Classes/MuCO/CustomizableObjectInstance.h
+++ b/Engine/Plugins/Experimental/Mutable/Source/CustomizableObject/Classes/MuCO/CustomizableObjectInstance.h
@@ -27,20 +27,6 @@ struct FPropertyChangedEvent;
 #define MUTABLE_VERTEXBUFFER_TEXCOORDS	2
 
 
-// Current state of a projector associated with a parameter.
-namespace EProjectorState
-{
-	typedef uint8 Type;
-	
-	const Type Hidden = 0;
-	const Type Translate = 1;
-	const Type Rotate = 2;
-	const Type Scale = 3;	
-	const Type Selected = 4;
-	const Type TypeChanged = 5;
-};
-
-
 /** FString with the possible errors from skeletal mesh update */
 namespace ESkeletalMeshState
 {
@@ -326,12 +312,6 @@ public:
 	bool bShowOnlyRuntimeParameters = true;
 	bool bShowOnlyRelevantParameters = true;
 
-	/** Control the display of support widgets to edit projectors of this instance. */
-	void SetProjectorState( const FString& ParamName, int32 RangeIndex, EProjectorState::Type state );
-	void ResetProjectorStates();
-	EProjectorState::Type GetProjectorState(const FString& ParamName, int32 RangeIndex) const;
-	FProjectorStateChangedDelegate ProjectorStateChangedDelegate;
-
 	// DEPRECATED: Use the method in the CustomizableObject instead which takes an index among all parameters
 	// Returns how many possible options an int parameter has
 	//int32 GetIntParameterNumOptions(int32 IntParamIndex);
@@ -394,8 +374,25 @@ public:
 		int32 RangeIndex = -1);
 
 	/** Set only the projector position. */
-	void SetProjectorPosition(const FString& ProjectorParamName, const FVector3f& Pos, int32 RangeIndex = -1);
+	UFUNCTION(BlueprintCallable, Category = CustomizableObjectInstance)
+	void SetProjectorPosition(const FString& ProjectorParamName, const FVector3f& Pos, int32 RangeIndex = -1); // TODO GMTFuture Make API consistent (FVector3f -> FVector)
+
+	/** Set only the projector direction vector. */
+	UFUNCTION(BlueprintCallable, Category = CustomizableObjectInstance)
+	void SetProjectorDirection(const FString& ProjectorParamName, const FVector& Direction, int32 RangeIndex = -1);
+	
+	/** Set only the projector up vector. */
+	UFUNCTION(BlueprintCallable, Category = CustomizableObjectInstance)
+	void SetProjectorUp(const FString& ProjectorParamName, const FVector& Up, int32 RangeIndex = -1);
+
+	/** Set only the projector scale. */
+	UFUNCTION(BlueprintCallable, Category = CustomizableObjectInstance)
+	void SetProjectorScale(const FString& ProjectorParamName, const FVector& Scale, int32 RangeIndex = -1);
 
+	/** Set only the cylindrical projector angle. */
+	UFUNCTION(BlueprintCallable, Category = CustomizableObjectInstance)
+	void SetProjectorAngle(const FString& ProjectorParamName, float Angle, int32 RangeIndex = -1);
+	
 	// Get the projector values of a projector parameter with index "ProjectorParamIndex"
 	UFUNCTION(BlueprintCallable, Category = CustomizableObjectInstance)
 	void GetProjectorValue(const FString& ProjectorParamName,
@@ -660,49 +657,14 @@ public:
 
 	UCustomizableInstancePrivateData* GetPrivate() const;
 
-	bool ProjectorUpdatedInViewport = false;
-
-	// TEMP VARIABLE to ease updating the gizmo in the editor after pasting a new projector's transform information
-	bool TempUpdateGizmoInViewport = false;
-
-	// TEMP VARIABLE to verify the projector parameter with transform modified by Paste Transform is set as selected
-	FString TempProjectorParameterName;
-
-	// TEMP VARIABLE to verify the projector parameter with transform modified by Paste Transform is set as selected
-	int32 TempProjectorParameterRangeIndex;
-
-	// TEMP VARIABLE to avoid the projector selection being reset after pasting transform
-	bool AvoidResetProjectorVisibilityForNonNode = false;
-
 	// TEMP VARIABLE to check the Min desired LODs for this instance
 	TWeakObjectPtr<UCustomizableSkeletalComponent> NearestToActor;
 	TWeakObjectPtr<const AActor> NearestToViewCenter;
 
 #if WITH_EDITOR
-	/** If there's a projector parameter pending to be hidden, name of that projector parameter */
-	FString LastSelectedProjectorParameter;
-
-	/** For the case of projector parameter in several layers (the name of the projector is the same and only the index will change */
-	FString LastSelectedProjectorParameterWithIndex;
-
-	/** Flag to ease detect a projector layer change event in FCustomizableObjectEditor::OnObjectModified */
-	bool ProjectorLayerChange = false;
-
-	/** Flag to avoid resetting projector state when editing the alpha value of a layer */
-	bool ProjectorAlphaChange = false;
-
-	/** Flag to unselect the projector */
-	bool UnselectProjector = false;
-
 	/** Profile index the instance parameters are in and if the profile needs to be refreshed */
 	int32 SelectedProfileIndex = INDEX_NONE;
 	bool bSelectedProfileDirty = false;
-
-	/** Tag required to avoid updating the wrong projector parameter range index once removed.
-	On a Group Projector Parameter, if a projector gizmo is removed while is selected, the projector is unselected and removed. The unselection will causes a late update
-	to that projector index, which is performed after the projector being removed. During the update, since the projector index has been removed, the update index is no
-	longer valid. This tag allows to check if the last removed index matches with the currently update index. In the case they match, the update does not occur.*/
-	FString RemovedProjectorParameterNameWithIndex;
 #endif 
 
 private:
diff --git a/Engine/Plugins/Experimental/Mutable/Source/CustomizableObject/Private/MuCO/CustomizableInstancePrivateData.h b/Engine/Plugins/Experimental/Mutable/Source/CustomizableObject/Private/MuCO/CustomizableInstancePrivateData.h
--- a/Engine/Plugins/Experimental/Mutable/Source/CustomizableObject/Private/MuCO/CustomizableInstancePrivateData.h
+++ b/Engine/Plugins/Experimental/Mutable/Source/CustomizableObject/Private/MuCO/CustomizableInstancePrivateData.h
@@ -201,9 +201,6 @@ public:
 	UPROPERTY( Transient )
 	TArray<FGeneratedTexture> GeneratedTextures;
 
-	// Projector state for each parameter name
-	TMap<TPair<FString, int32>, EProjectorState::Type> ProjectorStates;
-
 	// Indices of the parameters that are relevant for the given parameter values.
 	// This only gets updated if parameter decorations are generated.
 	TArray<int> RelevantParameters;
diff --git a/Engine/Plugins/Experimental/Mutable/Source/CustomizableObject/Private/MuCO/CustomizableObjectInstance.cpp b/Engine/Plugins/Experimental/Mutable/Source/CustomizableObject/Private/MuCO/CustomizableObjectInstance.cpp
--- a/Engine/Plugins/Experimental/Mutable/Source/CustomizableObject/Private/MuCO/CustomizableObjectInstance.cpp
+++ b/Engine/Plugins/Experimental/Mutable/Source/CustomizableObject/Private/MuCO/CustomizableObjectInstance.cpp
@@ -152,6 +152,8 @@ const FCustomizableInstanceComponentData* UCustomizableInstancePrivateData::GetC
 
 UCustomizableObjectInstance::UCustomizableObjectInstance()
 {
+	SetFlags(RF_Transactional);
+	
 	PrivateData = CreateDefaultSubobject<UCustomizableInstancePrivateData>(FName("PrivateData"), true);
 }
 
@@ -556,14 +558,6 @@ void UCustomizableInstancePrivateData::ReloadParameters(UCustomizableObjectInsta
 
 	InstanceUpdateFlags(*Public); // TODO Move somewhere else.
 
-#if WITH_EDITOR
-	if(!Public->ProjectorAlphaChange)
-#endif
-	{
-		ProjectorStates.Reset();
-	}
-	
-
 	if (bInvalidatePreviousData)
 	{
 		InvalidateGeneratedData();
@@ -634,8 +628,6 @@ void UCustomizableObjectInstance::SetState(const int32 InState)
 
 	if (OldState != InState)
 	{
-		PrivateData->ProjectorStates.Reset();
-
 		// State may change texture properties, so invalidate the texture reuse cache
 		PrivateData->TextureReuseCache.Empty();
 	}
@@ -654,32 +646,6 @@ void UCustomizableObjectInstance::SetCurrentState(const FString& StateName)
 }
 
 
-void UCustomizableObjectInstance::SetProjectorState(const FString& ParamName, int32 RangeIndex, EProjectorState::Type state)
-{
-	PrivateData->ProjectorStates.Add(TPair<FString, int32>(ParamName, RangeIndex), state);
-	ProjectorStateChangedDelegate.ExecuteIfBound(ParamName);
-}
-
-
-void UCustomizableObjectInstance::ResetProjectorStates()
-{
-	PrivateData->ProjectorStates.Empty();
-}
-
-
-EProjectorState::Type UCustomizableObjectInstance::GetProjectorState(const FString& ParamName, int32 RangeIndex) const
-{
-	if (const EProjectorState::Type* state = PrivateData->ProjectorStates.Find(TPair<FString, int32>(ParamName, RangeIndex)))
-	{
-		return *state;
-	}
-	else
-	{
-		return EProjectorState::Hidden;
-	}
-}
-
-
 // Marked as deprecated. Will be removed in future versions.
 UTexture2D* UCustomizableObjectInstance::GetParameterDescription(const FString& ParamName, int32 DescIndex)
 {
@@ -2419,10 +2385,34 @@ void UCustomizableObjectInstance::SetProjectorPosition(const FString& ProjectorP
 }
 
 
+void UCustomizableObjectInstance::SetProjectorDirection(const FString& ProjectorParamName, const FVector& Direction, int32 RangeIndex)
+{
+	Descriptor.SetProjectorDirection(ProjectorParamName, Direction, RangeIndex);
+}
+
+
+void UCustomizableObjectInstance::SetProjectorUp(const FString& ProjectorParamName, const FVector& Up, int32 RangeIndex)
+{
+	Descriptor.SetProjectorUp(ProjectorParamName, Up, RangeIndex);	
+}
+
+
+void UCustomizableObjectInstance::SetProjectorScale(const FString& ProjectorParamName, const FVector& Scale, int32 RangeIndex)
+{
+	Descriptor.SetProjectorScale(ProjectorParamName, Scale, RangeIndex);	
+}
+
+
+void UCustomizableObjectInstance::SetProjectorAngle(const FString& ProjectorParamName, float Angle, int32 RangeIndex)
+{
+	Descriptor.SetProjectorAngle(ProjectorParamName, Angle, RangeIndex);
+}
+
+
 void UCustomizableObjectInstance::GetProjectorValue(const FString& ProjectorParamName,
-	FVector& OutPos, FVector& OutDir, FVector& OutUp, FVector& OutScale,
-	float& OutAngle, ECustomizableObjectProjectorType& OutType,
-	const int32 RangeIndex) const
+                                                    FVector& OutPos, FVector& OutDir, FVector& OutUp, FVector& OutScale,
+                                                    float& OutAngle, ECustomizableObjectProjectorType& OutType,
+                                                    const int32 RangeIndex) const
 {
 	Descriptor.GetProjectorValue(ProjectorParamName, OutPos, OutDir, OutUp, OutScale, OutAngle, OutType, RangeIndex);
 }
diff --git a/Engine/Plugins/Experimental/Mutable/Source/CustomizableObject/Private/MuCO/CustomizableObjectInstanceDescriptor.cpp b/Engine/Plugins/Experimental/Mutable/Source/CustomizableObject/Private/MuCO/CustomizableObjectInstanceDescriptor.cpp
--- a/Engine/Plugins/Experimental/Mutable/Source/CustomizableObject/Private/MuCO/CustomizableObjectInstanceDescriptor.cpp
+++ b/Engine/Plugins/Experimental/Mutable/Source/CustomizableObject/Private/MuCO/CustomizableObjectInstanceDescriptor.cpp
@@ -1530,51 +1530,63 @@ void FCustomizableObjectInstanceDescriptor::SetProjectorValue(const FString& Pro
 
 void FCustomizableObjectInstanceDescriptor::SetProjectorPosition(const FString& ProjectorParamName, const FVector3f& Pos, const int32 RangeIndex)
 {
-	check(CustomizableObject);
-	RETURN_ON_UNCOMPILED_CO(CustomizableObject, TEXT("Error: Cannot set Projector parameter "))
+	FVector DummyPos, Direction, Up, Scale;
+	float Angle;
+	ECustomizableObjectProjectorType Type;
+   	GetProjectorValue(ProjectorParamName, DummyPos, Direction, Up, Scale, Angle, Type, RangeIndex);
+	
+	SetProjectorValue(ProjectorParamName, static_cast<FVector>(Pos), Direction, Up, Scale, Angle, RangeIndex);
+}
 
-	const int32 ParameterIndexInObject = CustomizableObject->FindParameter(ProjectorParamName);
-	const int32 ParameterIndexInInstance = FindProjectorParameterNameIndex(ProjectorParamName);
 
-	if (ParameterIndexInObject < 0 || ParameterIndexInInstance < 0)
-	{
-		// Early out since we could not find the parameter to set.
-		LogParameterNotFoundWarning(ProjectorParamName, ParameterIndexInObject, ParameterIndexInInstance, CustomizableObject, __FUNCTION__);
-		return;
-	}
+void FCustomizableObjectInstanceDescriptor::SetProjectorDirection(const FString& ProjectorParamName, const FVector& Direction, int32 RangeIndex)
+{
+	FVector Position, DummyDirection, Up, Scale;
+	float Angle;
+	ECustomizableObjectProjectorType Type;
+	GetProjectorValue(ProjectorParamName, Position, DummyDirection, Up, Scale, Angle, Type, RangeIndex);
+		
+	SetProjectorValue(ProjectorParamName, Position, Direction, Up, Scale, Angle, RangeIndex);
+}
 
-	// Parameter to modify
-	FCustomizableObjectProjectorParameterValue& ProjectorParameter = ProjectorParameters[ParameterIndexInInstance];
 
-	FCustomizableObjectProjector ProjectorData = ProjectorParameter.Value;
-	ProjectorData.Position = static_cast<FVector3f>(Pos);
+void FCustomizableObjectInstanceDescriptor::SetProjectorUp(const FString& ProjectorParamName, const FVector& Up, int32 RangeIndex)
+{
+	FVector Position, Direction, DummyUp, Scale;
+	float Angle;
+	ECustomizableObjectProjectorType Type;
+	GetProjectorValue(ProjectorParamName, Position, Direction, DummyUp, Scale, Angle, Type, RangeIndex);
 
-	if (RangeIndex == -1)
-	{
-		check(!CustomizableObject->IsParameterMultidimensional(ParameterIndexInObject)); // This param is multidimensional, it must have a RangeIndex of 0 or more
-		ProjectorParameter.Value = ProjectorData;
-	}
-	else
-	{
-		check(CustomizableObject->IsParameterMultidimensional(ParameterIndexInObject)); // This param is not multidimensional, it must have a RangeIndex of -1
+	SetProjectorValue(ProjectorParamName, Position, Direction, Up, Scale, Angle, RangeIndex);
+}
 
-		if (!ProjectorParameter.RangeValues.IsValidIndex(RangeIndex))
-		{
-			const int32 InsertionIndex = ProjectorParameter.RangeValues.Num();
-			const int32 NumInsertedElements = RangeIndex + 1 - ProjectorParameter.RangeValues.Num();
-			ProjectorParameter.RangeValues.InsertDefaulted(InsertionIndex, NumInsertedElements);
-		}
 
-		check(ProjectorParameter.RangeValues.IsValidIndex(RangeIndex));
-		ProjectorParameter.RangeValues[RangeIndex] = ProjectorData;
-	}
+void FCustomizableObjectInstanceDescriptor::SetProjectorScale(const FString& ProjectorParamName, const FVector& Scale, int32 RangeIndex)
+{
+	FVector Position, Direction, Up, DummyScale;
+	float Angle;
+	ECustomizableObjectProjectorType Type;
+	GetProjectorValue(ProjectorParamName, Position, Direction, Up, DummyScale, Angle, Type, RangeIndex);
+	
+	SetProjectorValue(ProjectorParamName, Position, Direction, Up, Scale, Angle, RangeIndex);
+}
+
+
+void FCustomizableObjectInstanceDescriptor::SetProjectorAngle(const FString& ProjectorParamName, float Angle, int32 RangeIndex)
+{
+	FVector Position, Direction, Up, Scale;
+	float DummyAngle;
+	ECustomizableObjectProjectorType Type;
+	GetProjectorValue(ProjectorParamName, Position, Direction, Up, Scale, DummyAngle, Type, RangeIndex);
+	
+	SetProjectorValue(ProjectorParamName, Position, Direction, Up, Scale, Angle, RangeIndex);
 }
 
 
 void FCustomizableObjectInstanceDescriptor::GetProjectorValue(const FString& ProjectorParamName,
-	FVector& OutPos, FVector& OutDirection, FVector& OutUp, FVector& OutScale,
-	float& OutAngle, ECustomizableObjectProjectorType& OutType,
-	const int32 RangeIndex) const
+                                                              FVector& OutPos, FVector& OutDirection, FVector& OutUp, FVector& OutScale,
+                                                              float& OutAngle, ECustomizableObjectProjectorType& OutType,
+                                                              const int32 RangeIndex) const
 {
 	FVector3f Pos, Direction, Up, Scale;
 	GetProjectorValueF(ProjectorParamName, Pos, Direction, Up, Scale, OutAngle, OutType, RangeIndex);
diff --git a/Engine/Plugins/Experimental/Mutable/Source/CustomizableObjectEditor/Classes/MuCOE/Nodes/CustomizableObjectNodeProjectorConstant.h b/Engine/Plugins/Experimental/Mutable/Source/CustomizableObjectEditor/Classes/MuCOE/Nodes/CustomizableObjectNodeProjectorConstant.h
--- a/Engine/Plugins/Experimental/Mutable/Source/CustomizableObjectEditor/Classes/MuCOE/Nodes/CustomizableObjectNodeProjectorConstant.h
+++ b/Engine/Plugins/Experimental/Mutable/Source/CustomizableObjectEditor/Classes/MuCOE/Nodes/CustomizableObjectNodeProjectorConstant.h
@@ -50,12 +50,6 @@ public:
 	/** Temporary variable where to put the direction information for bone combo box selection changes (in FCustomizableObjectNodeProjectorParameterDetails)*/
 	FVector BoneComboBoxUpDirection;
 
-	/** Flag to know which parameters have been modified of the node in the details tab.
-	* 0 Means projection type
-	* 1 projection angle
-	* 2 projector transform */
-	int32 ParameterSetModified;
-
 	// UObject interface.
 	void PostEditChangeProperty(FPropertyChangedEvent& PropertyChangedEvent) override;
 
@@ -67,6 +61,29 @@ public:
 	// UCustomizableObjectNode interface
 	void AllocateDefaultPins(UCustomizableObjectNodeRemapPins* RemapPins) override;
 
+	// Own interface
 	UEdGraphPin* ValuePin();
+
+	ECustomizableObjectProjectorType GetProjectorType() const;
+
+	FVector GetProjectorPosition() const;
+
+	void SetProjectorPosition(const FVector& Position);
+
+	FVector GetProjectorDirection() const;
+
+	void SetProjectorDirection(const FVector& Direction);
+
+	FVector GetProjectorUp() const;
+
+	void SetProjectorUp(const FVector& Up);
+
+	FVector GetProjectorScale() const;
+
+	void SetProjectorScale(const FVector& Scale);
+
+	float GetProjectorAngle();
+
+	void SetProjectorAngle(float Angle);
 };
 
diff --git a/Engine/Plugins/Experimental/Mutable/Source/CustomizableObjectEditor/Classes/MuCOE/Nodes/CustomizableObjectNodeProjectorParameter.h b/Engine/Plugins/Experimental/Mutable/Source/CustomizableObjectEditor/Classes/MuCOE/Nodes/CustomizableObjectNodeProjectorParameter.h
--- a/Engine/Plugins/Experimental/Mutable/Source/CustomizableObjectEditor/Classes/MuCOE/Nodes/CustomizableObjectNodeProjectorParameter.h
+++ b/Engine/Plugins/Experimental/Mutable/Source/CustomizableObjectEditor/Classes/MuCOE/Nodes/CustomizableObjectNodeProjectorParameter.h
@@ -50,12 +50,6 @@ public:
 	/** Temporary variable where to put the direction information for bone combo box selection changes (in FCustomizableObjectNodeProjectorParameterDetails)*/
 	FVector BoneComboBoxUpDirection = FVector::ZeroVector;
 
-	/** Flag to know which parameters have been modified of the node in the details tab.
-	* 0 Means projection type
-	* 1 projection angle
-	* 2 projector transform */
-	int32 ParameterSetModified = -1;
-
 	// UObject interface.
 	void PostEditChangeProperty(FPropertyChangedEvent& PropertyChangedEvent) override;
 
@@ -68,9 +62,32 @@ public:
 	void AllocateDefaultPins(UCustomizableObjectNodeRemapPins* RemapPins) override;
 	virtual bool IsAffectedByLOD() const override { return false; }
 
+	// Own interface
 	UEdGraphPin* ProjectorPin() const
 	{
 		return FindPin(TEXT("Value"));
 	}
+
+	ECustomizableObjectProjectorType GetProjectorType() const;
+
+	FVector GetProjectorDefaultPosition() const;
+	
+	void SetProjectorDefaultPosition(const FVector& Position);
+
+	FVector GetProjectorDefaultDirection() const;
+
+	void SetProjectorDefaultDirection(const FVector& Direction);
+
+	FVector GetProjectorDefaultUp() const;
+
+	void SetProjectorDefaultUp(const FVector& Up);
+
+	FVector GetProjectorDefaultScale() const;
+
+	void SetProjectorDefaultScale(const FVector& Scale);
+
+	float GetProjectorDefaultAngle() const;
+
+	void SetProjectorDefaultAngle(float Angle);
 };
 
diff --git a/Engine/Plugins/Experimental/Mutable/Source/CustomizableObjectEditor/Private/MuCOE/CustomizableInstanceDetails.cpp b/Engine/Plugins/Experimental/Mutable/Source/CustomizableObjectEditor/Private/MuCOE/CustomizableInstanceDetails.cpp
--- a/Engine/Plugins/Experimental/Mutable/Source/CustomizableObjectEditor/Private/MuCOE/CustomizableInstanceDetails.cpp
+++ b/Engine/Plugins/Experimental/Mutable/Source/CustomizableObjectEditor/Private/MuCOE/CustomizableInstanceDetails.cpp
@@ -24,8 +24,8 @@ void FCustomizableInstanceDetails::CustomizeDetails(const TSharedPtr<IDetailLayo
 	const IDetailsView* DetailsView = DetailBuilder->GetDetailsView();
 	check(DetailsView->GetSelectedObjects().Num());
 
-	CustomInstance = Cast<UCustomizableObjectInstance>(DetailsView->GetSelectedObjects()[0].Get());
-	check(CustomInstance.IsValid());
+	UCustomizableObjectInstance* CustomInstance = Cast<UCustomizableObjectInstance>(DetailsView->GetSelectedObjects()[0].Get());
+	check(CustomInstance);
 	
 	LayoutBuilder = DetailBuilder;
 	
@@ -33,7 +33,7 @@ void FCustomizableInstanceDetails::CustomizeDetails(const TSharedPtr<IDetailLayo
 
 	MainCategory.AddCustomRow( LOCTEXT("CustomizableInstanceDetails", "Instance Parameters") )
 	[
-		SAssignNew(InstancePropertiesWidget, SCustomizableInstanceProperties)
+		SNew(SCustomizableInstanceProperties)
 			.CustomInstance(CustomInstance)
 			.InstanceDetails(SharedThis(this))
 	];
diff --git a/Engine/Plugins/Experimental/Mutable/Source/CustomizableObjectEditor/Private/MuCOE/CustomizableInstanceDetails.h b/Engine/Plugins/Experimental/Mutable/Source/CustomizableObjectEditor/Private/MuCOE/CustomizableInstanceDetails.h
--- a/Engine/Plugins/Experimental/Mutable/Source/CustomizableObjectEditor/Private/MuCOE/CustomizableInstanceDetails.h
+++ b/Engine/Plugins/Experimental/Mutable/Source/CustomizableObjectEditor/Private/MuCOE/CustomizableInstanceDetails.h
@@ -29,10 +29,6 @@ public:
 	void Refresh() const;
 
 private:
-	TWeakObjectPtr<UCustomizableObjectInstance> CustomInstance;
-	TWeakObjectPtr<UCustomizableSkeletalComponent> CustomizableSkeletalComponent;
-	TSharedPtr<SCustomizableInstanceProperties> InstancePropertiesWidget;
-	
 	TWeakPtr<IDetailLayoutBuilder> LayoutBuilder;
 };
 
diff --git a/Engine/Plugins/Experimental/Mutable/Source/CustomizableObjectEditor/Private/MuCOE/CustomizableObjectCustomSettings.cpp b/Engine/Plugins/Experimental/Mutable/Source/CustomizableObjectEditor/Private/MuCOE/CustomizableObjectCustomSettings.cpp
--- a/Engine/Plugins/Experimental/Mutable/Source/CustomizableObjectEditor/Private/MuCOE/CustomizableObjectCustomSettings.cpp
+++ b/Engine/Plugins/Experimental/Mutable/Source/CustomizableObjectEditor/Private/MuCOE/CustomizableObjectCustomSettings.cpp
@@ -11,7 +11,12 @@
 #include "MuCOE/SCustomizableObjectEditorViewport.h"
 #include "MuCOE/UnrealEditorPortabilityHelpers.h"
 #include "PropertyCustomizationHelpers.h"
+#include "MuCO/CustomizableObject.h"
+#include "MuCO/CustomizableObjectInstance.h"
+#include "MuCOE/ICustomizableObjectEditor.h"
+#include "MuCOE/ICustomizableObjectInstanceEditor.h"
 #include "Subsystems/ImportSubsystem.h"
+#include "Toolkits/ToolkitManager.h"
 #include "Widgets/Colors/SColorBlock.h"
 #include "Widgets/Colors/SColorPicker.h"
 #include "Widgets/Input/SButton.h"
@@ -30,16 +35,28 @@ struct FGeometry;
 
 const FString SCustomizableObjectCustomSettings::LightPackagePath = "/Game/EditorAssets/Mutable/";
 
-void SCustomizableObjectCustomSettings::Construct(const FArguments& InArgs, UCustomizableObjectEmptyClassForSettings* InObject)
-{
-	if (!IsValid(InObject) || !InObject->Viewport.IsValid())
+void SCustomizableObjectCustomSettings::Construct(const FArguments& InArgs)
+{	
+	WeakEditor = InArgs._PreviewSettings->GetEditor();
+
+	const TSharedPtr<ICustomizableObjectInstanceEditor> Editor = GetEditorChecked();
+	
+	const TSharedPtr<SCustomizableObjectEditorViewportTabBody> Viewpot = Editor->GetViewport();
+	WeakViewport = Editor->GetViewport();
+
+	TSharedPtr<FString> SelectedLight;
+	for (ULightComponent* Light : Viewpot->GetViewportClient()->GetLightComponents())
 	{
-		return;
-	}
+		TSharedPtr<FString> LightName = MakeShareable(new FString(Light->GetName()));
+		LightNames.Add(LightName);
 
-	TSharedPtr<SCustomizableObjectEditorViewportTabBody> ViewportTabBody = InObject->Viewport.Pin();
+		if (Light == Editor->GetCustomSettings()->GetSelectedLight())
+		{
+			SelectedLight = LightName;
+		}
+	}		
 
-	ViewportClient = ViewportTabBody->GetViewportClient().Get();
+	TSharedPtr<STextBlock> LightAssetText;
 
 	// Viewport Lighting 
 	AddSlot()
@@ -77,7 +94,6 @@ void SCustomizableObjectCustomSettings::Construct(const FArguments& InArgs, UCus
 				.FillWidth(0.5)
 				[
 					SAssignNew(LightsAssetNameInputText, SEditableTextBox)
-					.Text(FText::FromString("Default Viewport Lights"))
 				]
 				+ SHorizontalBox::Slot()
 				.FillWidth(.15)
@@ -121,7 +137,6 @@ void SCustomizableObjectCustomSettings::Construct(const FArguments& InArgs, UCus
 					[
 						SAssignNew(LightAssetText, STextBlock)
 						.TextStyle(FAppStyle::Get(), "PropertyEditor.AssetClass")
-						.Text(FText::FromString("No Asset Selected"))
 					]
 				]
 			]
@@ -408,23 +423,46 @@ void SCustomizableObjectCustomSettings::Construct(const FArguments& InArgs, UCus
 			]
 		]
 	];
+
+	if (const UCustomizableObjectEditorViewportLights* LightPreset = Editor->GetCustomSettings()->GetLightsPreset())
+	{
+		const FText Name = FText::FromString(LightPreset->GetName());
+		LightAssetText->SetText(Name);		
+		LightsAssetNameInputText->SetText(Name);		
+	}
+	else
+	{
+		LightAssetText->SetText(FText::FromString("No Asset Selected"));		
+		LightsAssetNameInputText->SetText(FText::FromString("Default Viewport Lights"));			
+	}
+
+	if (SelectedLight)
+	{
+		LightsListView->SetSelection(SelectedLight);		
+	}
 }
 
 
 void SCustomizableObjectCustomSettings::SetViewportLightsByAsset(const FAssetData& InAsset)
 {
-	UCustomizableObjectEditorViewportLights* InLights = Cast<UCustomizableObjectEditorViewportLights>(InAsset.GetAsset());
+	const TSharedPtr<ICustomizableObjectInstanceEditor> Editor = GetEditorChecked();
+	
+	const TSharedPtr<SCustomizableObjectEditorViewportTabBody> Viewport = WeakViewport.Pin();
+	if (!Viewport)
+	{
+		return;	
+	}
 	
+	UCustomizableObjectEditorViewportLights* InLights = Cast<UCustomizableObjectEditorViewportLights>(InAsset.GetAsset());	
 	if (!InLights)
 	{
 		return;
 	}
+	
+	Editor->GetCustomSettings()->SetLightsPreset(*InLights);
 
 	OnNewViewportLightsAsset();
 
-	LightsAssetNameInputText->SetText(FText::FromString(InLights->GetName()));
-	LightAssetText->SetText(LightsAssetNameInputText->GetText());
-
 	for (FViewportLightData& LightData : InLights->LightsData)
 	{
 		UPointLightComponent* NewLight = nullptr;
@@ -448,24 +486,23 @@ void SCustomizableObjectCustomSettings::SetViewportLightsByAsset(const FAssetDat
 		NewLight->SetSourceLength(LightData.SourceLength);
 		NewLight->SetSourceRadius(LightData.SourceRadius);
 
-		Lights.Add(NewLight);
-		ViewportClient->AddLightToScene(NewLight);
+		Viewport->GetViewportClient()->AddLightToScene(NewLight);			
+		
 		LightNames.Add(MakeShareable(new FString(NewLight->GetName())));
 	}
 
-	SelectedLight = nullptr;
-	ViewportClient->SetSelectedLight(nullptr);
-
-	LightsListView->RequestListRefresh();
-
+	Editor->HideGizmoLight();
 }
 
+
 TSharedRef<SWidget> SCustomizableObjectCustomSettings::GetLightComboButtonContent()
 {
+	const TSharedPtr<ICustomizableObjectInstanceEditor> Editor = GetEditorChecked();
+	
 	TArray<const UClass*> filterClasses;
 	filterClasses.Add(UCustomizableObjectEditorViewportLights::StaticClass());
 	return PropertyCustomizationHelpers::MakeAssetPickerWithMenu(
-		FAssetData(),
+		FAssetData(Editor->GetCustomSettings()->GetLightsPreset()),
 		true,
 		filterClasses,
 		TArray<UFactory*>(),
@@ -485,8 +522,16 @@ void SCustomizableObjectCustomSettings::CloseLightComboButtonContent()
 }
 
 
-FReply SCustomizableObjectCustomSettings::OnSaveViewportLightsAsset()
+FReply SCustomizableObjectCustomSettings::OnSaveViewportLightsAsset() const
 {
+	const TSharedPtr<ICustomizableObjectInstanceEditor> Editor = GetEditorChecked();
+
+	const TSharedPtr<SCustomizableObjectEditorViewportTabBody> Viewport = WeakViewport.Pin();
+	if (!Viewport)
+	{
+		return FReply::Handled();	
+	}
+	
 	FString ObjName = LightsAssetNameInputText->GetText().ToString();
 
 	if (ObjName.IsEmpty() || ObjName == "None")
@@ -520,19 +565,20 @@ FReply SCustomizableObjectCustomSettings::OnSaveViewportLightsAsset()
 
 	if (ViewportLights)
 	{
-		for (int32 i = 0; i < Lights.Num(); ++i)
+		for (ULightComponent* Light : Viewport->GetViewportClient()->GetLightComponents())
 		{
-			UPointLightComponent* Light = Lights[i];
-
-			ViewportLights->LightsData.Add(FViewportLightData());
-			FViewportLightData& NewData = ViewportLights->LightsData.Last();
-			NewData.Intensity = Light->Intensity;
-			NewData.Color = Light->LightColor;
-			NewData.AttenuationRadius = Light->AttenuationRadius;
-			NewData.Transform = Light->GetComponentTransform();
-			NewData.SourceLength = Light->SourceLength;
-			NewData.SourceRadius = Light->SourceRadius;
+			FViewportLightData& NewData = ViewportLights->LightsData.AddDefaulted_GetRef();
 
+			if (const UPointLightComponent* PointLight = Cast<UPointLightComponent>(Light))
+			{
+				NewData.Intensity = PointLight->Intensity;
+				NewData.Color = PointLight->LightColor;
+				NewData.AttenuationRadius = PointLight->AttenuationRadius;
+				NewData.Transform = PointLight->GetComponentTransform();
+				NewData.SourceLength = PointLight->SourceLength;
+				NewData.SourceRadius = PointLight->SourceRadius;
+			}
+			
 			if (const USpotLightComponent* SpotLight = Cast<USpotLightComponent>(Light))
 			{
 				NewData.bIsSpotLight = true;
@@ -552,94 +598,101 @@ FReply SCustomizableObjectCustomSettings::OnSaveViewportLightsAsset()
 	return FReply::Handled();
 }
 
-FReply SCustomizableObjectCustomSettings::OnNewViewportLightsAsset()
+FReply SCustomizableObjectCustomSettings::OnNewViewportLightsAsset() const
 {
-	if (Lights.Num())
+	const TSharedPtr<ICustomizableObjectInstanceEditor> Editor = GetEditorChecked();
+	
+	const TSharedPtr<SCustomizableObjectEditorViewportTabBody> Viewport = WeakViewport.Pin();
+	if (!Viewport)
 	{
-		for (int32 i = 0; i < Lights.Num(); ++i)
-		{
-			ViewportClient->RemoveLightFromScene(Lights[i]);
-		}
-		Lights.Empty();
-		LightNames.Empty();
-
-		LightsAssetNameInputText->SetText(FText::FromString("*Enter Asset Name*"));
-		LightAssetText->SetText(FText::FromString("No Asset Selected"));
-		LightsListView->RequestListRefresh();
+		return FReply::Handled();	
 	}
 
+	Viewport->GetViewportClient()->RemoveAllLightsFromScene();
+	
+	Editor->HideGizmoLight();
+	
 	return FReply::Handled();
 }
 
-FReply SCustomizableObjectCustomSettings::OnPointLightAdded()
+FReply SCustomizableObjectCustomSettings::OnPointLightAdded() const
 {
-	UPointLightComponent* component = NewObject<UPointLightComponent>();
-	Lights.Add(component);
+	const TSharedPtr<ICustomizableObjectInstanceEditor> Editor = GetEditorChecked();
 
-	ViewportClient->AddLightToScene(component);
-
-	// Add Light to ListView and update widgets to reflect the changes
-	LightNames.Add(MakeShareable(new FString(component->GetName())));
-	LightsListView->RequestListRefresh();
-	LightsListView->SetSelection(LightNames.Last());
+	const TSharedPtr<SCustomizableObjectEditorViewportTabBody> Viewport = WeakViewport.Pin();
+	if (!Viewport)
+	{
+		return FReply::Handled();	
+	}
 
+	UPointLightComponent* Component = NewObject<UPointLightComponent>();
+	Viewport->GetViewportClient()->AddLightToScene(Component);
+	Editor->ShowGizmoLight(*Component);			
+	
 	return FReply::Handled();
 }
 
 
-FReply SCustomizableObjectCustomSettings::OnSpotLightAdded()
+FReply SCustomizableObjectCustomSettings::OnSpotLightAdded() const
 {
-	USpotLightComponent* component = NewObject<USpotLightComponent>();
-	Lights.Add(component);
-
-	ViewportClient->AddLightToScene(component);
-
-	// Add Light to ListView and update widgets to reflect the changes
-	LightNames.Add(MakeShareable(new FString(component->GetName())));
-	LightsListView->RequestListRefresh();
-	LightsListView->SetSelection(LightNames.Last());
+	const TSharedPtr<ICustomizableObjectInstanceEditor> Editor = GetEditorChecked();
+	
+	const TSharedPtr<SCustomizableObjectEditorViewportTabBody> Viewport = WeakViewport.Pin();
+	if (!Viewport)
+	{
+		return FReply::Handled();	
+	}
+	
+	USpotLightComponent* Component = NewObject<USpotLightComponent>();
+		
+	Viewport->GetViewportClient()->AddLightToScene(Component);
+	Editor->ShowGizmoLight(*Component);			
 
 	return FReply::Handled();
 }
 
 
-FReply SCustomizableObjectCustomSettings::OnLightRemoved()
+FReply SCustomizableObjectCustomSettings::OnLightRemoved() const
 {
-	TArray<TSharedPtr<FString>> SelectedItems = LightsListView->GetSelectedItems();
+	const TSharedPtr<ICustomizableObjectInstanceEditor> Editor = GetEditorChecked();
 
-	if (SelectedItems.Num())
+	const TSharedPtr<SCustomizableObjectEditorViewportTabBody> Viewport = WeakViewport.Pin();
+	if (!Viewport)
 	{
-		int32 LightIndex = LightNames.Find(SelectedItems[0]);
-
-		// Delete Light from ViewportClient  
-		ULightComponent* LightToRemove = Lights[LightIndex];
-		if (LightToRemove)
-		{
-			ViewportClient->RemoveLightFromScene(LightToRemove);
-			Lights.RemoveAt(LightIndex);
-		}
-		LightToRemove->DestroyComponent();
+		return FReply::Handled();	
+	}
+	
+	const TArray<TSharedPtr<FString>> SelectedItems = LightsListView->GetSelectedItems();
+	if (!SelectedItems.Num())
+	{
+		return FReply::Handled();
+	}
 
-		// Delete and refresh ListView
-		LightNames.RemoveAt(LightIndex);
-		LightsListView->RequestListRefresh();
+	const int32 LightIndex = LightNames.Find(SelectedItems[0]);
 
-		SelectedLight = nullptr;
+	// Delete Light from ViewportClient  
+	if (ULightComponent* LightToRemove = Viewport->GetViewportClient()->GetLightComponents()[LightIndex])
+	{
+		Viewport->GetViewportClient()->RemoveLightFromScene(LightToRemove);
+		LightToRemove->DestroyComponent();
 	}
 
+	Editor->HideGizmoLight();
+
 	return FReply::Handled();
 }
 
 FReply SCustomizableObjectCustomSettings::OnLightUnselected()
 {
-	if (SelectedLight)
-	{
-		ViewportClient->SetSelectedLight(nullptr);
-		SelectedLight = nullptr;
+	const TSharedPtr<ICustomizableObjectInstanceEditor> Editor = GetEditorChecked();
 
-		LightsListView->ClearSelection();
-		LightsListView->RequestListRefresh();
+	const TArray<TSharedPtr<FString>> SelectedItems = LightsListView->GetSelectedItems();
+	if (!SelectedItems.Num())
+	{
+		return FReply::Handled();
 	}
+	
+	Editor->HideGizmoLight();
 
 	return FReply::Handled();
 }
@@ -658,50 +711,62 @@ TSharedRef<ITableRow> SCustomizableObjectCustomSettings::OnGenerateWidgetForList
 
 void SCustomizableObjectCustomSettings::OnListSelectionChanged(TSharedPtr<FString> Selection, ESelectInfo::Type SelectInfo)
 {
-	TArray<TSharedPtr<FString>> SelectedItems = LightsListView->GetSelectedItems();
+	const TSharedPtr<ICustomizableObjectInstanceEditor> Editor = GetEditorChecked();
 
-	// Find selected component and enable viewport gizmo
-	if (SelectedItems.Num())
+	SpotLightProperties->SetVisibility(EVisibility::HitTestInvisible);
+	SpotLightProperties->SetExpanded(false);
+	PointLightProperties->SetVisibility(EVisibility::HitTestInvisible);
+	PointLightProperties->SetExpanded(false);
+	
+	TArray<TSharedPtr<FString>> SelectedItems = LightsListView->GetSelectedItems();
+	if (SelectedItems.Num() != 1)
 	{
-		int32 LightIndex = LightNames.Find(SelectedItems[0]);
-
-		SelectedLight = Lights[LightIndex];
-		ViewportClient->SetSelectedLight(SelectedLight);
+		return;
+	}
 
-		if (const USpotLightComponent* SpotLight = Cast<USpotLightComponent>(SelectedLight))
-		{
-			SpotLightProperties->SetVisibility(EVisibility::Visible);
-			SpotLightProperties->SetExpanded(true);
-		}
-		else
-		{
-			SpotLightProperties->SetExpanded(false);
-			SpotLightProperties->SetVisibility(EVisibility::HitTestInvisible);
-		}
+	const TSharedPtr<SCustomizableObjectEditorViewportTabBody> Viewport = WeakViewport.Pin();
+	if (!Viewport)
+	{
+		return;	
+	}
+	
+	const int32 LightIndex = LightNames.Find(SelectedItems[0]);
+	
+	ULightComponent* SelectedLight = Viewport->GetViewportClient()->GetLightComponents()[LightIndex];
+	Editor->ShowGizmoLight(*SelectedLight);			
 
-		PointLightProperties->SetVisibility(EVisibility::Visible);
-		PointLightProperties->SetExpanded(true);
+	if (Cast<USpotLightComponent>(SelectedLight))
+	{
+		SpotLightProperties->SetVisibility(EVisibility::Visible);
+		SpotLightProperties->SetExpanded(true);
 	}
 	else
 	{
-		SpotLightProperties->SetVisibility(EVisibility::HitTestInvisible);
 		SpotLightProperties->SetExpanded(false);
-		PointLightProperties->SetVisibility(EVisibility::HitTestInvisible);
-		PointLightProperties->SetExpanded(false);
+		SpotLightProperties->SetVisibility(EVisibility::HitTestInvisible);
 	}
 
-
+	PointLightProperties->SetVisibility(EVisibility::Visible);
+	PointLightProperties->SetExpanded(true);	
 }
 
 
 TOptional<float> SCustomizableObjectCustomSettings::GetIntensityValue() const
 {
+	const TSharedPtr<ICustomizableObjectInstanceEditor> Editor = GetEditorChecked();
+
+	ULightComponent* SelectedLight = Editor->GetCustomSettings()->GetSelectedLight();
+	
 	return (SelectedLight) ? SelectedLight->Intensity : 1.0f;
 }
 
 
 void SCustomizableObjectCustomSettings::OnIntensityValueCommited(float Value, ETextCommit::Type CommitType)
 {
+	const TSharedPtr<ICustomizableObjectInstanceEditor> Editor = GetEditorChecked();
+
+	ULightComponent* SelectedLight = Editor->GetCustomSettings()->GetSelectedLight();
+	
 	if (SelectedLight && CommitType != ETextCommit::OnCleared)
 	{
 		SelectedLight->SetIntensity(Value);
@@ -711,6 +776,10 @@ void SCustomizableObjectCustomSettings::OnIntensityValueCommited(float Value, ET
 
 FLinearColor SCustomizableObjectCustomSettings::GetLightColorValue() const
 {
+	const TSharedPtr<ICustomizableObjectInstanceEditor> Editor = GetEditorChecked();
+
+	const ULightComponent* SelectedLight = Editor->GetCustomSettings()->GetSelectedLight();
+	
 	return (SelectedLight) ? SelectedLight->LightColor : FLinearColor::White;
 }
 
@@ -737,7 +806,9 @@ FReply SCustomizableObjectCustomSettings::OnLightColorBlockMouseButtonDown(const
 
 void SCustomizableObjectCustomSettings::OnSetLightColorFromColorPicker(FLinearColor InColor)
 {
-	if (SelectedLight)
+	const TSharedPtr<ICustomizableObjectInstanceEditor> Editor = GetEditorChecked();
+
+	if (ULightComponent* SelectedLight = Editor->GetCustomSettings()->GetSelectedLight())
 	{
 		SelectedLight->SetLightColor(InColor);
 	}
@@ -746,6 +817,10 @@ void SCustomizableObjectCustomSettings::OnSetLightColorFromColorPicker(FLinearCo
 
 TOptional<float> SCustomizableObjectCustomSettings::GetAttenuationRadius() const
 {
+	const TSharedPtr<ICustomizableObjectInstanceEditor> Editor = GetEditorChecked();
+
+	ULightComponent* SelectedLight = Editor->GetCustomSettings()->GetSelectedLight();
+	
 	if (const UPointLightComponent* PointLight = Cast<const UPointLightComponent>(SelectedLight))
 	{
 		return PointLight->AttenuationRadius;
@@ -756,6 +831,10 @@ TOptional<float> SCustomizableObjectCustomSettings::GetAttenuationRadius() const
 
 void SCustomizableObjectCustomSettings::OnAttenuationRadiusValueCommited(float Value, ETextCommit::Type CommitType)
 {
+	const TSharedPtr<ICustomizableObjectInstanceEditor> Editor = GetEditorChecked();
+
+	ULightComponent* SelectedLight = Editor->GetCustomSettings()->GetSelectedLight();
+	
 	if (UPointLightComponent* PointLight = Cast<UPointLightComponent>(SelectedLight))
 	{
 		PointLight->SetAttenuationRadius(Value);
@@ -764,6 +843,10 @@ void SCustomizableObjectCustomSettings::OnAttenuationRadiusValueCommited(float V
 
 TOptional<float> SCustomizableObjectCustomSettings::GetLightSourceRadius() const
 {
+	const TSharedPtr<ICustomizableObjectInstanceEditor> Editor = GetEditorChecked();
+
+	ULightComponent* SelectedLight = Editor->GetCustomSettings()->GetSelectedLight();
+	
 	if (const UPointLightComponent* PointLight = Cast<const UPointLightComponent>(SelectedLight))
 	{
 		return PointLight->SourceRadius;
@@ -774,6 +857,10 @@ TOptional<float> SCustomizableObjectCustomSettings::GetLightSourceRadius() const
 
 void SCustomizableObjectCustomSettings::OnLightSourceRadiusValueCommited(float Value, ETextCommit::Type CommitType)
 {
+	const TSharedPtr<ICustomizableObjectInstanceEditor> Editor = GetEditorChecked();
+
+	ULightComponent* SelectedLight = Editor->GetCustomSettings()->GetSelectedLight();
+	
 	if (UPointLightComponent* PointLight = Cast<UPointLightComponent>(SelectedLight))
 	{
 		PointLight->SetSourceRadius(Value);
@@ -782,6 +869,10 @@ void SCustomizableObjectCustomSettings::OnLightSourceRadiusValueCommited(float V
 
 TOptional<float> SCustomizableObjectCustomSettings::GetLightSourceLength() const
 {
+	const TSharedPtr<ICustomizableObjectInstanceEditor> Editor = GetEditorChecked();
+
+	ULightComponent* SelectedLight = Editor->GetCustomSettings()->GetSelectedLight();
+	
 	if (const UPointLightComponent* PointLight = Cast<const UPointLightComponent>(SelectedLight))
 	{
 		return PointLight->SourceLength;
@@ -792,6 +883,10 @@ TOptional<float> SCustomizableObjectCustomSettings::GetLightSourceLength() const
 
 void SCustomizableObjectCustomSettings::OnLightSourceLengthValueCommited(float Value, ETextCommit::Type CommitType)
 {
+	const TSharedPtr<ICustomizableObjectInstanceEditor> Editor = GetEditorChecked();
+
+	ULightComponent* SelectedLight = Editor->GetCustomSettings()->GetSelectedLight();
+	
 	if (UPointLightComponent* PointLight = Cast<UPointLightComponent>(SelectedLight))
 	{
 		PointLight->SetSourceLength(Value);
@@ -800,6 +895,10 @@ void SCustomizableObjectCustomSettings::OnLightSourceLengthValueCommited(float V
 
 TOptional<float> SCustomizableObjectCustomSettings::GetLightInnerConeAngle() const
 {
+	const TSharedPtr<ICustomizableObjectInstanceEditor> Editor = GetEditorChecked();
+
+	ULightComponent* SelectedLight = Editor->GetCustomSettings()->GetSelectedLight();
+	
 	if (const USpotLightComponent* SpotLight = Cast<const USpotLightComponent>(SelectedLight))
 	{
 		return SpotLight->InnerConeAngle;
@@ -810,6 +909,10 @@ TOptional<float> SCustomizableObjectCustomSettings::GetLightInnerConeAngle() con
 
 void SCustomizableObjectCustomSettings::OnLightInnerConeAngleValueCommited(float Value, ETextCommit::Type CommitType)
 {
+	const TSharedPtr<ICustomizableObjectInstanceEditor> Editor = GetEditorChecked();
+
+	ULightComponent* SelectedLight = Editor->GetCustomSettings()->GetSelectedLight();
+	
 	if (USpotLightComponent* SpotLight = Cast<USpotLightComponent>(SelectedLight))
 	{
 		SpotLight->SetInnerConeAngle(Value);
@@ -818,6 +921,10 @@ void SCustomizableObjectCustomSettings::OnLightInnerConeAngleValueCommited(float
 
 TOptional<float> SCustomizableObjectCustomSettings::GetLightOuterConeAngle() const
 {
+	const TSharedPtr<ICustomizableObjectInstanceEditor> Editor = GetEditorChecked();
+
+	ULightComponent* SelectedLight = Editor->GetCustomSettings()->GetSelectedLight();
+	
 	if (const USpotLightComponent* SpotLight = Cast<const USpotLightComponent>(SelectedLight))
 	{
 		return SpotLight->OuterConeAngle;
@@ -828,8 +935,21 @@ TOptional<float> SCustomizableObjectCustomSettings::GetLightOuterConeAngle() con
 
 void SCustomizableObjectCustomSettings::OnLightOuterConeAngleValueCommited(float Value, ETextCommit::Type CommitType)
 {
+	const TSharedPtr<ICustomizableObjectInstanceEditor> Editor = GetEditorChecked();
+
+	ULightComponent* SelectedLight = Editor->GetCustomSettings()->GetSelectedLight();
+	
 	if (USpotLightComponent* SpotLight = Cast<USpotLightComponent>(SelectedLight))
 	{
 		SpotLight->SetOuterConeAngle(Value);
 	}
 }
+
+
+TSharedPtr<ICustomizableObjectInstanceEditor> SCustomizableObjectCustomSettings::GetEditorChecked() const
+{
+	const TSharedPtr<ICustomizableObjectInstanceEditor> Editor = WeakEditor.Pin();
+	check(Editor);
+	
+	return Editor;
+}
diff --git a/Engine/Plugins/Experimental/Mutable/Source/CustomizableObjectEditor/Private/MuCOE/CustomizableObjectCustomSettings.h b/Engine/Plugins/Experimental/Mutable/Source/CustomizableObjectEditor/Private/MuCOE/CustomizableObjectCustomSettings.h
--- a/Engine/Plugins/Experimental/Mutable/Source/CustomizableObjectEditor/Private/MuCOE/CustomizableObjectCustomSettings.h
+++ b/Engine/Plugins/Experimental/Mutable/Source/CustomizableObjectEditor/Private/MuCOE/CustomizableObjectCustomSettings.h
@@ -2,52 +2,48 @@
 
 #pragma once
 
+#include "CustomizableObjectInstanceEditor.h"
 #include "Widgets/SBoxPanel.h"
 
-#include "CustomizableObjectCustomSettings.generated.h"
+class UPointLightComponent;
+class SCustomizableObjectEditorViewportTabBody;
+class ICustomizableObjectInstanceEditor;
 
 namespace ESelectInfo { enum Type : int; }
 namespace ETextCommit { enum Type : int; }
 template <typename ItemType> class SListView;
 
 class SWidget;
+class SEditableTextBox;
+class SComboButton;
+class UCustomizableObject;
 struct FAssetData;
 struct FGeometry;
 struct FPointerEvent;
 
-UCLASS()
-class UCustomizableObjectEmptyClassForSettings : public UObject
-{
-	GENERATED_BODY()
-public:
-	TWeakPtr<class SCustomizableObjectEditorViewportTabBody> Viewport;
-	TObjectPtr<class UDebugSkelMeshComponent>* PreviewSkeletalMeshComp = nullptr;
-};
 
 class SCustomizableObjectCustomSettings : public SVerticalBox
 {
 public:
-	void Construct(const FArguments& InArgs, UCustomizableObjectEmptyClassForSettings* InObject);
-
-private:
-	class FCustomizableObjectEditorViewportClient* ViewportClient;
+	SLATE_BEGIN_ARGS(SCustomizableObjectCustomSettings) {}
+	SLATE_ARGUMENT(UCustomSettings*, PreviewSettings)
+	SLATE_END_ARGS()
+	
+	void Construct(const FArguments& InArgs);
 
-	// Lighting
-public:
 	void SetViewportLightsByAsset(const FAssetData& InAsset);
 
 private:
-
 	TSharedRef<SWidget> GetLightComboButtonContent();
 	void CloseLightComboButtonContent();
-	FReply OnSaveViewportLightsAsset();
-	FReply OnNewViewportLightsAsset();
+	FReply OnSaveViewportLightsAsset() const;
+	FReply OnNewViewportLightsAsset() const;
 
-	FReply OnPointLightAdded();
-	FReply OnSpotLightAdded();
-	FReply OnLightRemoved();
+	FReply OnPointLightAdded() const;
+	FReply OnSpotLightAdded() const;
+	FReply OnLightRemoved() const;
 	FReply OnLightUnselected();
-	TSharedRef<class ITableRow> OnGenerateWidgetForList(TSharedPtr<FString> Item, const TSharedRef<class STableViewBase>& OwnerTable);
+	TSharedRef<ITableRow> OnGenerateWidgetForList(TSharedPtr<FString> Item, const TSharedRef<STableViewBase>& OwnerTable);
 	void OnListSelectionChanged(TSharedPtr<FString> Selection, ESelectInfo::Type SelectInfo);
 
 	// LightProperties
@@ -73,18 +69,18 @@ private:
 	TOptional<float> GetLightOuterConeAngle() const;
 	void OnLightOuterConeAngleValueCommited(float Value, ETextCommit::Type CommitType);
 
-private:
+	TSharedPtr<ICustomizableObjectInstanceEditor> GetEditorChecked() const;
+	
 	TArray<TSharedPtr<FString>> LightNames;
-	TArray<class UPointLightComponent*> Lights;
 
 	static const FString LightPackagePath;
-	TSharedPtr<class SEditableTextBox> LightsAssetNameInputText;
+	TSharedPtr<SEditableTextBox> LightsAssetNameInputText;
 
-	TSharedPtr<class SComboButton> LightComboButton;
-	TSharedPtr<class STextBlock> LightAssetText;
+	TSharedPtr<SComboButton> LightComboButton;
 	TSharedPtr<SListView<TSharedPtr<FString>>> LightsListView;
-	class ULightComponent* SelectedLight;
-	TSharedPtr<class SExpandableArea> PointLightProperties;
-	TSharedPtr<class SExpandableArea> SpotLightProperties;
+	TSharedPtr<SExpandableArea> PointLightProperties;
+	TSharedPtr<SExpandableArea> SpotLightProperties;
 
+	TWeakPtr<ICustomizableObjectInstanceEditor> WeakEditor;
+	TWeakPtr<SCustomizableObjectEditorViewportTabBody> WeakViewport;
 };
diff --git a/Engine/Plugins/Experimental/Mutable/Source/CustomizableObjectEditor/Private/MuCOE/CustomizableObjectCustomSettingsDetails.cpp b/Engine/Plugins/Experimental/Mutable/Source/CustomizableObjectEditor/Private/MuCOE/CustomizableObjectCustomSettingsDetails.cpp
--- a/Engine/Plugins/Experimental/Mutable/Source/CustomizableObjectEditor/Private/MuCOE/CustomizableObjectCustomSettingsDetails.cpp
+++ b/Engine/Plugins/Experimental/Mutable/Source/CustomizableObjectEditor/Private/MuCOE/CustomizableObjectCustomSettingsDetails.cpp
@@ -5,6 +5,7 @@
 #include "DetailLayoutBuilder.h"
 #include "DetailWidgetRow.h"
 #include "IDetailsView.h"
+#include "MuCO/CustomizableObject.h"
 #include "MuCOE/CustomizableObjectCustomSettings.h"
 
 TSharedRef<IDetailCustomization> FCustomizableObjectCustomSettingsDetails::MakeInstance()
@@ -20,7 +21,8 @@ void FCustomizableObjectCustomSettingsDetails::CustomizeDetails(IDetailLayoutBui
 		IDetailCategoryBuilder& MainCategory = DetailBuilder.EditCategory("Custom Settings");
 		MainCategory.AddCustomRow(FText::FromString("Custom Settings"))
 		[
-			SNew(SCustomizableObjectCustomSettings, Cast<UCustomizableObjectEmptyClassForSettings>(DetailsView->GetSelectedObjects()[0].Get()))
+			SNew(SCustomizableObjectCustomSettings)
+				.PreviewSettings(Cast<UCustomSettings>(DetailsView->GetSelectedObjects()[0]))
 		];
 	}
 }
diff --git a/Engine/Plugins/Experimental/Mutable/Source/CustomizableObjectEditor/Private/MuCOE/CustomizableObjectEditor.cpp b/Engine/Plugins/Experimental/Mutable/Source/CustomizableObjectEditor/Private/MuCOE/CustomizableObjectEditor.cpp
--- a/Engine/Plugins/Experimental/Mutable/Source/CustomizableObjectEditor/Private/MuCOE/CustomizableObjectEditor.cpp
+++ b/Engine/Plugins/Experimental/Mutable/Source/CustomizableObjectEditor/Private/MuCOE/CustomizableObjectEditor.cpp
@@ -51,6 +51,7 @@
 #include "MuCOE/SCustomizableObjectEditorViewport.h"
 #include "PropertyEditorModule.h"
 #include "ScopedTransaction.h"
+#include "MuCOE/SCustomizableInstanceProperties.h"
 #include "UObject/EnumProperty.h"
 #include "Widgets/Docking/SDockTab.h"
 #include "Widgets/Input/SSearchBox.h"
@@ -198,6 +199,11 @@ FCustomizableObjectEditor::~FCustomizableObjectEditor()
 
 void FCustomizableObjectEditor::InitCustomizableObjectEditor(const EToolkitMode::Type Mode, const TSharedPtr<IToolkitHost>& InitToolkitHost)
 {
+	ProjectorParameter = NewObject<UProjectorParameter>();
+
+	CustomSettings = NewObject<UCustomSettings>();
+	CustomSettings->SetEditor(SharedThis(this));
+
 	FAssetRegistryModule& AssetRegistryModule = FModuleManager::LoadModuleChecked<FAssetRegistryModule>(TEXT("AssetRegistry"));
 	if (AssetRegistryModule.Get().IsLoadingAssets())
 	{
@@ -234,7 +240,6 @@ void FCustomizableObjectEditor::InitCustomizableObjectEditor(const EToolkitMode:
 	DetailsViewArgs.bShowScrollBar = false;
 
 	CustomizableObjectDetailsView = PropPlugin.CreateDetailView( DetailsViewArgs );
-	CustomizableObjectDetailsView->SetObject(CustomizableObject);
 
 	CustomizableInstanceDetailsView = PropPlugin.CreateDetailView( DetailsViewArgs );
 	GraphNodeDetailsView = PropPlugin.CreateDetailView(DetailsViewArgs);
@@ -252,20 +257,13 @@ void FCustomizableObjectEditor::InitCustomizableObjectEditor(const EToolkitMode:
 	// \TODO: Create only when needed?
 	TagExplorer = SNew(SCustomizableObjectEditorTagExplorer).CustomizableObjectEditor(this);
 	
-	AdditionalSettings = NewObject<UCustomizableObjectEmptyClassForSettings>();
-	if(AdditionalSettings != nullptr)
-		{
-		AdditionalSettings->Viewport = Viewport;
-		AdditionalSettings->PreviewSkeletalMeshComp = !PreviewSkeletalMeshComponents.IsEmpty() ? &PreviewSkeletalMeshComponents[0] : nullptr;
-	}
-
 	FAdvancedPreviewSceneModule& AdvancedPreviewSceneModule = FModuleManager::LoadModuleChecked<FAdvancedPreviewSceneModule>("AdvancedPreviewScene");
 
 	TSharedPtr<FAdvancedPreviewScene> AdvancedPreviewScene = StaticCastSharedPtr<FAdvancedPreviewScene>(Viewport->GetPreviewScene());
 
 	CustomizableObjectEditorAdvancedPreviewSettings =
 		SNew(SCustomizableObjectEditorAdvancedPreviewSettings, AdvancedPreviewScene.ToSharedRef())
-		.AdditionalSettings(AdditionalSettings)
+		.CustomSettings(CustomSettings)
 		.CustomizableObjectEditor(this);
 	CustomizableObjectEditorAdvancedPreviewSettings->LoadProfileEnvironment();
 	AdvancedPreviewSettingsWidget = CustomizableObjectEditorAdvancedPreviewSettings;
@@ -325,12 +323,14 @@ void FCustomizableObjectEditor::InitCustomizableObjectEditor(const EToolkitMode:
 					->SetHideTabWell(true)
 				)
 			)
-		)
+		)	
 	);
 
 	const bool bCreateDefaultStandaloneMenu = true;
 	const bool bCreateDefaultToolbar = true;
 	FAssetEditorToolkit::InitAssetEditor( Mode, InitToolkitHost, CustomizableObjectEditorAppIdentifier, StandaloneDefaultLayout, bCreateDefaultStandaloneMenu, bCreateDefaultToolbar, CustomizableObject);
+	
+	CustomizableObjectDetailsView->SetObject(CustomizableObject); // Can only be called after initializing the Asset Editor
 
 	ExtendToolbar();
 	RegenerateMenusAndToolbars();
@@ -340,8 +340,6 @@ void FCustomizableObjectEditor::InitCustomizableObjectEditor(const EToolkitMode:
 	OnInstancePropertySelectionChanged(NULL);
 	OnObjectModifiedHandle = FCoreUObjectDelegates::OnObjectModified.AddRaw(this, &FCustomizableObjectEditor::OnObjectModified);
 
-	FCoreUObjectDelegates::OnObjectPropertyChanged.AddSP(this, &FCustomizableObjectEditor::OnObjectPropertyChanged);
-
 	// Compile for the first time if necessary
 	if (!CustomizableObject->IsCompiled())
 	{
@@ -472,12 +470,6 @@ void FCustomizableObjectEditor::CreatePreviewInstance()
 }
 
 
-void FCustomizableObjectEditor::OnPreviewInstanceUpdated()
-{
-	RefreshViewport();
-}
-
-
 void FCustomizableObjectEditor::AddReferencedObjects( FReferenceCollector& Collector )
 {
 	Collector.AddReferencedObject( CustomizableObject );
@@ -486,7 +478,8 @@ void FCustomizableObjectEditor::AddReferencedObjects( FReferenceCollector& Colle
 	Collector.AddReferencedObject( PreviewStaticMeshComponent );
 	Collector.AddReferencedObjects( PreviewSkeletalMeshComponents );
 	Collector.AddReferencedObject( HelperCallback );
-	Collector.AddReferencedObject( AdditionalSettings );
+	Collector.AddReferencedObject( ProjectorParameter );
+	Collector.AddReferencedObject( CustomSettings );
 }
 
 
@@ -926,6 +919,315 @@ void FCustomizableObjectEditor::ReconstructAllChildNodes(UCustomizableObjectNode
 }
 
 
+UProjectorParameter* FCustomizableObjectEditor::GetProjectorParameter()
+{
+	return ProjectorParameter;
+}
+
+
+UCustomSettings* FCustomizableObjectEditor::GetCustomSettings()
+{
+	return CustomSettings;
+}
+
+
+void FCustomizableObjectEditor::HideGizmo()
+{
+	HideGizmoProjectorNodeProjectorConstant();
+	HideGizmoProjectorNodeProjectorParameter();
+	HideGizmoClipMorph();
+	HideGizmoClipMesh();
+	HideGizmoLight();
+
+	FCustomizableObjectInstanceEditor::HideGizmo(SharedThis(this), Viewport, CustomizableInstanceDetailsView);
+}
+
+
+void FCustomizableObjectEditor::ShowGizmoProjectorNodeProjectorConstant(UCustomizableObjectNodeProjectorConstant& Node)
+{
+	HideGizmo();
+	
+	ProjectorGizmo = EProjectorGizmo::NodeProjectorConstant;
+	
+	GraphEditor->ClearSelectionSet();
+	GraphEditor->SetNodeSelection(&Node, true);
+	
+	FProjectorTypeDelegate ProjectorTypeDelegate;
+	ProjectorTypeDelegate.BindUObject(&Node, &UCustomizableObjectNodeProjectorConstant::GetProjectorType);		
+
+	FWidgetColorDelegate WidgetColorDelegate;
+	WidgetColorDelegate.BindLambda([]() { return FColor::Red; });
+
+	FWidgetLocationDelegate WidgetLocationDelegate;
+	WidgetLocationDelegate.BindUObject(&Node, &UCustomizableObjectNodeProjectorConstant::GetProjectorPosition);
+
+	FOnWidgetLocationChangedDelegate OnWidgetLocationChangedDelegate;
+	OnWidgetLocationChangedDelegate.BindUObject(&Node, &UCustomizableObjectNodeProjectorConstant::SetProjectorPosition);
+
+	FWidgetDirectionDelegate WidgetDirectionDelegate;
+	WidgetDirectionDelegate.BindUObject(&Node, &UCustomizableObjectNodeProjectorConstant::GetProjectorDirection);
+
+	FOnWidgetDirectionChangedDelegate OnWidgetDirectionChangedDelegate;
+	OnWidgetDirectionChangedDelegate.BindUObject(&Node, &UCustomizableObjectNodeProjectorConstant::SetProjectorDirection);
+
+	FWidgetUpDelegate WidgetUpDelegate;
+	WidgetUpDelegate.BindUObject(&Node, &UCustomizableObjectNodeProjectorConstant::GetProjectorUp);
+
+	FOnWidgetUpChangedDelegate OnWidgetUpChangedDelegate;
+	OnWidgetUpChangedDelegate.BindUObject(&Node, &UCustomizableObjectNodeProjectorConstant::SetProjectorUp);
+
+	FWidgetScaleDelegate WidgetScaleDelegate;
+	WidgetScaleDelegate.BindUObject(&Node, &UCustomizableObjectNodeProjectorConstant::GetProjectorScale);
+
+	FOnWidgetScaleChangedDelegate OnWidgetScaleChangedDelegate;
+	OnWidgetScaleChangedDelegate.BindUObject(&Node, &UCustomizableObjectNodeProjectorConstant::SetProjectorScale);
+
+	FWidgetAngleDelegate WidgetAngleDelegate;
+	WidgetAngleDelegate.BindUObject(&Node, &UCustomizableObjectNodeProjectorConstant::GetProjectorAngle);
+
+	FWidgetTrackingStartedDelegate WidgetTrackingStartedDelegate;
+	WidgetTrackingStartedDelegate.BindLambda([WeakNode = MakeWeakObjectPtr(&Node)]()
+	{
+		if (UCustomizableObjectNodeProjectorConstant* Node = WeakNode.Get())
+		{
+			Node->Modify();
+		}
+	});
+	
+	Viewport->ShowGizmoProjector(WidgetLocationDelegate, OnWidgetLocationChangedDelegate,
+		WidgetDirectionDelegate, OnWidgetDirectionChangedDelegate,
+		WidgetUpDelegate, OnWidgetUpChangedDelegate,
+		WidgetScaleDelegate, OnWidgetScaleChangedDelegate,
+		WidgetAngleDelegate,
+		ProjectorTypeDelegate,
+		WidgetColorDelegate,
+		WidgetTrackingStartedDelegate);
+}
+
+
+void FCustomizableObjectEditor::HideGizmoProjectorNodeProjectorConstant()
+{
+	if (ProjectorGizmo != EProjectorGizmo::NodeProjectorConstant)
+	{
+		return;
+	}
+	
+	Viewport->HideGizmoProjector();
+
+	const FGraphPanelSelectionSet SelectedNodes = GraphEditor->GetSelectedNodes();
+	for (FGraphPanelSelectionSet::TConstIterator NodeIt( SelectedNodes ); NodeIt; ++NodeIt)
+	{
+		const UObject* Node = *NodeIt;
+		if (Node->IsA<UCustomizableObjectNodeProjectorConstant>())
+		{
+			GraphEditor->ClearSelectionSet();
+			break;
+		}
+	}			
+}
+
+
+void FCustomizableObjectEditor::ShowGizmoProjectorNodeProjectorParameter(UCustomizableObjectNodeProjectorParameter& Node)
+{
+	HideGizmo();
+
+	ProjectorGizmo = EProjectorGizmo::NodeProjectorParameter;
+	
+	GraphEditor->ClearSelectionSet();
+	GraphEditor->SetNodeSelection(&Node, true);
+	
+	FProjectorTypeDelegate ProjectorTypeDelegate;
+	ProjectorTypeDelegate.BindUObject(&Node, &UCustomizableObjectNodeProjectorParameter::GetProjectorType);		
+
+	FWidgetColorDelegate WidgetColorDelegate;
+	WidgetColorDelegate.BindLambda([]() { return FColor::Red; });
+	
+	FWidgetLocationDelegate WidgetLocationDelegate;
+	WidgetLocationDelegate.BindUObject(&Node, &UCustomizableObjectNodeProjectorParameter::GetProjectorDefaultPosition);
+
+	FOnWidgetLocationChangedDelegate OnWidgetLocationChangedDelegate;
+	OnWidgetLocationChangedDelegate.BindUObject(&Node, &UCustomizableObjectNodeProjectorParameter::SetProjectorDefaultPosition);
+
+	FWidgetDirectionDelegate WidgetDirectionDelegate;
+	WidgetDirectionDelegate.BindUObject(&Node, &UCustomizableObjectNodeProjectorParameter::GetProjectorDefaultDirection);
+
+	FOnWidgetDirectionChangedDelegate OnWidgetDirectionChangedDelegate;
+	OnWidgetDirectionChangedDelegate.BindUObject(&Node, &UCustomizableObjectNodeProjectorParameter::SetProjectorDefaultDirection);
+
+	FWidgetUpDelegate WidgetUpDelegate;
+	WidgetUpDelegate.BindUObject(&Node, &UCustomizableObjectNodeProjectorParameter::GetProjectorDefaultUp);
+
+	FOnWidgetUpChangedDelegate OnWidgetUpChangedDelegate;
+	OnWidgetUpChangedDelegate.BindUObject(&Node, &UCustomizableObjectNodeProjectorParameter::SetProjectorDefaultUp);
+
+	FWidgetScaleDelegate WidgetScaleDelegate;
+	WidgetScaleDelegate.BindUObject(&Node, &UCustomizableObjectNodeProjectorParameter::GetProjectorDefaultScale);
+
+	FOnWidgetScaleChangedDelegate OnWidgetScaleChangedDelegate;
+	OnWidgetScaleChangedDelegate.BindUObject(&Node, &UCustomizableObjectNodeProjectorParameter::SetProjectorDefaultScale);
+
+	FWidgetAngleDelegate WidgetAngleDelegate;
+	WidgetAngleDelegate.BindUObject(&Node, &UCustomizableObjectNodeProjectorParameter::GetProjectorDefaultAngle);
+	
+	FWidgetTrackingStartedDelegate WidgetTrackingStartedDelegate;
+	WidgetTrackingStartedDelegate.BindLambda([WeakNode = MakeWeakObjectPtr(&Node)]()
+	{
+		if (UCustomizableObjectNodeProjectorParameter* Node = WeakNode.Get())
+		{
+			Node->Modify();
+		}
+	});
+	
+	Viewport->ShowGizmoProjector(WidgetLocationDelegate, OnWidgetLocationChangedDelegate,
+		WidgetDirectionDelegate, OnWidgetDirectionChangedDelegate,
+		WidgetUpDelegate, OnWidgetUpChangedDelegate,
+		WidgetScaleDelegate, OnWidgetScaleChangedDelegate,
+		WidgetAngleDelegate,
+		ProjectorTypeDelegate,
+		WidgetColorDelegate,
+		WidgetTrackingStartedDelegate);
+}
+
+
+void FCustomizableObjectEditor::HideGizmoProjectorNodeProjectorParameter()
+{
+	if (ProjectorGizmo != EProjectorGizmo::NodeProjectorParameter)
+	{
+		return;
+	}
+	
+	Viewport->HideGizmoProjector();
+
+	const FGraphPanelSelectionSet SelectedNodes = GraphEditor->GetSelectedNodes();
+	for (FGraphPanelSelectionSet::TConstIterator NodeIt( SelectedNodes ); NodeIt; ++NodeIt)
+	{
+		const UObject* Node = *NodeIt;
+		if (Node->IsA<UCustomizableObjectNodeProjectorParameter>())
+		{
+			GraphEditor->ClearSelectionSet();
+			break;
+		}
+	}	
+}
+
+
+void FCustomizableObjectEditor::ShowGizmoProjectorParameter(const FString& ParamName, int32 RangeIndex)
+{
+	ProjectorGizmo = EProjectorGizmo::Parameter;
+	
+	FCustomizableObjectInstanceEditor::ShowGizmoProjectorParameter(ParamName, RangeIndex, SharedThis(this), Viewport, CustomizableInstanceDetailsView, ProjectorParameter, PreviewInstance);
+}
+
+
+void FCustomizableObjectEditor::HideGizmoProjectorParameter()
+{
+	if (ProjectorGizmo != EProjectorGizmo::Parameter)
+	{
+		return;	
+	}
+	
+	FCustomizableObjectInstanceEditor::HideGizmoProjectorParameter(SharedThis(this), Viewport, CustomizableInstanceDetailsView);
+}
+
+
+void FCustomizableObjectEditor::ShowGizmoClipMorph(UCustomizableObjectNodeMeshClipMorph& Node)
+{
+	if (Node.BoneName != FName())
+	{
+		HideGizmo();
+
+		GraphEditor->ClearSelectionSet();
+		GraphEditor->SetNodeSelection(&Node, true);
+
+		Viewport->ShowGizmoClipMorph(Node);
+	}
+}
+
+
+void FCustomizableObjectEditor::HideGizmoClipMorph()
+{
+	Viewport->HideGizmoClipMorph();
+
+	const FGraphPanelSelectionSet SelectedNodes = GraphEditor->GetSelectedNodes();
+	for (FGraphPanelSelectionSet::TConstIterator NodeIt( SelectedNodes ); NodeIt; ++NodeIt)
+	{
+		const UObject* Node = *NodeIt;
+		if (Node->IsA<UCustomizableObjectNodeMeshClipMorph>())
+		{
+			GraphEditor->ClearSelectionSet();
+			break;
+		}
+	}	
+}
+
+
+void FCustomizableObjectEditor::ShowGizmoClipMesh(UCustomizableObjectNodeMeshClipWithMesh& Node)
+{
+	UStaticMesh* ClipMesh = nullptr;
+
+	if (const UEdGraphPin* ConnectedPin = FollowInputPin(*Node.ClipMeshPin()))
+	{
+		if (const UEdGraphNode* ConnectedNode = ConnectedPin->GetOwningNode())
+		{
+			if (const UCustomizableObjectNodeStaticMesh* TypedNode = Cast<UCustomizableObjectNodeStaticMesh>(ConnectedNode))
+			{
+				ClipMesh = TypedNode->StaticMesh;
+			}
+			else if (const UCustomizableObjectNodeTable* TableNode = Cast<UCustomizableObjectNodeTable>(ConnectedNode))
+			{
+				ClipMesh = TableNode->GetColumnDefaultAssetByType<UStaticMesh>(ConnectedPin);
+			}
+		}
+	}
+
+	if (ClipMesh)
+	{
+		HideGizmo();
+
+		GraphEditor->ClearSelectionSet();
+		GraphEditor->SetNodeSelection(&Node, true);
+
+		Viewport->ShowGizmoClipMesh(Node, *ClipMesh);
+	}
+}
+
+
+void FCustomizableObjectEditor::HideGizmoClipMesh()
+{
+	Viewport->HideGizmoClipMesh();
+
+	const FGraphPanelSelectionSet SelectedNodes = GraphEditor->GetSelectedNodes();
+	for (FGraphPanelSelectionSet::TConstIterator NodeIt( SelectedNodes ); NodeIt; ++NodeIt)
+	{
+		const UObject* Node = *NodeIt;
+		if (Node->IsA<UCustomizableObjectNodeMeshClipWithMesh>())
+		{
+			GraphEditor->ClearSelectionSet();
+			break;
+		}
+	}	
+}
+
+
+void FCustomizableObjectEditor::ShowGizmoLight(ULightComponent& InSelectedLight)
+{
+	HideGizmo();
+
+	CustomSettings->SetSelectedLight(&InSelectedLight);
+
+	Viewport->ShowGizmoLight(InSelectedLight);
+	
+	CustomizableObjectEditorAdvancedPreviewSettings->Refresh();
+}
+
+
+void FCustomizableObjectEditor::HideGizmoLight()
+{
+	CustomSettings->SetSelectedLight(nullptr);
+	CustomizableObjectEditorAdvancedPreviewSettings->Refresh();
+}
+
+
 void FCustomizableObjectEditor::PostUndo(bool bSuccess)
 {
 	if (bSuccess)
@@ -1111,12 +1413,6 @@ UCustomizableObject* FCustomizableObjectEditor::GetCustomizableObject()
 
 
 void FCustomizableObjectEditor::RefreshTool()
-{
-	RefreshViewport();
-}
-
-
-void FCustomizableObjectEditor::RefreshViewport()
 {
 	if (ViewportClient)
 	{
@@ -1125,145 +1421,38 @@ void FCustomizableObjectEditor::RefreshViewport()
 }
 
 
+TSharedPtr<SCustomizableObjectEditorViewportTabBody> FCustomizableObjectEditor::GetViewport()
+{
+	return Viewport;
+}
+
+
 void FCustomizableObjectEditor::OnObjectPropertySelectionChanged(FProperty* InProperty)
 {
 	CustomizableObject->PostEditChange();
-	RefreshViewport();
+
+	if (ViewportClient)
+	{
+		ViewportClient->Invalidate();
+	}
 }
 
 
 void FCustomizableObjectEditor::OnInstancePropertySelectionChanged(FProperty* InProperty)
 {
-	RefreshViewport();
+	if (ViewportClient)
+	{
+		ViewportClient->Invalidate();
+	}
 }
 
 
 void FCustomizableObjectEditor::OnObjectModified(UObject* Object)
 {
-	UCustomizableObjectInstance* Instance = Cast<UCustomizableObjectInstance>(Object);
-
-	if ((Instance != nullptr) && (Instance == PreviewInstance))
-	{
-		if (ManagingProjector)
-		{
-			return;
-		}
-
-		if(Instance->ProjectorLayerChange)
-		{
-			// Don't make any changes in the projector if only a projector layer has changed
-			Instance->ProjectorLayerChange = false;
-			return;
-		}
-
-		if (Viewport->GetGizmoHasAssignedData() && !Viewport->GetGizmoAssignedDataIsFromNode())
-		{
-			Instance->LastSelectedProjectorParameter = Viewport->GetGizmoProjectorParameterName();
-			Instance->LastSelectedProjectorParameterWithIndex = Viewport->GetGizmoProjectorParameterNameWithIndex();
-		}
-
-		//UE_LOG(LogMutable, Warning, TEXT("LastSelectedProjectorParameter=%s, LastSelectedProjectorParameterWithIndex=%s"), *(PreviewInstance->LastSelectedProjectorParameter), *(PreviewInstance->LastSelectedProjectorParameterWithIndex));
-
-		// Update projector parameter information (not projector parameter node information,
-		// this is the case where a projector parameter in the object details is selected)
-		const TArray<FCustomizableObjectProjectorParameterValue>& Parameters = PreviewInstance->GetProjectorParameters();
-		const int32 MaxIndex = Parameters.Num();
-		bool AnySelected = false;
-
-		for (int32 i = 0; (!AnySelected && (i < MaxIndex)); ++i)
-		{
-			int32 RangeIndex = Parameters[i].RangeValues.Num() > 0 ? 0 : -1;
-
-			for (; RangeIndex < Parameters[i].RangeValues.Num(); ++RangeIndex)
-			{
-				const FCustomizableObjectProjector& Value = (RangeIndex == -1) ? Parameters[i].Value : Parameters[i].RangeValues[RangeIndex];
-
-				FString ParameterNameWithIndex = Parameters[i].ParameterName;
-
-				if (RangeIndex != -1)
-				{
-					ParameterNameWithIndex += FString::Printf(TEXT("__%d"), RangeIndex);
-				}
-
-				if ((PreviewInstance->GetProjectorState(Parameters[i].ParameterName, RangeIndex) == EProjectorState::Selected) ||
-					((PreviewInstance->GetProjectorState(Parameters[i].ParameterName, RangeIndex) == EProjectorState::TypeChanged) &&
-					(ParameterNameWithIndex != Instance->LastSelectedProjectorParameterWithIndex)))
-				{
-					AnySelected = true;
-					SetProjectorVisibilityForParameter = true;
-					ProjectorParameterName = Parameters[i].ParameterName;
-					ProjectorParameterNameWithIndex = Parameters[i].ParameterName;
-					if (RangeIndex != -1)
-					{
-						ProjectorParameterNameWithIndex += FString::Printf(TEXT("__%d"), RangeIndex);
-					}
-					ProjectorRangeIndex = RangeIndex;
-					ProjectorParameterIndex = i;
-					ProjectorParameterPosition = Value.Position;
-					ProjectorParameterDirection = Value.Direction;
-					ProjectorParameterUp = Value.Up;
-					ProjectorParameterScale = Value.Scale;
-					ProjectorParameterProjectionType = Value.ProjectionType;
-					ProjectionAngle = Value.Angle;
-				}
-				else if ((PreviewInstance->GetProjectorState(Parameters[i].ParameterName, RangeIndex) == EProjectorState::TypeChanged) &&
-					(Parameters[i].ParameterName == Instance->LastSelectedProjectorParameter))
-				{
-					AnySelected = true;
-					SetProjectorTypeForParameter = true;
-					ProjectorParameterName = Parameters[i].ParameterName;
-					ProjectorParameterNameWithIndex = Parameters[i].ParameterName;
-					if (RangeIndex != -1)
-					{
-						ProjectorParameterNameWithIndex += FString::Printf(TEXT("__%d"), RangeIndex);
-					}
-					ProjectorRangeIndex = RangeIndex;
-					ProjectorParameterIndex = i;
-					ProjectorParameterPosition = Value.Position;
-					ProjectorParameterDirection = Value.Direction;
-					ProjectorParameterUp = Value.Up;
-					ProjectorParameterScale = Value.Scale;
-					ProjectorParameterProjectionType = Value.ProjectionType;
-					ProjectionAngle = Value.Angle;
-				}
-			}
-		}
-
-		// Another projector parameter is being selected, and there already a projector parameter selected pending to be hidden,
-		// update its information before selecting the new one
-		bool ResetPendingDone = false;
-		if (AnySelected && SetProjectorVisibilityForParameter && !Instance->LastSelectedProjectorParameter.IsEmpty())
-		{
-			ResetProjectorVisibilityNoUpdate();
-			Instance->LastSelectedProjectorParameter = "";
-			Instance->LastSelectedProjectorParameterWithIndex = "";
-			ResetPendingDone = true;
-		}
-
-		if (!ProjectorConstantNodeSelected && !ProjectorParameterNodeSelected && !ResetPendingDone && !AnySelected && !Viewport->GetIsManipulatingGizmo())
-		{
-			ResetProjectorVisibilityForNonNode = true;
-		}
-		else if (ProjectorParameterNodeSelected)
-		{
-			SelectedProjectorParameterNotNode = true;
-		}
-
-		if (Instance->UnselectProjector)
-		{
-			Instance->UnselectProjector = false;
-
-			ResetProjectorVisibilityNoUpdate();
-			
-			Instance->LastSelectedProjectorParameter = "";
-			Instance->LastSelectedProjectorParameterWithIndex = "";
-		}
-	}
-
-	if (!Instance)
+	if (const UCustomizableObjectInstance* Instance = Cast<UCustomizableObjectInstance>(Object); !Instance)
 	{
 		// Sometimes when another CO is open in another editor window/tab, it triggers this callback, so prevent the modification of this object by a callback triggered by another one
-		if(UCustomizableObject* AuxCustomizableObject = Cast<UCustomizableObject>(Object))
+		if (UCustomizableObject* AuxCustomizableObject = Cast<UCustomizableObject>(Object))
 		{
 			AuxCustomizableObject->UpdateVersionId();
 		}
@@ -1290,11 +1479,7 @@ void FCustomizableObjectEditor::OnObjectModified(UObject* Object)
 
 void FCustomizableObjectEditor::CompileObject()
 {
-	// If any projector selected, unselect it in case the projector node is removed with this CO compile
-	ResetProjectorVisibilityNoUpdate();
-
 	// Resetting viewport parameters
-	Viewport->SetClipMorphPlaneVisibility(false, nullptr);
 	Viewport->SetDrawDefaultUVMaterial(true);
 
 	UE_LOG(LogMutable, Verbose, TEXT("PROFILE: -----------------------------------------------------------"));
@@ -1311,8 +1496,6 @@ void FCustomizableObjectEditor::CompileObject()
 		return;
 	}
 
-	Viewport->UpdateGizmoDataToOrigin();
-
 	if (CustomizableObject->Source)
 	{
 		FCompilationOptions Options = CustomizableObject->CompileOptions;
@@ -1468,57 +1651,59 @@ bool FCustomizableObjectEditor::CanDuplicateSelectedNodes() const
 void FCustomizableObjectEditor::OnSelectedGraphNodesChanged(const FGraphPanelSelectionSet& NewSelection)
 {
 	TArray<UObject*> Objects;
-	for ( FGraphPanelSelectionSet::TConstIterator It(NewSelection); It; ++It)
+	for (FGraphPanelSelectionSet::TConstIterator It(NewSelection); It; ++It)
 	{
 		Objects.Add(*It);
 	}
 
 	// Standard details
-	if ( GraphNodeDetailsView.IsValid() )
+	if (GraphNodeDetailsView.IsValid())
 	{
-		GraphNodeDetailsView->SetObjects( Objects );
-	}
-
-	if (Objects.Num())
+		GraphNodeDetailsView->SetObjects(Objects);
+	}		
+	
+	if (!bRecursionGuard) // Calling the following functions will unselect some nodes causing OnSelectedGraphNodesChanged to be called again
 	{
-		UObject* FirstNode = nullptr;
-		FirstNode = Objects[0];
-		SelectedMeshClipMorphNode = Cast<UCustomizableObjectNodeMeshClipMorph>(FirstNode);
+		TGuardValue<bool> RecursionGuard(bRecursionGuard, true);
 
-		if (!SelectedMeshClipMorphNode)
+		if (Objects.Num() != 1)
 		{
-			SelectedMeshClipWithMeshNode = Cast<UCustomizableObjectNodeMeshClipWithMesh>(FirstNode);
+			HideGizmoClipMorph();
+			HideGizmoClipMesh();
+			HideGizmoProjectorNodeProjectorConstant();
+			HideGizmoProjectorNodeProjectorParameter();
 
-			if (!SelectedProjectorNode || (SelectedProjectorNode && FirstNode && (SelectedProjectorNode != FirstNode)))
+			for (UObject* Object : Objects) // Reselect the multiple selection. Clearly showing gizmos when selecting a node is a really bad idea. Remove on MTBL-1684
 			{
-				SelectedProjectorNode = Cast<UCustomizableObjectNodeProjectorConstant>(FirstNode);
-				if (SelectedProjectorNode != nullptr)
-				{
-					ProjectorConstantNodeSelected = true;
-				}
-			}
-
-			if (!SelectedProjectorParameterNode || (SelectedProjectorParameterNode && FirstNode && (SelectedProjectorParameterNode != FirstNode)))
-			{
-				SelectedProjectorParameterNode = Cast<UCustomizableObjectNodeProjectorParameter>(FirstNode);
-				if (SelectedProjectorParameterNode != nullptr)
-				{
-					ProjectorParameterNodeSelected = true;
-				}
+				GraphEditor->SetNodeSelection(Cast<UEdGraphNode>(Object), true);
 			}
+			
+			return;
 		}
-	}
-	else 
-	{
-		SelectedMeshClipMorphNode = nullptr;
-		SelectedMeshClipWithMeshNode = nullptr;
-		SelectedProjectorNode = nullptr;
-		SelectedProjectorParameterNode = nullptr;
-	}
 
-	if (!ManagingProjector)
-	{
-		SelectedGraphNodesChanged = true;
+		if (UCustomizableObjectNodeMeshClipMorph* NodeMeshClipMorph = Cast<UCustomizableObjectNodeMeshClipMorph>(Objects[0]))
+		{		
+			ShowGizmoClipMorph(*NodeMeshClipMorph);			
+		}
+		else if (UCustomizableObjectNodeMeshClipWithMesh* NodeMeshClipWithMesh = Cast<UCustomizableObjectNodeMeshClipWithMesh>(Objects[0]))
+		{
+			ShowGizmoClipMesh(*NodeMeshClipWithMesh);			
+		}
+		else if (UCustomizableObjectNodeProjectorConstant* NodeProjectorConstant = Cast<UCustomizableObjectNodeProjectorConstant>(Objects[0]))
+		{
+			ShowGizmoProjectorNodeProjectorConstant(*NodeProjectorConstant);
+		}
+		else if (UCustomizableObjectNodeProjectorParameter* NodeProjectorParameter = Cast<UCustomizableObjectNodeProjectorParameter>(Objects[0]))
+		{
+			ShowGizmoProjectorNodeProjectorParameter(*NodeProjectorParameter);		
+		}
+		else
+		{
+			HideGizmoClipMorph();	
+			HideGizmoClipMesh();
+			HideGizmoProjectorNodeProjectorParameter();
+			HideGizmoProjectorNodeProjectorConstant();
+		}
 	}
 }
 
@@ -1548,7 +1733,7 @@ bool FCustomizableObjectEditor::IsTickable() const
 }
 
 
-void FCustomizableObjectEditor::Tick( float InDeltaTime )
+void FCustomizableObjectEditor::Tick(float InDeltaTime)
 {
 	const bool bUpdated = Compiler.Tick();
 	
@@ -1565,143 +1750,6 @@ void FCustomizableObjectEditor::Tick( float InDeltaTime )
 			CreatePreviewInstance();
 		}
 	}
-
-	if (SelectedMeshClipMorphNode)
-	{
-		bool bVisible = SelectedMeshClipMorphNode->BoneName != FName();
-
-		Viewport->SetClipMorphPlaneVisibility(bVisible, SelectedMeshClipMorphNode);
-	}
-	else if (SelectedMeshClipWithMeshNode)
-	{
-		UStaticMesh* ClipMesh = nullptr;
-
-		if (const UEdGraphPin* ConnectedPin = FollowInputPin(*SelectedMeshClipWithMeshNode->ClipMeshPin()))
-		{
-			if (const UEdGraphNode* Node = ConnectedPin->GetOwningNode())
-			{
-				if (const UCustomizableObjectNodeStaticMesh* TypedNode = Cast<UCustomizableObjectNodeStaticMesh>(Node))
-				{
-					ClipMesh = TypedNode->StaticMesh;
-				}
-				else if (const UCustomizableObjectNodeTable* TableNode = Cast<UCustomizableObjectNodeTable>(Node))
-				{
-					ClipMesh = TableNode->GetColumnDefaultAssetByType<UStaticMesh>(ConnectedPin);
-				}
-			}
-		}
-
-		bool bVisible = ClipMesh != nullptr;
-		Viewport->SetClipMeshVisibility(bVisible, ClipMesh, SelectedMeshClipWithMeshNode);
-	}
-	else if (ProjectorConstantNodeSelected)
-	{
-		ManagingProjector = true;
-		Viewport->SetProjectorVisibility(true, SelectedProjectorNode);
-		Viewport->ProjectorParameterChanged(SelectedProjectorNode);
-		ManagingProjector = false;
-		ProjectorConstantNodeSelected = false;
-	}
-	else if (ProjectorParameterNodeSelected)
-	{
-		ManagingProjector = true;
-		Viewport->SetProjectorParameterVisibility(true, SelectedProjectorParameterNode);
-		Viewport->ProjectorParameterChanged(SelectedProjectorParameterNode);
-		ManagingProjector = false;
-		ProjectorParameterNodeSelected = false;
-	}
-	else
-	{
-		if (SetProjectorVisibilityForParameter)
-		{
-			ManagingProjector = true;
-			if (Viewport->AnyProjectorNodeSelected())
-			{
-				GraphEditor->ClearSelectionSet();
-			}
-			FCustomizableObjectProjector ProjectorParameterValue;
-			ProjectorParameterValue.Position = (FVector3f)ProjectorParameterPosition;
-			ProjectorParameterValue.Direction = (FVector3f)ProjectorParameterDirection;
-			ProjectorParameterValue.Up = (FVector3f)ProjectorParameterUp;
-			ProjectorParameterValue.Scale = (FVector3f)ProjectorParameterScale;
-			ProjectorParameterValue.ProjectionType = ProjectorParameterProjectionType;
-			ProjectorParameterValue.Angle = ProjectionAngle;
-			Viewport->SetProjectorVisibility(true, ProjectorParameterName, ProjectorParameterNameWithIndex, ProjectorRangeIndex, ProjectorParameterValue, ProjectorParameterIndex);
-			ManagingProjector = false;
-			SetProjectorVisibilityForParameter = false;
-		}
-
-		if (SetProjectorTypeForParameter)
-		{
-			ManagingProjector = true;
-			if (Viewport->AnyProjectorNodeSelected())
-			{
-				GraphEditor->ClearSelectionSet();
-			}
-			FCustomizableObjectProjector ProjectorParameterValue;
-			ProjectorParameterValue.Position = (FVector3f)ProjectorParameterPosition;
-			ProjectorParameterValue.Direction = (FVector3f)ProjectorParameterDirection;
-			ProjectorParameterValue.Up = (FVector3f)ProjectorParameterUp;
-			ProjectorParameterValue.Scale = (FVector3f)ProjectorParameterScale;
-			ProjectorParameterValue.ProjectionType = ProjectorParameterProjectionType;
-			ProjectorParameterValue.Angle = ProjectionAngle;
-			Viewport->SetProjectorType(true, ProjectorParameterName, ProjectorParameterNameWithIndex, ProjectorRangeIndex, ProjectorParameterValue, ProjectorParameterIndex);
-			ManagingProjector = false;
-			SetProjectorTypeForParameter = false;
-		}
-
-		if (ResetProjectorVisibilityForNonNode && PreviewInstance && !PreviewInstance->AvoidResetProjectorVisibilityForNonNode)
-		{
-			ResetProjectorVisibilityNoUpdate();
-		}
-
-		if ((SelectedGraphNodesChanged && !SelectedProjectorParameterNotNode) || (PreviewInstance && PreviewInstance->TempUpdateGizmoInViewport))
-		{
-			ManagingProjector = true;
-			Viewport->ResetProjectorVisibility(false);
-			ManagingProjector = false;
-		}
-		Viewport->SetClipMorphPlaneVisibility(false, nullptr);
-		Viewport->SetClipMeshVisibility(false, nullptr, nullptr);
-	}
-
-	if (SelectedGraphNodesChanged)
-	{
-		SelectedGraphNodesChanged = false;
-	}
-
-	if (SelectedProjectorParameterNotNode)
-	{
-		SelectedProjectorParameterNotNode = false;
-	}
-
-	if (ResetProjectorVisibilityForNonNode)
-	{
-		ResetProjectorVisibilityForNonNode = false;
-    }
-
-	if (PreviewInstance != nullptr)
-	{
-		PreviewInstance->AvoidResetProjectorVisibilityForNonNode = false;
-	}
-
-    // TEMP CODE
-    if ((PreviewInstance != nullptr) && PreviewInstance->TempUpdateGizmoInViewport)
-    {
-        PreviewInstance->TempUpdateGizmoInViewport = false;
-		PreviewInstance->AvoidResetProjectorVisibilityForNonNode = true;
-        Viewport->CopyTransformFromOriginData();
-        PreviewInstance->UpdateSkeletalMeshAsync(true);
-
-        EProjectorState::Type ProjectorState = PreviewInstance->GetProjectorState(PreviewInstance->TempProjectorParameterName, PreviewInstance->TempProjectorParameterRangeIndex);
-        if (ProjectorState != EProjectorState::Selected)
-        {
-            //ResetProjectorVisibilityForNonNode = false;
-			PreviewInstance->ResetProjectorStates();
-			PreviewInstance->SetProjectorState(PreviewInstance->TempProjectorParameterName, PreviewInstance->TempProjectorParameterRangeIndex, EProjectorState::Selected);
-            FCoreUObjectDelegates::BroadcastOnObjectModified(PreviewInstance);
-        }
-	}
 }
 
 
@@ -1969,45 +2017,12 @@ void FCustomizableObjectEditor::OnAssetRegistryLoadComplete()
 }
 
 
-void FCustomizableObjectEditor::OnObjectPropertyChanged(UObject* Object, FPropertyChangedEvent& PropertyChangedEvent)
-{
-	if (Object->IsA(UCustomizableObjectGraph::StaticClass()) && PropertyChangedEvent.Property)
-	{
-		UClass* OuterClass = PropertyChangedEvent.Property->GetOwnerClass();
-		if (OuterClass->IsChildOf(UCustomizableObjectNodeProjectorConstant::StaticClass()))
-		{
-			if (Viewport.IsValid() && SelectedProjectorNode)
-			{
-				Viewport->ProjectorParameterChanged(SelectedProjectorNode);
-			}
-		}
-		else if (OuterClass->IsChildOf(UCustomizableObjectNodeProjectorParameter::StaticClass()))
-		{
-			if (Viewport.IsValid() && SelectedProjectorParameterNode)
-			{
-				Viewport->ProjectorParameterChanged(SelectedProjectorParameterNode);
-			}
-		}
-	}
-}
-
-
 TSharedPtr<SCustomizableObjectEditorAdvancedPreviewSettings> FCustomizableObjectEditor::GetCustomizableObjectEditorAdvancedPreviewSettings()
 {
 	return CustomizableObjectEditorAdvancedPreviewSettings;
 }
 
 
-void FCustomizableObjectEditor::ResetProjectorVisibilityNoUpdate()
-{
-	ManagingProjector = true;
-	Viewport->SetGizmoCallUpdateSkeletalMesh(false);
-	Viewport->ResetProjectorVisibility(true);
-	Viewport->SetGizmoCallUpdateSkeletalMesh(true);
-	ManagingProjector = false;
-}
-
-
 void FCustomizableObjectEditor::OnEnterText(const FText & NewText, ETextCommit::Type TextType)
 {
 	bool found = false;
@@ -2504,35 +2519,6 @@ void FCustomizableObjectEditor::OnUpdatePreviewInstance()
 	ViewportClient->ReSetAnimation();
 	ViewportClient->SetReferenceMeshMissingWarningMessage(false);
 
-	if (AssetRegistryLoaded)
-	{
-		// If the instance is updated due to a change in a parameter projector, set again the LastSelectedProjectorParameter variable with the
-		// current projector the viewport gizmo has to continue having it as selected when rebuilding the parameter details widget.
-		// Otherwise, reset the value.
-
-		if (Viewport->GetGizmoHasAssignedData() && !Viewport->GetGizmoAssignedDataIsFromNode())
-		{
-			PreviewInstance->LastSelectedProjectorParameter = Viewport->GetGizmoProjectorParameterName();
-			PreviewInstance->LastSelectedProjectorParameterWithIndex = Viewport->GetGizmoProjectorParameterNameWithIndex();
-
-			if (!PreviewInstance->ProjectorAlphaChange && !Viewport->GetIsManipulatingGizmo())
-			{
-				TArray<UObject*> instances;
-				instances.Add(PreviewInstance);
-				CustomizableInstanceDetailsView->SetObjects(instances, true);
-			}
-
-			PreviewInstance->ProjectorAlphaChange = false;
-		}
-		else
-		{
-			PreviewInstance->LastSelectedProjectorParameter = "";
-			PreviewInstance->LastSelectedProjectorParameterWithIndex = "";
-		}
-
-		//UE_LOG(LogMutable, Warning, TEXT("LastSelectedProjectorParameter=%s, LastSelectedProjectorParameterWithIndex=%s"), *(PreviewInstance->LastSelectedProjectorParameter), *(PreviewInstance->LastSelectedProjectorParameterWithIndex));
-	}
-
 	if (TextureAnalyzer.IsValid())
 	{
 		TextureAnalyzer->RefreshTextureAnalyzerTable(PreviewInstance);
@@ -2807,5 +2793,4 @@ void FCustomizableObjectEditor::CreatePreviewComponents()
 }
 
 
-
 #undef LOCTEXT_NAMESPACE
diff --git a/Engine/Plugins/Experimental/Mutable/Source/CustomizableObjectEditor/Private/MuCOE/CustomizableObjectEditor.h b/Engine/Plugins/Experimental/Mutable/Source/CustomizableObjectEditor/Private/MuCOE/CustomizableObjectEditor.h
--- a/Engine/Plugins/Experimental/Mutable/Source/CustomizableObjectEditor/Private/MuCOE/CustomizableObjectEditor.h
+++ b/Engine/Plugins/Experimental/Mutable/Source/CustomizableObjectEditor/Private/MuCOE/CustomizableObjectEditor.h
@@ -2,6 +2,7 @@
 
 #pragma once
 
+#include "CustomizableObjectInstanceEditor.h"
 #include "EditorUndoClient.h"
 #include "GraphEditor.h"
 #include "Misc/NotifyHook.h"
@@ -184,12 +185,27 @@ public:
 	// ICustomizableObjectEditor interface
 	virtual UCustomizableObject* GetCustomizableObject() override;
 	virtual void RefreshTool() override;
-	virtual void RefreshViewport() override;
+	virtual TSharedPtr<SCustomizableObjectEditorViewportTabBody> GetViewport() override;
 	virtual UCustomizableObjectInstance* GetPreviewInstance() override;
 	virtual bool CanPasteNodes() const override;
 	virtual void PasteNodesHere(const FVector2D& Location) override;
 	virtual void SelectNode(const UCustomizableObjectNode* Node) override;
 	virtual void ReconstructAllChildNodes(UCustomizableObjectNode& StartNode, const UClass& NodeType) override;
+	virtual UProjectorParameter* GetProjectorParameter() override;
+	virtual UCustomSettings* GetCustomSettings() override;
+	virtual void HideGizmo() override;
+	virtual void ShowGizmoProjectorNodeProjectorConstant(UCustomizableObjectNodeProjectorConstant& Node) override;
+	virtual void HideGizmoProjectorNodeProjectorConstant() override;
+	virtual void ShowGizmoProjectorNodeProjectorParameter(UCustomizableObjectNodeProjectorParameter& Node) override;
+	virtual void HideGizmoProjectorNodeProjectorParameter() override;
+	virtual void ShowGizmoProjectorParameter(const FString& ParamName, int32 RangeIndex = -1) override;
+	virtual void HideGizmoProjectorParameter() override;
+	virtual void ShowGizmoClipMorph(UCustomizableObjectNodeMeshClipMorph& Node) override;
+	virtual void HideGizmoClipMorph() override;
+	virtual void ShowGizmoClipMesh(UCustomizableObjectNodeMeshClipWithMesh& Node) override;
+	virtual void HideGizmoClipMesh() override;
+	virtual void ShowGizmoLight(ULightComponent& SelectedLight) override;
+	virtual void HideGizmoLight() override;
 
 	/** Called to undo the last action */
 	void UndoGraphAction();
@@ -302,10 +318,7 @@ private:
 
 	/** Callback for the object modified event */
 	void OnObjectModified(UObject* Object);
-	
-	/** */
-	void OnPreviewInstanceUpdated();
-	
+		
 	void CreateGraphEditorWidget(UEdGraph* InGraph);
 
 	/** Copy the currently selected nodes */
@@ -329,12 +342,6 @@ private:
 	/** Updates the visibility of PreviewSkeletalMeshComponent */
 	void UpdatePreviewVisibility();
 
-	/** Handler for when an asset's property has changed */
-	void OnObjectPropertyChanged(UObject* Object, FPropertyChangedEvent& PropertyChangedEvent);
-
-	/** Utility method to reset current's projector visibility with no Skeletal Mesh Update */
-	void ResetProjectorVisibilityNoUpdate();
-
 	/** Searches a node that contains the inserted word */
 	void OnEnterText(const FText& NewText, ETextCommit::Type TextType);
 
@@ -356,13 +363,11 @@ private:
 	void CreatePreviewComponents();
 
 public:
-
 	// Helpers to get the absolute parent of a Customizable Object
 	static UCustomizableObject* GetAbsoluteCOParent(const UCustomizableObjectNodeObject* const Root);
 	static void AddCachedReferencers(const FName& PathName, TArray<FName>& ArrayReferenceNames, TArray<FAssetData>& ArrayAssetData);
 	static void GetExternalChildObjects(const UCustomizableObject* const Object, TArray<UCustomizableObject*>& ExternalChildren, const bool bRecursively = true, const EObjectFlags ExcludeFlags = EObjectFlags::RF_Transient);
 
-
 	/**	The tab ids for all the tabs used */
 	static const FName ViewportTabId;
 	static const FName ObjectPropertiesTabId;
@@ -395,7 +400,7 @@ private:
 	TSharedPtr<SCustomizableObjectEditorViewportTabBody> Viewport;
 	TSharedPtr<FCustomizableObjectEditorViewportClient> ViewportClient;
 
-	TSharedPtr<class IDetailsView> CustomizableInstanceDetailsView;
+	TSharedPtr<IDetailsView> CustomizableInstanceDetailsView;
 
 
 	/** Property View */
@@ -412,36 +417,9 @@ private:
 	/** Widget to select which node pins are visible. */
 	TSharedPtr<class SCustomizableObjectNodePinViewer> NodePinViewer;
 
-	UCustomizableObjectNodeMeshClipMorph* SelectedMeshClipMorphNode = nullptr;
-
-	UCustomizableObjectNodeMeshClipWithMesh* SelectedMeshClipWithMeshNode = nullptr;
-
-	UCustomizableObjectNodeProjectorConstant* SelectedProjectorNode = nullptr;
-
-	UCustomizableObjectNodeProjectorParameter* SelectedProjectorParameterNode = nullptr;
-
 	/** Handle for the OnObjectModified event */
 	FDelegateHandle OnObjectModifiedHandle;
 
-	bool ProjectorConstantNodeSelected = false;
-	bool ProjectorParameterNodeSelected = false;
-	bool SelectedGraphNodesChanged = false;
-	bool SelectedProjectorParameterNotNode = false;
-	bool ResetProjectorVisibilityForNonNode = false;
-	bool SetProjectorVisibilityForParameter = false;
-	bool SetProjectorTypeForParameter = false;
-	FString ProjectorParameterName;
-	FString ProjectorParameterNameWithIndex;
-	int32 ProjectorRangeIndex;
-	int32 ProjectorParameterIndex = -1;
-	FVector3f ProjectorParameterPosition;
-	FVector3f ProjectorParameterDirection;
-	FVector3f ProjectorParameterUp;
-	FVector3f ProjectorParameterScale;
-	ECustomizableObjectProjectorType ProjectorParameterProjectionType;
-	float ProjectionAngle;
-	bool ManagingProjector = false;
-
 	/** Flag to know when the asset registry initial loading has completed */
 	bool AssetRegistryLoaded = false;
 
@@ -456,9 +434,6 @@ private:
 
 	/** Scene preview settings widget */
 	TSharedPtr<class SCustomizableObjectEditorAdvancedPreviewSettings> CustomizableObjectEditorAdvancedPreviewSettings;
-
-	/** Advanced scene preview settings */
-	TObjectPtr<class UCustomizableObjectEmptyClassForSettings> AdditionalSettings;
 	
 	/** Texture Analyzer table widget which shows the information of the transient textures used in the customizable object instance */
 	TSharedPtr<class SCustomizableObjecEditorTextureAnalyzer> TextureAnalyzer;
@@ -477,7 +452,15 @@ private:
 
 	/** Postponed work to do when OnUpdatePreviewInstance is called. Emptied at the end on each OnUpdatePreviewInstance call. */
 	TArray<TFunction<void()>> OnUpdatePreviewInstanceWork;
+
+	TObjectPtr<UProjectorParameter> ProjectorParameter = nullptr;
+
+	TObjectPtr<UCustomSettings> CustomSettings = nullptr;
+
+	bool bRecursionGuard = false;
 	
+	EProjectorGizmo ProjectorGizmo = EProjectorGizmo::None;
+
 protected:
 	/** @return the documentation location for this editor */
 	virtual FString GetDocumentationLink() const override;
diff --git a/Engine/Plugins/Experimental/Mutable/Source/CustomizableObjectEditor/Private/MuCOE/CustomizableObjectEditorActions.cpp b/Engine/Plugins/Experimental/Mutable/Source/CustomizableObjectEditor/Private/MuCOE/CustomizableObjectEditorActions.cpp
--- a/Engine/Plugins/Experimental/Mutable/Source/CustomizableObjectEditor/Private/MuCOE/CustomizableObjectEditorActions.cpp
+++ b/Engine/Plugins/Experimental/Mutable/Source/CustomizableObjectEditor/Private/MuCOE/CustomizableObjectEditorActions.cpp
@@ -58,13 +58,7 @@ void FCustomizableObjectEditorViewportCommands::RegisterCommands()
 	UI_COMMAND( SetCameraLock, "Camera Lock", "Toggles viewport navigation between orbit and freely moving about.", EUserInterfaceActionType::ToggleButton, FInputChord() );
 	UI_COMMAND( SaveThumbnail, "Save Thumbnail", "Saves the viewpoint position and orientation in the Preview Pane for use as the thumbnail preview in the Content Browser.", EUserInterfaceActionType::Button, FInputChord() );
 
-	UI_COMMAND( SetShowNormals, "Normals", "Toggles display of vertex normals in the Preview Pane.", EUserInterfaceActionType::ToggleButton, FInputChord() );
-	UI_COMMAND( SetShowTangents, "Tangents", "Toggles display of vertex tangents in the Preview Pane.", EUserInterfaceActionType::ToggleButton, FInputChord() );
-	UI_COMMAND( SetShowBinormals, "Binormals", "Toggles display of vertex binormals (orthogonal vector to normal and tangent) in the Preview Pane.", EUserInterfaceActionType::ToggleButton, FInputChord() );
-	UI_COMMAND( SetShowPivot, "Show Pivot", "Display the pivot location of the static mesh.", EUserInterfaceActionType::ToggleButton, FInputChord() );
-
 	UI_COMMAND(BakeInstance, "Bake Instance from Object", "Create baked unreal resources for the current preview instance.", EUserInterfaceActionType::Button, FInputChord());
-	UI_COMMAND(StateChangeTest, "State test", "Make a state change test, iterating through all possible states and runtime parameters for this instance", EUserInterfaceActionType::Button, FInputChord());
 	UI_COMMAND(StateChangeShowData, "Show or hide test results", "Show or hide test results", EUserInterfaceActionType::ToggleButton, FInputChord());
 	UI_COMMAND(StateChangeShowGeometryData, "Show instance geometry data", "Show instance geometry data", EUserInterfaceActionType::ToggleButton, FInputChord());
 }
diff --git a/Engine/Plugins/Experimental/Mutable/Source/CustomizableObjectEditor/Private/MuCOE/CustomizableObjectEditorModule.cpp b/Engine/Plugins/Experimental/Mutable/Source/CustomizableObjectEditor/Private/MuCOE/CustomizableObjectEditorModule.cpp
--- a/Engine/Plugins/Experimental/Mutable/Source/CustomizableObjectEditor/Private/MuCOE/CustomizableObjectEditorModule.cpp
+++ b/Engine/Plugins/Experimental/Mutable/Source/CustomizableObjectEditor/Private/MuCOE/CustomizableObjectEditorModule.cpp
@@ -176,7 +176,7 @@ void FCustomizableObjectEditorModule::StartupModule()
 	RegisterCustomDetails(PropertyModule, UCustomizableObjectNodeMeshClipMorph::StaticClass(), FOnGetDetailCustomizationInstance::CreateStatic(&FCustomizableObjectNodeMeshClipMorphDetails::MakeInstance));
 	RegisterCustomDetails(PropertyModule, UCustomizableObjectNodeMeshClipWithMesh::StaticClass(), FOnGetDetailCustomizationInstance::CreateStatic(&FCustomizableObjectNodeMeshClipWithMeshDetails::MakeInstance));
 	RegisterCustomDetails(PropertyModule, UCustomizableObjectNodeExternalPin::StaticClass(), FOnGetDetailCustomizationInstance::CreateStatic(&FCustomizableObjectNodeExternalPinDetails::MakeInstance));
-	RegisterCustomDetails(PropertyModule, UCustomizableObjectEmptyClassForSettings::StaticClass(), FOnGetDetailCustomizationInstance::CreateStatic(&FCustomizableObjectCustomSettingsDetails::MakeInstance));
+	RegisterCustomDetails(PropertyModule, UCustomSettings::StaticClass(), FOnGetDetailCustomizationInstance::CreateStatic(&FCustomizableObjectCustomSettingsDetails::MakeInstance));
 	RegisterCustomDetails(PropertyModule, UCustomizableObjectNodeMaterial::StaticClass(), FOnGetDetailCustomizationInstance::CreateStatic(&FCustomizableObjectNodePinViewerDetails::MakeInstance));
 	RegisterCustomDetails(PropertyModule, UCustomizableObjectNodeCopyMaterial::StaticClass(), FOnGetDetailCustomizationInstance::CreateStatic(&FCustomizableObjectNodeCopyMaterialDetails::MakeInstance));
 	RegisterCustomDetails(PropertyModule, UCustomizableObjectNodeSkeletalMesh::StaticClass(), FOnGetDetailCustomizationInstance::CreateStatic(&FCustomizableObjectNodeSkeletalMeshDetails::MakeInstance));
