[Mutable] Projector Gizmo code refactor.

diff --git a/Engine/Plugins/Experimental/Mutable/Source/CustomizableObjectEditor/Private/MuCOE/CustomizableObjectEditorViewportClient.cpp b/Engine/Plugins/Experimental/Mutable/Source/CustomizableObjectEditor/Private/MuCOE/CustomizableObjectEditorViewportClient.cpp
--- a/Engine/Plugins/Experimental/Mutable/Source/CustomizableObjectEditor/Private/MuCOE/CustomizableObjectEditorViewportClient.cpp
+++ b/Engine/Plugins/Experimental/Mutable/Source/CustomizableObjectEditor/Private/MuCOE/CustomizableObjectEditorViewportClient.cpp
@@ -28,7 +28,6 @@
 #include "MuCO/CustomizableObjectMipDataProvider.h"
 #include "MuCOE/UnrealBakeHelpers.h"
 #include "MuCOE/CustomizableObjectPreviewScene.h"
-#include "MuCOE/CustomizableObjectWidget.h"
 #include "MuCOE/ICustomizableObjectInstanceEditor.h"
 #include "MuCOE/Nodes/CustomizableObjectNodeMeshClipMorph.h"
 #include "MuCOE/Nodes/CustomizableObjectNodeMeshClipWithMesh.h"
@@ -54,340 +53,6 @@ class UTextureMipDataProviderFactory;
 #define LOCTEXT_NAMESPACE "CustomizableObjectEditor" 
 
 
-GizmoRTSProxy::GizmoRTSProxy() :
-	AnyGizmoSelected(false)
-	, DataOriginParameter(nullptr)
-	, DataOriginConstant(nullptr)
-	, HasAssignedData(false)
-	, AssignedDataIsFromNode(false)
-	, ProjectorParameterIndex(-1)
-	, ProjectorGizmoEdited(false)
-	, CallUpdateSkeletalMesh(true)
-	, bManipulatingGizmo(false)
-	, ProjectionType(ECustomizableObjectProjectorType::Planar)
-{
-
-}
-
-
-void GizmoRTSProxy::CopyTransformFromOriginData()
-{
-    if (const TSharedPtr<ICustomizableObjectInstanceEditor> Editor = CustomizableObjectEditorPtr.Pin())
-    {
-	    if (const UCustomizableObjectInstance* Instance = Editor->GetPreviewInstance())
-        {
-			Value = Instance->GetProjector(ProjectorParameterName, ProjectorRangeIndex);
-        }
-    }
-}
-
-
-bool GizmoRTSProxy::UpdateOriginData()
-{
-	if (!HasAssignedData)
-	{
-		return false;
-	}
-
-	if (DataOriginParameter != nullptr)
-	{
-		DataOriginParameter->DefaultValue = Value;
-		return true;
-	}
-	else if (DataOriginConstant != nullptr)
-	{
-		DataOriginConstant->Value = Value;
-		return true;
-	}
-	else
-	{
-		if (CustomizableObjectEditorPtr.IsValid())
-		{
-			if (UCustomizableObjectInstance* Instance = CustomizableObjectEditorPtr.Pin()->GetPreviewInstance())
-			{
-				TArray<FCustomizableObjectProjectorParameterValue>& ProjectorParameters = Instance->GetProjectorParameters();
-				int32 Index = Instance->FindProjectorParameterNameIndex(ProjectorParameterName);
-				
-				for (int32 i = 0; i < Instance->GetProjectorParameters().Num(); ++i)
-				{
-					if (ProjectorParameters[i].ParameterName == ProjectorParameterName)
-					{
-						if (ProjectorRangeIndex >= 0 && !ProjectorParameters[i].RangeValues.IsValidIndex(ProjectorRangeIndex))
-						{
-							return false;
-						}
-					}
-				}
-				
-				if (!Instance->RemovedProjectorParameterNameWithIndex.IsEmpty() && Instance->RemovedProjectorParameterNameWithIndex == ProjectorParameterNameWithIndex) // If the projector has been removed, do not update
-				{
-					Index = -1;
-				}
-				Instance->RemovedProjectorParameterNameWithIndex = FString(""); // Consume the projector removed tag
-
-				if (Index != -1)
-				{
-					if (!ProjectorGizmoEdited)
-					{
-						Instance->PreEditChange(NULL);
-					}
-
-					Instance->SetProjectorValue(ProjectorParameterName,
-						(FVector)Value.Position,
-						(FVector)Value.Direction,
-						(FVector)Value.Up,
-						(FVector)Value.Scale,
-						Value.Angle,
-						ProjectorRangeIndex);
-
-					if (!ProjectorGizmoEdited)
-					{
-						Instance->PostEditChange(); // Avoid unnecessary updates
-					}
-
-					if (CallUpdateSkeletalMesh)
-					{
-						Instance->UpdateSkeletalMeshAsync(true);
-					}
-					return true;
-				}
-			}
-		}
-	}
-
-	return false;
-}
-
-
-bool GizmoRTSProxy::Modify()
-{
-	if (DataOriginParameter != nullptr)
-	{
-		DataOriginParameter->Modify();
-		return true;
-	}
-	else if (DataOriginConstant != nullptr)
-	{
-		DataOriginConstant->Modify();
-		return true;
-	}
-
-	return false;
-}
-
-
-bool GizmoRTSProxy::ModifyGraph()
-{
-	if ((DataOriginParameter == nullptr) && (DataOriginConstant == nullptr))
-	{
-		return false;
-	}
-
-	if (DataOriginParameter != nullptr)
-	{
-		DataOriginParameter->GetGraph()->Modify();
-		return true;
-	}
-	else if (DataOriginConstant != nullptr)
-	{
-		DataOriginConstant->GetGraph()->Modify();
-		return true;
-	}
-
-	return false;
-}
-
-
-void GizmoRTSProxy::CleanOriginData()
-{
-	HasAssignedData = false;
-	AssignedDataIsFromNode = false;
-	AnyGizmoSelected = false;
-	DataOriginParameter = nullptr;
-	DataOriginConstant = nullptr;
-	ProjectorParameterName = FString("");
-	ProjectorParameterNameWithIndex = FString("");
-	ProjectorRangeIndex = -1;
-	ProjectorParameterIndex = -1;
-}
-
-
-bool GizmoRTSProxy::ProjectorHasInitialValues(FCustomizableObjectProjector& Parameter)
-{
-	if ((Parameter.Position == FVector3f(0, 0, 0)) &&
-		(Parameter.Direction == FVector3f(1, 0, 0)) &&
-		(Parameter.Up == FVector3f(0, 1, 0)) &&
-		(Parameter.Scale == FVector3f(10, 10, 100)) &&
-		(Parameter.ProjectionType == ECustomizableObjectProjectorType::Planar) &&
-		(Parameter.Angle == 2.0f * PI))
-	{
-		return true;
-	}
-
-	return false;
-}
-
-
-FCustomizableObjectProjector GizmoRTSProxy::SetProjectorInitialValue(TArray<TWeakObjectPtr<UDebugSkelMeshComponent>>& SkeletalMeshComponents, float TotalLength)
-{
-	FBoxSphereBounds::Builder BoundsBuilder;
-	for (TWeakObjectPtr<USkeletalMeshComponent> SkeletalMeshComponent : SkeletalMeshComponents)
-	{
-		BoundsBuilder += SkeletalMeshComponent->Bounds;
-	}
-
-	FCustomizableObjectProjector Result;
-
-	if (BoundsBuilder.IsValid())
-	{
-		FBoxSphereBounds Bounds(BoundsBuilder);
-		Result.Position = FVector3f(Bounds.Origin + FVector(0.0f, 1.0f, 0.0f) * TotalLength * 0.5f);	// LWC_TODO: Precision Loss
-		Result.Scale = FVector3f(40.0f, 40.0f, 40.0f);
-		Result.Up = FVector3f(0.0f, 0.0f, 1.0f);
-		Result.Direction = FVector3f(0.0f, -1.0f, 0.0f);
-		Result.ProjectionType = ECustomizableObjectProjectorType::Planar;
-		Result.Angle = 2.0f * PI;
-	}
-
-	return Result;
-}
-
-
-bool GizmoRTSProxy::IsProjectorParameterSelected()
-{
-	if (AnyGizmoSelected && HasAssignedData && !AssignedDataIsFromNode && (DataOriginParameter == nullptr) && (DataOriginConstant == nullptr))
-	{
-		return true;
-	}
-
-	return false;
-}
-
-
-void GizmoRTSProxy::ProjectorParameterChanged(UCustomizableObjectNodeProjectorParameter* Node)
-{
-	if (AnyGizmoSelected && AssignedDataIsFromNode && (Node == DataOriginParameter))
-	{
-		switch (Node->ParameterSetModified)
-		{
-			case 0:
-			{
-				ProjectionType = Node->ProjectionType;
-				Value.ProjectionType = Node->ProjectionType;
-				break;
-			}
-			case 1:
-			{
-				Value.Angle = Node->DefaultValue.Angle;
-				break;
-			}
-			case 2:
-			{
-				Value.Position = Node->DefaultValue.Position;
-				Value.Direction = Node->DefaultValue.Direction;
-				Value.Scale = Node->DefaultValue.Scale;
-				Value.Up = Node->DefaultValue.Up;
-				Node->ParameterSetModified = -1;
-				break;
-			}
-			default:
-			{
-				UE_LOG(LogMutable, Warning, TEXT("ERROR: wrong parameter set modified value %d"), Node->ParameterSetModified);
-				break;
-			}
-		}
-	}
-}
-
-
-void GizmoRTSProxy::ProjectorParameterChanged(UCustomizableObjectNodeProjectorConstant* Node)
-{
-	if (AnyGizmoSelected && AssignedDataIsFromNode && (Node == DataOriginConstant))
-	{
-		switch (Node->ParameterSetModified)
-		{
-			case 0:
-			{
-				ProjectionType = Node->ProjectionType;
-				Value.ProjectionType = Node->ProjectionType;
-				break;
-			}
-			case 1:
-			{
-				Value.Angle = Node->Value.Angle;
-				break;
-			}
-			case 2:
-			{
-				Value.Position = Node->Value.Position;
-				Value.Direction = Node->Value.Direction;
-				Value.Up = Node->Value.Up;
-				break;
-			}
-			default:
-			{
-				UE_LOG(LogMutable, Warning, TEXT("ERROR: wrong parameter set modified value %d"), Node->ParameterSetModified);
-				break;
-			}
-		}
-	}
-}
-
-
-void GizmoRTSProxy::SetCallUpdateSkeletalMesh(bool InValue)
-{
-	CallUpdateSkeletalMesh = InValue;
-}
-
-
-void GizmoRTSProxy::SetProjectorUpdatedInViewport(bool InValue)
-{
-	if (!AnyGizmoSelected || DataOriginParameter || DataOriginConstant)
-	{
-		return;
-	}
-
-	if(InValue != bManipulatingGizmo)
-	{
-		bManipulatingGizmo = InValue;
-
-		if (CustomizableObjectEditorPtr.IsValid())
-		{
-			UCustomizableObjectInstance* Inst = CustomizableObjectEditorPtr.Pin()->GetPreviewInstance();
-
-			if (Inst)
-			{
-				Inst->ProjectorUpdatedInViewport = InValue;
-			}
-		}
-	}
-}
-
-
-FString GizmoRTSProxy::GetProjectorParameterName()
-{
-	return ProjectorParameterName;
-}
-
-
-FString GizmoRTSProxy::GetProjectorParameterNameWithIndex()
-{
-	return ProjectorParameterNameWithIndex;
-}
-
-
-bool GizmoRTSProxy::GetHasAssignedData()
-{
-	return HasAssignedData;
-}
-
-
-bool GizmoRTSProxy::GetAssignedDataIsFromNode()
-{
-	return AssignedDataIsFromNode;
-}
-
-
 FCustomizableObjectEditorViewportClient::FCustomizableObjectEditorViewportClient(TWeakPtr<ICustomizableObjectInstanceEditor> InCustomizableObjectEditor, FPreviewScene* InPreviewScene)
 	: FEditorViewportClient(&GLevelEditorModeTools(), InPreviewScene)
 	, CustomizableObjectEditorPtr(InCustomizableObjectEditor)
@@ -399,15 +64,9 @@ FCustomizableObjectEditorViewportClient::FCustomizableObjectEditorViewportClient
 	ConfigOption = UPersonaOptions::StaticClass()->GetDefaultObject<UPersonaOptions>();
 	check (ConfigOption);
 
-	bCameraMove = false;
 	bUsingOrbitCamera = true;
 
-	bShowSockets = true;
 	bDrawUVs = false;
-	bDrawNormals = false;
-	bDrawTangents = false;
-	bDrawBinormals = false;
-	bShowPivot = false;
 	Widget->SetDefaultVisibility(false);
 	bCameraLock = true;
 	bDrawSky = true;
@@ -426,8 +85,6 @@ FCustomizableObjectEditorViewportClient::FCustomizableObjectEditorViewportClient
 	UVChannelToDrawInUVs = 0;
 	MaterialToDrawInUVsComponent = 0;
 
-	bManipulating = false;
-	LastHitProxyWidget = 0;
 
 	bReferenceMeshMissingWarningMessageVisible = false;
 
@@ -465,7 +122,6 @@ FCustomizableObjectEditorViewportClient::FCustomizableObjectEditorViewportClient
 
 	// now add the ClipMorph plane
 	ClipMorphNode = nullptr;
-	bClipMorphVisible = false;
 	bClipMorphLocalStartOffset = true;
 	ClipMorphMaterial = LoadObject<UMaterial>(NULL, TEXT("Material'/Engine/EditorMaterials/LevelGridMaterial.LevelGridMaterial'"), NULL, LOAD_None, NULL);
 	check(ClipMorphMaterial);
@@ -476,7 +132,6 @@ FCustomizableObjectEditorViewportClient::FCustomizableObjectEditorViewportClient
 	PreviewScene->AddComponent(ClipMeshComp, FTransform());
 	ClipMeshComp->SetVisibility(false);
 
-	WidgetMode = UE::Widget::WM_Translate;
 	BoundSphere.W = 100.f;
 
 	const float FOVMin = 5.f;
@@ -489,16 +144,12 @@ FCustomizableObjectEditorViewportClient::FCustomizableObjectEditorViewportClient
 		AddRealtimeOverride(false, LOCTEXT("RealtimeOverrideMessage_InstanceViewport", "Instance Viewport")); // We are PIE, don't start in realtime mode
 	}
 
-	GizmoProxy.CustomizableObjectEditorPtr = InCustomizableObjectEditor;
-	WidgetVisibility = false;
-
 	IsPlayingAnimation = false;
 	AnimationBeingPlayed = nullptr;
 
 	// Lighting 
 	SelectedLightComponent = nullptr;
-	bIsEditingLightEnabled = false;
-
+	
 	StateChangeShowGeometryDataFlag = false;
 
 	// Register delegate to update the show flags when the post processing is turned on or off
@@ -543,7 +194,6 @@ void FCustomizableObjectEditorViewportClient::UpdateCameraSetup()
 void FCustomizableObjectEditorViewportClient::UpdateFloor()
 {
 	// Move the floor to the bottom of the bounding box of the mesh, rather than on the origin
-	FVector Bottom(0.0f);
 	bool bFoundSkelMesh = false;
 
 	for (TWeakObjectPtr<USkeletalMeshComponent> SkeletalMeshComponent : SkeletalMeshComponents)
@@ -553,7 +203,6 @@ void FCustomizableObjectEditorViewportClient::UpdateFloor()
 			SkeletalMeshComponent->bComponentUseFixedSkelBounds = true;
 			SkeletalMeshComponent->UpdateBounds();
 
-			Bottom = SkeletalMeshComponent->Bounds.GetBoxExtrema(0);
 			bFoundSkelMesh = true;
 		}
 	}
@@ -566,8 +215,6 @@ void FCustomizableObjectEditorViewportClient::UpdateFloor()
 	else if (StaticMeshComponent.IsValid())
 	{
 		StaticMeshComponent->UpdateBounds();
-
-		Bottom = StaticMeshComponent->Bounds.GetBoxExtrema(0);
 	}
 
 	FAdvancedPreviewScene* AdvancedScene = static_cast<FAdvancedPreviewScene*>(PreviewScene);
@@ -594,8 +241,6 @@ void FCustomizableObjectEditorViewportClient::Tick(float DeltaSeconds)
 	FEditorViewportClient::Tick(DeltaSeconds);
 
 	UpdateFloor();
-
-	//PreviewScene->GetWorld()->Tick(LEVELTICK_All, DeltaSeconds);
 }
 
 
@@ -616,210 +261,169 @@ void DrawEllipse(FPrimitiveDrawInterface* PDI, const FVector& Base, const FVecto
 void FCustomizableObjectEditorViewportClient::Draw(const FSceneView* View, FPrimitiveDrawInterface* PDI)
 {
 	FEditorViewportClient::Draw(View, PDI);
-
-	//DrawHelper.Draw( View, PDI );	
-	if( StaticMeshComponent.IsValid() && StaticMeshComponent->GetStaticMesh() && (bDrawNormals || bDrawTangents || bDrawBinormals ) )
+	
+	switch (WidgetType)
 	{
-		/* TODO
-		FStaticMeshRenderData* RenderData = &StaticMeshComponent->StaticMesh->LODModels[CustomizableObjectEditorPtr.Pin()->GetCurrentLODIndex()];
-		uint16* Indices = (uint16*)RenderData->IndexBuffer.Indices.GetData();
-		uint32 NumIndices = RenderData->IndexBuffer.Indices.Num();
-
-		FMatrix LocalToWorldInverseTranspose = StaticMeshComponent->ComponentToWorld.ToMatrixWithScale().Inverse().GetTransposed();
-		for (uint32 i = 0; i < NumIndices; i++)
+	case EWidgetType::Light:
 		{
-			const FVector& VertexPos = RenderData->PositionVertexBuffer.VertexPosition( Indices[i] );
+			check(SelectedLightComponent);
 
-			const FVector WorldPos = StaticMeshComponent->ComponentToWorld.TransformPosition( VertexPos );
-			const FVector& Normal = RenderData->VertexBuffer.VertexTangentZ( Indices[i] ); 
-			const FVector& Binormal = RenderData->VertexBuffer.VertexTangentY( Indices[i] ); 
-			const FVector& Tangent = RenderData->VertexBuffer.VertexTangentX( Indices[i] ); 
+			if (USpotLightComponent* SpotLightComp = Cast<USpotLightComponent>(SelectedLightComponent))
+			{
+				FTransform TransformNoScale = SpotLightComp->GetComponentToWorld();
+				TransformNoScale.RemoveScaling();
 
-			const float Len = 5.0f;
+				// Draw point light source shape
+				DrawWireCapsule(PDI, TransformNoScale.GetTranslation(), -TransformNoScale.GetUnitAxis(EAxis::Z), TransformNoScale.GetUnitAxis(EAxis::Y), TransformNoScale.GetUnitAxis(EAxis::X),
+					FColor(231, 239, 0, 255), SpotLightComp->SourceRadius, 0.5f * SpotLightComp->SourceLength + SpotLightComp->SourceRadius, 25, SDPG_World);
 
-			if( bDrawNormals )
-			{
-				PDI->DrawLine( WorldPos, WorldPos+LocalToWorldInverseTranspose.TransformVector( Normal ).SafeNormal() * Len, FLinearColor( 0.0f, 1.0f, 0.0f), SDPG_World );
-			}
+				// Draw outer light cone
+				DrawWireSphereCappedCone(PDI, TransformNoScale, SpotLightComp->AttenuationRadius, SpotLightComp->OuterConeAngle, 32, 8, 10, FColor(200, 255, 255), SDPG_World);
 
-			if( bDrawTangents )
-			{
-				PDI->DrawLine( WorldPos, WorldPos+LocalToWorldInverseTranspose.TransformVector( Tangent ).SafeNormal() * Len, FLinearColor( 1.0f, 0.0f, 0.0f), SDPG_World );
+				// Draw inner light cone (if non zero)
+				if (SpotLightComp->InnerConeAngle > UE_KINDA_SMALL_NUMBER)
+				{
+					DrawWireSphereCappedCone(PDI, TransformNoScale, SpotLightComp->AttenuationRadius, SpotLightComp->InnerConeAngle, 32, 8, 10, FColor(150, 200, 255), SDPG_World);
+				}
 			}
-
-			if( bDrawBinormals )
+			else if (UPointLightComponent* PointLightComp = Cast<UPointLightComponent>(SelectedLightComponent))
 			{
-				PDI->DrawLine( WorldPos, WorldPos+LocalToWorldInverseTranspose.TransformVector( Binormal ).SafeNormal() * Len, FLinearColor( 0.0f, 0.0f, 1.0f), SDPG_World );
-			}
-		}	
-		*/
-	}
-
+				FTransform LightTM = PointLightComp->GetComponentToWorld();
 
-	//if( bShowPivot )
-	//{
-	//	FMatrix Transform;
-	//	if (StaticMeshComponent.IsValid())
-	//	{
-	//		Transform = StaticMeshComponent->ComponentToWorld.ToMatrixWithScale();
-	//	}
-	//	else if (SkeletalMeshComponent.IsValid())
-	//	{
-	//		Transform = SkeletalMeshComponent->ComponentToWorld.ToMatrixWithScale();
-	//	}
+				// Draw light radius
+				DrawWireSphereAutoSides(PDI, FTransform(LightTM.GetTranslation()), FColor(200, 255, 255), PointLightComp->AttenuationRadius, SDPG_World);
 
-	//	FUnrealEdUtils::DrawWidget(View, PDI, Transform, 0, 0, EAxisList::All, EWidgetMovementMode::WMM_Translate, false);
-	//}
-
-	// Draw the widgets
-	for ( TMap< int, TSharedPtr<FCustomizableObjectWidget> >::TIterator it(ProjectorWidgets); it; ++it )
-	{
-		it.Value()->Render( View, PDI );
-	}
-	
-	// Draw Selected light Visualizer
-	if (bIsEditingLightEnabled && SelectedLightComponent)
-	{
-		if (USpotLightComponent* SpotLightComp = Cast<USpotLightComponent>(SelectedLightComponent))
+				// Draw point light source shape
+				DrawWireCapsule(PDI, LightTM.GetTranslation(), -LightTM.GetUnitAxis(EAxis::Z), LightTM.GetUnitAxis(EAxis::Y), LightTM.GetUnitAxis(EAxis::X),
+					FColor(231, 239, 0, 255), PointLightComp->SourceRadius, 0.5f * PointLightComp->SourceLength + PointLightComp->SourceRadius, 25, SDPG_World);
+			}
+			
+			break;
+		}
+	case EWidgetType::ClipMorph:
 		{
-			FTransform TransformNoScale = SpotLightComp->GetComponentToWorld();
-			TransformNoScale.RemoveScaling();
-
-			// Draw point light source shape
-			DrawWireCapsule(PDI, TransformNoScale.GetTranslation(), -TransformNoScale.GetUnitAxis(EAxis::Z), TransformNoScale.GetUnitAxis(EAxis::Y), TransformNoScale.GetUnitAxis(EAxis::X),
-				FColor(231, 239, 0, 255), SpotLightComp->SourceRadius, 0.5f * SpotLightComp->SourceLength + SpotLightComp->SourceRadius, 25, SDPG_World);
+			float MaxSphereRadius = 0.f;
 
-			// Draw outer light cone
-			DrawWireSphereCappedCone(PDI, TransformNoScale, SpotLightComp->AttenuationRadius, SpotLightComp->OuterConeAngle, 32, 8, 10, FColor(200, 255, 255), SDPG_World);
+			for (TWeakObjectPtr<USkeletalMeshComponent> SkeletalMeshComponent : SkeletalMeshComponents)
+			{
+				if (SkeletalMeshComponent.IsValid())
+				{
+					MaxSphereRadius = FMath::Max(MaxSphereRadius, SkeletalMeshComponent->Bounds.SphereRadius);
+				}
+			}
 
-			// Draw inner light cone (if non zero)
-			if (SpotLightComp->InnerConeAngle > UE_KINDA_SMALL_NUMBER)
+			if (MaxSphereRadius <= 0.f)
 			{
-				DrawWireSphereCappedCone(PDI, TransformNoScale, SpotLightComp->AttenuationRadius, SpotLightComp->InnerConeAngle, 32, 8, 10, FColor(150, 200, 255), SDPG_World);
+				MaxSphereRadius = 1.f;
 			}
-		}
-		else if (UPointLightComponent* PointLightComp = Cast<UPointLightComponent>(SelectedLightComponent))
-		{
-			FTransform LightTM = PointLightComp->GetComponentToWorld();
 
-			// Draw light radius
-			DrawWireSphereAutoSides(PDI, FTransform(LightTM.GetTranslation()), FColor(200, 255, 255), PointLightComp->AttenuationRadius, SDPG_World);
+			float PlaneRadius1 = MaxSphereRadius * 0.1f;
+			float PlaneRadius2 = PlaneRadius1 * 0.5f;
 
-			// Draw point light source shape
-			DrawWireCapsule(PDI, LightTM.GetTranslation(), -LightTM.GetUnitAxis(EAxis::Z), LightTM.GetUnitAxis(EAxis::Y), LightTM.GetUnitAxis(EAxis::X),
-				FColor(231, 239, 0, 255), PointLightComp->SourceRadius, 0.5f * PointLightComp->SourceLength + PointLightComp->SourceRadius, 25, SDPG_World);
-		}
-	}
+			FMatrix PlaneMatrix = FMatrix(ClipMorphNormal, ClipMorphYAxis, ClipMorphXAxis, ClipMorphOrigin + ClipMorphOffset);
 
-	
-	// Draw the clip morph axis, planes and bounds if necessary
-	if (bClipMorphVisible && SkeletalMeshComponents.Num())
-	{
-		float MaxSphereRadius = 0.f;
+			// Start Plane
+			DrawDirectionalArrow(PDI, PlaneMatrix, FColor::Red, MorphLength, MorphLength * 0.1f, 0, 0.1f);
+			DrawBox(PDI, PlaneMatrix, FVector(0.01f, PlaneRadius1, PlaneRadius1), ClipMorphMaterial->GetRenderProxy(), 0);
 
-		for (TWeakObjectPtr<USkeletalMeshComponent> SkeletalMeshComponent : SkeletalMeshComponents)
-		{
-			if (SkeletalMeshComponent.IsValid())
-			{
-				MaxSphereRadius = FMath::Max(MaxSphereRadius, SkeletalMeshComponent->Bounds.SphereRadius);
-			}
+			// End Plane + Ellipse
+			PlaneMatrix.SetOrigin(ClipMorphOrigin + ClipMorphOffset + ClipMorphNormal * MorphLength);
+			DrawBox(PDI, PlaneMatrix, FVector(0.01f, PlaneRadius2, PlaneRadius2), ClipMorphMaterial->GetRenderProxy(), 0);
+			DrawEllipse(PDI, ClipMorphOrigin + ClipMorphOffset + ClipMorphNormal * MorphLength, ClipMorphXAxis, ClipMorphYAxis, FColor::Red, Radius1, Radius2, 15, 1, 0.f, 0, false);
+			
+			break;
 		}
 
-		if (MaxSphereRadius <= 0.f)
+	case EWidgetType::Projector:
 		{
-			MaxSphereRadius = 1.f;
-		}
-
-		float PlaneRadius1 = MaxSphereRadius * 0.1f;
-		float PlaneRadius2 = PlaneRadius1 * 0.5f;
+			const FColor Color = WidgetColorDelegate.IsBound() ?
+			WidgetColorDelegate.Execute() :
+			FColor::Green;
 
-		FMatrix PlaneMatrix = FMatrix(ClipMorphNormal, ClipMorphYAxis, ClipMorphXAxis, ClipMorphOrigin + ClipMorphOffset);
+			const ECustomizableObjectProjectorType ProjectorType = ProjectorTypeDelegate.IsBound() ?
+				ProjectorTypeDelegate.Execute() :
+				ECustomizableObjectProjectorType::Planar;
 
-		// Start Plane
-		DrawDirectionalArrow(PDI, PlaneMatrix, FColor::Red, MorphLength, MorphLength * 0.1f, 0, 0.1f);
-		DrawBox(PDI, PlaneMatrix, FVector(0.01f, PlaneRadius1, PlaneRadius1), ClipMorphMaterial->GetRenderProxy(), 0);
+			const FVector WidgetScale = WidgetScaleDelegate.IsBound() ?
+				WidgetScaleDelegate.Execute() :
+				FVector::OneVector;
 
-		// End Plane + Ellipse
-		PlaneMatrix.SetOrigin(ClipMorphOrigin + ClipMorphOffset + ClipMorphNormal * MorphLength);
-		DrawBox(PDI, PlaneMatrix, FVector(0.01f, PlaneRadius2, PlaneRadius2), ClipMorphMaterial->GetRenderProxy(), 0);
-		DrawEllipse(PDI, ClipMorphOrigin + ClipMorphOffset + ClipMorphNormal * MorphLength, ClipMorphXAxis, ClipMorphYAxis, FColor::Red, Radius1, Radius2, 15, 1, 0.f, 0, false);
-	}
+			const float CylindricalAngle = WidgetAngleDelegate.IsBound() ? 
+				FMath::DegreesToRadians<float>(WidgetAngleDelegate.Execute()) :
+				0.0f;
 
-	if (GizmoProxy.AnyGizmoSelected)
-	{
-		FColor Color;
-		if (GizmoProxy.AssignedDataIsFromNode)
-		{
-			Color = FColor::Red;
-		}
-		else
-		{
-			Color = FColor::Emerald;
-		}
+			const FVector CorrectedWidgetScale = FVector(WidgetScale.Z, WidgetScale.X, WidgetScale.Y);
 
-		switch (GizmoProxy.ProjectionType)
-		{
-			case ECustomizableObjectProjectorType::Planar:
-			{
-				FVector Scale = FVector(GizmoProxy.Value.Scale.Z, GizmoProxy.Value.Scale.X, GizmoProxy.Value.Scale.Y);
-				FVector Min = FVector(0.f, -0.5f, -0.5f);
-				FVector Max = FVector(1.0f, 0.5f, 0.5f);
-				FBox Box = FBox(Min * Scale, Max * Scale);
-				FMatrix Mat = GetWidgetCoordSystem();
-				Mat.SetOrigin(GetWidgetLocation());
-				DrawWireBox(PDI, Mat, Box, Color, 1, 0.f);
-				break;
-			}
-			case ECustomizableObjectProjectorType::Cylindrical:
+			switch (ProjectorType)
 			{
-				// Draw the cylinder
-				FVector Scale = FVector(GizmoProxy.Value.Scale.Z, GizmoProxy.Value.Scale.X, GizmoProxy.Value.Scale.Y);
-				FMatrix Mat = GetWidgetCoordSystem();
-				FVector Location = GetWidgetLocation();
-				Mat.SetOrigin(Location);
-				FVector TransformedX = Mat.TransformVector(FVector(1, 0, 0));
-				FVector TransformedY = Mat.TransformVector(FVector(0, 1, 0));
-				FVector TransformedZ = Mat.TransformVector(FVector(0, 0, 1));
-
-				FVector Min = FVector(0.f, -0.5f, -0.5f);
-				FVector Max = FVector(1.0f, 0.5f, 0.5f);
-				FBox Box = FBox(Min * Scale, Max * Scale);
-				FVector BoxExtent = Box.GetExtent();
-				float CylinderHalfHeight = BoxExtent.X;
-				//float CylinderRadius = (BoxExtent.Y + BoxExtent.Z) * 0.5f;
-				float CylinderRadius = FMath::Abs(BoxExtent.Y);
-
-				DrawWireCylinder(PDI, Location + TransformedX * CylinderHalfHeight, TransformedY, TransformedZ, TransformedX, Color, CylinderRadius, CylinderHalfHeight, 16, SDPG_World, 0.1f, 0, false);
-
-				// Draw the arcs: the locations are Location with an offset towards the local forward direction
-				FVector Location0 = Location - TransformedX * CylinderHalfHeight * 0.8f + TransformedX * CylinderHalfHeight;
-				FVector Location1 = Location + TransformedX * CylinderHalfHeight * 0.8f + TransformedX * CylinderHalfHeight;
-				FMatrix Mat0 = Mat;
-				FMatrix Mat1 = Mat;
-				Mat0.SetOrigin(Location0);
-				Mat1.SetOrigin(Location1);
-				DrawCylinderArc(PDI, Mat0, FVector(0.0f, 0.0f, 0.0f), FVector(0, 1, 0), FVector(0, 0, 1), FVector(1, 0, 0),  CylinderRadius, CylinderHalfHeight * 0.1f, 16, TransparentPlaneMaterialXY->GetRenderProxy(), SDPG_World, FColor(255, 85, 0, 192), GizmoProxy.Value.Angle);
-				DrawCylinderArc(PDI, Mat1, FVector(0.0f, 0.0f, 0.0f), FVector(0, 1, 0), FVector(0, 0, 1), FVector(1, 0, 0), CylinderRadius, CylinderHalfHeight * 0.1f, 16, TransparentPlaneMaterialXY->GetRenderProxy(), SDPG_World, FColor(255, 85, 0, 192), GizmoProxy.Value.Angle);
-				break;
-			}
-			case ECustomizableObjectProjectorType::Wrapping:
-            {
-                FVector Scale = FVector(GizmoProxy.Value.Scale.Z, GizmoProxy.Value.Scale.X, GizmoProxy.Value.Scale.Y);
-                FVector Min = FVector(0.f, -0.5f, -0.5f);
-                FVector Max = FVector(1.0f, 0.5f, 0.5f);
-                FBox Box = FBox(Min * Scale, Max * Scale);
-                FMatrix Mat = GetWidgetCoordSystem();
-                Mat.SetOrigin(GetWidgetLocation());
-                DrawWireBox(PDI, Mat, Box, Color, 1, 0.f);
-				break;
-			}
-			default:
-			{
-				UE_LOG(LogMutable, Warning, TEXT("ERROR: wrong projector type for projector %d"), *GizmoProxy.ProjectorParameterName);
-				break;
+				case ECustomizableObjectProjectorType::Planar:
+				{
+					FVector Min = FVector(0.f, -0.5f, -0.5f);
+					FVector Max = FVector(1.0f, 0.5f, 0.5f);
+					FBox Box = FBox(Min * CorrectedWidgetScale, Max * CorrectedWidgetScale);
+					FMatrix Mat = GetWidgetCoordSystem();
+					Mat.SetOrigin(GetWidgetLocation());
+					DrawWireBox(PDI, Mat, Box, Color, 1, 0.f);
+					break;
+				}
+				case ECustomizableObjectProjectorType::Cylindrical:
+				{
+					// Draw the cylinder
+					FMatrix Mat = GetWidgetCoordSystem();
+					FVector Location = GetWidgetLocation();
+					Mat.SetOrigin(Location);
+					FVector TransformedX = Mat.TransformVector(FVector(1, 0, 0));
+					FVector TransformedY = Mat.TransformVector(FVector(0, 1, 0));
+					FVector TransformedZ = Mat.TransformVector(FVector(0, 0, 1));
+
+					FVector Min = FVector(0.f, -0.5f, -0.5f);
+					FVector Max = FVector(1.0f, 0.5f, 0.5f);
+					FBox Box = FBox(Min * CorrectedWidgetScale, Max * CorrectedWidgetScale);
+					FVector BoxExtent = Box.GetExtent();
+					float CylinderHalfHeight = BoxExtent.X;
+					//float CylinderRadius = (BoxExtent.Y + BoxExtent.Z) * 0.5f;
+					float CylinderRadius = FMath::Abs(BoxExtent.Y);
+
+					DrawWireCylinder(PDI, Location + TransformedX * CylinderHalfHeight, TransformedY, TransformedZ, TransformedX, Color, CylinderRadius, CylinderHalfHeight, 16, SDPG_World, 0.1f, 0, false);
+
+					// Draw the arcs: the locations are Location with an offset towards the local forward direction
+					FVector Location0 = Location - TransformedX * CylinderHalfHeight * 0.8f + TransformedX * CylinderHalfHeight;
+					FVector Location1 = Location + TransformedX * CylinderHalfHeight * 0.8f + TransformedX * CylinderHalfHeight;
+					FMatrix Mat0 = Mat;
+					FMatrix Mat1 = Mat;
+					Mat0.SetOrigin(Location0);
+					Mat1.SetOrigin(Location1);
+					DrawCylinderArc(PDI, Mat0, FVector(0.0f, 0.0f, 0.0f), FVector(0, 1, 0), FVector(0, 0, 1), FVector(1, 0, 0),  CylinderRadius, CylinderHalfHeight * 0.1f, 16, TransparentPlaneMaterialXY->GetRenderProxy(), SDPG_World, FColor(255, 85, 0, 192), CylindricalAngle);
+					DrawCylinderArc(PDI, Mat1, FVector(0.0f, 0.0f, 0.0f), FVector(0, 1, 0), FVector(0, 0, 1), FVector(1, 0, 0), CylinderRadius, CylinderHalfHeight * 0.1f, 16, TransparentPlaneMaterialXY->GetRenderProxy(), SDPG_World, FColor(255, 85, 0, 192), CylindricalAngle);
+					break;
+				}
+				case ECustomizableObjectProjectorType::Wrapping:
+		        {
+		            FVector Min = FVector(0.f, -0.5f, -0.5f);
+		            FVector Max = FVector(1.0f, 0.5f, 0.5f);
+		            FBox Box = FBox(Min * CorrectedWidgetScale, Max * CorrectedWidgetScale);
+		            FMatrix Mat = GetWidgetCoordSystem();
+		            Mat.SetOrigin(GetWidgetLocation());
+		            DrawWireBox(PDI, Mat, Box, Color, 1, 0.f);
+					break;
+				}
+				default:
+				{
+					check(false);
+					break;
+				}
 			}
+			break;
 		}
+
+	case EWidgetType::ClipMesh:
+	case EWidgetType::Hidden:
+		break;
+
+	default:
+		unimplemented(); // Case not implemented	
 	}
+	
 
 	if (bShowBones)
 	{
@@ -848,79 +452,10 @@ void FCustomizableObjectEditorViewportClient::Draw(FViewport* InViewport, FCanva
 	FEditorViewportClient::Draw(InViewport, Canvas);
 
 	FSceneViewFamilyContext ViewFamily( FSceneViewFamily::ConstructionValues( InViewport, GetScene(), EngineShowFlags ));
-	FSceneView* View = CalcSceneView(&ViewFamily);
-
-	const int32 HalfX = InViewport->GetSizeXY().X/2;
-	const int32 HalfY = InViewport->GetSizeXY().Y/2;
-
-	//int32 CurrentLODLevel = CustomizableObjectEditorPtr.Pin()->GetCurrentLODIndex();
-
-	int32 YPos = 6;
-	//if ( StaticMeshComponent.IsValid() || SkeletalMeshComponent.IsValid() )
-	//{
-		/* TODO
-		int32 NumVertices = 0;
-		int32 NumTriangles = 0;
-		int32 NumUVChannels = 0;
-		FBoxSphereBounds Bounds(ForceInit);
-		if (StaticMeshComponent && StaticMeshComponent->StaticMesh)
-		{
-			Bounds = StaticMeshComponent->StaticMesh->Bounds;
-			NumTriangles = StaticMeshComponent->StaticMesh->LODModels[CurrentLODLevel].IndexBuffer.Indices.Num()/3;
-			NumVertices = StaticMeshComponent->StaticMesh->LODModels[CurrentLODLevel].VertexBuffer.GetNumVertices();
-			NumUVChannels = StaticMeshComponent->StaticMesh->LODModels[CurrentLODLevel].VertexBuffer.GetNumTexCoords();
-		}
-		else if (SkeletalMeshComponent && SkeletalMeshComponent->SkeletalMesh)
-		{
-			Bounds = SkeletalMeshComponent->SkeletalMesh->Bounds;
-			NumVertices = SkeletalMeshComponent->SkeletalMesh->LODModels[CurrentLODLevel].NumVertices;
-			NumTriangles = SkeletalMeshComponent->SkeletalMesh->LODModels[CurrentLODLevel].MultiSizeIndexContainer.GetIndexBuffer()->Num()/3;
-			NumUVChannels = SkeletalMeshComponent->SkeletalMesh->LODModels[CurrentLODLevel].NumTexCoords;
-		}
-
-		DrawShadowedString(Canvas,
-			6,
-			YPos,
-			*FString::Printf(LocalizeSecure(LocalizeUnrealEd("Triangles_F"), NumTriangles )),
-			GEngine->GetSmallFont(),
-			FLinearColor::White
-			);
-		YPos += 18;
-
-		DrawShadowedString(Canvas,
-			6,
-			YPos,
-			*FString::Printf(LocalizeSecure(LocalizeUnrealEd("Vertices_F"), NumVertices )),
-			GEngine->GetSmallFont(),
-			FLinearColor::White
-			);
-		YPos += 18;
-
-		DrawShadowedString(Canvas,
-			6,
-			YPos,
-			*FString::Printf(LocalizeSecure(LocalizeUnrealEd("UVChannels_F"), NumUVChannels )),
-			GEngine->GetSmallFont(),
-			FLinearColor::White
-			);
-		YPos += 18;
-
-
-		DrawShadowedString(Canvas,
-			6,
-			YPos,
-			*FString::Printf( LocalizeSecure( LocalizeUnrealEd("ApproxSize_F"), int32(Bounds.BoxExtent.X * 2.0f),
-			int32(Bounds.BoxExtent.Y * 2.0f),
-			int32(Bounds.BoxExtent.Z * 2.0f) ) ),
-			GEngine->GetSmallFont(),
-			FLinearColor::White
-			);
-			*/
-	//}
-	YPos += 18;
 
 	if(bDrawUVs)
 	{
+		constexpr int32 YPos = 24;
 		DrawUVs(InViewport, Canvas, YPos, MaterialToDrawInUVs);
 	}
 
@@ -1030,31 +565,6 @@ void FCustomizableObjectEditorViewportClient::DrawUVs(FViewport* InViewport, FCa
 				constexpr Vector2DRealType One     = static_cast<Vector2DRealType>(1);
 				constexpr Vector2DRealType Zero    = static_cast<Vector2DRealType>(0);
 
-				FLinearColor UV12LineColor = FLinearColor::Black;
-				if (UV1.X < -Epsilon || UV1.X > One + Epsilon ||
-					UV2.X < -Epsilon || UV2.X > One + Epsilon ||
-					UV1.Y < -Epsilon || UV1.Y > One + Epsilon ||
-					UV2.Y < -Epsilon || UV2.Y > One + Epsilon)
-				{
-					UV12LineColor = FLinearColor(0.6f, 0.0f, 0.0f);
-				}
-				FLinearColor UV23LineColor = FLinearColor::Black;
-				if (UV3.X < -Epsilon || UV3.X > One + Epsilon ||
-					UV2.X < -Epsilon || UV2.X > One + Epsilon ||
-					UV3.Y < -Epsilon || UV3.Y > One + Epsilon ||
-					UV2.Y < -Epsilon || UV2.Y > One + Epsilon)
-				{
-					UV23LineColor = FLinearColor(0.6f, 0.0f, 0.0f);
-				}
-				FLinearColor UV31LineColor = FLinearColor::Black;
-				if (UV3.X < -Epsilon || UV3.X > One + Epsilon ||
-					UV1.X < -Epsilon || UV1.X > One + Epsilon ||
-					UV3.Y < -Epsilon || UV3.Y > One + Epsilon ||
-					UV1.Y < -Epsilon || UV1.Y > One + Epsilon)
-				{
-					UV31LineColor = FLinearColor(0.6f, 0.0f, 0.0f);
-				}
-
 				UV1 = ClampUVRange(UV1.X, UV1.Y) * UVBoxScale + UVBoxOrigin;
 				UV2 = ClampUVRange(UV2.X, UV2.Y) * UVBoxScale + UVBoxOrigin;
 				UV3 = ClampUVRange(UV3.X, UV3.Y) * UVBoxScale + UVBoxOrigin;
@@ -1146,31 +656,6 @@ void FCustomizableObjectEditorViewportClient::DrawUVs(FViewport* InViewport, FCa
 					constexpr Vector2DRealType One = static_cast<Vector2DRealType>(1);
 					constexpr Vector2DRealType Zero = static_cast<Vector2DRealType>(0);
 
-					FLinearColor UV12LineColor = FLinearColor::Black;
-					if (UV1.X < -Epsilon || UV1.X > One + Epsilon ||
-						UV2.X < -Epsilon || UV2.X > One + Epsilon ||
-						UV1.Y < -Epsilon || UV1.Y > One + Epsilon ||
-						UV2.Y < -Epsilon || UV2.Y > One + Epsilon)
-					{
-						UV12LineColor = FLinearColor(0.6f, 0.0f, 0.0f);
-					}
-					FLinearColor UV23LineColor = FLinearColor::Black;
-					if (UV3.X < -Epsilon || UV3.X > One + Epsilon ||
-						UV2.X < -Epsilon || UV2.X > One + Epsilon ||
-						UV3.Y < -Epsilon || UV3.Y > One + Epsilon ||
-						UV2.Y < -Epsilon || UV2.Y > One + Epsilon)
-					{
-						UV23LineColor = FLinearColor(0.6f, 0.0f, 0.0f);
-					}
-					FLinearColor UV31LineColor = FLinearColor::Black;
-					if (UV3.X < -Epsilon || UV3.X > One + Epsilon ||
-						UV1.X < -Epsilon || UV1.X > One + Epsilon ||
-						UV3.Y < -Epsilon || UV3.Y > One + Epsilon ||
-						UV1.Y < -Epsilon || UV1.Y > One + Epsilon)
-					{
-						UV31LineColor = FLinearColor(0.6f, 0.0f, 0.0f);
-					}
-
 					UV1 = ClampUVRange(UV1.X, UV1.Y) * UVBoxScale + UVBoxOrigin;
 					UV2 = ClampUVRange(UV2.X, UV2.Y) * UVBoxScale + UVBoxOrigin;
 					UV3 = ClampUVRange(UV3.X, UV3.Y) * UVBoxScale + UVBoxOrigin;
@@ -1206,45 +691,23 @@ float FCustomizableObjectEditorViewportClient::GetFloorOffset() const
 }
 
 
-void FCustomizableObjectEditorViewportClient::SetClipMorphPlaneVisibility(bool bVisible, const FVector& Origin, const FVector& Normal, float InMorphLength, const FBoxSphereBounds& Bounds, float InRadius1, float InRadius2, float InRotationAngle)
-{
-	/*
-	//EditorClipMorphPlaneComp->SetVisibility(bVisible);
-	bClipMorphVisible = bVisible;
-	//EditorClipMorphPlaneComp->SetWorldLocationAndRotation(Origin, Normal.Rotation());
-	ClipMorphOrigin = Origin;
-	ClipMorphNormal = Normal;
-	MorphLength = InMorphLength;
-	MorphBounds = Bounds;
-	Radius1 = InRadius1;
-	Radius2 = InRadius2;
-	RotationAngle = InRotationAngle;
-	*/
-}
-
-
-void FCustomizableObjectEditorViewportClient::SetClipMorphPlaneVisibility(bool bVisible, UCustomizableObjectNodeMeshClipMorph* NodeMeshClipMorph)
+void FCustomizableObjectEditorViewportClient::ShowGizmoClipMorph(UCustomizableObjectNodeMeshClipMorph& NodeMeshClipMorph)
 {
-	if (!bClipMorphVisible && !bVisible)
-	{
-		return;
-	}
+	SetWidgetType(EWidgetType::ClipMorph);
 
-	bClipMorphVisible = bVisible;
+	if (ClipMorphNode != &NodeMeshClipMorph || NodeMeshClipMorph.bUpdateViewportWidget)
+	{	
+		NodeMeshClipMorph.bUpdateViewportWidget = false;
 
-	if (NodeMeshClipMorph && (ClipMorphNode != NodeMeshClipMorph || NodeMeshClipMorph->bUpdateViewportWidget))
-	{
-		NodeMeshClipMorph->bUpdateViewportWidget = false;
+		bClipMorphLocalStartOffset = NodeMeshClipMorph.bLocalStartOffset;
+		MorphLength = NodeMeshClipMorph.B;
+		Radius1 = NodeMeshClipMorph.Radius;
+		Radius2 = NodeMeshClipMorph.Radius2;
+		RotationAngle = NodeMeshClipMorph.RotationAngle;
+		ClipMorphOrigin = NodeMeshClipMorph.Origin;
+		ClipMorphLocalOffset = NodeMeshClipMorph.StartOffset;
 
-		bClipMorphLocalStartOffset = NodeMeshClipMorph->bLocalStartOffset;
-		MorphLength = NodeMeshClipMorph->B;
-		Radius1 = NodeMeshClipMorph->Radius;
-		Radius2 = NodeMeshClipMorph->Radius2;
-		RotationAngle = NodeMeshClipMorph->RotationAngle;
-		ClipMorphOrigin = NodeMeshClipMorph->Origin;
-		ClipMorphLocalOffset = NodeMeshClipMorph->StartOffset;
-
-		NodeMeshClipMorph->FindLocalAxes(ClipMorphXAxis, ClipMorphYAxis, ClipMorphNormal);
+		NodeMeshClipMorph.FindLocalAxes(ClipMorphXAxis, ClipMorphYAxis, ClipMorphNormal);
 
 		if (bClipMorphLocalStartOffset)
 		{
@@ -1258,195 +721,90 @@ void FCustomizableObjectEditorViewportClient::SetClipMorphPlaneVisibility(bool b
 		}
 	}
 
-	ClipMorphNode = NodeMeshClipMorph;
-
-	if (WidgetVisibility != bClipMorphVisible)
-	{
-		Widget->SetDefaultVisibility(bClipMorphVisible);
-		WidgetVisibility = bClipMorphVisible;
-	}
+	ClipMorphNode = &NodeMeshClipMorph;
 }
 
 
-void FCustomizableObjectEditorViewportClient::SetClipMeshVisibility(bool bVisible, UStaticMesh* ClipMesh, UCustomizableObjectNodeMeshClipWithMesh* MeshNode)
+void FCustomizableObjectEditorViewportClient::HideGizmoClipMorph()
 {
-	if (!ClipMeshNode && !MeshNode)
+	if (WidgetType == EWidgetType::ClipMorph)
 	{
-		return;
-	}
-
-	ClipMeshComp->SetStaticMesh(ClipMesh);
-	ClipMeshComp->SetVisibility(bVisible);
-	ClipMeshNode = bVisible ? MeshNode : nullptr;
-
-	if (ClipMeshNode)
-	{
-		ClipMeshComp->SetWorldTransform(ClipMeshNode->Transform);
-	}
-
-	if (WidgetVisibility != bVisible)
-	{
-		Widget->SetDefaultVisibility(bVisible);
-		WidgetVisibility = bVisible;
+		SetWidgetType(EWidgetType::Hidden);
 	}
 }
 
 
-void FCustomizableObjectEditorViewportClient::SetProjectorVisibility(bool bVisible, FString ProjectorParameterName, FString ProjectorParameterNameWithIndex, int32 ProjectorRangeIndex, const FCustomizableObjectProjector& Data, int32 ProjectorParameterIndex)
+void FCustomizableObjectEditorViewportClient::ShowGizmoClipMesh(UCustomizableObjectNodeMeshClipWithMesh& InClipMeshNode, UStaticMesh& ClipMesh)
 {
-	if (GizmoProxy.ProjectorParameterIndex == ProjectorParameterIndex && GizmoProxy.ProjectorRangeIndex == ProjectorRangeIndex)
-	{
-		return;
-	}
+	SetWidgetType(EWidgetType::ClipMesh);
 
-	if (GizmoProxy.HasAssignedData)
-	{
-		GizmoProxy.UpdateOriginData();
-		GizmoProxy.CleanOriginData();
-	}
-	GizmoProxy.ProjectorParameterName = ProjectorParameterName;
-	GizmoProxy.ProjectorParameterNameWithIndex = ProjectorParameterNameWithIndex;
-	GizmoProxy.ProjectorRangeIndex = ProjectorRangeIndex;
-	GizmoProxy.Value = Data;
-	GizmoProxy.ProjectionType = Data.ProjectionType;
+	ClipMeshNode = &InClipMeshNode;
 
-	if (GizmoProxy.ProjectorHasInitialValues(GizmoProxy.Value))
-	{
-		GizmoProxy.Value = GizmoProxy.SetProjectorInitialValue(SkeletalMeshComponents, BoundSphere.W * 2.0f);
-	}
+	ClipMeshComp->SetStaticMesh(&ClipMesh);
+	ClipMeshComp->SetVisibility(true);
+	ClipMeshComp->SetWorldTransform(InClipMeshNode.Transform);
+}
 
-	GizmoProxy.ProjectorParameterIndex = ProjectorParameterIndex;
-	GizmoProxy.AnyGizmoSelected = true;
-	GizmoProxy.HasAssignedData = true;
 
-	if (!WidgetVisibility)
+void FCustomizableObjectEditorViewportClient::HideGizmoClipMesh()
+{
+	if (WidgetType == EWidgetType::ClipMesh)
 	{
-		Widget->SetDefaultVisibility(true);
-		WidgetVisibility = true;
-	}
+		SetWidgetType(EWidgetType::Hidden);
+	}	
 }
 
 
-void FCustomizableObjectEditorViewportClient::SetProjectorType(bool bVisible, FString ProjectorParameterName, FString ProjectorParameterNameWithIndex, int32 ProjectorRangeIndex, const FCustomizableObjectProjector& Data, int32 ProjectorParameterIndex)
+void FCustomizableObjectEditorViewportClient::ShowGizmoProjector(
+	const FWidgetLocationDelegate& InWidgetLocationDelegate,
+	const FOnWidgetLocationChangedDelegate& InOnWidgetLocationChangedDelegate,
+	const FWidgetDirectionDelegate& InWidgetDirectionDelegate,
+	const FOnWidgetDirectionChangedDelegate& InOnWidgetDirectionChangedDelegate,
+	const FWidgetUpDelegate& InWidgetUpDelegate, const FOnWidgetUpChangedDelegate& InOnWidgetUpChangedDelegate,
+	const FWidgetScaleDelegate& InWidgetScaleDelegate, const FOnWidgetScaleChangedDelegate& InOnWidgetScaleChangedDelegate,
+	const FWidgetAngleDelegate& InWidgetAngleDelegate, const FProjectorTypeDelegate& InProjectorTypeDelegate,
+	const FWidgetColorDelegate& InWidgetColorDelegate,
+	const FWidgetTrackingStartedDelegate& InWidgetTrackingStartedDelegate)
 {
-	GizmoProxy.ProjectorParameterName = ProjectorParameterName;
-	GizmoProxy.ProjectorParameterNameWithIndex = ProjectorParameterNameWithIndex;
-	GizmoProxy.ProjectorRangeIndex = ProjectorRangeIndex;
-	GizmoProxy.Value = Data;
-	GizmoProxy.ProjectionType = Data.ProjectionType;
+	SetWidgetType(EWidgetType::Projector);
 
-	GizmoProxy.ProjectorParameterIndex = ProjectorParameterIndex;
-	GizmoProxy.AnyGizmoSelected = true;
-	GizmoProxy.HasAssignedData = true;
+	WidgetLocationDelegate = InWidgetLocationDelegate;
+	OnWidgetLocationChangedDelegate = InOnWidgetLocationChangedDelegate;
+	WidgetDirectionDelegate = InWidgetDirectionDelegate;
+	OnWidgetDirectionChangedDelegate = InOnWidgetDirectionChangedDelegate;
+	WidgetUpDelegate = InWidgetUpDelegate;
+	OnWidgetUpChangedDelegate = InOnWidgetUpChangedDelegate;
+	WidgetScaleDelegate = InWidgetScaleDelegate;
+	OnWidgetScaleChangedDelegate = InOnWidgetScaleChangedDelegate;
+	WidgetAngleDelegate = InWidgetAngleDelegate;
+	ProjectorTypeDelegate = InProjectorTypeDelegate;
+	WidgetColorDelegate = InWidgetColorDelegate;
+	WidgetTrackingStartedDelegate = InWidgetTrackingStartedDelegate;
 }
 
 
-void FCustomizableObjectEditorViewportClient::SetProjectorVisibility(bool bVisible, UCustomizableObjectNodeProjectorConstant* InProjector)
+void FCustomizableObjectEditorViewportClient::HideGizmoProjector()
 {
-	if (InProjector == nullptr)
-	{
-		if ((GizmoProxy.HasAssignedData) && (GizmoProxy.DataOriginConstant != nullptr))
-		{
-			GizmoProxy.UpdateOriginData();
-			GizmoProxy.CleanOriginData();
-		}
-
-		if (WidgetVisibility)
-		{
-			Widget->SetDefaultVisibility(false);
-			WidgetVisibility = false;
-		}
-	}
-	else
+	if (WidgetType == EWidgetType::Projector)
 	{
-		if (!GizmoProxy.HasAssignedData ||
-		   (GizmoProxy.HasAssignedData && (GizmoProxy.DataOriginConstant != InProjector)))
-		{
-			GizmoProxy.UpdateOriginData();
-			GizmoProxy.CleanOriginData();
-			GizmoProxy.DataOriginConstant = InProjector;
-			GizmoProxy.Value = InProjector->Value;
-
-			if (GizmoProxy.ProjectorHasInitialValues(GizmoProxy.Value))
-			{
-				GizmoProxy.Value = GizmoProxy.SetProjectorInitialValue(SkeletalMeshComponents, BoundSphere.W * 2.0f);
-			}
-
-			GizmoProxy.AnyGizmoSelected = true;
-			GizmoProxy.HasAssignedData = true;
-			GizmoProxy.AssignedDataIsFromNode = true;
-			GizmoProxy.ProjectionType = InProjector->ProjectionType;
-
-			if (!WidgetVisibility)
-			{
-				Widget->SetDefaultVisibility(true);
-				WidgetVisibility = true;
-			}
-		}
+		SetWidgetType(EWidgetType::Hidden);
 	}
 }
 
 
-void FCustomizableObjectEditorViewportClient::SetProjectorParameterVisibility(bool bVisible, UCustomizableObjectNodeProjectorParameter* InProjectorParameter)
+void FCustomizableObjectEditorViewportClient::ShowGizmoLight(ULightComponent& Light)
 {
-	if (InProjectorParameter == nullptr)
-	{
-		if ((GizmoProxy.HasAssignedData) && (GizmoProxy.DataOriginParameter != nullptr))
-		{
-			GizmoProxy.UpdateOriginData();
-			GizmoProxy.CleanOriginData();
-
-			if (WidgetVisibility)
-			{
-				Widget->SetDefaultVisibility(false);
-				WidgetVisibility = false;
-			}
-		}
-	}
-	else
-	{
-		if (!GizmoProxy.HasAssignedData ||
-			(GizmoProxy.HasAssignedData && (GizmoProxy.DataOriginParameter != InProjectorParameter)))
-		{
-			GizmoProxy.UpdateOriginData();
-			GizmoProxy.CleanOriginData();
-			GizmoProxy.DataOriginParameter = InProjectorParameter;
-			GizmoProxy.Value = InProjectorParameter->DefaultValue;
-
-			if (GizmoProxy.ProjectorHasInitialValues(GizmoProxy.Value))
-			{
-				GizmoProxy.Value = GizmoProxy.SetProjectorInitialValue(SkeletalMeshComponents, BoundSphere.W * 2.0f);
-			}
-
-			GizmoProxy.AnyGizmoSelected = true;
-			GizmoProxy.HasAssignedData = true;
-			GizmoProxy.AssignedDataIsFromNode = true;
-			GizmoProxy.ProjectionType = InProjectorParameter->ProjectionType;
-
-			if (!WidgetVisibility)
-			{
-				Widget->SetDefaultVisibility(true);
-				WidgetVisibility = true;
-			}
-		}
-	}
+	SelectedLightComponent = &Light;
+	
+	SetWidgetType(EWidgetType::Light);
 }
 
 
-void FCustomizableObjectEditorViewportClient::ResetProjectorVisibility()
+void FCustomizableObjectEditorViewportClient::HideGizmoLight()
 {
-	if (bManipulating || !GizmoProxy.AnyGizmoSelected)
-	{
-		return;
-	}
-
-	GizmoProxy.UpdateOriginData();
-	GizmoProxy.CleanOriginData();
-
-	if (WidgetVisibility)
+	if (WidgetType == EWidgetType::Light)
 	{
-		Widget->SetDefaultVisibility(false);
-		WidgetVisibility = false;
+		SetWidgetType(EWidgetType::Hidden);
 	}
 }
 
@@ -1616,181 +974,162 @@ void FCustomizableObjectEditorViewportClient::SetShowGrid()
 		AdvancedScene->SetFloorVisibility(DrawHelper.bDrawGrid,true);
 	}
 
-	EngineShowFlags.Grid = DrawHelper.bDrawGrid;
-
-	Invalidate();
-}
-
-bool FCustomizableObjectEditorViewportClient::IsSetShowGridChecked() const
-{
-	return DrawHelper.bDrawGrid;
-}
-
-void FCustomizableObjectEditorViewportClient::SetShowSky()
-{
-	bDrawSky = !bDrawSky;
-	FAdvancedPreviewScene* PreviewSceneCasted = static_cast<FAdvancedPreviewScene*>(PreviewScene);
-	PreviewSceneCasted->SetEnvironmentVisibility(bDrawSky, true);
-	Invalidate();
-}
-
-bool FCustomizableObjectEditorViewportClient::IsSetShowSkyChecked() const
-{
-	return bDrawSky;
-}
-
-void FCustomizableObjectEditorViewportClient::SetShowBounds()
-{
-	EngineShowFlags.Bounds = 1 - EngineShowFlags.Bounds;
+	EngineShowFlags.Grid = DrawHelper.bDrawGrid;
+
 	Invalidate();
 }
 
-bool FCustomizableObjectEditorViewportClient::IsSetShowBoundsChecked() const
+bool FCustomizableObjectEditorViewportClient::IsSetShowGridChecked() const
 {
-	return EngineShowFlags.Bounds ? true : false;
+	return DrawHelper.bDrawGrid;
 }
 
-void FCustomizableObjectEditorViewportClient::SetShowCollision()
+void FCustomizableObjectEditorViewportClient::SetShowSky()
 {
-	EngineShowFlags.Collision = 1 - EngineShowFlags.Collision;
+	bDrawSky = !bDrawSky;
+	FAdvancedPreviewScene* PreviewSceneCasted = static_cast<FAdvancedPreviewScene*>(PreviewScene);
+	PreviewSceneCasted->SetEnvironmentVisibility(bDrawSky, true);
 	Invalidate();
 }
 
-bool FCustomizableObjectEditorViewportClient::IsSetShowCollisionChecked() const
+bool FCustomizableObjectEditorViewportClient::IsSetShowSkyChecked() const
 {
-	return EngineShowFlags.Collision ? true : false;
+	return bDrawSky;
 }
 
-void FCustomizableObjectEditorViewportClient::SetShowPivot()
+void FCustomizableObjectEditorViewportClient::SetShowBounds()
 {
-	bShowPivot = !bShowPivot;
-	Widget->SetDefaultVisibility(bShowPivot);
-	WidgetVisibility = bShowPivot;
+	EngineShowFlags.Bounds = 1 - EngineShowFlags.Bounds;
 	Invalidate();
 }
 
-bool FCustomizableObjectEditorViewportClient::IsSetShowPivotChecked() const
-{
-	return bShowPivot;
-}
-
 
 bool FCustomizableObjectEditorViewportClient::InputKey(const FInputKeyEventArgs& EventArgs)
 {
-	const int32 HitX = EventArgs.Viewport->GetMouseX();
-	const int32 HitY = EventArgs.Viewport->GetMouseY();
-
 	const bool bMouseButtonDown = EventArgs.Viewport->KeyState(EKeys::LeftMouseButton) || EventArgs.Viewport->KeyState(EKeys::MiddleMouseButton) || EventArgs.Viewport->KeyState(EKeys::RightMouseButton);
 
-	bool bHandled = false;
-
-	if (EventArgs.Event == IE_Pressed && !bManipulating && !bMouseButtonDown)
+	if (EventArgs.Event == IE_Pressed && !bMouseButtonDown)
 	{
 		if (EventArgs.Key == EKeys::F)
 		{
-			bHandled = true;
 			UpdateCameraSetup();
+			return true;
 		}
-		else if (EventArgs.Key == EKeys::W)
-		{
-			bHandled = true;
-			WidgetMode = UE::Widget::WM_Translate;
-		}
-		else if (EventArgs.Key == EKeys::E)
+		else if (WidgetType != EWidgetType::Hidden) // Do not change the type when hidden.
 		{
-			bHandled = true;
-			WidgetMode = UE::Widget::WM_Rotate;
+			if (EventArgs.Key == EKeys::W)
+			{
+				SetWidgetMode(UE::Widget::WM_Translate);
+				return true;
+			}
+			else if (EventArgs.Key == EKeys::E)
+			{
+				SetWidgetMode(UE::Widget::WM_Rotate);
+				return true;
+			}
+			else if (EventArgs.Key == EKeys::R)
+			{
+				SetWidgetMode(UE::Widget::WM_Scale);
+				return true;
+			}	
 		}
-		else if (EventArgs.Key == EKeys::R)
+		else if (EventArgs.Key == EKeys::Q) // Not sure why, pressing Q the super class hides the widget.
 		{
-			bHandled = true;
-			WidgetMode = UE::Widget::WM_Scale;
+			SetWidgetType(EWidgetType::Hidden);
+			return true;
 		}
 	}
 
-	if (EventArgs.Event == IE_Released && bManipulating)
-	{
-		GizmoProxy.SetProjectorUpdatedInViewport(false);
-	}
-
 	// Pass keys to standard controls, if we didn't consume input
-	return (bHandled)
-		? true
-		: FEditorViewportClient::InputKey(EventArgs);
+	return FEditorViewportClient::InputKey(EventArgs);
 }
 
 
 bool FCustomizableObjectEditorViewportClient::InputWidgetDelta(FViewport* InViewport, EAxisList::Type CurrentAxis, FVector& Drag, FRotator& Rot, FVector& Scale)
 {
-	if (!GizmoProxy.AnyGizmoSelected && !bClipMorphVisible && !ClipMeshNode && !bIsEditingLightEnabled)
+	if (CurrentAxis == EAxisList::None)
 	{
 		return false;
 	}
 
-	// Get some useful info about buttons being held down
-	const bool bCtrlDown = InViewport->KeyState(EKeys::LeftControl) || InViewport->KeyState(EKeys::RightControl);
-	const bool bShiftDown = InViewport->KeyState(EKeys::LeftShift) || InViewport->KeyState(EKeys::RightShift);
-	const bool bMouseButtonDown = InViewport->KeyState(EKeys::LeftMouseButton) || InViewport->KeyState(EKeys::MiddleMouseButton) || InViewport->KeyState(EKeys::RightMouseButton);
-
-	bool bHandled = false;
-
-	if (bManipulating && CurrentAxis != EAxisList::None)
+	const UE::Widget::EWidgetMode WidgetMode = GetWidgetMode();
+	
+	switch (WidgetType)
 	{
-		bHandled = true;
-
-		if (GizmoProxy.AnyGizmoSelected)
+	case EWidgetType::Projector:
 		{
-			if (WidgetMode == UE::Widget::WM_Translate)
+			if (WidgetLocationDelegate.IsBound() && OnWidgetLocationChangedDelegate.IsBound())
 			{
-				GizmoProxy.Value.Position += (FVector3f)Drag;
-				GizmoProxy.ProjectorGizmoEdited = true;
+				if (Drag != FVector::ZeroVector)
+				{
+					OnWidgetLocationChangedDelegate.Execute(WidgetLocationDelegate.Execute() + Drag);				
+				}
 			}
-			else if (WidgetMode == UE::Widget::WM_Rotate)
+
+			if (WidgetDirectionDelegate.IsBound() && OnWidgetDirectionChangedDelegate.IsBound())
 			{
-				GizmoProxy.Value.Direction = (FVector3f)Rot.RotateVector((FVector)GizmoProxy.Value.Direction);
-				GizmoProxy.Value.Up = (FVector3f)Rot.RotateVector((FVector)GizmoProxy.Value.Up);
-				GizmoProxy.ProjectorGizmoEdited = true;
+				const FVector WidgetDirection = WidgetDirectionDelegate.Execute();
+				const FVector NewWidgetDirection = Rot.RotateVector(WidgetDirection);
+
+				if (WidgetDirection != NewWidgetDirection)
+				{
+					OnWidgetDirectionChangedDelegate.Execute(NewWidgetDirection);				
+				}
 			}
-			else if (WidgetMode == UE::Widget::WM_Scale)
+
+			if (WidgetUpDelegate.IsBound() && OnWidgetUpChangedDelegate.IsBound())
 			{
-				GizmoProxy.Value.Scale.X += Scale.Y;
-				GizmoProxy.Value.Scale.Y += Scale.Z;
-				GizmoProxy.Value.Scale.Z += Scale.X;
-				GizmoProxy.ProjectorGizmoEdited = true;
-			}
+				const FVector WidgetUp = WidgetUpDelegate.Execute();
+				const FVector NewWidgetUp = Rot.RotateVector(WidgetUp);
 
-			InViewport->Invalidate();
+				if (WidgetUp != NewWidgetUp)
+				{
+					OnWidgetUpChangedDelegate.Execute(NewWidgetUp);				
+				}
+			}
 
-			if (GizmoProxy.ProjectorGizmoEdited)
+			if (WidgetScaleDelegate.IsBound() && OnWidgetScaleChangedDelegate.IsBound())
 			{
-				GizmoProxy.SetProjectorUpdatedInViewport(true);
-				GizmoProxy.UpdateOriginData();
-				GizmoProxy.ProjectorGizmoEdited = false;
+				const FVector CorrectedScale(Scale.Y, Scale.Z, Scale.X);
+				if (CorrectedScale != FVector::ZeroVector)
+				{
+					OnWidgetScaleChangedDelegate.Execute(WidgetScaleDelegate.Execute() + CorrectedScale);
+				}
 			}
+		
+			return true;
 		}
-
-		else if (bClipMorphVisible && ClipMorphNode)
+		
+	case EWidgetType::ClipMorph:
 		{
 			if (WidgetMode == UE::Widget::WM_Translate)
 			{
-				if (CurrentAxis == EAxisList::Z)
+				if (CurrentAxis == EAxisList::Screen) // true when selecting the widget center
+				{
+					CurrentAxis = EAxisList::XYZ;
+				}
+				
+				if (CurrentAxis & EAxisList::Z)
 				{
-					float dragZ = bClipMorphLocalStartOffset ? FVector::DotProduct(Drag, ClipMorphNormal) : Drag.Z;
+					const float dragZ = bClipMorphLocalStartOffset ? FVector::DotProduct(Drag, ClipMorphNormal) : Drag.Z;
 					ClipMorphLocalOffset.Z += dragZ;
 					ClipMorphOffset += (bClipMorphLocalStartOffset) ? dragZ * ClipMorphNormal : FVector(0,0,dragZ);
 				}
-				else if(CurrentAxis == EAxisList::X)
+
+				if(CurrentAxis & EAxisList::X)
 				{
-					float dragX = bClipMorphLocalStartOffset ? FVector::DotProduct(Drag, ClipMorphXAxis) : Drag.X;
+					const float dragX = bClipMorphLocalStartOffset ? FVector::DotProduct(Drag, ClipMorphXAxis) : Drag.X;
 					ClipMorphLocalOffset.X += dragX;
 					ClipMorphOffset += (bClipMorphLocalStartOffset) ? dragX * ClipMorphXAxis : FVector(dragX, 0, 0);
 				}
-				else if (CurrentAxis == EAxisList::Y)
+
+				if (CurrentAxis & EAxisList::Y)
 				{
-					float dragY = bClipMorphLocalStartOffset ? FVector::DotProduct(Drag, ClipMorphYAxis) : Drag.Y;
+					const float dragY = bClipMorphLocalStartOffset ? FVector::DotProduct(Drag, ClipMorphYAxis) : Drag.Y;
 					ClipMorphLocalOffset.Y += dragY;
 					ClipMorphOffset += (bClipMorphLocalStartOffset) ? dragY * ClipMorphYAxis : FVector(0, dragY, 0);
 				}
+				
 				ClipMorphNode->StartOffset = ClipMorphLocalOffset;
 			}
 			else if (WidgetMode == UE::Widget::WM_Rotate)
@@ -1826,9 +1165,10 @@ bool FCustomizableObjectEditorViewportClient::InputWidgetDelta(FViewport* InView
 					ClipMorphNode->StartOffset = ClipMorphLocalOffset;
 				}
 			}
-		}
 
-		else if (ClipMeshNode)
+			return true;
+		}
+	case EWidgetType::ClipMesh:
 		{
 			if (WidgetMode == UE::Widget::WM_Translate)
 			{
@@ -1844,36 +1184,57 @@ bool FCustomizableObjectEditorViewportClient::InputWidgetDelta(FViewport* InView
 			}
 
 			ClipMeshComp->SetWorldTransform(ClipMeshNode->Transform);
-		}
 
-		else if (bIsEditingLightEnabled && SelectedLightComponent)
+			return true;
+		}
+		
+	case EWidgetType::Light:
 		{
 			if (WidgetMode == UE::Widget::WM_Translate)
 			{
 				SelectedLightComponent->AddWorldOffset(Drag);
 				SelectedLightComponent->MarkForNeededEndOfFrameRecreate();
 			}
-
 			else if (WidgetMode == UE::Widget::WM_Rotate)
 			{
 				SelectedLightComponent->AddWorldRotation(Rot.Quaternion());
 				SelectedLightComponent->MarkForNeededEndOfFrameRecreate();
 			}
-		}
-	}
 
-	return bHandled;
+			return true;
+		}
+		
+	case EWidgetType::Hidden:
+		{
+			return false;
+		}
+		
+	default:
+		{
+			unimplemented()
+			return false;
+		}
+	}	
 }
 
 
-void FCustomizableObjectEditorViewportClient::TrackingStarted(const struct FInputEventState& InInputState, bool bIsDraggingWidget, bool bNudge)
+void FCustomizableObjectEditorViewportClient::TrackingStarted(const FInputEventState& InInputState, bool bIsDraggingWidget, bool bNudge)
 {
-	if (bIsDraggingWidget)
+	switch (WidgetType)
 	{
-		if (InInputState.IsLeftMouseButtonPressed() && (Widget->GetCurrentAxis() & EAxisList::XYZ) != 0)
+	case EWidgetType::Projector:
 		{
+			if (!bIsDraggingWidget ||
+				!InInputState.IsLeftMouseButtonPressed() ||
+				(Widget->GetCurrentAxis() & EAxisList::All) == 0)
+			{
+				return;
+			}
+
 			bManipulating = true;
 
+			const UE::Widget::EWidgetMode WidgetMode = GetWidgetMode();
+				
 			if (WidgetMode == UE::Widget::WM_Translate)
 			{
 				GEditor->BeginTransaction(LOCTEXT("CustomizableObjectEditor_TranslateProjector", "Translate Projector"));
@@ -1887,91 +1248,121 @@ void FCustomizableObjectEditorViewportClient::TrackingStarted(const struct FInpu
 				GEditor->BeginTransaction(LOCTEXT("CustomizableObjectEditor_ScaleProjector", "Scale Projector"));
 			}
 
-			if (GizmoProxy.AnyGizmoSelected)
-			{
-				GizmoProxy.Modify();
-				GizmoProxy.ModifyGraph();
-			}
+			WidgetTrackingStartedDelegate.ExecuteIfBound();
 		}
+
+	// The following cases are missing Undo/Redo functionality MTBL-391.
+	case EWidgetType::ClipMorph:
+	case EWidgetType::ClipMesh:
+	case EWidgetType::Light:
+	case EWidgetType::Hidden:
+		break;
+	default:
+		unimplemented();
 	}
 }
 
 
 void FCustomizableObjectEditorViewportClient::TrackingStopped()
 {
-	if (bManipulating)
+	switch (WidgetType)
 	{
-		bManipulating = false;
-		GEditor->EndTransaction();
-	}
-
-	Invalidate();
-}
-
+	case EWidgetType::Projector:
+		if (bManipulating)
+		{
+			bManipulating = false;
+			GEditor->EndTransaction();
+		}
 
-UE::Widget::EWidgetMode FCustomizableObjectEditorViewportClient::GetWidgetMode() const
-{
-	return WidgetMode;
+	case EWidgetType::Hidden:
+	case EWidgetType::ClipMorph:
+	case EWidgetType::ClipMesh:
+	case EWidgetType::Light:
+		break;
+	default:
+		unimplemented();
+	}
 }
 
 
 FVector FCustomizableObjectEditorViewportClient::GetWidgetLocation() const
 {
-	if (GizmoProxy.AnyGizmoSelected)
-	{
-		return (FVector)GizmoProxy.Value.Position;
-	}
-
-	if (bClipMorphVisible)
+	switch (WidgetType)
 	{
+	case EWidgetType::Projector:
+		return WidgetLocationDelegate.IsBound() ?WidgetLocationDelegate.Execute() : FVector::ZeroVector;
+		
+	case EWidgetType::ClipMorph:
 		return ClipMorphOrigin + ClipMorphOffset;
-	}
 
-	if (ClipMeshNode)
-	{
+	case EWidgetType::ClipMesh:
 		return ClipMeshNode->Transform.GetTranslation();
-	}
 
-	if (bIsEditingLightEnabled && SelectedLightComponent)
-	{
+	case EWidgetType::Light:
 		return SelectedLightComponent->GetComponentLocation();
-	}
 
-	return FVector::ZeroVector;
+	case EWidgetType::Hidden:
+		return FVector::ZeroVector;
+
+	default:
+		unimplemented()
+		return FVector::ZeroVector;
+	}
 }
 
 
 FMatrix FCustomizableObjectEditorViewportClient::GetWidgetCoordSystem() const
 {
-	if (GizmoProxy.AnyGizmoSelected)
+	switch (WidgetType)
 	{
-		FVector3f YVector = FVector3f::CrossProduct(GizmoProxy.Value.Direction, GizmoProxy.Value.Up);
-		return FMatrix((FVector)GizmoProxy.Value.Direction, (FVector)YVector, (FVector)GizmoProxy.Value.Up, FVector::ZeroVector);
-	}
+	case EWidgetType::Projector:
+		{
+			const FVector WidgetDirection = WidgetDirectionDelegate.IsBound() ?
+				WidgetDirectionDelegate.Execute() :
+				FVector::ForwardVector;
 
-	if (bClipMorphVisible)
-	{
-		if (bClipMorphLocalStartOffset)
+			const FVector WidgetUp = WidgetUpDelegate.IsBound() ?
+				WidgetUpDelegate.Execute() :
+				FVector::UpVector;
+	
+			const FVector YVector = FVector::CrossProduct(WidgetDirection, WidgetUp);
+			return FMatrix(WidgetDirection, YVector, WidgetUp, FVector::ZeroVector);
+		}		
+	case EWidgetType::ClipMorph:
 		{
-			return FMatrix(-ClipMorphXAxis, -ClipMorphYAxis, -ClipMorphNormal, FVector::ZeroVector);
+			if (bClipMorphLocalStartOffset)
+			{
+				return FMatrix(-ClipMorphXAxis, -ClipMorphYAxis, -ClipMorphNormal, FVector::ZeroVector);
+			}
+			else
+			{			
+				return FMatrix(FVector(1, 0, 0), FVector(0,1,0), FVector(0,0,1), FVector::ZeroVector);
+			}
+		}
+		
+	case EWidgetType::ClipMesh:
+		{
+			return ClipMeshNode->Transform.ToMatrixNoScale().RemoveTranslation();			
+		}
+		
+	case EWidgetType::Light:
+		{
+			FMatrix Rotation = SelectedLightComponent->GetComponentTransform().ToMatrixNoScale();
+			Rotation.SetOrigin(FVector::ZeroVector);
+			return Rotation;
+		}
+		
+	case EWidgetType::Hidden:
+		{
+			return FMatrix::Identity;
 		}
 
-		return FMatrix(FVector(1, 0, 0), FVector(0,1,0), FVector(0,0,1), FVector::ZeroVector);
-	}
-
-	if (ClipMeshNode)
-	{
-		return ClipMeshNode->Transform.ToMatrixNoScale().RemoveTranslation();
-	}
-
-	if (bIsEditingLightEnabled && SelectedLightComponent)
-	{
-		FMatrix Rotation = SelectedLightComponent->GetComponentTransform().ToMatrixNoScale();
-		Rotation.SetOrigin(FVector::ZeroVector);
-		return Rotation;
+	default:
+		{
+			unimplemented()
+			return FMatrix::Identity;
+		}
 	}
-
-	return FMatrix::Identity;
 }
 
 
@@ -2088,58 +1479,36 @@ void FCustomizableObjectEditorViewportClient::ReSetAnimation()
 
 void FCustomizableObjectEditorViewportClient::AddLightToScene(ULightComponent* AddedLight)
 {
-	if (AddedLight)
-	{
-		PreviewScene->AddComponent(AddedLight, AddedLight->GetComponentTransform());
-	}
-
-	if (WidgetVisibility && !bIsEditingLightEnabled)
+	if (!AddedLight)
 	{
 		return;
 	}
 
-	SelectedLightComponent = AddedLight;
-	bIsEditingLightEnabled = AddedLight != nullptr;
-
-	if (!WidgetVisibility)
-	{
-		Widget->SetDefaultVisibility(true);
-		WidgetVisibility = true;
-	}
+	LightComponents.Add(AddedLight);
+	PreviewScene->AddComponent(AddedLight, AddedLight->GetComponentTransform());
 }
 
 
 void FCustomizableObjectEditorViewportClient::RemoveLightFromScene(ULightComponent* RemovedLight)
 {
-	PreviewScene->RemoveComponent(RemovedLight);
-
-	if (RemovedLight == SelectedLightComponent)
+	if (!RemovedLight)
 	{
-		SelectedLightComponent = nullptr;
-		bIsEditingLightEnabled = false;
-
-		Widget->SetDefaultVisibility(false);
-		WidgetVisibility = false;
+		return;
 	}
+
+	LightComponents.Remove(RemovedLight);
+	PreviewScene->RemoveComponent(RemovedLight);
 }
 
 
-void FCustomizableObjectEditorViewportClient::SetSelectedLight(ULightComponent* SelectedLight)
+void FCustomizableObjectEditorViewportClient::RemoveAllLightsFromScene()
 {
-	if (WidgetVisibility && !bIsEditingLightEnabled)
+	for (ULightComponent* Light : LightComponents)
 	{
-		return;
+		PreviewScene->RemoveComponent(Light);
 	}
 
-	SelectedLightComponent = SelectedLight;
-	bIsEditingLightEnabled = SelectedLight != nullptr;
-
-	// Activate Gizmo and Set Light Transform or disable lights if !bIsEditingLightEnabled
-	if (WidgetVisibility != bIsEditingLightEnabled)
-	{
-		Widget->SetDefaultVisibility(bIsEditingLightEnabled);
-		WidgetVisibility = bIsEditingLightEnabled;
-	}
+	LightComponents.Empty();
 }
 
 
@@ -2894,57 +2263,8 @@ void FCustomizableObjectEditorViewportClient::StateChangeShowGeometryData()
 }
 
 
-void FCustomizableObjectEditorViewportClient::SetProjectorWidgetMode(UE::Widget::EWidgetMode InMode)
-{
-	WidgetMode = InMode;
-}
-
-
-const GizmoRTSProxy& FCustomizableObjectEditorViewportClient::GetGizmoProxy()
-{
-	return GizmoProxy;
-}
-
-
-bool FCustomizableObjectEditorViewportClient::GetGizmoIsProjectorParameterSelected()
-{
-	return GizmoProxy.IsProjectorParameterSelected();
-}
-
-
-bool FCustomizableObjectEditorViewportClient::GetIsManipulating()
-{
-	return bManipulating;
-}
-
-
-bool FCustomizableObjectEditorViewportClient::GetWidgetVisibility()
-{
-	return WidgetVisibility;
-}
-
-
-void FCustomizableObjectEditorViewportClient::UpdateGizmoDataToOrigin()
-{
-	GizmoProxy.UpdateOriginData();
-}
-
-
-void FCustomizableObjectEditorViewportClient::CopyTransformFromOriginData()
-{
-    GizmoProxy.CopyTransformFromOriginData();
-}
-
-
-bool FCustomizableObjectEditorViewportClient::AnyProjectorNodeSelected()
-{
-	return (GizmoProxy.HasAssignedData && GizmoProxy.AssignedDataIsFromNode);
-}
-
-
 void FCustomizableObjectEditorViewportClient::ShowInstanceGeometryInformation(FCanvas* InCanvas)
 {
-	UCustomizableObjectInstance* Instance = CustomizableObjectEditorPtr.Pin()->GetPreviewInstance();
 	float YOffset = 50.0f;
 	int32 ComponentIndex = 0;
 
@@ -2989,74 +2309,6 @@ void FCustomizableObjectEditorViewportClient::ShowInstanceGeometryInformation(FC
 }
 
 
-void FCustomizableObjectEditorViewportClient::BuildMeanTimesBoxSize(const float MinValue, const float MaxValue, const float MeanValue, const float MaxWidth, const float Data, float& BoxSize, FLinearColor& BoxColor) const
-{
-	const float GreenLimit = MeanValue * 0.75f;
-	const float YellowLimit = MeanValue * 0.99;
-
-	if (Data < GreenLimit)
-	{
-		BoxColor = FLinearColor::Green;
-	}
-	else if (Data < YellowLimit)
-	{
-		BoxColor = FLinearColor::Yellow;
-	}
-	else
-	{
-		BoxColor = FLinearColor::Red;
-	}
-
-	BoxSize = (Data / MaxValue) * MaxWidth;
-}
-
-
-void FCustomizableObjectEditorViewportClient::BuildMeanTimesBoxSizes(float MaxWidth, TArray<float>& ArrayData, TArray<float>& ArrayBoxSize, TArray<FLinearColor>& ArrayBoxColor)
-{
-	const int32 NumElement = ArrayData.Num();
-	ArrayBoxSize.AddUninitialized(NumElement);
-	ArrayBoxColor.AddUninitialized(NumElement);
-
-	float MinValue = FLT_MAX;
-	float MaxValue = -1.0f * FLT_MAX;
-	float MeanValue = 0.0f;
-	for (int32 i = 0; i < ArrayData.Num(); ++i)
-	{
-		MeanValue += ArrayData[i];
-		MinValue = FMath::Min(MinValue, ArrayData[i]);
-		MaxValue = FMath::Max(MaxValue, ArrayData[i]);
-	}
-
-	float MaxMinRatio = 1.0f;
-	if (MinValue > 0.0f)
-	{
-		MaxMinRatio = MaxValue / MinValue;
-	}
-
-	MeanValue /= float(ArrayData.Num());
-	float GreenLimit = MeanValue + (MaxValue - MeanValue) * 0.1f;
-	float YellowLimit = MeanValue + (MaxValue - MeanValue) * 0.5f;
-
-	for (int32 i = 0; i < ArrayData.Num(); ++i)
-	{
-		if (ArrayData[i] < GreenLimit)
-		{
-			ArrayBoxColor[i] = FLinearColor::Green;
-		}
-		else if (ArrayData[i] < YellowLimit)
-		{
-			ArrayBoxColor[i] = FLinearColor::Yellow;
-		}
-		else
-		{
-			ArrayBoxColor[i] = FLinearColor::Red;
-		}
-
-		ArrayBoxSize[i] = (ArrayData[i] / MaxValue) * MaxWidth;
-	}
-}
-
-
 void FCustomizableObjectEditorViewportClient::SetCustomizableObject(UCustomizableObject* CustomizableObjectParameter)
 {
 	CustomizableObject = CustomizableObjectParameter;
@@ -3130,53 +2382,6 @@ void FCustomizableObjectEditorViewportClient::OnAssetViewerSettingsChanged(const
 }
 
 
-void FCustomizableObjectEditorViewportClient::ProjectorParameterChanged(UCustomizableObjectNodeProjectorParameter* Node)
-{
-	GizmoProxy.ProjectorParameterChanged(Node);
-}
-
-
-void FCustomizableObjectEditorViewportClient::ProjectorParameterChanged(UCustomizableObjectNodeProjectorConstant* Node)
-{
-	GizmoProxy.ProjectorParameterChanged(Node);
-}
-
-
-void FCustomizableObjectEditorViewportClient::SetGizmoCallUpdateSkeletalMesh(bool Value)
-{
-	GizmoProxy.SetCallUpdateSkeletalMesh(Value);
-}
-
-
-FString FCustomizableObjectEditorViewportClient::GetGizmoProjectorParameterName()
-{
-	return GizmoProxy.ProjectorParameterName;
-}
-
-
-FString FCustomizableObjectEditorViewportClient::GetGizmoProjectorParameterNameWithIndex()
-{
-	return GizmoProxy.ProjectorParameterNameWithIndex;
-}
-
-
-bool FCustomizableObjectEditorViewportClient::GetGizmoHasAssignedData()
-{
-	return GizmoProxy.HasAssignedData;
-}
-
-
-bool FCustomizableObjectEditorViewportClient::GetGizmoAssignedDataIsFromNode()
-{
-	return GizmoProxy.AssignedDataIsFromNode;
-}
-
-
-bool FCustomizableObjectEditorViewportClient::IsNodeMeshClipMorphSelected()
-{
-	return bClipMorphVisible;
-}
-
 void FCustomizableObjectEditorViewportClient::DrawCylinderArc(FPrimitiveDrawInterface* PDI, const FMatrix& CylToWorld, const FVector& Base, const FVector& XAxis, const FVector& YAxis, const FVector& ZAxis, float Radius, float HalfHeight, uint32 Sides, const FMaterialRenderProxy* MaterialRenderProxy, uint8 DepthPriority, FColor Color, float MaxAngle)
 {
 	TArray<FDynamicMeshVertex> MeshVerts;
@@ -3351,6 +2556,12 @@ bool FCustomizableObjectEditorViewportClient::IsShowingBones() const
 }
 
 
+const TArray<ULightComponent*>& FCustomizableObjectEditorViewportClient::GetLightComponents() const
+{
+	return LightComponents;
+}
+
+
 void FCustomizableObjectEditorViewportClient::DrawMeshBones(UDebugSkelMeshComponent* MeshComponent, FPrimitiveDrawInterface* PDI)
 {
 	if (!MeshComponent ||
@@ -3416,6 +2627,15 @@ void FCustomizableObjectEditorViewportClient::DrawMeshBones(UDebugSkelMeshCompon
 }
 
 
+void FCustomizableObjectEditorViewportClient::SetWidgetType(EWidgetType Type)
+{
+	WidgetType = Type;
+	
+	//SetWidgetMode(UE::Widget::WM_Translate); // TODO GMTFuture Uncomment once UE-191354 fixed
+	Widget->SetDefaultVisibility(Type != EWidgetType::Hidden);
+}	
+
+
 /////////////////////////////////////////////////
 // select folder dialog \todo: move to its own file
 //////////////////////////////////////////////////
diff --git a/Engine/Plugins/Experimental/Mutable/Source/CustomizableObjectEditor/Private/MuCOE/CustomizableObjectInstanceEditor.cpp b/Engine/Plugins/Experimental/Mutable/Source/CustomizableObjectEditor/Private/MuCOE/CustomizableObjectInstanceEditor.cpp
--- a/Engine/Plugins/Experimental/Mutable/Source/CustomizableObjectEditor/Private/MuCOE/CustomizableObjectInstanceEditor.cpp
+++ b/Engine/Plugins/Experimental/Mutable/Source/CustomizableObjectEditor/Private/MuCOE/CustomizableObjectInstanceEditor.cpp
@@ -8,6 +8,8 @@
 #include "Animation/PoseAsset.h"
 #include "AssetRegistry/AssetRegistryModule.h"
 #include "ContentBrowserModule.h"
+#include "CustomizableObjectCustomSettings.h"
+#include "CustomizableObjectEditor.h"
 #include "DetailsViewArgs.h"
 #include "FileHelpers.h"
 #include "Framework/Commands/UICommandList.h"
@@ -24,9 +26,9 @@
 #include "MuCOE/CustomizableObjectInstanceEditorActions.h"
 #include "MuCOE/CustomizableObjectPreviewScene.h"
 #include "MuCOE/SCustomizableObjectEditorTextureAnalyzer.h"
-#include "MuCOE/SCustomizableObjectEditorViewport.h"
 #include "MuCOE/UnrealEditorPortabilityHelpers.h"
 #include "PropertyEditorModule.h"
+#include "SCustomizableObjectEditorViewport.h"
 #include "Selection.h"
 #include "Widgets/Docking/SDockTab.h"
 #include "Widgets/Notifications/SNotificationList.h"
@@ -51,6 +53,81 @@ void UUpdateClassWrapperClass::DelegatedCallback(UCustomizableObjectInstance* In
 	Delegate.ExecuteIfBound();
 }
 
+
+UProjectorParameter::UProjectorParameter()
+{
+	SetFlags(RF_Transactional);
+}
+
+
+void UProjectorParameter::SelectProjector(const FString& InParamName, const int32 InRangeIndex)
+{
+	ParamName = InParamName;
+	RangeIndex = InRangeIndex;
+}
+
+
+void UProjectorParameter::UnselectProjector()
+{
+	ParamName = "";
+	RangeIndex = -1;	
+}
+
+
+bool UProjectorParameter::IsProjectorSelected(const FString& InParamName, const int32 InRangeIndex) const
+{
+	return ParamName == InParamName && RangeIndex == InRangeIndex;
+}
+
+
+FVector UProjectorParameter::GetPosition() const
+{
+	return Position;
+}
+
+
+void UProjectorParameter::SetPosition(const FVector& InPosition)
+{
+	Position = InPosition;
+}
+
+
+FVector UProjectorParameter::GetDirection() const
+{
+	return Direction;
+}
+
+
+void UProjectorParameter::SetDirection(const FVector& InDirection)
+{
+	Direction = InDirection;
+}
+
+
+FVector UProjectorParameter::GetUp() const
+{
+	return Up;
+}
+
+
+void UProjectorParameter::SetUp(const FVector& InUp)
+{
+	Up = InUp;
+}
+
+
+FVector UProjectorParameter::GetScale() const
+{
+	return Scale;
+}
+
+
+void UProjectorParameter::SetScale(const FVector& InScale)
+{
+	Scale = InScale;
+}
+
+
 void FCustomizableObjectInstanceEditor::RegisterTabSpawners(const TSharedRef<class FTabManager>& InTabManager)
 {
 	const IWorkspaceMenuStructure& MenuStructure = WorkspaceMenu::GetMenuStructure();
@@ -80,6 +157,42 @@ void FCustomizableObjectInstanceEditor::UnregisterTabSpawners(const TSharedRef<c
 }
 
 
+ULightComponent* UCustomSettings::GetSelectedLight() const
+{
+	return SelectedLight;
+}
+
+
+void UCustomSettings::SetSelectedLight(ULightComponent* Light)
+{
+	SelectedLight = Light;
+}
+
+
+UCustomizableObjectEditorViewportLights* UCustomSettings::GetLightsPreset() const
+{
+	return LightsPreset;
+}
+
+
+void UCustomSettings::SetLightsPreset(UCustomizableObjectEditorViewportLights& InLightsPreset)
+{
+	LightsPreset = &InLightsPreset;
+}
+
+
+TWeakPtr<ICustomizableObjectInstanceEditor> UCustomSettings::GetEditor() const
+{
+	return Editor;
+}
+
+
+void UCustomSettings::SetEditor(TSharedPtr<ICustomizableObjectInstanceEditor> InEditor)
+{
+	Editor = InEditor;
+}
+
+
 FCustomizableObjectInstanceEditor::FCustomizableObjectInstanceEditor()
 {
 	CustomizableObjectInstance = nullptr;
@@ -130,32 +243,31 @@ FCustomizableObjectInstanceEditor::~FCustomizableObjectInstanceEditor()
 
 void FCustomizableObjectInstanceEditor::InitCustomizableObjectInstanceEditor( const EToolkitMode::Type Mode, const TSharedPtr< class IToolkitHost >& InitToolkitHost, UCustomizableObjectInstance* InCustomizableObjectInstance )
 {
+	ProjectorParameter = NewObject<UProjectorParameter>();
+	
+	CustomSettings = NewObject<UCustomSettings>();
+	CustomSettings->SetEditor(SharedThis(this));
+
 	// Register our commands. This will only register them if not previously registered
 	FCustomizableObjectInstanceEditorCommands::Register();
 	FCustomizableObjectEditorViewportCommands::Register();
 
 	BindCommands();
 
-	bool bUpdateFromSelection = false;
-	bool bIsLockable = false;
-	bool bAllowSearch = true;
-
 	FPropertyEditorModule& PropPlugin = FModuleManager::LoadModuleChecked<FPropertyEditorModule>("PropertyEditor");
 
 	FDetailsViewArgs DetailsViewArgs;
 	DetailsViewArgs.NameAreaSettings = FDetailsViewArgs::ENameAreaSettings::HideNameArea;
 	DetailsViewArgs.bAllowSearch = false;
-	//DetailsViewArgs.bShowActorLabel = false;
 	DetailsViewArgs.bShowObjectLabel = false;
+	
 	CustomizableInstanceDetailsView = PropPlugin.CreateDetailView( DetailsViewArgs );
-	CustomizableInstanceDetailsView->SetObject(InCustomizableObjectInstance, true);
-	//CustomizableInstanceDetailsView->SetIsPropertyVisibleDelegate( FIsPropertyVisible::CreateStatic(&IsPropertyVisible) );
-
+	
 	TextureAnalyzer = SNew(SCustomizableObjecEditorTextureAnalyzer).CustomizableObjectInstanceEditor(this).CustomizableObjectEditor(nullptr);
 
 	Viewport = SNew(SCustomizableObjectEditorViewportTabBody)
 		.CustomizableObjectEditor(SharedThis(this));
-
+	
 	// Set the instance
 	check(InCustomizableObjectInstance);
 	CustomizableObjectInstance = InCustomizableObjectInstance;
@@ -220,13 +332,13 @@ void FCustomizableObjectInstanceEditor::InitCustomizableObjectInstanceEditor( co
 	const bool bCreateDefaultToolbar = true;
 	FAssetEditorToolkit::InitAssetEditor( Mode, InitToolkitHost, CustomizableObjectInstanceEditorAppIdentifier, StandaloneDefaultLayout, bCreateDefaultStandaloneMenu, bCreateDefaultToolbar, CustomizableObjectInstance );
 
+	CustomizableInstanceDetailsView->SetObject(InCustomizableObjectInstance, true); // Can only be called after initializing the Asset Editor
+	
 	ExtendToolbar();
 	RegenerateMenusAndToolbars();
 
 	// Clears selection highlight.
 	OnInstancePropertySelectionChanged(NULL);
-
-	OnObjectModifiedHandle = FCoreUObjectDelegates::OnObjectModified.AddRaw(this, &FCustomizableObjectInstanceEditor::OnObjectModified);
 }
 
 
@@ -391,6 +503,8 @@ void FCustomizableObjectInstanceEditor::AddReferencedObjects( FReferenceCollecto
 	}
 
 	Collector.AddReferencedObject( HelperCallback );
+	Collector.AddReferencedObject( ProjectorParameter );
+	Collector.AddReferencedObject( CustomSettings );
 }
 
 
@@ -521,16 +635,13 @@ FLinearColor FCustomizableObjectInstanceEditor::GetWorldCentricTabColorScale() c
 
 void FCustomizableObjectInstanceEditor::RefreshTool()
 {
-	RefreshViewport();
+	Viewport->GetViewportClient()->Invalidate();
 }
 
 
-void FCustomizableObjectInstanceEditor::RefreshViewport()
+TSharedPtr<SCustomizableObjectEditorViewportTabBody> FCustomizableObjectInstanceEditor::GetViewport()
 {
-	if (Viewport)
-	{
-		Viewport->GetViewportClient()->Invalidate();
-	}
+	return Viewport;		
 }
 
 
@@ -562,74 +673,33 @@ void FCustomizableObjectInstanceEditor::SetPoseAsset(class UPoseAsset* PoseAsset
 }
 
 
-void FCustomizableObjectInstanceEditor::OnObjectModified(UObject* Object)
+UProjectorParameter* FCustomizableObjectInstanceEditor::GetProjectorParameter()
 {
-	UCustomizableObjectInstance* Instance = Cast<UCustomizableObjectInstance>(Object);
-
-	if ((Instance != nullptr) && (Instance == CustomizableObjectInstance))
-	{
-		if (Instance->ProjectorLayerChange)
-		{
-			// Variable used in a non multilayer projector to not lose the focus of the gizmo
-			Instance->ProjectorLayerChange = false;
-			return;
-		}
-
-		if (Viewport->GetGizmoHasAssignedData())
-		{
-			Instance->LastSelectedProjectorParameter = Viewport->GetGizmoProjectorParameterName();
-			Instance->LastSelectedProjectorParameterWithIndex = Viewport->GetGizmoProjectorParameterNameWithIndex();
-		}
+	return ProjectorParameter;
+}
 
-		// Update projector parameter information (not projector parameter node information,
-		// this is the case where a projector parameter in the object details is selected)
-		const TArray<FCustomizableObjectProjectorParameterValue>& Parameters = CustomizableObjectInstance->GetProjectorParameters();
-		const int32 MaxIndex = Parameters.Num();
-		bool AnySelected = false;
 
-		for (int32 i = 0; (!AnySelected && (i < MaxIndex)); ++i)
-		{
-			int32 RangeIndex = Parameters[i].RangeValues.Num() > 0 ? 0 : -1;
+UCustomSettings* FCustomizableObjectInstanceEditor::GetCustomSettings()
+{
+	return CustomSettings;
+}
 
-			for (; RangeIndex < Parameters[i].RangeValues.Num(); ++RangeIndex)
-			{
-				const FCustomizableObjectProjector& Value = (RangeIndex == -1) ? Parameters[i].Value : Parameters[i].RangeValues[RangeIndex];
 
-				FString ParameterNameWithIndex = Parameters[i].ParameterName;
+void FCustomizableObjectInstanceEditor::HideGizmo()
+{
+	HideGizmo(SharedThis(this), Viewport, CustomizableInstanceDetailsView);
+}
 
-				if (RangeIndex != -1)
-				{
-					ParameterNameWithIndex += FString::Printf(TEXT("__%d"), RangeIndex);
-				}
 
-				if (CustomizableObjectInstance->GetProjectorState(Parameters[i].ParameterName, RangeIndex) == EProjectorState::Selected ||
-					(CustomizableObjectInstance->GetProjectorState(Parameters[i].ParameterName, RangeIndex) == EProjectorState::TypeChanged &&
-						ParameterNameWithIndex != Instance->LastSelectedProjectorParameterWithIndex))
-				{
-					Viewport->SetGizmoCallUpdateSkeletalMesh(false);
-					Viewport->SetProjectorVisibility(true, Parameters[i].ParameterName, ParameterNameWithIndex, RangeIndex, Value, i);
-					Viewport->SetGizmoCallUpdateSkeletalMesh(true);
-				AnySelected = true;
-			}
-				else if (CustomizableObjectInstance->GetProjectorState(Parameters[i].ParameterName, RangeIndex) == EProjectorState::TypeChanged &&
-					Parameters[i].ParameterName == Instance->LastSelectedProjectorParameter)
-					{
-					Viewport->SetProjectorType(true, Parameters[i].ParameterName, ParameterNameWithIndex, RangeIndex, Value, i);
-						AnySelected = true;
-					}
-				}
+void FCustomizableObjectInstanceEditor::ShowGizmoProjectorParameter(const FString& ParamName, int32 RangeIndex)
+{
+	ShowGizmoProjectorParameter(ParamName, RangeIndex, SharedThis(this), Viewport, CustomizableInstanceDetailsView, ProjectorParameter, CustomizableObjectInstance);
+}
 
-		}
 
-		if (!AnySelected && !Viewport->GetIsManipulatingGizmo())
-		{
-			Viewport->SetGizmoCallUpdateSkeletalMesh(false);
-			Viewport->ResetProjectorVisibility(false);
-			Viewport->SetGizmoCallUpdateSkeletalMesh(true);
-			Instance->LastSelectedProjectorParameter = "";
-			Instance->LastSelectedProjectorParameterWithIndex = "";
-		}
-	}
+void FCustomizableObjectInstanceEditor::HideGizmoProjectorParameter()
+{
+	HideGizmoProjectorParameter(SharedThis(this), Viewport, CustomizableInstanceDetailsView);
 }
 
 
@@ -641,13 +711,13 @@ bool FCustomizableObjectInstanceEditor::GetAssetRegistryLoaded()
 
 void FCustomizableObjectInstanceEditor::OnInstancePropertySelectionChanged(FProperty* InProperty)
 {
-	RefreshViewport();
+	Viewport->GetViewportClient()->Invalidate();
 }
 
 
 void FCustomizableObjectInstanceEditor::OnUpdatePreviewInstance()
 {
-	RefreshViewport();
+	Viewport->GetViewportClient()->Invalidate();
 
 	UpdatePreviewVisibility();
 
@@ -669,29 +739,11 @@ void FCustomizableObjectInstanceEditor::OnUpdatePreviewInstance()
 		Viewport->SetPreviewComponents(ObjectPtrDecay(PreviewSkeletalMeshComponents));
 	}
 
-	if (!CustomizableObjectInstance) return;
-	
-	// If the instance is updated due to a change in a parameter projector, set again the LastSelectedProjectorParameter variable with the
-	// current projector the viewport gizmo has to continue having it as selected when rebuilding the parameter details widget.
-	// Otherwise, reset the value.
-	bool UpdateDetails = false;
-	if (Viewport->GetGizmoHasAssignedData())
-	{
-		CustomizableObjectInstance->LastSelectedProjectorParameter = Viewport->GetGizmoProjectorParameterName();
-		CustomizableObjectInstance->LastSelectedProjectorParameterWithIndex = Viewport->GetGizmoProjectorParameterNameWithIndex();
-		UpdateDetails = true;
-	}
-	else
-	{
-		CustomizableObjectInstance->LastSelectedProjectorParameter = "";
-		CustomizableObjectInstance->LastSelectedProjectorParameterWithIndex = "";
-	}
-
-	if (UpdateDetails && !CustomizableObjectInstance->ProjectorUpdatedInViewport)
+	if (!CustomizableObjectInstance)
 	{
-		CustomizableInstanceDetailsView->ForceRefresh();
+		return;
 	}
-
+	
 	if (TextureAnalyzer.IsValid())
 	{
 		TextureAnalyzer.Get()->RefreshTextureAnalyzerTable(CustomizableObjectInstance);
@@ -701,14 +753,11 @@ void FCustomizableObjectInstanceEditor::OnUpdatePreviewInstance()
 
 void FCustomizableObjectInstanceEditor::CompileObject(UCustomizableObject* Object)
 {
-
 	if (!Object || !Object->Source)
 	{
 		return;
 	}
 	
-	Viewport->UpdateGizmoDataToOrigin();
-
 	if (!Compiler)
 	{
 		Compiler = MakeUnique<FCustomizableObjectCompiler>();
@@ -749,23 +798,6 @@ void FCustomizableObjectInstanceEditor::Tick(float InDeltaTime)
 		Compiler.Reset();
 	}
 
-	// TEMP CODE
-	if (CustomizableObjectInstance->TempUpdateGizmoInViewport)
-	{
-		CustomizableObjectInstance->TempUpdateGizmoInViewport = false;
-		CustomizableObjectInstance->AvoidResetProjectorVisibilityForNonNode = true;
-		Viewport->CopyTransformFromOriginData();
-		CustomizableObjectInstance->UpdateSkeletalMeshAsync(true);
-
-		EProjectorState::Type ProjectorState = CustomizableObjectInstance->GetProjectorState(CustomizableObjectInstance->TempProjectorParameterName, CustomizableObjectInstance->TempProjectorParameterRangeIndex);
-		if (ProjectorState != EProjectorState::Selected)
-		{
-			// ResetProjectorVisibilityForNonNode = false;
-			CustomizableObjectInstance->SetProjectorState(CustomizableObjectInstance->TempProjectorParameterName, CustomizableObjectInstance->TempProjectorParameterRangeIndex, EProjectorState::Selected);
-			FCoreUObjectDelegates::BroadcastOnObjectModified(CustomizableObjectInstance);
-		}
-	}
-
 	// If we want to show the Relevant/Runtime parameters, we need to refresh the details view to make sure that the scroll bar appears		
 	if (bOnlyRelevantParameters != CustomizableObjectInstance->bShowOnlyRelevantParameters)
 	{
@@ -792,6 +824,266 @@ void FCustomizableObjectInstanceEditor::OpenTextureAnalyzerTab()
 }
 
 
+void FCustomizableObjectInstanceEditor::HideGizmo(const TSharedPtr<ICustomizableObjectInstanceEditor>& Editor,
+	const TSharedPtr<SCustomizableObjectEditorViewportTabBody>& Viewport,
+	const TSharedPtr<IDetailsView>& InstanceDetailsView)
+{
+	HideGizmoProjectorParameter(Editor, Viewport, InstanceDetailsView);
+}
+
+
+void FCustomizableObjectInstanceEditor::ShowGizmoProjectorParameter(const FString& ParamName, int32 RangeIndex,
+	const TSharedPtr<ICustomizableObjectInstanceEditor>& Editor, const TSharedPtr<SCustomizableObjectEditorViewportTabBody>& Viewport,
+	const TSharedPtr<IDetailsView>& InstanceDetailsView, UProjectorParameter* ProjectorParameter,
+	UCustomizableObjectInstance* Instance)
+{
+	Editor->HideGizmo();
+	
+	ProjectorParameter->SelectProjector(ParamName, RangeIndex);
+	
+	ProjectorParameter->SetPosition(Instance->GetProjectorPosition(ParamName, RangeIndex));
+	ProjectorParameter->SetDirection(Instance->GetProjectorDirection(ParamName, RangeIndex));
+	ProjectorParameter->SetUp(Instance->GetProjectorUp(ParamName, RangeIndex));
+	ProjectorParameter->SetScale(Instance->GetProjectorScale(ParamName, RangeIndex));
+
+	TWeakObjectPtr<UCustomizableObjectInstance> WeakInstance = Instance;
+
+	const TWeakPtr<ICustomizableObjectInstanceEditor> WeakEditor = Editor.ToWeakPtr();
+	
+	FProjectorTypeDelegate  ProjectorTypeDelegate;
+	ProjectorTypeDelegate.BindLambda([WeakInstance, ParamName, RangeIndex]() -> ECustomizableObjectProjectorType
+	{
+		const UCustomizableObjectInstance* Instance = WeakInstance.Get();
+		if (!Instance)
+		{
+			return {};
+		}
+
+		return Instance->GetProjectorParameterType(ParamName, RangeIndex);
+	});
+	
+	FWidgetColorDelegate WidgetColorDelegate;
+	WidgetColorDelegate.BindLambda([]() { return FColor::Green; });
+		
+	// Position
+	FWidgetLocationDelegate WidgetLocationDelegate;
+	WidgetLocationDelegate.BindLambda([WeakEditor, WeakInstance]() -> FVector
+	{
+		const UCustomizableObjectInstance* Instance = WeakInstance.Get();
+		if (!Instance)
+		{
+			return {};
+		}
+
+		const TSharedPtr<ICustomizableObjectInstanceEditor> Editor = WeakEditor.Pin();
+		if (!Editor)
+		{
+			return {};
+		}
+		
+		return Editor->GetProjectorParameter()->GetPosition(); // We are not getting the value directly from the parameters since they get updated on ReloadParameters, making the gizmo jittery.
+	});
+
+	FOnWidgetLocationChangedDelegate OnWidgetLocationChangedDelegate;
+	OnWidgetLocationChangedDelegate.BindLambda([WeakEditor, WeakInstance, ParamName, RangeIndex](const FVector& Location)
+	{		
+		UCustomizableObjectInstance* Instance = WeakInstance.Get();
+		if (!Instance)
+		{
+			return;
+		}
+
+		const TSharedPtr<ICustomizableObjectInstanceEditor> Editor = WeakEditor.Pin();
+		if (!Editor)
+		{
+			return;
+		}
+		
+		Editor->GetProjectorParameter()->SetPosition(Location);
+		
+		Instance->SetProjectorPosition(ParamName, static_cast<FVector3f>(Location), RangeIndex);
+		Instance->UpdateSkeletalMeshAsync(true, true);
+	});
+
+	// Direction
+	FWidgetDirectionDelegate WidgetDirectionDelegate;
+	WidgetDirectionDelegate.BindLambda([WeakEditor, WeakInstance]() -> FVector
+	{
+		const UCustomizableObjectInstance* Instance = WeakInstance.Get();
+		if (!Instance)
+		{
+			return {};
+		}
+
+		const TSharedPtr<ICustomizableObjectInstanceEditor> Editor = WeakEditor.Pin();
+		if (!Editor)
+		{
+			return {};
+		}
+		
+		return Editor->GetProjectorParameter()->GetDirection();
+	});
+	
+	FOnWidgetDirectionChangedDelegate OnWidgetDirectionChangedDelegate;
+	OnWidgetDirectionChangedDelegate.BindLambda([WeakEditor, WeakInstance, ParamName, RangeIndex](const FVector& Direction)
+	{
+		UCustomizableObjectInstance* Instance = WeakInstance.Get();
+		if (!Instance)
+		{
+			return;
+		}
+
+		const TSharedPtr<ICustomizableObjectInstanceEditor> Editor = WeakEditor.Pin();
+		if (!Editor)
+		{
+			return;
+		}
+		
+		Editor->GetProjectorParameter()->SetDirection(Direction);
+		
+		Instance->SetProjectorDirection(ParamName, Direction, RangeIndex);
+		Instance->UpdateSkeletalMeshAsync(true, true);
+	});
+
+	// Up
+	FWidgetUpDelegate WidgetUpDelegate;
+	WidgetUpDelegate.BindLambda([WeakEditor, WeakInstance]() -> FVector
+	{
+		const UCustomizableObjectInstance* Instance = WeakInstance.Get();
+		if (!Instance)
+		{
+			return {};
+		}
+
+		const TSharedPtr<ICustomizableObjectInstanceEditor> Editor = WeakEditor.Pin();
+		if (!Editor)
+		{
+			return {};
+		}
+		
+		return Editor->GetProjectorParameter()->GetUp();
+	});
+
+	FOnWidgetUpChangedDelegate OnWidgetUpChangedDelegate;
+	OnWidgetUpChangedDelegate.BindLambda([WeakEditor, WeakInstance, ParamName, RangeIndex](const FVector& Up)
+	{
+		UCustomizableObjectInstance* Instance = WeakInstance.Get();
+		if (!Instance)
+		{
+			return;
+		}
+
+		const TSharedPtr<ICustomizableObjectInstanceEditor> Editor = WeakEditor.Pin();
+		if (!Editor)
+		{
+			return;
+		}
+		
+		Editor->GetProjectorParameter()->SetUp(Up);
+
+		Instance->SetProjectorUp(ParamName, Up, RangeIndex);
+		Instance->UpdateSkeletalMeshAsync(true, true);
+	});
+
+	// Scale
+	FWidgetScaleDelegate WidgetScaleDelegate;
+	WidgetScaleDelegate.BindLambda([WeakEditor, WeakInstance]() -> FVector
+	{
+		const UCustomizableObjectInstance* Instance = WeakInstance.Get();
+		if (!Instance)
+		{
+			return {};
+		}
+
+		const TSharedPtr<ICustomizableObjectInstanceEditor> Editor = WeakEditor.Pin();
+		if (!Editor)
+		{
+			return {};
+		}
+		
+		return Editor->GetProjectorParameter()->GetScale();
+	});
+	
+	FOnWidgetScaleChangedDelegate OnWidgetScaleChangedDelegate;
+	OnWidgetScaleChangedDelegate.BindLambda([WeakEditor, WeakInstance, ParamName, RangeIndex](const FVector& Scale)
+	{
+		UCustomizableObjectInstance* Instance = WeakInstance.Get();
+		if (!Instance)
+		{
+			return;
+		}
+
+		const TSharedPtr<ICustomizableObjectInstanceEditor> Editor = WeakEditor.Pin();
+		if (!Editor)
+		{
+			return;
+		}
+		
+		Editor->GetProjectorParameter()->SetScale(Scale);
+		
+		Instance->SetProjectorScale(ParamName, Scale, RangeIndex);
+		Instance->UpdateSkeletalMeshAsync(true, true);
+	});
+
+	// Angle
+	FWidgetAngleDelegate WidgetAngleDelegate;
+	WidgetAngleDelegate.BindLambda([WeakInstance, ParamName, RangeIndex]() -> float
+	{
+		const UCustomizableObjectInstance* Instance = WeakInstance.Get();
+		if (!Instance)
+		{
+			return {};
+		}
+
+		return Instance->GetProjectorAngle(ParamName, RangeIndex);
+	});
+
+	// UObject transactions
+	FWidgetTrackingStartedDelegate WidgetTrackingStartedDelegate;
+	WidgetTrackingStartedDelegate.BindLambda([WeakEditor, WeakInstance]()
+	{
+		UCustomizableObjectInstance* Instance = WeakInstance.Get();
+		if (!Instance)
+		{
+			return;
+		}
+
+		const TSharedPtr<ICustomizableObjectInstanceEditor> Editor = WeakEditor.Pin();
+		if (!Editor)
+		{
+			return;
+		}
+		
+		Instance->Modify();
+		Editor->GetProjectorParameter()->Modify();
+	});
+
+	Viewport->ShowGizmoProjector(WidgetLocationDelegate, OnWidgetLocationChangedDelegate,
+		WidgetDirectionDelegate, OnWidgetDirectionChangedDelegate,
+		WidgetUpDelegate, OnWidgetUpChangedDelegate,
+		WidgetScaleDelegate, OnWidgetScaleChangedDelegate,
+		WidgetAngleDelegate,
+		ProjectorTypeDelegate,
+		WidgetColorDelegate,
+		WidgetTrackingStartedDelegate);
+	
+	InstanceDetailsView->ForceRefresh();
+}
+
+
+void FCustomizableObjectInstanceEditor::HideGizmoProjectorParameter(const TSharedPtr<ICustomizableObjectInstanceEditor>& Editor,
+		const TSharedPtr<SCustomizableObjectEditorViewportTabBody>& Viewport,
+		const TSharedPtr<IDetailsView>& InstanceDetailsView)
+{
+	Viewport->HideGizmoProjector();
+
+	UProjectorParameter* ProjectorParameter = Editor->GetProjectorParameter();
+	
+	ProjectorParameter->UnselectProjector();
+	InstanceDetailsView->ForceRefresh();		
+}
+
+
 TSharedRef<SDockTab> FCustomizableObjectInstanceEditor::SpawnTab_TextureAnalyzer(const FSpawnTabArgs& Args)
 {
 	check(Args.GetTabId() == TextureAnalyzerTabId);
