[Mutable] Projector Gizmo code refactor.

diff --git a/Engine/Plugins/Experimental/Mutable/Source/CustomizableObjectEditor/Private/MuCOE/CustomizableObjectInstanceEditor.h b/Engine/Plugins/Experimental/Mutable/Source/CustomizableObjectEditor/Private/MuCOE/CustomizableObjectInstanceEditor.h
--- a/Engine/Plugins/Experimental/Mutable/Source/CustomizableObjectEditor/Private/MuCOE/CustomizableObjectInstanceEditor.h
+++ b/Engine/Plugins/Experimental/Mutable/Source/CustomizableObjectEditor/Private/MuCOE/CustomizableObjectInstanceEditor.h
@@ -3,6 +3,7 @@
 #pragma once
 
 #include "MuCOE/ICustomizableObjectInstanceEditor.h"
+#include "MuCOE/CustomizableObjectEditorViewportLights.h"
 #include "TickableEditorObject.h"
 
 #include "CustomizableObjectInstanceEditor.generated.h"
@@ -14,10 +15,25 @@ class SWidget;
 class UCustomizableObject;
 class UCustomizableObjectInstance;
 class UPoseAsset;
+class SLevelOfDetailSettings;
+class SCustomizableObjectEditorViewportTabBody;
+class SCustomizableObjecEditorTextureAnalyzer;
+class FCustomizableObjectCompiler;
+class UCustomizableSkeletalComponent;
+class UDebugSkelMeshComponent;
+class SDockableTab;
 struct FFrame;
 
 DECLARE_DELEGATE(FCreatePreviewInstanceFlagDelegate);
 
+enum class EProjectorGizmo : uint8
+{
+	None,
+	Parameter,
+	NodeProjectorConstant,
+	NodeProjectorParameter
+};
+
 /**
 * Wrapper UObject class for the UCustomizableObjectInstance::FObjectInstanceUpdatedDelegate dynamic multicast delegate
 */
@@ -35,6 +51,89 @@ public:
 };
 
 
+/** Currently selected projector parameter and its position.
+ *
+ * Only used by Instance Parameters. Default values projector values modify the node directly. */
+UCLASS(Transient)
+class UProjectorParameter : public UObject
+{
+	GENERATED_BODY()
+	
+public:
+	UProjectorParameter();
+
+	void SelectProjector(const FString& ParamName, int32 RangeIndex = -1);
+	
+	void UnselectProjector();
+
+	bool IsProjectorSelected(const FString& InParamName, int32 InRangeIndex = -1) const;
+	
+	FVector GetPosition() const;
+
+	void SetPosition(const FVector& Position);
+
+	FVector GetDirection() const;
+
+	void SetDirection(const FVector& Direction);
+
+	FVector GetUp() const;
+
+	void SetUp(const FVector& Up);
+
+	FVector GetScale() const;
+
+	void SetScale(const FVector& Scale);
+	
+private:
+	UPROPERTY()
+	FString ParamName = {};
+
+	UPROPERTY()
+	int32 RangeIndex = -1;
+	
+	UPROPERTY()
+	FVector Position = FVector::ZeroVector;
+
+	UPROPERTY()
+	FVector Direction = FVector::ForwardVector;
+
+	UPROPERTY()
+	FVector Up = FVector::UpVector;
+
+	UPROPERTY()
+	FVector Scale = FVector::OneVector;
+};
+
+
+UCLASS(Transient)
+class UCustomSettings : public UObject
+{
+	GENERATED_BODY()
+
+public:
+	ULightComponent* GetSelectedLight() const;
+
+	void SetSelectedLight(ULightComponent* Light);
+
+	UCustomizableObjectEditorViewportLights* GetLightsPreset() const;
+
+	void SetLightsPreset(UCustomizableObjectEditorViewportLights& LightsPreset);
+
+	TWeakPtr<ICustomizableObjectInstanceEditor> GetEditor() const;
+
+	void SetEditor(TSharedPtr<ICustomizableObjectInstanceEditor> Editor);
+
+private:
+	UPROPERTY()
+	TObjectPtr<ULightComponent> SelectedLight = nullptr;
+
+	UPROPERTY()
+	UCustomizableObjectEditorViewportLights* LightsPreset;
+
+	TWeakPtr<ICustomizableObjectInstanceEditor> Editor;
+};
+
+
 /**
  * CustomizableObject Editor class
  */
@@ -77,12 +176,14 @@ public:
 	/** ICustomizableObjectInstanceEditor interface */
 	UCustomizableObjectInstance* GetPreviewInstance() override;
 	virtual void RefreshTool() override;
-	virtual void RefreshViewport() override;
-	virtual void SetPoseAsset(class UPoseAsset* PoseAssetParameter) override;
-
-	/** Callback for the object modified event */
-	void OnObjectModified(UObject* Object);
-
+	virtual TSharedPtr<SCustomizableObjectEditorViewportTabBody> GetViewport() override;
+	virtual void SetPoseAsset(UPoseAsset* PoseAssetParameter) override;
+	virtual UProjectorParameter* GetProjectorParameter() override;
+	virtual UCustomSettings* GetCustomSettings() override;
+	virtual void HideGizmo() override;
+	virtual void ShowGizmoProjectorParameter(const FString& ParamName, int32 RangeIndex) override;
+	virtual void HideGizmoProjectorParameter() override;
+	
 	/** Getter of AssetRegistryLoaded */
 	bool GetAssetRegistryLoaded();
 
@@ -104,8 +205,6 @@ private:
 	TSharedRef<SDockTab> SpawnTab_AdvancedPreviewSettings(const FSpawnTabArgs& Args);
 	TSharedRef<SDockTab> SpawnTab_TextureAnalyzer(const FSpawnTabArgs& Args);
 
-private:
-
 	void CreatePreviewInstance();
 
 	/** Binds commands associated with the Static Mesh Editor. */
@@ -114,9 +213,6 @@ private:
 	/** Callback when selection changes in the Property Tree. */
 	void OnInstancePropertySelectionChanged(FProperty* InProperty);
 
-	/** The instance notifies the editor of projector states changes through this. */
-	void OnProjectorStateChanged();
-
 	/** Callback for the asset registry initial load */
 	void OnAssetRegistryLoadComplete();
 
@@ -126,47 +222,63 @@ private:
 	/** Save Customizable Object Instance open in the editor */
 	void SaveAsset_Execute() override;
 
-	/** Says if the customizable Object can be shown or be oppened in the editor */
+	/** Says if the customizable Object can be shown or be opened in the editor */
 	bool CanOpenOrShowParent();
 
-	/** Show Customizable Object Intance's Parent In Content Browser */
+	/** Show Customizable Object Instance's Parent In Content Browser */
 	void ShowParentInContentBrowser();
 
-	/** Open Customizable Object Intance's Parent In Editor */
+	/** Open Customizable Object Instance's Parent In Editor */
 	void OpenParentInEditor();
 
 	/** Open the Texture Analyzer tab */
 	void OpenTextureAnalyzerTab();
 
+public:
+	/** Static to reuse code with FCustomizableObjectEditor. */
+	static void HideGizmo(const TSharedPtr<ICustomizableObjectInstanceEditor>& Editor,
+		const TSharedPtr<SCustomizableObjectEditorViewportTabBody>& Viewport,
+		const TSharedPtr<IDetailsView>& InstanceDetailsView);
+
+	/** Static to reuse code with FCustomizableObjectEditor. */
+	static void ShowGizmoProjectorParameter(
+		const FString& ParamName, int32 RangeIndex,
+		const TSharedPtr<ICustomizableObjectInstanceEditor>& Editor, const TSharedPtr<SCustomizableObjectEditorViewportTabBody>& Viewport, const TSharedPtr<IDetailsView>& InstanceDetailsView,
+		UProjectorParameter* ProjectorParameter, UCustomizableObjectInstance* Instance);
+
+	/** Static to reuse code with FCustomizableObjectEditor. */
+	static void HideGizmoProjectorParameter(const TSharedPtr<ICustomizableObjectInstanceEditor>& Editor,
+		const TSharedPtr<SCustomizableObjectEditorViewportTabBody>& Viewport,
+		const TSharedPtr<IDetailsView>& InstanceDetailsView);
+	
 private:
-
 	/** The currently viewed object. */
-	TObjectPtr<class UCustomizableObjectInstance> CustomizableObjectInstance;
-	TArray<TObjectPtr<class UCustomizableSkeletalComponent>> PreviewCustomizableSkeletalComponents;
-	TObjectPtr<class UStaticMeshComponent> PreviewStaticMeshComponent;
-	TArray<TObjectPtr<class UDebugSkelMeshComponent>> PreviewSkeletalMeshComponents;
+	TObjectPtr<UCustomizableObjectInstance> CustomizableObjectInstance;
+	TArray<TObjectPtr<UCustomizableSkeletalComponent>> PreviewCustomizableSkeletalComponents;
+	TObjectPtr<UStaticMeshComponent> PreviewStaticMeshComponent;
+	TArray<TObjectPtr<UDebugSkelMeshComponent>> PreviewSkeletalMeshComponents;
 
 	/** List of open tool panels; used to ensure only one exists at any one time */
-	TMap< FName, TWeakPtr<class SDockableTab> > SpawnedToolPanels;
+	TMap<FName, TWeakPtr<SDockableTab>> SpawnedToolPanels;
 
 	/** Preview Viewport widget */
-	TSharedPtr<class SCustomizableObjectEditorViewportTabBody> Viewport;
-	TSharedPtr<class IDetailsView> CustomizableInstanceDetailsView;
+	TSharedPtr<SCustomizableObjectEditorViewportTabBody> Viewport;
+	TSharedPtr<IDetailsView> CustomizableInstanceDetailsView;
 
 	/** Level of Details Settings widget. */
-	TSharedPtr< class SLevelOfDetailSettings> LevelOfDetailSettings;
+	TSharedPtr<SLevelOfDetailSettings> LevelOfDetailSettings;
 
 	/** Widget for displaying the available UV Channels. */
-	TSharedPtr< class STextComboBox > UVChannelCombo;
+	TSharedPtr<STextComboBox> UVChannelCombo;
 
 	/** List of available UV Channels. */
-	TArray< TSharedPtr< FString > > UVChannels;
+	TArray<TSharedPtr<FString>> UVChannels;
 
 	/** Widget for displaying the available LOD. */
-	TSharedPtr< class STextComboBox > LODLevelCombo;
+	TSharedPtr<STextComboBox> LODLevelCombo;
 
 	/** List of LODs. */
-	TArray< TSharedPtr< FString > > LODLevels;
+	TArray<TSharedPtr<FString>> LODLevels;
 
 	/**	The tab ids for all the tabs used */
 	static const FName ViewportTabId;
@@ -187,19 +299,25 @@ private:
 	TSharedPtr<SWidget> AdvancedPreviewSettingsWidget;
 
 	/** Object compiler */
-	TUniquePtr<class FCustomizableObjectCompiler> Compiler;
+	TUniquePtr<FCustomizableObjectCompiler> Compiler;
 
 	/** Pose asset when doing drag and drop of an UPoseAsset to the viewport */
 	UPoseAsset* PoseAsset;
 
 	/** Texture Analyzer table widget which shows the information of the transient textures used in the customizable object instance */
-	TSharedPtr<class SCustomizableObjecEditorTextureAnalyzer> TextureAnalyzer;
+	TSharedPtr<SCustomizableObjecEditorTextureAnalyzer> TextureAnalyzer;
 
 	/** Variables used to force the refresh of the details view widget. These are needed because sometimes the scrollbar of the window doesn't appear
 	    until we force the refresh */
 	bool bOnlyRuntimeParameters;
 	bool bOnlyRelevantParameters;
 
+	TObjectPtr<UProjectorParameter> ProjectorParameter = nullptr;
+
+	TObjectPtr<UCustomSettings> CustomSettings = nullptr;
+
+	EProjectorGizmo ProjectorGizmo = EProjectorGizmo::None;
+	
 	/** Adds the customizable Object Instance Editor commands to the default toolbar */
 	void ExtendToolbar();
 };
diff --git a/Engine/Plugins/Experimental/Mutable/Source/CustomizableObjectEditor/Private/MuCOE/CustomizableObjectInstanceEditorActions.cpp b/Engine/Plugins/Experimental/Mutable/Source/CustomizableObjectEditor/Private/MuCOE/CustomizableObjectInstanceEditorActions.cpp
--- a/Engine/Plugins/Experimental/Mutable/Source/CustomizableObjectEditor/Private/MuCOE/CustomizableObjectInstanceEditorActions.cpp
+++ b/Engine/Plugins/Experimental/Mutable/Source/CustomizableObjectEditor/Private/MuCOE/CustomizableObjectInstanceEditorActions.cpp
@@ -30,18 +30,9 @@ void FCustomizableObjectInstanceEditorCommands::RegisterCommands()
 	UI_COMMAND( SetCameraLock, "Camera Lock", "Toggles viewport navigation between orbit and freely moving about.", EUserInterfaceActionType::ToggleButton, FInputChord() );
 	UI_COMMAND( SaveThumbnail, "Save Thumbnail", "Saves the viewpoint position and orientation in the Preview Pane for use as the thumbnail preview in the Content Browser.", EUserInterfaceActionType::Button, FInputChord() );
 
-	UI_COMMAND( SetShowNormals, "Normals", "Toggles display of vertex normals in the Preview Pane.", EUserInterfaceActionType::ToggleButton, FInputChord() );
-	UI_COMMAND( SetShowTangents, "Tangents", "Toggles display of vertex tangents in the Preview Pane.", EUserInterfaceActionType::ToggleButton, FInputChord() );
-	UI_COMMAND( SetShowBinormals, "Binormals", "Toggles display of vertex binormals (orthogonal vector to normal and tangent) in the Preview Pane.", EUserInterfaceActionType::ToggleButton, FInputChord() );
-	UI_COMMAND( SetShowPivot, "Show Pivot", "Display the pivot location of the static mesh.", EUserInterfaceActionType::ToggleButton, FInputChord() );
-
-	UI_COMMAND( BakeInstanceFromInstance, "Bake Instance from Instance", "Create baked unreal resources for the current instance.", EUserInterfaceActionType::Button, FInputChord());
-	UI_COMMAND( StateChangeTest, "State test", "Make a state change test, iterating through all possible states and runtime parameters for this instance", EUserInterfaceActionType::Button, FInputChord());
-
 	UI_COMMAND(ShowParentCO, "Show Parent", "Show the Parent of the Customizable Object in the Content Browser", EUserInterfaceActionType::Button, FInputChord());
 	UI_COMMAND(EditParentCO, "Edit Parent", "Open the Parent of the Customizable Object to Edit", EUserInterfaceActionType::Button, FInputChord());
 	UI_COMMAND(TextureAnalyzer, "Texture Memory Analyzer", "Open the Texture Analyzer window to check all the information of the textures created by Mutable.", EUserInterfaceActionType::Button, FInputChord());
-
 }
 
 #undef LOCTEXT_NAMESPACE
diff --git a/Engine/Plugins/Experimental/Mutable/Source/CustomizableObjectEditor/Private/MuCOE/Nodes/CustomizableObjectNodeProjectorConstant.cpp b/Engine/Plugins/Experimental/Mutable/Source/CustomizableObjectEditor/Private/MuCOE/Nodes/CustomizableObjectNodeProjectorConstant.cpp
--- a/Engine/Plugins/Experimental/Mutable/Source/CustomizableObjectEditor/Private/MuCOE/Nodes/CustomizableObjectNodeProjectorConstant.cpp
+++ b/Engine/Plugins/Experimental/Mutable/Source/CustomizableObjectEditor/Private/MuCOE/Nodes/CustomizableObjectNodeProjectorConstant.cpp
@@ -17,7 +17,6 @@ UCustomizableObjectNodeProjectorConstant::UCustomizableObjectNodeProjectorConsta
 	, BoneComboBoxLocation(FVector::ZeroVector)
 	, BoneComboBoxForwardDirection(FVector::ZeroVector)
 	, BoneComboBoxUpDirection(FVector::ZeroVector)
-	, ParameterSetModified(-1)
 {
 
 }
@@ -25,28 +24,17 @@ UCustomizableObjectNodeProjectorConstant::UCustomizableObjectNodeProjectorConsta
 
 void UCustomizableObjectNodeProjectorConstant::PostEditChangeProperty(FPropertyChangedEvent& PropertyChangedEvent)
 {
-	FProperty* PropertyThatChanged = PropertyChangedEvent.Property;
-	//if ( PropertyThatChanged && PropertyThatChanged->GetName() == TEXT("NumLODs") )
-	{
-	}
-
 	const FName PropertyName = PropertyChangedEvent.GetPropertyName();
 
-	if (PropertyName == "ProjectionType")
-	{
-		ParameterSetModified = 0;
-	}
-	else if (PropertyName == "ProjectionAngle")
+	if (PropertyName == "ProjectionAngle")
 	{
 		Value.Angle = FMath::DegreesToRadians(ProjectionAngle);
-		ParameterSetModified = 1;
 	}
 	else if (PropertyName == "ProjectorBone")
 	{
 		Value.Position = (FVector3f)BoneComboBoxLocation;
 		Value.Direction = (FVector3f)BoneComboBoxForwardDirection;
 		Value.Up = (FVector3f)BoneComboBoxUpDirection;
-		ParameterSetModified = 2;
 	}
 
 	Super::PostEditChangeProperty(PropertyChangedEvent);
@@ -63,6 +51,72 @@ void UCustomizableObjectNodeProjectorConstant::AllocateDefaultPins(UCustomizable
 }
 
 
+ECustomizableObjectProjectorType UCustomizableObjectNodeProjectorConstant::GetProjectorType() const
+{
+	return ProjectionType;
+}
+
+
+FVector UCustomizableObjectNodeProjectorConstant::GetProjectorPosition() const
+{
+	return static_cast<FVector>(Value.Position);
+}
+
+
+void UCustomizableObjectNodeProjectorConstant::SetProjectorPosition(const FVector& Position)
+{
+	Value.Position = static_cast<FVector3f>(Position);
+}
+
+
+FVector UCustomizableObjectNodeProjectorConstant::GetProjectorDirection() const
+{
+	return static_cast<FVector>(Value.Direction);
+}
+
+
+void UCustomizableObjectNodeProjectorConstant::SetProjectorDirection(const FVector& Direction)
+{
+	Value.Direction = static_cast<FVector3f>(Direction);
+}
+
+
+FVector UCustomizableObjectNodeProjectorConstant::GetProjectorUp() const
+{
+	return static_cast<FVector>(Value.Up);
+}
+
+
+void UCustomizableObjectNodeProjectorConstant::SetProjectorUp(const FVector& Up)
+{
+	Value.Up = static_cast<FVector3f>(Up);
+}
+
+
+FVector UCustomizableObjectNodeProjectorConstant::GetProjectorScale() const
+{
+	return static_cast<FVector>(Value.Scale);
+}
+
+
+void UCustomizableObjectNodeProjectorConstant::SetProjectorScale(const FVector& Scale)
+{
+	Value.Scale = static_cast<FVector3f>(Scale);
+}
+
+
+float UCustomizableObjectNodeProjectorConstant::GetProjectorAngle()
+{
+	return ProjectionAngle;
+}
+
+
+void UCustomizableObjectNodeProjectorConstant::SetProjectorAngle(float Angle)
+{
+	ProjectionAngle = Angle;
+}
+
+
 FText UCustomizableObjectNodeProjectorConstant::GetNodeTitle(ENodeTitleType::Type TitleType) const
 {
 	return LOCTEXT("Projector_Constant", "Projector Constant");
diff --git a/Engine/Plugins/Experimental/Mutable/Source/CustomizableObjectEditor/Private/MuCOE/Nodes/CustomizableObjectNodeProjectorParameter.cpp b/Engine/Plugins/Experimental/Mutable/Source/CustomizableObjectEditor/Private/MuCOE/Nodes/CustomizableObjectNodeProjectorParameter.cpp
--- a/Engine/Plugins/Experimental/Mutable/Source/CustomizableObjectEditor/Private/MuCOE/Nodes/CustomizableObjectNodeProjectorParameter.cpp
+++ b/Engine/Plugins/Experimental/Mutable/Source/CustomizableObjectEditor/Private/MuCOE/Nodes/CustomizableObjectNodeProjectorParameter.cpp
@@ -11,25 +11,20 @@ class UCustomizableObjectNodeRemapPins;
 
 void UCustomizableObjectNodeProjectorParameter::PostEditChangeProperty(FPropertyChangedEvent& PropertyChangedEvent)
 {
-	FProperty* PropertyThatChanged = PropertyChangedEvent.Property;
-
 	const FName PropertyName = PropertyChangedEvent.GetPropertyName();
 
 	if (PropertyName == "ProjectionType")
 	{
-		ParameterSetModified = 0;
 	}
 	else if (PropertyName == "ProjectionAngle")
 	{
 		DefaultValue.Angle = FMath::DegreesToRadians(ProjectionAngle);
-		ParameterSetModified = 1;
 	}
 	else if (PropertyName == "ProjectorBone")
 	{
 		DefaultValue.Position = (FVector3f)BoneComboBoxLocation;
 		DefaultValue.Direction = (FVector3f)BoneComboBoxForwardDirection;
 		DefaultValue.Up = (FVector3f)BoneComboBoxUpDirection;
-		ParameterSetModified = 2;
 	}
 
 	Super::PostEditChangeProperty(PropertyChangedEvent);
@@ -46,6 +41,72 @@ void UCustomizableObjectNodeProjectorParameter::AllocateDefaultPins(UCustomizabl
 }
 
 
+ECustomizableObjectProjectorType UCustomizableObjectNodeProjectorParameter::GetProjectorType() const
+{
+	return ProjectionType;
+}
+
+
+FVector UCustomizableObjectNodeProjectorParameter::GetProjectorDefaultPosition() const
+{
+	return static_cast<FVector>(DefaultValue.Position);
+}
+
+
+void UCustomizableObjectNodeProjectorParameter::SetProjectorDefaultPosition(const FVector& Position)
+{
+	DefaultValue.Position = static_cast<FVector3f>(Position);
+}
+
+
+FVector UCustomizableObjectNodeProjectorParameter::GetProjectorDefaultDirection() const
+{
+	return static_cast<FVector>(DefaultValue.Direction);
+}
+
+
+void UCustomizableObjectNodeProjectorParameter::SetProjectorDefaultDirection(const FVector& Direction)
+{
+	DefaultValue.Direction = static_cast<FVector3f>(Direction);
+}
+
+
+FVector UCustomizableObjectNodeProjectorParameter::GetProjectorDefaultUp() const
+{
+	return static_cast<FVector>(DefaultValue.Up);
+}
+
+
+void UCustomizableObjectNodeProjectorParameter::SetProjectorDefaultUp(const FVector& Up)
+{
+	DefaultValue.Up = static_cast<FVector3f>(Up);
+}
+
+
+FVector UCustomizableObjectNodeProjectorParameter::GetProjectorDefaultScale() const
+{
+	return static_cast<FVector>(DefaultValue.Scale);
+}
+
+
+void UCustomizableObjectNodeProjectorParameter::SetProjectorDefaultScale(const FVector& Scale)
+{
+	DefaultValue.Scale = static_cast<FVector3f>(Scale);
+}
+
+
+float UCustomizableObjectNodeProjectorParameter::GetProjectorDefaultAngle() const
+{
+	return ProjectionAngle;
+}
+
+
+void UCustomizableObjectNodeProjectorParameter::SetProjectorDefaultAngle(float Angle)
+{
+	ProjectionAngle = Angle;
+}
+
+
 FText UCustomizableObjectNodeProjectorParameter::GetNodeTitle(ENodeTitleType::Type TitleType) const
 {
 	return LOCTEXT("Projector_Parameter", "Projector Parameter");
diff --git a/Engine/Plugins/Experimental/Mutable/Source/CustomizableObjectEditor/Private/MuCOE/Nodes/CustomizableObjectNodeProjectorParameterDetails.cpp b/Engine/Plugins/Experimental/Mutable/Source/CustomizableObjectEditor/Private/MuCOE/Nodes/CustomizableObjectNodeProjectorParameterDetails.cpp
--- a/Engine/Plugins/Experimental/Mutable/Source/CustomizableObjectEditor/Private/MuCOE/Nodes/CustomizableObjectNodeProjectorParameterDetails.cpp
+++ b/Engine/Plugins/Experimental/Mutable/Source/CustomizableObjectEditor/Private/MuCOE/Nodes/CustomizableObjectNodeProjectorParameterDetails.cpp
@@ -6,6 +6,7 @@
 #include "Engine/SkeletalMesh.h"
 #include "HAL/PlatformApplicationMisc.h"
 #include "IDetailsView.h"
+#include "ScopedTransaction.h"
 #include "MuCOE/CustomizableObjectEditorUtilities.h"
 #include "MuCOE/CustomizableObjectGraph.h"
 #include "MuCOE/ICustomizableObjectEditor.h"
@@ -195,6 +196,8 @@ FReply FCustomizableObjectNodeProjectorParameterDetails::OnProjectorCopyPressed(
 
 FReply FCustomizableObjectNodeProjectorParameterDetails::OnProjectorPastePressed()
 {
+	FScopedTransaction Transaction(LOCTEXT("PasteTransform", "Paste Transform"));
+	
 	FString ClipText;
 
 	FPlatformApplicationMisc::ClipboardPaste(ClipText);
@@ -203,16 +206,18 @@ FReply FCustomizableObjectNodeProjectorParameterDetails::OnProjectorPastePressed
 
 	if (NodeParameter)
 	{
+		NodeParameter->Modify();
+	
 		UScriptStruct* Struct = NodeParameter->DefaultValue.StaticStruct();
 		Struct->ImportText(*ClipText, &NodeParameter->DefaultValue, nullptr, 0, GLog, GetPathNameSafe(Struct));
-		NodeParameter->ParameterSetModified = 2;
 		Editor = NodeParameter->GetGraphEditor();
 	}
 	else if (NodeConstant)
 	{
+		NodeConstant->Modify();
+		
 		UScriptStruct* Struct = NodeConstant->Value.StaticStruct();
 		Struct->ImportText(*ClipText, &NodeConstant->Value, nullptr, 0, GLog, GetPathNameSafe(Struct));
-		NodeConstant->ParameterSetModified = 2;
 		Editor = NodeConstant->GetGraphEditor();
 	}
 
diff --git a/Engine/Plugins/Experimental/Mutable/Source/CustomizableObjectEditor/Private/MuCOE/SCustomizableInstanceProperties.cpp b/Engine/Plugins/Experimental/Mutable/Source/CustomizableObjectEditor/Private/MuCOE/SCustomizableInstanceProperties.cpp
--- a/Engine/Plugins/Experimental/Mutable/Source/CustomizableObjectEditor/Private/MuCOE/SCustomizableInstanceProperties.cpp
+++ b/Engine/Plugins/Experimental/Mutable/Source/CustomizableObjectEditor/Private/MuCOE/SCustomizableInstanceProperties.cpp
@@ -3,7 +3,9 @@
 #include "MuCOE/SCustomizableInstanceProperties.h"
 
 #include "ContentBrowserModule.h"
+#include "CustomizableObjectInstanceEditor.h"
 #include "Editor.h"
+#include "ScopedTransaction.h"
 #include "HAL/PlatformApplicationMisc.h"
 #include "Misc/Paths.h"
 #include "MuCO/CustomizableObject.h"
@@ -13,8 +15,10 @@
 #include "MuCOE/CustomizableObjectEditorUtilities.h"
 #include "MuCOE/SMutableTextSearchBox.h"
 #include "MuCOE/UnrealEditorPortabilityHelpers.h"
+#include "MuCOE/SCustomizableObjectEditorViewport.h"
 #include "Serialization/BufferArchive.h"
 #include "Slate/DeferredCleanupSlateBrush.h"
+#include "Toolkits/ToolkitManager.h"
 #include "Widgets/Colors/SColorBlock.h"
 #include "Widgets/Colors/SColorPicker.h"
 #include "Widgets/Input/SButton.h"
@@ -35,9 +39,17 @@ void SCustomizableInstanceProperties::Construct(const FArguments& InArgs)
 	check(CustomInstance.IsValid());
 	
 	InstanceDetails = InArgs._InstanceDetails;
+
+	TSharedPtr<IToolkit> FoundAssetEditor = FToolkitManager::Get().FindEditorForAsset(InArgs._CustomInstance); // Tab spawned in a COEInstanceEditor
+	if (!FoundAssetEditor)
+	{
+		FoundAssetEditor = FToolkitManager::Get().FindEditorForAsset(CustomInstance->GetCustomizableObject()); // Tab spawned in a COEditor
+	}
+	check(FoundAssetEditor);
+
+	WeakEditor = StaticCastSharedPtr<ICustomizableObjectInstanceEditor>(FoundAssetEditor).ToWeakPtr();
 	
 	NoInstanceMessage = LOCTEXT("Model not compiled", "Model not compiled");
-	SimplifiedProjectorInterface = false;
 	
 	if (CustomInstance->IsSelectedParameterProfileDirty())
 	{
@@ -387,12 +399,6 @@ void SCustomizableInstanceProperties::SetNoInstanceMessage(const FText& InNoInst
 }
 
 
-void SCustomizableInstanceProperties::SetSimplifiedProjectorInterface(bool Enabled)
-{
-	SimplifiedProjectorInterface = Enabled;
-}
-
-
 const UCustomizableObjectInstance* SCustomizableInstanceProperties::GetInstance() const
 {
 	return CustomInstance.Get();
@@ -444,9 +450,6 @@ void SCustomizableInstanceProperties::ResetParamBox()
 	ParamHasParent.Empty();
 	ParamNameToExpandableAreaMap.Empty();
 
-	MapProjectorButton.Empty();
-	MapProjectorButtonTextBlock.Empty();
-
 	if (UCustomizableObject* CustomizableObject = CustomInstance->GetCustomizableObject())
 	{
 		if (CustomInstance->bShowOnlyRuntimeParameters)
@@ -547,6 +550,8 @@ void SCustomizableInstanceProperties::AddParameter(int32 ParamIndexInObject)
 {
 	UCustomizableObject* CustomizableObject = CustomInstance->GetCustomizableObject();
 	check(CustomizableObject);
+
+	TSharedPtr<ICustomizableObjectInstanceEditor> Editor = GetEditorChecked();
 	
 	FString ParamName = CustomizableObject->GetParameterName(ParamIndexInObject);
 
@@ -903,427 +908,353 @@ void SCustomizableInstanceProperties::AddParameter(int32 ParamIndexInObject)
 
 	case EMutableParameterType::Projector:
 	{
-		if (SimplifiedProjectorInterface)
+		bool bIsParamMultidimensional = CustomInstance->GetCustomizableObject()->IsParameterMultidimensional(ParamIndexInObject);
+
+		if (!bIsParamMultidimensional)
 		{
+			const UProjectorParameter* ProjectorParameter = Editor->GetProjectorParameter();
+			const bool bSelectedProjector = ProjectorParameter->IsProjectorSelected(ParamName);
+
+			TSharedPtr<SButton> Button;
+			TSharedPtr<STextBlock> TextBlock;
+
 			ParameterBox->AddSlot()
 			.AutoWidth()
-			.FillWidth(1.0f)
-			.Padding(2.0f)
 			[
 				SNew(SHorizontalBox)
 				+ SHorizontalBox::Slot()
 				.AutoWidth()
+				.VAlign(VAlign_Center)
 				[
-					SNew(SCheckBox)
-					.Type(ESlateCheckBoxType::ToggleButton)
-					.ToolTipText(LOCTEXT("Hide_Projector", "Hide projector"))
-					.IsChecked(this, &SCustomizableInstanceProperties::IsProjectorHidden, ParamName)
+					SAssignNew(Button, SButton)
+					.OnClicked(this, &SCustomizableInstanceProperties::OnProjectorSelectChanged, ParamName, -1)
+					.HAlign(HAlign_Center)
+					.Content()
 					[
-						SNew(SImage)
-						// todo: editor style not available in game
-						//.Image(UE_MUTABLE_GET_BRUSH(TEXT("LevelEditor.ViewOptions.Small")))
+						SNew(STextBlock)
+						.Text(bSelectedProjector ? LOCTEXT("Unselect Projector", "Unselect Projector") : LOCTEXT("Select Projector", "Select Projector"))
+						.ToolTipText(bSelectedProjector ? LOCTEXT("Unselect Projector", "Unselect Projector") : LOCTEXT("Select Projector", "Select Projector"))
 					]
 				]
 			
+				+SHorizontalBox::Slot()
+				.VAlign(VAlign_Center)
+				.AutoWidth()
+				[
+					SNew(SButton)
+					.HAlign(HAlign_Center)
+					.ToolTipText(LOCTEXT("Copy Transform", "Copy Transform"))
+					.Text(LOCTEXT("Copy Transform", "Copy Transform"))
+					.OnClicked(this, &SCustomizableInstanceProperties::OnProjectorCopyTransform, ParamName, -1)
+				]
+				
 				+ SHorizontalBox::Slot()
+				.VAlign(VAlign_Center)
 				.AutoWidth()
 				[
-					SNew(SCheckBox)
-					.Type(ESlateCheckBoxType::ToggleButton)
-					.ToolTipText(LOCTEXT("Show_Projector", "Show projector widget"))
-					.IsChecked(this, &SCustomizableInstanceProperties::IsProjectorTranslate, ParamName)
-					.OnCheckStateChanged(this, &SCustomizableInstanceProperties::OnProjectorTranslateStateChanged, ParamName)
-					.Padding(2.f)
-					[
-						SNew(SImage)
-						// todo: editor style not available in game
-						// .Image(UE_MUTABLE_GET_BRUSH(TEXT("LevelEditor.TranslateMode.Small")))
-					]
+					SNew(SButton)
+					.HAlign(HAlign_Center)
+					.ToolTipText(LOCTEXT("Paste Transform", "Paste Transform"))
+					.Text(LOCTEXT("Paste Transform", "Paste Transform"))
+					.OnClicked(this, &SCustomizableInstanceProperties::OnProjectorPasteTransform, ParamName, -1)
+				]
+				
+				+ SHorizontalBox::Slot()
+				.VAlign(VAlign_Center)
+				.AutoWidth()
+				[
+					SNew(SButton)
+					.HAlign(HAlign_Center)
+					.ToolTipText(LOCTEXT("Reset Transform", "Reset Transform"))
+					.Text(LOCTEXT("Reset Transform", "Reset Transform"))
+					.OnClicked(this, &SCustomizableInstanceProperties::OnProjectorResetTransform, ParamName, -1)
 				]
 			];
+
+			Button->SetBorderBackgroundColor(bSelectedProjector ? FLinearColor::Green : FLinearColor::White);
 		}
 		else
 		{
-			bool bIsParamMultidimensional = CustomInstance->GetCustomizableObject()->IsParameterMultidimensional(ParamIndexInObject);
-
-			if (!bIsParamMultidimensional)
-			{
-				FString ParamNameWithIndex = ParamName + FString::Printf(TEXT("__%d"), -1);
+			TArray<FCustomizableObjectProjectorParameterValue>& ProjectorParameters = CustomInstance->GetProjectorParameters();
+			const int32 ProjectorParamIndex = CustomInstance->FindProjectorParameterNameIndex(ParamName);
+			check(ProjectorParamIndex < ProjectorParameters.Num());
 
-				TSharedPtr<SButton> Button;
-				TSharedPtr<STextBlock> TextBlock;
+			// Selected Pose UI
+			const FString PoseSwitchEnumParamName = ParamName + FMultilayerProjector::POSE_PARAMETER_POSTFIX;
+			const int32 PoseSwitchEnumParamIndexInObject = CustomizableObject->FindParameter(PoseSwitchEnumParamName);
 
-				ParameterBox->AddSlot()
-				.AutoWidth()
-				[
-					SNew(SHorizontalBox)
-					+ SHorizontalBox::Slot()
-					.AutoWidth()
-					.VAlign(VAlign_Center)
-					[
-						SAssignNew(Button, SButton)
-						.OnClicked(this, &SCustomizableInstanceProperties::OnProjectorSelectStateChanged, ParamNameWithIndex)
-						.HAlign(HAlign_Center)
-						.Content()
-						[
-							SAssignNew(TextBlock, STextBlock)
-							.Text(LOCTEXT("Select Projector", "Select Projector"))
-							.ToolTipText(LOCTEXT("Select Projector", "Select Projector"))
-						]
-					]
-				
-					+SHorizontalBox::Slot()
-					.VAlign(VAlign_Center)
-					.AutoWidth()
-					[
-						SNew(SButton)
-						.HAlign(HAlign_Center)
-						.ToolTipText(LOCTEXT("Copy Transform", "Copy Transform"))
-						.Text(LOCTEXT("Copy Transform", "Copy Transform"))
-						.OnClicked(this, &SCustomizableInstanceProperties::OnProjectorCopyTransform, ParamNameWithIndex)
-					]
-					
-					+ SHorizontalBox::Slot()
-					.VAlign(VAlign_Center)
-					.AutoWidth()
-					[
-						SNew(SButton)
-						.HAlign(HAlign_Center)
-						.ToolTipText(LOCTEXT("Paste Transform", "Paste Transform"))
-						.Text(LOCTEXT("Paste Transform", "Paste Transform"))
-						.OnClicked(this, &SCustomizableInstanceProperties::OnProjectorPasteTransform, ParamNameWithIndex)
-					]
-					
-					+ SHorizontalBox::Slot()
-					.VAlign(VAlign_Center)
-					.AutoWidth()
-					[
-						SNew(SButton)
-						.HAlign(HAlign_Center)
-						.ToolTipText(LOCTEXT("Reset Transform", "Reset Transform"))
-						.Text(LOCTEXT("Reset Transform", "Reset Transform"))
-						.OnClicked(this, &SCustomizableInstanceProperties::OnProjectorResetTransform, ParamNameWithIndex)
-					]
-				];
-				
-				ArrayProjectorParameterNameWithIndex.Add(ParamNameWithIndex);
+			if (PoseSwitchEnumParamIndexInObject != INDEX_NONE)
+			{
+				const int32 NumPoseValues = CustomizableObject->GetIntParameterNumOptions(PoseSwitchEnumParamIndexInObject);
 
-				MapProjectorButton.Add(ParamNameWithIndex, Button);
-				MapProjectorButtonTextBlock.Add(ParamNameWithIndex, TextBlock);
+				TArray<FString> PoseOptionNamesAttribute;
+				FString PoseValue = GetIntParameterValue(PoseSwitchEnumParamName, -1);
+				int32 PoseValueIndex = 0;
 
-				if (!CustomInstance->LastSelectedProjectorParameter.IsEmpty() && (CustomInstance->LastSelectedProjectorParameter == ParamName))
+				for (int32 j = 0; j < NumPoseValues; ++j)
 				{
-					if (Button.Get() != nullptr)
+					const FString PossibleValue = CustomizableObject->GetIntParameterAvailableOption(PoseSwitchEnumParamIndexInObject, j);
+					if (PossibleValue == PoseValue)
 					{
-						Button->SetBorderBackgroundColor(FLinearColor::Green);
+						PoseValueIndex = j;
 					}
 
-					if (TextBlock.Get() != nullptr)
-					{
-						TextBlock->SetText(LOCTEXT("Unselect Projector", "Unselect Projector"));
-					}
+					PoseOptionNamesAttribute.Add(PossibleValue);
 				}
-			}
-			else
-			{
-				TArray<FCustomizableObjectProjectorParameterValue>& ProjectorParameters = CustomInstance->GetProjectorParameters();
-				const int32 ProjectorParamIndex = CustomInstance->FindProjectorParameterNameIndex(ParamName);
-				check(ProjectorParamIndex < ProjectorParameters.Num());
 
-				// Selected Pose UI
-				FString PoseSwitchEnumParamName = ParamName + FMultilayerProjector::POSE_PARAMETER_POSTFIX;
-				FString PoseSwitchEnumParamNameWithRange = PoseSwitchEnumParamName + FString::Printf(TEXT("__%d"), -1);
-				const int32 PoseSwitchEnumParamIndexInObject = CustomizableObject->FindParameter(PoseSwitchEnumParamName);
+				ParameterBox->AddSlot()
+					.HAlign(HAlign_Right)
+					.VAlign(VAlign_Fill)
+					.FillWidth(10.f)
+					[
+						SNew(SHorizontalBox)
+						+ SHorizontalBox::Slot()
+							.HAlign(HAlign_Fill)
+							.VAlign(VAlign_Center)
+							.FillWidth(0.45f)
+							[
+								SNew(SMutableTextSearchBox)
+								.PossibleSuggestions(PoseOptionNamesAttribute)
+								.InitialText(FText::FromString(PoseOptionNamesAttribute[PoseValueIndex]))
+								.MustMatchPossibleSuggestions(TAttribute<bool>(true))
+								.SuggestionListPlacement(EMenuPlacement::MenuPlacement_ComboBox)
+								.OnTextCommitted(this, &SCustomizableInstanceProperties::OnProjectorTextureParameterComboBoxChanged, PoseSwitchEnumParamName, -1)
+								.ToolTipText(LOCTEXT("Pose selector tooltip", "Select the skeletal mesh pose used for projection. This does not control the actual visual mesh pose in the viewport (or during gameplay for that matter). It has to be manually set. You can drag&drop a pose onto the preview viewport."))
+							]
 
-				if (PoseSwitchEnumParamIndexInObject != INDEX_NONE)
-				{
-					const int32 NumPoseValues = CustomizableObject->GetIntParameterNumOptions(PoseSwitchEnumParamIndexInObject);
+						+ SHorizontalBox::Slot()
+							.HAlign(HAlign_Fill)
+							.VAlign(VAlign_Center)
+							.FillWidth(0.3f)
+							[
+								SNew(SButton)
+								.ToolTipText(LOCTEXT("Add Layer", "Add Layer"))
+								.Text(LOCTEXT("Add Layer", "Add Layer"))
+								.OnClicked(this, &SCustomizableInstanceProperties::OnProjectorLayerAdded, ParamName)
+								.HAlign(HAlign_Fill)
+							]
+					];
 
-					TArray<FString> PoseOptionNamesAttribute;
-					FString PoseValue = GetIntParameterValue(PoseSwitchEnumParamName, -1);
-					int32 PoseValueIndex = 0;
+			}
+			else
+			{
+				ParameterBox->AddSlot()
+					.HAlign(HAlign_Right)
+					.VAlign(VAlign_Fill)
+					.FillWidth(10.f)
+					[
+						SNew(SButton)
+						.ToolTipText(LOCTEXT("Add Layer", "Add Layer"))
+						.Text(LOCTEXT("Add Layer", "Add Layer"))
+						.OnClicked(this, &SCustomizableInstanceProperties::OnProjectorLayerAdded, ParamName)
+						.HAlign(HAlign_Fill)
+					];
+			}
 
-					for (int32 j = 0; j < NumPoseValues; ++j)
-					{
-						const FString PossibleValue = CustomizableObject->GetIntParameterAvailableOption(PoseSwitchEnumParamIndexInObject, j);
-						if (PossibleValue == PoseValue)
-						{
-							PoseValueIndex = j;
-						}
+			const FString TextureSwitchEnumParamName = ParamName + FMultilayerProjector::IMAGE_PARAMETER_POSTFIX;
+			
+			for (int32 RangeIndex = 0; RangeIndex < ProjectorParameters[ProjectorParamIndex].RangeValues.Num(); ++RangeIndex)
+			{
+				const int32 TextureSwitchEnumParamIndexInObject = CustomizableObject->FindParameter(TextureSwitchEnumParamName);
+				check(TextureSwitchEnumParamIndexInObject >= 0);
+				int32 NumValues = CustomizableObject->GetIntParameterNumOptions(TextureSwitchEnumParamIndexInObject);
 
-						PoseOptionNamesAttribute.Add(PossibleValue);
-					}
+				FString OpacitySliderParamName = ParamName + FMultilayerProjector::OPACITY_PARAMETER_POSTFIX;
+				FString OpacitySliderParamNameWithRange = OpacitySliderParamName + FString::Printf(TEXT("__%d"), RangeIndex);
 
-					ParameterBox->AddSlot()
-						.HAlign(HAlign_Right)
-						.VAlign(VAlign_Fill)
-						.FillWidth(10.f)
-						[
-							SNew(SHorizontalBox)
-							+ SHorizontalBox::Slot()
-								.HAlign(HAlign_Fill)
-								.VAlign(VAlign_Center)
-								.FillWidth(0.45f)
-								[
-									SNew(SMutableTextSearchBox)
-									.PossibleSuggestions(PoseOptionNamesAttribute)
-									.InitialText(FText::FromString(PoseOptionNamesAttribute[PoseValueIndex]))
-									.MustMatchPossibleSuggestions(TAttribute<bool>(true))
-									.SuggestionListPlacement(EMenuPlacement::MenuPlacement_ComboBox)
-									.OnTextCommitted(this, &SCustomizableInstanceProperties::OnProjectorTextureParameterComboBoxChanged, PoseSwitchEnumParamNameWithRange)
-									.ToolTipText(LOCTEXT("Pose selector tooltip", "Select the skeletal mesh pose used for projection. This does not control the actual visual mesh pose in the viewport (or during gameplay for that matter). It has to be manually set. You can drag&drop a pose onto the preview viewport."))
-								]
+				TArray<FString> OptionNamesAttribute;
+				FString Value = GetIntParameterValue(TextureSwitchEnumParamName, RangeIndex);
+				int32 ValueIndex = 0;
 
-							+ SHorizontalBox::Slot()
-								.HAlign(HAlign_Fill)
-								.VAlign(VAlign_Center)
-								.FillWidth(0.3f)
-								[
-									SNew(SButton)
-									.ToolTipText(LOCTEXT("Add Layer", "Add Layer"))
-									.Text(LOCTEXT("Add Layer", "Add Layer"))
-									.OnClicked(this, &SCustomizableInstanceProperties::OnProjectorLayerAdded, ParamName)
-									.HAlign(HAlign_Fill)
-								]
-						];
+				const UProjectorParameter* ProjectorParameter = Editor->GetProjectorParameter();
+				const bool bSelectedProjector = ProjectorParameter->IsProjectorSelected(ParamName, RangeIndex);
 
-				}
-				else
+				for (int32 CandidateIndex = 0; CandidateIndex < NumValues; ++CandidateIndex)
 				{
-					ParameterBox->AddSlot()
-						.HAlign(HAlign_Right)
-						.VAlign(VAlign_Fill)
-						.FillWidth(10.f)
-						[
-							SNew(SButton)
-							.ToolTipText(LOCTEXT("Add Layer", "Add Layer"))
-							.Text(LOCTEXT("Add Layer", "Add Layer"))
-							.OnClicked(this, &SCustomizableInstanceProperties::OnProjectorLayerAdded, ParamName)
-							.HAlign(HAlign_Fill)
-						];
+					FString PossibleValue = CustomizableObject->GetIntParameterAvailableOption(TextureSwitchEnumParamIndexInObject, CandidateIndex);
+					if (PossibleValue == Value)
+					{
+						ValueIndex = CandidateIndex;
+					}
+					OptionNamesAttribute.Add(CustomizableObject->GetIntParameterAvailableOption(TextureSwitchEnumParamIndexInObject, CandidateIndex));
 				}
-
-				const FString TextureSwitchEnumParamName = ParamName + FMultilayerProjector::IMAGE_PARAMETER_POSTFIX;
 				
-				for (int32 i = 0; i < ProjectorParameters[ProjectorParamIndex].RangeValues.Num(); ++i)
-				{
-					const FString ParamNameWithIndex = ParamName + FString::Printf(TEXT("__%d"), i);
-
-					const FString TextureSwitchEnumParamNameWithRange = TextureSwitchEnumParamName + FString::Printf(TEXT("__%d"), i);
-					const int32 TextureSwitchEnumParamIndexInObject = CustomizableObject->FindParameter(TextureSwitchEnumParamName);
-					check(TextureSwitchEnumParamIndexInObject >= 0);
-					int32 NumValues = CustomizableObject->GetIntParameterNumOptions(TextureSwitchEnumParamIndexInObject);
-
-					FString OpacitySliderParamName = ParamName + FMultilayerProjector::OPACITY_PARAMETER_POSTFIX;
-					FString OpacitySliderParamNameWithRange = OpacitySliderParamName + FString::Printf(TEXT("__%d"), i);
+				TSharedPtr<SHorizontalBox> SliderBox;
+				TSharedPtr<SSpinBox<float>> Slider;
+				TSharedPtr<SHorizontalBox> PropertiesHB;
+				int32 SliderIndex = FloatSliders.Num();
 
-					TArray<FString> OptionNamesAttribute;
-					FString Value = GetIntParameterValue(TextureSwitchEnumParamName, i);
-					int32 ValueIndex = 0;
+				TSharedPtr<SButton> Button;
+				TSharedPtr<STextBlock> TextBlock;
 
-					for (int32 CandidateIndex = 0; CandidateIndex < NumValues; ++CandidateIndex)
-					{
-						FString PossibleValue = CustomizableObject->GetIntParameterAvailableOption(TextureSwitchEnumParamIndexInObject, CandidateIndex);
-						if (PossibleValue == Value)
-						{
-							ValueIndex = CandidateIndex;
-						}
-						OptionNamesAttribute.Add(CustomizableObject->GetIntParameterAvailableOption(TextureSwitchEnumParamIndexInObject, CandidateIndex));
-					}
+				//Horizontal box that owns all the hidden layer properties
+				ActualParamBox->AddSlot()
+				.HAlign(HAlign_Fill)
+				.Padding(0, 10.f)
+				[
+					SAssignNew(PropertiesHB, SHorizontalBox)
 					
-					TSharedPtr<SHorizontalBox> SliderBox;
-					TSharedPtr<SSpinBox<float>> Slider;
-					TSharedPtr<SHorizontalBox> PropertiesHB;
-					int32 SliderIndex = FloatSliders.Num();
-
-					TSharedPtr<SButton> Button;
-					TSharedPtr<STextBlock> TextBlock;
-
-					//Horizontal box that owns all the hidden layer properties
-					ActualParamBox->AddSlot()
-					.HAlign(HAlign_Fill)
-					.Padding(0, 10.f)
+					//Select Projector slot
+					+ SHorizontalBox::Slot()
+					.Padding(1, 0)
+					.AutoWidth()
 					[
-						SAssignNew(PropertiesHB, SHorizontalBox)
-						
-						//Select Projector slot
+						SNew(SHorizontalBox)
 						+ SHorizontalBox::Slot()
 						.Padding(1, 0)
-						.AutoWidth()
 						[
-							SNew(SHorizontalBox)
-							+ SHorizontalBox::Slot()
-							.Padding(1, 0)
+							SNew(SBox)
+							.MinDesiredWidth(115.f)
+							.MaxDesiredWidth(115.f)
 							[
-								SNew(SBox)
-								.MinDesiredWidth(115.f)
-								.MaxDesiredWidth(115.f)
+								SAssignNew(Button, SButton)
+								.OnClicked(this, &SCustomizableInstanceProperties::OnProjectorSelectChanged, ParamName, RangeIndex)
+								.VAlign(VAlign_Center)
+								.HAlign(HAlign_Center)
+								.Content()
 								[
-									SAssignNew(Button, SButton)
-									.OnClicked(this, &SCustomizableInstanceProperties::OnProjectorSelectStateChanged, ParamNameWithIndex)
-									.VAlign(VAlign_Center)
-									.HAlign(HAlign_Center)
-									.Content()
-									[
-										SAssignNew(TextBlock, STextBlock)
-										.Text(LOCTEXT("Select Projector", "Select Projector"))
-										.ToolTipText(LOCTEXT("Select Projector", "Select Projector"))
-										.Justification(ETextJustify::Center)
-										//.AutoWrapText(true)
-									]
+									SNew(STextBlock)
+									.Text(bSelectedProjector ? LOCTEXT("Unselect Projector", "Unselect Projector") : LOCTEXT("Select Projector", "Select Projector"))
+									.ToolTipText(bSelectedProjector ? LOCTEXT("Unselect Projector", "Unselect Projector") : LOCTEXT("Select Projector", "Select Projector"))
+									.Justification(ETextJustify::Center)
+									//.AutoWrapText(true)
 								]
 							]
+						]
 
-							+ SHorizontalBox::Slot()
-							.Padding(1, 0)
+						+ SHorizontalBox::Slot()
+						.Padding(1, 0)
+						[
+							SNew(SBox)
+							.MinDesiredWidth(120.f)
 							[
-								SNew(SBox)
-								.MinDesiredWidth(120.f)
-								[
-									SNew(SMutableTextSearchBox)
-									.PossibleSuggestions(OptionNamesAttribute)
-									.InitialText(FText::FromString(OptionNamesAttribute[ValueIndex]))
-									.MustMatchPossibleSuggestions(TAttribute<bool>(true))
-									.SuggestionListPlacement(EMenuPlacement::MenuPlacement_ComboBox)
-									.OnTextCommitted(this, &SCustomizableInstanceProperties::OnProjectorTextureParameterComboBoxChanged, TextureSwitchEnumParamNameWithRange)
-								]
+								SNew(SMutableTextSearchBox)
+								.PossibleSuggestions(OptionNamesAttribute)
+								.InitialText(FText::FromString(OptionNamesAttribute[ValueIndex]))
+								.MustMatchPossibleSuggestions(TAttribute<bool>(true))
+								.SuggestionListPlacement(EMenuPlacement::MenuPlacement_ComboBox)
+								.OnTextCommitted(this, &SCustomizableInstanceProperties::OnProjectorTextureParameterComboBoxChanged, TextureSwitchEnumParamName, RangeIndex)
 							]
 						]
-					];
-						
-					MapProjectorButton.Add(ParamNameWithIndex, Button);
-					MapProjectorButtonTextBlock.Add(ParamNameWithIndex, TextBlock);
+					]
+				];
+
+				Button->SetBorderBackgroundColor(bSelectedProjector ? FLinearColor::Green : FLinearColor::White);
 
-					PropertiesHB->AddSlot()
+				PropertiesHB->AddSlot()
+				.HAlign(HAlign_Fill)
+				[
+					SNew(SHorizontalBox)
+					+ SHorizontalBox::Slot()
+					.Padding(1, 0)
 					.HAlign(HAlign_Fill)
 					[
 						SNew(SHorizontalBox)
 						+ SHorizontalBox::Slot()
-						.Padding(1, 0)
-						.HAlign(HAlign_Fill)
 						[
-							SNew(SHorizontalBox)
-							+ SHorizontalBox::Slot()
-							[
-								SAssignNew(Slider, SSpinBox<float>)
-								.MinValue(0.0f)
-								.MaxValue(1.0f)
-								.Value(this, &SCustomizableInstanceProperties::GetSliderValue, SliderIndex)
-								.OnValueChanged(this, &SCustomizableInstanceProperties::OnProjectorFloatParameterChanged, SliderIndex, ParamNameWithIndex)
-								.OnBeginSliderMovement(this, &SCustomizableInstanceProperties::OnProjectorFloatParameterSliderBegin)
-								.OnEndSliderMovement(this, &SCustomizableInstanceProperties::OnProjectorFloatParameterSliderEnd)
-							]
-
+							SAssignNew(Slider, SSpinBox<float>)
+							.MinValue(0.0f)
+							.MaxValue(1.0f)
+							.Value(this, &SCustomizableInstanceProperties::GetSliderValue, SliderIndex)
+							.OnValueChanged(this, &SCustomizableInstanceProperties::OnProjectorFloatParameterChanged, SliderIndex)
+							.OnBeginSliderMovement(this, &SCustomizableInstanceProperties::OnProjectorFloatParameterSliderBegin)
+							.OnEndSliderMovement(this, &SCustomizableInstanceProperties::OnProjectorFloatParameterSliderEnd)
 						]
-					];
 
-					PropertiesHB->AddSlot()
-					.HAlign(HAlign_Right)
-					.AutoWidth()
+					]
+				];
+
+				PropertiesHB->AddSlot()
+				.HAlign(HAlign_Right)
+				.AutoWidth()
+				[
+					SNew(SHorizontalBox)
+					+ SHorizontalBox::Slot()
+					.Padding(1, 0)
 					[
-						SNew(SHorizontalBox)
-						+ SHorizontalBox::Slot()
-						.Padding(1, 0)
+						SNew(SBox)
+						.MinDesiredWidth(80.f)
 						[
-							SNew(SBox)
-							.MinDesiredWidth(80.f)
+							SNew(SButton)
+							.OnClicked(this, &SCustomizableInstanceProperties::OnProjectorCopyTransform, ParamName, RangeIndex)
+							.VAlign(VAlign_Center)
+							.HAlign(HAlign_Center)
+							.Content()
 							[
-								SNew(SButton)
-								.OnClicked(this, &SCustomizableInstanceProperties::OnProjectorCopyTransform, ParamNameWithIndex)
-								.VAlign(VAlign_Center)
-								.HAlign(HAlign_Center)
-								.Content()
-								[
-									SNew(STextBlock)
-									.ToolTipText(LOCTEXT("Copy Transform", "Copy Transform"))
-									.Text(LOCTEXT("Copy Transform", "Copy Transform"))
-									.AutoWrapText(true)
-									.Justification(ETextJustify::Center)
-								]
+								SNew(STextBlock)
+								.ToolTipText(LOCTEXT("Copy Transform", "Copy Transform"))
+								.Text(LOCTEXT("Copy Transform", "Copy Transform"))
+								.AutoWrapText(true)
+								.Justification(ETextJustify::Center)
 							]
 						]
+					]
 
-						+ SHorizontalBox::Slot()
-						.Padding(1, 0)
+					+ SHorizontalBox::Slot()
+					.Padding(1, 0)
+					[
+						SNew(SBox)
+						.MinDesiredWidth(80.f)
 						[
-							SNew(SBox)
-							.MinDesiredWidth(80.f)
+							SNew(SButton)
+							.OnClicked(this, &SCustomizableInstanceProperties::OnProjectorPasteTransform, ParamName, RangeIndex)
+							.VAlign(VAlign_Center)
+							.HAlign(HAlign_Center)
+							.Content()
 							[
-								SNew(SButton)
-								.OnClicked(this, &SCustomizableInstanceProperties::OnProjectorPasteTransform, ParamNameWithIndex)
-								.VAlign(VAlign_Center)
-								.HAlign(HAlign_Center)
-								.Content()
-								[
-									SNew(STextBlock)
-									.ToolTipText(LOCTEXT("Paste Transform", "Paste Transform"))
-									.Text(LOCTEXT("Paste Transform", "Paste Transform"))
-									.Justification(ETextJustify::Center)
-									.AutoWrapText(true)
-								]
+								SNew(STextBlock)
+								.ToolTipText(LOCTEXT("Paste Transform", "Paste Transform"))
+								.Text(LOCTEXT("Paste Transform", "Paste Transform"))
+								.Justification(ETextJustify::Center)
+								.AutoWrapText(true)
 							]
 						]
+					]
 
-						+ SHorizontalBox::Slot()
-						.Padding(1, 0)
+					+ SHorizontalBox::Slot()
+					.Padding(1, 0)
+					[
+						SNew(SBox)
+						.MinDesiredWidth(80.f)
 						[
-							SNew(SBox)
-							.MinDesiredWidth(80.f)
+							SNew(SButton)
+							.OnClicked(this, &SCustomizableInstanceProperties::OnProjectorResetTransform, ParamName, RangeIndex)
+							.VAlign(VAlign_Center)
+							.HAlign(HAlign_Center)
+							.Content()
 							[
-								SNew(SButton)
-								.OnClicked(this, &SCustomizableInstanceProperties::OnProjectorResetTransform, ParamNameWithIndex)
-								.VAlign(VAlign_Center)
-								.HAlign(HAlign_Center)
-								.Content()
-								[
-									SNew(STextBlock)
-									.ToolTipText(LOCTEXT("Reset Transform", "Reset Transform"))
-									.Text(LOCTEXT("Reset Transform", "Reset Transform"))
-									.Justification(ETextJustify::Center)
-									.AutoWrapText(true)
-								]
+								SNew(STextBlock)
+								.ToolTipText(LOCTEXT("Reset Transform", "Reset Transform"))
+								.Text(LOCTEXT("Reset Transform", "Reset Transform"))
+								.Justification(ETextJustify::Center)
+								.AutoWrapText(true)
 							]
 						]
+					]
 
-						+ SHorizontalBox::Slot()
-						.Padding(1, 0)
+					+ SHorizontalBox::Slot()
+					.Padding(1, 0)
+					[
+						SNew(SBox)
+						.MinDesiredWidth(80.f)
 						[
-							SNew(SBox)
-							.MinDesiredWidth(80.f)
+							SNew(SButton)
+							.OnClicked(this, &SCustomizableInstanceProperties::OnProjectorLayerRemoved, ParamName, RangeIndex)
+							.VAlign(VAlign_Center)
+							.HAlign(HAlign_Center)
+							.Content()
 							[
-								SNew(SButton)
-								.OnClicked(this, &SCustomizableInstanceProperties::OnProjectorLayerRemoved, ParamNameWithIndex)
-								.VAlign(VAlign_Center)
-								.HAlign(HAlign_Center)
-								.Content()
-								[
-									SNew(STextBlock)
-									.ToolTipText(LOCTEXT("Remove Layer", "Remove Layer"))
-									.Text(LOCTEXT("Remove Layer", "Remove Layer"))
-									.Justification(ETextJustify::Center)
-									.AutoWrapText(true)
-								]
+								SNew(STextBlock)
+								.ToolTipText(LOCTEXT("Remove Layer", "Remove Layer"))
+								.Text(LOCTEXT("Remove Layer", "Remove Layer"))
+								.Justification(ETextJustify::Center)
+								.AutoWrapText(true)
 							]
 						]
-					];
-
-					ArrayProjectorParameterNameWithIndex.Add(ParamNameWithIndex);
-
-					if (!CustomInstance->LastSelectedProjectorParameterWithIndex.IsEmpty() && (CustomInstance->LastSelectedProjectorParameterWithIndex == ParamNameWithIndex))
-					{
-						if (Button.Get() != nullptr)
-						{
-							Button->SetBorderBackgroundColor(FLinearColor::Green);
-						}
-
-						if (TextBlock.Get() != nullptr)
-						{
-							TextBlock->SetText(LOCTEXT("Unselect Projector", "Unselect Projector"));
-						}
-					}
+					]
+				];
 
-					FloatSliders.Add(FSliderData(Slider, OpacitySliderParamName, i, GetFloatParameterValue(OpacitySliderParamName, i)));
-				}
+				FloatSliders.Add(FSliderData(Slider, OpacitySliderParamName, RangeIndex, GetFloatParameterValue(OpacitySliderParamName, RangeIndex)));
 			}
 		}
 		break;
@@ -1648,19 +1579,6 @@ void SCustomizableInstanceProperties::OnIntParameterComboBoxChanged(TSharedPtr<F
 
 					if (ValueName == *Selection)
 					{
-						if (ValueName == "None" || CustomInstance->FindProjectorParameterNameIndex(ParamName) == INDEX_NONE)
-						{
-							// Hiding the projector when the selection doesn't have one
-							CustomInstance->ResetProjectorStates();
-							CustomInstance->UnselectProjector = true;
-						}
-						else
-						{
-							// Variable used in a non multilayer projector to not lose the focus of the gizmo
-							CustomInstance->ProjectorLayerChange = true;
-						}
-						
-
 						IntParameters[i].ParameterValueName = ValueName;
 						break;
 					}
@@ -1764,69 +1682,6 @@ void SCustomizableInstanceProperties::OnSetColorFromColorPicker(FLinearColor New
 }
 
 
-FString ParseNameAndIndex(const FString& ParamNameWithIndex, int& OutIndex)
-{
-	FString ParamName, Index;
-	const bool bSplit = ParamNameWithIndex.Split(FString("__"), &ParamName, &Index);
-	check(bSplit && Index.IsNumeric());
-
-	OutIndex = FCString::Atoi(*Index);
-
-	return ParamName;
-}
-
-
-ECheckBoxState SCustomizableInstanceProperties::IsProjectorHidden(FString ParamNameWithIndex) const
-{
-	int32 RangeIndex;
-	const  FString ParamName = ParseNameAndIndex(ParamNameWithIndex, RangeIndex);
-
-	if (CustomInstance->GetProjectorState(ParamName, RangeIndex) == EProjectorState::Hidden)
-	{
-		return ECheckBoxState::Checked;
-	}
-	return ECheckBoxState::Unchecked;
-}
-
-
-ECheckBoxState SCustomizableInstanceProperties::IsProjectorTranslate(FString ParamNameWithIndex) const
-{
-	int32 RangeIndex;
-	const FString ParamName = ParseNameAndIndex(ParamNameWithIndex, RangeIndex);
-
-	if (CustomInstance->GetProjectorState(ParamName, RangeIndex) == EProjectorState::Translate)
-	{
-		return ECheckBoxState::Checked;
-	}
-	return ECheckBoxState::Unchecked;
-}
-
-
-ECheckBoxState SCustomizableInstanceProperties::IsProjectorRotate(FString ParamNameWithIndex) const
-{
-	int32 RangeIndex;
-	const FString ParamName = ParseNameAndIndex(ParamNameWithIndex, RangeIndex);
-
-	if (CustomInstance->GetProjectorState(ParamName, RangeIndex) == EProjectorState::Rotate)
-	{
-		return ECheckBoxState::Checked;
-	}
-	return ECheckBoxState::Unchecked;
-}
-
-
-ECheckBoxState SCustomizableInstanceProperties::IsProjectorScale(FString ParamNameWithIndex) const
-{
-	int32 RangeIndex;
-	const FString ParamName = ParseNameAndIndex(ParamNameWithIndex, RangeIndex);
-
-	if (CustomInstance->GetProjectorState(ParamName, RangeIndex) == EProjectorState::Scale)
-	{
-		return ECheckBoxState::Checked;
-	}
-	return ECheckBoxState::Unchecked;
-}
-
 FReply SCustomizableInstanceProperties::OnProjectorLayerAdded(FString ParamName) const
 {
 	const int32 NumLayers = CustomInstance->AddValueToProjectorRange(ParamName) + 1;
@@ -1864,11 +1719,8 @@ FReply SCustomizableInstanceProperties::OnProjectorLayerAdded(FString ParamName)
 }
 
 
-FReply SCustomizableInstanceProperties::OnProjectorLayerRemoved(FString ParamNameWithIndex) const
+FReply SCustomizableInstanceProperties::OnProjectorLayerRemoved(const FString ParamName, const int32 RangeIndex) const
 {
-	int32 RangeIndex;
-	const FString ParamName = ParseNameAndIndex(ParamNameWithIndex, RangeIndex);
-
 	const int32 projectorParameterIndex = CustomInstance->FindProjectorParameterNameIndex(ParamName);
 	if (projectorParameterIndex == INDEX_NONE
 		|| CustomInstance->GetProjectorParameters()[projectorParameterIndex].RangeValues.Num() <= 0)
@@ -1876,10 +1728,7 @@ FReply SCustomizableInstanceProperties::OnProjectorLayerRemoved(FString ParamNam
 		return FReply::Handled();
 	}
 
-	CustomInstance->SetProjectorState(ParamName, RangeIndex, EProjectorState::Hidden);  // If the Projector was selected, unselect it
-
 	const int32 NumLayers = CustomInstance->RemoveValueFromProjectorRange(ParamName, RangeIndex) + 1;
-	CustomInstance->RemovedProjectorParameterNameWithIndex = ParamNameWithIndex; // Since the gizmo updates the CustomObjectInstance when unselected, a check is required to avoid updating the removed projector parameter range index
 
 	const FString TextureSwitchEnumParamName = ParamName + FMultilayerProjector::IMAGE_PARAMETER_POSTFIX;
 	CustomInstance->RemoveValueFromIntRange(TextureSwitchEnumParamName, RangeIndex);
@@ -1910,110 +1759,28 @@ FReply SCustomizableInstanceProperties::OnProjectorLayerRemoved(FString ParamNam
 }
 
 
-FReply SCustomizableInstanceProperties::OnProjectorSelectStateChanged(FString ParamNameWithIndex)
+FReply SCustomizableInstanceProperties::OnProjectorSelectChanged(const FString ParamName, const int32 RangeIndex) const
 {
-	const int32 MaxIndex = ArrayProjectorParameterNameWithIndex.Num();
-
-	const TSharedPtr<SButton> Button = MapProjectorButton[ParamNameWithIndex];
-	const TSharedPtr<STextBlock> TextBlock = MapProjectorButtonTextBlock[ParamNameWithIndex];
+	const TSharedPtr<ICustomizableObjectInstanceEditor> Editor = GetEditorChecked();
 
-	CustomInstance->ProjectorLayerChange = false;
-
-	for (TPair<FString, TSharedPtr<SButton>>& Element : MapProjectorButton)
-	{
-		if (/*Element.Value && */Element.Value != Button)
-		{
-			Element.Value->SetBorderBackgroundColor(FLinearColor(1.0f, 1.0f, 1.0f, 1.0f));
-		}
-	}
-
-	for (TPair<FString, TSharedPtr<STextBlock>>& Element : MapProjectorButtonTextBlock)
-	{
-		if (/*Element.Value &&*/ Element.Value != TextBlock)
-		{
-			Element.Value->SetText(LOCTEXT("Select Projector", "Select Projector"));
-		}
-	}
-
-	bool SelectProjector = true;
-
-	if ((TextBlock.Get() != nullptr) && (TextBlock.Get()->GetText().ToString() == FString("Unselect Projector")))
+	const UProjectorParameter* ProjectorParameter = Editor->GetProjectorParameter();
+	if (ProjectorParameter->IsProjectorSelected(ParamName, RangeIndex))
 	{
-		SelectProjector = false;
+		Editor->HideGizmoProjectorParameter();
 	}
-
-	if (Button.Get() != nullptr)
-	{
-		Button->SetBorderBackgroundColor(SelectProjector ? FLinearColor::Green : FLinearColor::White);
-	}
-
-	if (TextBlock.Get() != nullptr)
+	else
 	{
-		TextBlock->SetText(SelectProjector ? LOCTEXT("Unselect Projector", "Unselect Projector") : LOCTEXT("Select Projector", "Select Projector"));
+		Editor->ShowGizmoProjectorParameter(ParamName, RangeIndex);		
 	}
-
-	for (int32 i = 0; i < MaxIndex; ++i)
-	{
-		int32 ArrayRangeIndex;
-		FString ArrayParamName = ParseNameAndIndex(ArrayProjectorParameterNameWithIndex[i], ArrayRangeIndex);
-
-		if ((ArrayProjectorParameterNameWithIndex[i] != ParamNameWithIndex) &&
-			(CustomInstance->GetProjectorState(ArrayParamName, ArrayRangeIndex) != EProjectorState::Hidden))
-		{
-			CustomInstance->SetProjectorState(ArrayParamName, ArrayRangeIndex, EProjectorState::Hidden);
-		}
-
-		if (ArrayProjectorParameterNameWithIndex[i] == ParamNameWithIndex)
-		{
-			CustomInstance->SetProjectorState(ArrayParamName, ArrayRangeIndex, SelectProjector ? EProjectorState::Selected : EProjectorState::Hidden);
-
-			if (SelectProjector)
-			{
-				CustomInstance->LastSelectedProjectorParameterWithIndex = ParamNameWithIndex;
-			}
-		}
-	}
-
-	FCoreUObjectDelegates::BroadcastOnObjectModified(CustomInstance.Get());
-
+	
 	return FReply::Handled();
 }
 
 
-void SCustomizableInstanceProperties::OnProjectorTranslateStateChanged(ECheckBoxState InCheckboxState, FString ParamNameWithIndex) const
-{
-	int32 RangeIndex;
-	const FString ParamName = ParseNameAndIndex(ParamNameWithIndex, RangeIndex);
-
-	CustomInstance->SetProjectorState(ParamName, RangeIndex, EProjectorState::Translate);
-}
-
-
-void SCustomizableInstanceProperties::OnProjectorRotateStateChanged(ECheckBoxState InCheckboxState, FString ParamNameWithIndex) const
-{
-	int32 RangeIndex;
-	const FString ParamName = ParseNameAndIndex(ParamNameWithIndex, RangeIndex);
-
-	CustomInstance->SetProjectorState(ParamName, RangeIndex, EProjectorState::Rotate);
-}
-
-
-void SCustomizableInstanceProperties::OnProjectorScaleStateChanged(ECheckBoxState InCheckboxState, FString ParamNameWithIndex) const
-{
-	int32 RangeIndex;
-	const FString ParamName = ParseNameAndIndex(ParamNameWithIndex, RangeIndex);
-
-	CustomInstance->SetProjectorState(ParamName, RangeIndex, EProjectorState::Scale);
-}
-
-
-void SCustomizableInstanceProperties::OnProjectorTextureParameterComboBoxChanged(TSharedPtr<FString> Selection, ESelectInfo::Type SelectInfo, FString ParamNameWithIndex) const
+void SCustomizableInstanceProperties::OnProjectorTextureParameterComboBoxChanged(TSharedPtr<FString> Selection, ESelectInfo::Type SelectInfo, FString ParamName, int32 RangeIndex) const
 {
 	TArray<FCustomizableObjectIntParameterValue>& IntParameters = CustomInstance->GetIntParameters();
 
-	int32 RangeIndex;
-	const FString ParamName = ParseNameAndIndex(ParamNameWithIndex, RangeIndex);
-
 	const UCustomizableObject* CustomObject = CustomInstance->GetCustomizableObject();
 	for (int32 i = 0; i < IntParameters.Num(); ++i)
 	{
@@ -2049,8 +1816,6 @@ void SCustomizableInstanceProperties::OnProjectorTextureParameterComboBoxChanged
 
 	CustomInstance->SetSelectedParameterProfileDirty();
 
-	CustomInstance->ProjectorLayerChange = true;
-
 	CustomInstance->UpdateSkeletalMeshAsync(true, true);
 
 	CustomInstance->PostEditChange();
@@ -2060,47 +1825,26 @@ void SCustomizableInstanceProperties::OnProjectorTextureParameterComboBoxChanged
 }
 
 
-void SCustomizableInstanceProperties::OnProjectorFloatParameterChanged(float Value, int SliderIndex, FString ParamNameWithIndex)
+void SCustomizableInstanceProperties::OnProjectorFloatParameterChanged(float Value, int SliderIndex)
 {
 	OnFloatParameterChanged(Value, SliderIndex);
-
-	// Code commented since it was causing problems with the alpha slider. Selecting the projector removed the focus on the slider.
-	//int32 RangeIndex;
-	//FString ParamName = ParseNameAndIndex(ParamNameWithIndex, RangeIndex);
-
-	//EProjectorState::Type ProjectorState = CustomInstance->GetProjectorState(ParamName, RangeIndex);
-	//if (ProjectorState != EProjectorState::Selected)
-	//{
-	//	CustomInstance->ResetProjectorStates();
-	//	CustomInstance->SetProjectorState(ParamName, RangeIndex, EProjectorState::Selected);
-
-	//	CustomInstance->LastSelectedProjectorParameterWithIndex = ParamNameWithIndex;
-	//	FCoreUObjectDelegates::BroadcastOnObjectModified(CustomInstance.Get());
-	//}
 }
 
 
 void SCustomizableInstanceProperties::OnProjectorFloatParameterSliderBegin()
 {
 	OnFloatParameterSliderBegin();
-
-	CustomInstance->ProjectorAlphaChange = true;
 }
 
 
 void SCustomizableInstanceProperties::OnProjectorFloatParameterSliderEnd(float Value)
 {
 	OnFloatParameterSliderEnd(Value);
-
-	CustomInstance->ProjectorAlphaChange = false;
 }
 
 
-FReply SCustomizableInstanceProperties::OnProjectorCopyTransform(FString ParamNameWithIndex) const
+FReply SCustomizableInstanceProperties::OnProjectorCopyTransform(const FString ParamName, const int32 RangeIndex) const
 {
-	int32 RangeIndex;
-	const FString ParamName = ParseNameAndIndex(ParamNameWithIndex, RangeIndex);
-
 	const int32 ParameterIndexInObject = CustomInstance->GetCustomizableObject()->FindParameter(ParamName);
 	const int32 ProjectorParamIndex = CustomInstance->FindProjectorParameterNameIndex(ParamName);
 
@@ -2133,11 +1877,11 @@ FReply SCustomizableInstanceProperties::OnProjectorCopyTransform(FString ParamNa
 }
 
 
-FReply SCustomizableInstanceProperties::OnProjectorPasteTransform(FString ParamNameWithIndex) const
+FReply SCustomizableInstanceProperties::OnProjectorPasteTransform(const FString ParamName, const int32 RangeIndex)
 {
-	int32 RangeIndex;
-	FString ParamName = ParseNameAndIndex(ParamNameWithIndex, RangeIndex);
-
+	FScopedTransaction Transaction(LOCTEXT("PasteTransform", "Paste Transform"));
+	CustomInstance->Modify();
+	
 	FString ClipboardText;
 	FPlatformApplicationMisc::ClipboardPaste(ClipboardText);
 	FCustomizableObjectProjector DefaultValue;
@@ -2152,23 +1896,20 @@ FReply SCustomizableInstanceProperties::OnProjectorPasteTransform(FString ParamN
 		DefaultValue.Angle,
 		RangeIndex);
 
-	//Instance->UpdateSkeletalMeshAsync(true);
-	CustomInstance->TempUpdateGizmoInViewport = true;
-	CustomInstance->TempProjectorParameterName = ParamName;
-	CustomInstance->TempProjectorParameterRangeIndex = RangeIndex;
+	const TSharedPtr<ICustomizableObjectInstanceEditor> Editor = GetEditorChecked();
 
-	CustomInstance->ProjectorLayerChange = false;
+	Editor->ShowGizmoProjectorParameter(ParamName, RangeIndex);
+	CustomInstance->UpdateSkeletalMeshAsync(true, true);
 
 	return FReply::Handled();
 }
 
 
-FReply SCustomizableInstanceProperties::OnProjectorResetTransform(FString ParamNameWithIndex) const
+FReply SCustomizableInstanceProperties::OnProjectorResetTransform(const FString ParamName, const int32 RangeIndex)
 {
-	int32 RangeIndex;
-
-	FString ParamName = ParseNameAndIndex(ParamNameWithIndex, RangeIndex);
-
+	FScopedTransaction Transaction(LOCTEXT("ResetTransform", "Reset Transform"));
+	CustomInstance->Modify();
+	
 	const FCustomizableObjectProjector DefaultValue = CustomInstance->GetCustomizableObject()->GetProjectorParameterDefaultValue(ParamName);
 
 	CustomInstance->SetProjectorValue(ParamName,
@@ -2179,11 +1920,10 @@ FReply SCustomizableInstanceProperties::OnProjectorResetTransform(FString ParamN
 		DefaultValue.Angle,
 		RangeIndex);
 
-	CustomInstance->TempUpdateGizmoInViewport = true;
-	CustomInstance->TempProjectorParameterName = ParamName;
-	CustomInstance->TempProjectorParameterRangeIndex = RangeIndex;
+	const TSharedPtr<ICustomizableObjectInstanceEditor> Editor = GetEditorChecked();
 
-	CustomInstance->ProjectorLayerChange = false;
+	Editor->ShowGizmoProjectorParameter(ParamName, RangeIndex);
+	CustomInstance->UpdateSkeletalMeshAsync(true, true);
 
 	return FReply::Handled();
 }
@@ -2192,7 +1932,6 @@ FReply SCustomizableInstanceProperties::OnProjectorResetTransform(FString ParamN
 void SCustomizableInstanceProperties::InstanceUpdated(UCustomizableObjectInstance* Instance) const
 {
 	if (CustomInstance->bShowOnlyRelevantParameters &&
-		!CustomInstance->ProjectorUpdatedInViewport &&
 		bShouldResetParamBox)
 	{
 		if (const TSharedPtr<FCustomizableInstanceDetails> Details = InstanceDetails.Pin())
@@ -2203,6 +1942,15 @@ void SCustomizableInstanceProperties::InstanceUpdated(UCustomizableObjectInstanc
 }
 
 
+TSharedPtr<ICustomizableObjectInstanceEditor> SCustomizableInstanceProperties::GetEditorChecked() const
+{
+	TSharedPtr<ICustomizableObjectInstanceEditor> Editor = WeakEditor.Pin();
+	check(Editor);
+	
+	return Editor;
+}
+
+
 class SCreateProfileParameters : public SWindow
 {
 public:
diff --git a/Engine/Plugins/Experimental/Mutable/Source/CustomizableObjectEditor/Private/MuCOE/SCustomizableObjectEditorAdvancedPreviewSettings.cpp b/Engine/Plugins/Experimental/Mutable/Source/CustomizableObjectEditor/Private/MuCOE/SCustomizableObjectEditorAdvancedPreviewSettings.cpp
--- a/Engine/Plugins/Experimental/Mutable/Source/CustomizableObjectEditor/Private/MuCOE/SCustomizableObjectEditorAdvancedPreviewSettings.cpp
+++ b/Engine/Plugins/Experimental/Mutable/Source/CustomizableObjectEditor/Private/MuCOE/SCustomizableObjectEditorAdvancedPreviewSettings.cpp
@@ -415,7 +415,7 @@ SCustomizableObjectEditorAdvancedPreviewSettings::SCustomizableObjectEditorAdvan
 void SCustomizableObjectEditorAdvancedPreviewSettings::Construct(const FArguments& InArgs, const TSharedRef<FAdvancedPreviewScene>& InPreviewScene)
 {
 	SAdvancedPreviewDetailsTab::FArguments ParentArgs;
-	ParentArgs.AdditionalSettings( InArgs._AdditionalSettings );
+	ParentArgs.AdditionalSettings( InArgs._CustomSettings );
 	SAdvancedPreviewDetailsTab::Construct(ParentArgs, InPreviewScene);
 
 	//SkyMaterial = nullptr;
diff --git a/Engine/Plugins/Experimental/Mutable/Source/CustomizableObjectEditor/Private/MuCOE/SCustomizableObjectEditorAdvancedPreviewSettings.h b/Engine/Plugins/Experimental/Mutable/Source/CustomizableObjectEditor/Private/MuCOE/SCustomizableObjectEditorAdvancedPreviewSettings.h
--- a/Engine/Plugins/Experimental/Mutable/Source/CustomizableObjectEditor/Private/MuCOE/SCustomizableObjectEditorAdvancedPreviewSettings.h
+++ b/Engine/Plugins/Experimental/Mutable/Source/CustomizableObjectEditor/Private/MuCOE/SCustomizableObjectEditorAdvancedPreviewSettings.h
@@ -2,6 +2,7 @@
 
 #pragma once
 
+#include "CustomizableObjectInstanceEditor.h"
 #include "AssetRegistry/AssetData.h"
 #include "SAdvancedPreviewDetailsTab.h"
 #include "Widgets/SWindow.h"
@@ -98,7 +99,7 @@ class SCustomizableObjectEditorAdvancedPreviewSettings : public SAdvancedPreview
 {
 public:
 	SLATE_BEGIN_ARGS(SCustomizableObjectEditorAdvancedPreviewSettings) {}
-		SLATE_ARGUMENT(UObject*, AdditionalSettings) // From SAdvancedPreviewDetailsTab
+		SLATE_ARGUMENT(UCustomSettings*, CustomSettings) // From SAdvancedPreviewDetailsTab
 		SLATE_ATTRIBUTE(FCustomizableObjectEditor*, CustomizableObjectEditor)
 	SLATE_END_ARGS()
 
diff --git a/Engine/Plugins/Experimental/Mutable/Source/CustomizableObjectEditor/Private/MuCOE/SCustomizableObjectEditorViewport.cpp b/Engine/Plugins/Experimental/Mutable/Source/CustomizableObjectEditor/Private/MuCOE/SCustomizableObjectEditorViewport.cpp
--- a/Engine/Plugins/Experimental/Mutable/Source/CustomizableObjectEditor/Private/MuCOE/SCustomizableObjectEditorViewport.cpp
+++ b/Engine/Plugins/Experimental/Mutable/Source/CustomizableObjectEditor/Private/MuCOE/SCustomizableObjectEditorViewport.cpp
@@ -32,15 +32,8 @@ struct FGeometry;
 
 void SCustomizableObjectEditorViewport::Construct(const FArguments& InArgs, const FCustomizableObjectEditorViewportRequiredArgs& InRequiredArgs)
 {
-	//SkeletonTreePtr = InRequiredArgs.SkeletonTree;
 	PreviewScenePtr = InRequiredArgs.PreviewScene;
 	TabBodyPtr = InRequiredArgs.TabBody;
-	//AssetEditorToolkitPtr = InRequiredArgs.AssetEditorToolkit;
-	bShowStats = InArgs._ShowStats;
-
-	// Not necessary, since there is no undo on instance parameter change
-	//InRequiredArgs.OnPostUndo.Add(FSimpleDelegate::CreateSP(this, &SCustomizableObjectEditorViewport::OnUndoRedo));
-
 
 	SEditorViewport::Construct(
 		SEditorViewport::FArguments()
@@ -61,12 +54,6 @@ TSharedRef<FEditorViewportClient> SCustomizableObjectEditorViewport::MakeEditorV
 	LevelViewportClient->SetViewRotation(EditorViewportDefs::DefaultPerspectiveViewRotation);
 
 	SceneViewport = MakeShareable(new FSceneViewport(LevelViewportClient.Get(), ViewportWidget));
-	//Viewport = MakeShareable(new FSceneViewport(LevelViewportClient.Get(), ViewportWidget));
-	//LevelViewportClient->Viewport = Viewport.Get();
-	//LevelViewportClient->Viewport->SetUserFocus(true);
-
-	// The viewport widget needs an interface so it knows what should render
-	//ViewportWidget->SetViewportInterface(Viewport.ToSharedRef());
 
 	return LevelViewportClient.ToSharedRef();
 }
@@ -99,9 +86,6 @@ void SCustomizableObjectEditorViewportTabBody::Construct(const FArguments& InArg
 
 	CustomizableObjectEditorPtr = InArgs._CustomizableObjectEditor;
 
-	CurrentViewMode = VMI_Lit;
-	LODSelection = 0;
-
 	FCustomizableObjectEditorViewportMenuCommands::Register();
 	//FAnimViewportShowCommands::Register();
 	FCustomizableObjectEditorViewportLODCommands::Register();
@@ -143,8 +127,6 @@ void SCustomizableObjectEditorViewportTabBody::Construct(const FArguments& InArg
 	//PreviewSkeletalMeshComponent = 0;
 
 	BindCommands();
-
-	AssetRegistryLoaded = false;
 }
 
 
@@ -168,193 +150,139 @@ void SCustomizableObjectEditorViewportTabBody::AddReferencedObjects( FReferenceC
 }
 
 
-void SCustomizableObjectEditorViewportTabBody::SetClipMorphPlaneVisibility(bool bVisible, const FVector& Origin, const FVector& Normal, float MorphLength, const FBoxSphereBounds& Bounds, float Radius1, float Radius2, float RotationAngle)
+void SCustomizableObjectEditorViewportTabBody::Tick(const FGeometry& AllottedGeometry, const double InCurrentTime, const float InDeltaTime)
 {
-	LevelViewportClient->SetClipMorphPlaneVisibility(bVisible, Origin, Normal, MorphLength, Bounds, Radius1, Radius2, RotationAngle);
-}
+	SCompoundWidget::Tick(AllottedGeometry, InCurrentTime, InDeltaTime);
 
+	// Needed?
+	PreviewScenePtr->GetWorld()->Tick(LEVELTICK_All, InDeltaTime);
 
-void SCustomizableObjectEditorViewportTabBody::SetClipMorphPlaneVisibility(bool bVisible, UCustomizableObjectNodeMeshClipMorph* NodeMeshClipMorph)
-{
-	TransformWidgetVisibility.State.bClipMorphPlaneVisible = bVisible && bVisible != LevelViewportClient->IsNodeMeshClipMorphSelected();
+	// Update the material list. Not ideal to do it every tick, but tracking changes on materials in the current instance is not easy right now.
+	if (PreviewSkeletalMeshComponents.Num()>0)
+	{
+		MaterialNames.Empty();
 
-	EVisibility ChildVisibility = TransformWidgetVisibility.Value ? EVisibility::Visible : EVisibility::Hidden;
-	if (ViewportToolbarTransformWidget.IsValid())
+		for (UDebugSkelMeshComponent* PreviewSkeletalMeshComponent : PreviewSkeletalMeshComponents)
+		{
+			const TArray<UMaterialInterface*> Materials = PreviewSkeletalMeshComponent->GetMaterials();
+			for (UMaterialInterface* m : Materials)
+			{
+				if (m)
+				{
+					const UMaterial* BaseMaterial = m->GetBaseMaterial();
+					MaterialNames.Add(MakeShareable(new FString(BaseMaterial->GetName())));
+				}
+			}
+		}
+	}
+	else
 	{
-		ViewportToolbarTransformWidget.Pin()->SetVisibility(ChildVisibility);
+		MaterialNames.Empty();
 	}
 
-	LevelViewportClient->SetClipMorphPlaneVisibility(bVisible, NodeMeshClipMorph);
 }
 
 
-void SCustomizableObjectEditorViewportTabBody::SetClipMeshVisibility(bool bVisible, UStaticMesh* ClipMesh, UCustomizableObjectNodeMeshClipWithMesh* ClipMeshNode)
-{	
-	TransformWidgetVisibility.State.bClipMeshVisible = ClipMesh && ClipMeshNode && bVisible;
-
-	EVisibility ChildVisibility = TransformWidgetVisibility.Value ? EVisibility::Visible : EVisibility::Hidden;
-	if (ViewportToolbarTransformWidget.IsValid())
+void SCustomizableObjectEditorViewportTabBody::ShowGizmoClipMorph(UCustomizableObjectNodeMeshClipMorph& ClipPlainNode) const
+{
+	if (const TSharedPtr<SWidget> Toolbar = ViewportToolbarTransformWidget.Pin())
 	{
-		ViewportToolbarTransformWidget.Pin()->SetVisibility(ChildVisibility);
+		Toolbar->SetVisibility(EVisibility::Visible);		
 	}
-	
 
-	LevelViewportClient->SetClipMeshVisibility(bVisible, ClipMesh, ClipMeshNode);
+	LevelViewportClient->ShowGizmoClipMorph(ClipPlainNode);
 }
 
 
-void SCustomizableObjectEditorViewportTabBody::SetProjectorVisibility(bool bVisible, FString ProjectorParameterName, FString ProjectorParameterNameWithIndex, int32 RangeIndex, const FCustomizableObjectProjector& Data, int32 ProjectorParameterIndex)
+void SCustomizableObjectEditorViewportTabBody::HideGizmoClipMorph() const
 {
-	if (LevelViewportClient->GetGizmoProxy().HasAssignedData &&
-		!LevelViewportClient->GetGizmoProxy().AssignedDataIsFromNode &&
-		LevelViewportClient->GetGizmoProxy().ProjectorParameterIndex == ProjectorParameterIndex &&
-		LevelViewportClient->GetGizmoProxy().ProjectorRangeIndex == RangeIndex)
+	if (const TSharedPtr<SWidget> Toolbar = ViewportToolbarTransformWidget.Pin())
 	{
-		return;
-	}
-	
-	TransformWidgetVisibility.State.bProjectorVisible = bVisible;
-
-	EVisibility ChildVisibility = TransformWidgetVisibility.Value ? EVisibility::Visible : EVisibility::Hidden;
-	if (ViewportToolbarTransformWidget.IsValid())
-	{
-		ViewportToolbarTransformWidget.Pin()->SetVisibility(ChildVisibility);
+		Toolbar->SetVisibility(EVisibility::Hidden);		
 	}
 
-	if ((LevelViewportClient->GetWidgetVisibility() != bVisible) || 
-		(LevelViewportClient->GetGizmoProxy().HasAssignedData &&
-		(
-			LevelViewportClient->GetGizmoProxy().AssignedDataIsFromNode ||
-			LevelViewportClient->GetGizmoProxy().ProjectorParameterIndex != ProjectorParameterIndex ||
-		    LevelViewportClient->GetGizmoProxy().ProjectorRangeIndex != RangeIndex)
-		))
-	{
-		LevelViewportClient->SetProjectorVisibility(bVisible, ProjectorParameterName, ProjectorParameterNameWithIndex, RangeIndex, Data, ProjectorParameterIndex);
-		
-	}
+	LevelViewportClient->HideGizmoClipMorph();
 }
 
 
-void SCustomizableObjectEditorViewportTabBody::SetProjectorType(bool bVisible, FString ProjectorParameterName, FString ProjectorParameterNameWithIndex, int32 RangeIndex, const FCustomizableObjectProjector& Data, int32 ProjectorParameterIndex)
+void SCustomizableObjectEditorViewportTabBody::ShowGizmoClipMesh(UCustomizableObjectNodeMeshClipWithMesh& ClipMeshNode, UStaticMesh& ClipMesh) const
 {
-	LevelViewportClient->SetProjectorType(bVisible, ProjectorParameterName, ProjectorParameterNameWithIndex, RangeIndex, Data, ProjectorParameterIndex);
-}
-
-
-void SCustomizableObjectEditorViewportTabBody::SetProjectorVisibility(bool bVisible, UCustomizableObjectNodeProjectorConstant* Projector)
-{
-	if (((Projector == nullptr) && !LevelViewportClient->GetGizmoProxy().HasAssignedData) ||
-	    (bVisible && (Projector != nullptr) && LevelViewportClient->GetGizmoProxy().HasAssignedData &&
-		(LevelViewportClient->GetGizmoProxy().DataOriginConstant == Projector)))
-	{
-		return;
-	}
-
-	TransformWidgetVisibility.State.bProjectorVisible = Projector && bVisible;
-	EVisibility ChildVisibility = TransformWidgetVisibility.Value ? EVisibility::Visible : EVisibility::Hidden;
-	if (ViewportToolbarTransformWidget.IsValid())
+	if (const TSharedPtr<SWidget> Toolbar = ViewportToolbarTransformWidget.Pin())
 	{
-		ViewportToolbarTransformWidget.Pin()->SetVisibility(ChildVisibility);
+		Toolbar->SetVisibility(EVisibility::Visible);		
 	}
-
-	LevelViewportClient->SetProjectorVisibility(bVisible, Projector);
+	
+	LevelViewportClient->ShowGizmoClipMesh(ClipMeshNode, ClipMesh);
 }
 
 
-void SCustomizableObjectEditorViewportTabBody::SetProjectorParameterVisibility(bool bVisible, class UCustomizableObjectNodeProjectorParameter* ProjectorParameter)
+void SCustomizableObjectEditorViewportTabBody::HideGizmoClipMesh() const
 {
-	if (((ProjectorParameter == nullptr) && !LevelViewportClient->GetGizmoProxy().HasAssignedData) ||
-	    (bVisible && (ProjectorParameter != nullptr) && LevelViewportClient->GetGizmoProxy().HasAssignedData &&
-		(LevelViewportClient->GetGizmoProxy().DataOriginParameter == ProjectorParameter)))
+	if (const TSharedPtr<SWidget> Toolbar = ViewportToolbarTransformWidget.Pin())
 	{
-		return;
+		Toolbar->SetVisibility(EVisibility::Hidden);		
 	}
+	
+	LevelViewportClient->HideGizmoClipMesh();
+}
 
-	TransformWidgetVisibility.State.bProjectorVisible = ProjectorParameter && bVisible;
 
-	EVisibility ChildVisibility = TransformWidgetVisibility.Value ? EVisibility::Visible : EVisibility::Hidden;
-	if (ViewportToolbarTransformWidget.IsValid())
+void SCustomizableObjectEditorViewportTabBody::ShowGizmoProjector(
+	const FWidgetLocationDelegate& WidgetLocationDelegate, const FOnWidgetLocationChangedDelegate& OnWidgetLocationChangedDelegate,
+	const FWidgetDirectionDelegate& WidgetDirectionDelegate, const FOnWidgetDirectionChangedDelegate& OnWidgetDirectionChangedDelegate,
+	const FWidgetUpDelegate& WidgetUpDelegate, const FOnWidgetUpChangedDelegate& OnWidgetUpChangedDelegate,
+	const FWidgetScaleDelegate& WidgetScaleDelegate, const FOnWidgetScaleChangedDelegate& OnWidgetScaleChangedDelegate,
+	const FWidgetAngleDelegate& WidgetAngleDelegate,
+	const FProjectorTypeDelegate& ProjectorTypeDelegate,
+	const FWidgetColorDelegate& WidgetColorDelegate,
+	const FWidgetTrackingStartedDelegate& WidgetTrackingStartedDelegate) const
+{
+	if (const TSharedPtr<SWidget> Toolbar = ViewportToolbarTransformWidget.Pin())
 	{
-		ViewportToolbarTransformWidget.Pin()->SetVisibility(ChildVisibility);
+		Toolbar->SetVisibility(EVisibility::Visible);		
 	}
 
-	LevelViewportClient->SetProjectorParameterVisibility(bVisible, ProjectorParameter);
+	LevelViewportClient->ShowGizmoProjector(WidgetLocationDelegate, OnWidgetLocationChangedDelegate,
+		WidgetDirectionDelegate, OnWidgetDirectionChangedDelegate,
+		WidgetUpDelegate, OnWidgetUpChangedDelegate,
+		WidgetScaleDelegate, OnWidgetScaleChangedDelegate,
+		WidgetAngleDelegate,
+		ProjectorTypeDelegate,
+		WidgetColorDelegate,
+		WidgetTrackingStartedDelegate);
 }
 
 
-void SCustomizableObjectEditorViewportTabBody::ResetProjectorVisibility(bool OnlyNonNodeProjector)
+void SCustomizableObjectEditorViewportTabBody::HideGizmoProjector() const
 {
-	if(LevelViewportClient->GetIsManipulating())
-	{
-		return;
-	}
-
-	if (OnlyNonNodeProjector &&
-		LevelViewportClient->GetGizmoProxy().HasAssignedData &&
-		LevelViewportClient->GetGizmoProxy().AssignedDataIsFromNode)
+	if (const TSharedPtr<SWidget> Toolbar = ViewportToolbarTransformWidget.Pin())
 	{
-		return;
+		Toolbar->SetVisibility(EVisibility::Hidden);		
 	}
+	
+	LevelViewportClient->HideGizmoProjector();
+}
 
-	LevelViewportClient->ResetProjectorVisibility();
-
-	// Check if the projector's alpha is being modified.
-	// ResetProjectorVisibility is called every tick while modifying a projector's alpha,
-	// that makes the widget UI flash in and out.
-	bool bProjectorAlphaChange = false;
-	if (CustomizableObjectEditorPtr.IsValid())
-	{
-		UCustomizableObjectInstance* PreviewInstance = CustomizableObjectEditorPtr.Pin()->GetPreviewInstance();
-		if (PreviewInstance)
-		{
-			bProjectorAlphaChange = PreviewInstance->ProjectorAlphaChange;
-		}
-	}
 
-	if (ViewportToolbarTransformWidget.IsValid() && TransformWidgetVisibility.Value && !bProjectorAlphaChange)
+void SCustomizableObjectEditorViewportTabBody::ShowGizmoLight(ULightComponent& SelectedLight) const
+{
+	if (const TSharedPtr<SWidget> Toolbar = ViewportToolbarTransformWidget.Pin())
 	{
-		ViewportToolbarTransformWidget.Pin()->SetVisibility(EVisibility::Hidden);
-		TransformWidgetVisibility.State.bProjectorVisible = false;
+		Toolbar->SetVisibility(EVisibility::Visible);		
 	}
+	
+	LevelViewportClient->ShowGizmoLight(SelectedLight);
 }
 
 
-void SCustomizableObjectEditorViewportTabBody::Tick(const FGeometry& AllottedGeometry, const double InCurrentTime, const float InDeltaTime)
+void SCustomizableObjectEditorViewportTabBody::HideGizmoLight() const
 {
-	SCompoundWidget::Tick(AllottedGeometry, InCurrentTime, InDeltaTime);
-
-	// Needed?
-	PreviewScenePtr->GetWorld()->Tick(LEVELTICK_All, InDeltaTime);
-
-	//if (PreviewSkeletalMeshComponent != nullptr)
-	//{
-	//	PreviewSkeletalMeshComponent->TickAnimation(InDeltaTime, false);
-	//	PreviewSkeletalMeshComponent->RefreshBoneTransforms();
-	//}
-
-	// Update the material list. Not ideal to do it every tick, but tracking changes on materials in the current instance is not easy right now.
-	if (PreviewSkeletalMeshComponents.Num()>0)
+	if (const TSharedPtr<SWidget> Toolbar = ViewportToolbarTransformWidget.Pin())
 	{
-		MaterialNames.Empty();
-
-		for (UDebugSkelMeshComponent* PreviewSkeletalMeshComponent : PreviewSkeletalMeshComponents)
-		{
-			const TArray<UMaterialInterface*> Materials = PreviewSkeletalMeshComponent->GetMaterials();
-			for (UMaterialInterface* m : Materials)
-			{
-				if (m)
-				{
-					const UMaterial* BaseMaterial = m->GetBaseMaterial();
-					MaterialNames.Add(MakeShareable(new FString(BaseMaterial->GetName())));
-				}
-			}
-		}
-	}
-	else
-	{
-		MaterialNames.Empty();
+		Toolbar->SetVisibility(EVisibility::Hidden);		
 	}
-
+	
+	LevelViewportClient->HideGizmoLight();
 }
 
 
@@ -368,7 +296,6 @@ void SCustomizableObjectEditorViewportTabBody::SetPreviewComponents(const TArray
 		{
 			Transform = PreviewSkeletalMeshComponent->GetComponentToWorld();
 			PreviewScenePtr->RemoveComponent(PreviewSkeletalMeshComponent);
-			//PreviewSkeletalMeshComponent->DestroyComponent();
 		}
 	}
 
@@ -449,65 +376,12 @@ void SCustomizableObjectEditorViewportTabBody::BindCommands()
 		FIsActionChecked::CreateSP( EditorViewportClientRef, &FCustomizableObjectEditorViewportClient::IsSetShowCollisionChecked ) );
 
 	// Menu
-	CommandList.MapAction(
-		Commands.SetShowPivot,
-		FExecuteAction::CreateSP( EditorViewportClientRef, &FCustomizableObjectEditorViewportClient::SetShowPivot ),
-		FCanExecuteAction(),
-		FIsActionChecked::CreateSP( EditorViewportClientRef, &FCustomizableObjectEditorViewportClient::IsSetShowPivotChecked ) );
-
-
 	CommandList.MapAction(
 		Commands.BakeInstance,
 		FExecuteAction::CreateSP(EditorViewportClientRef, &FCustomizableObjectEditorViewportClient::BakeInstance),
 		FCanExecuteAction(),
 		FIsActionChecked());
 
-	// In-viewport tool bar
-
-	//Bind menu commands
-	const FCustomizableObjectEditorViewportMenuCommands& MenuActions = FCustomizableObjectEditorViewportMenuCommands::Get();
-
-	//CommandList.MapAction(
-	//	MenuActions.CameraFollow,
-	//	FExecuteAction::CreateSP(this, &SCustomizableObjectEditorViewport::ToggleCameraFollow),
-	//	FCanExecuteAction::CreateSP(this, &SCustomizableObjectEditorViewport::CanChangeCameraMode),
-	//	FIsActionChecked::CreateSP(this, &SCustomizableObjectEditorViewport::IsCameraFollowEnabled));
-
-	//CommandList.MapAction(
-	//	MenuActions.PreviewSceneSettings,
-	//	FExecuteAction::CreateSP(this, &SCustomizableObjectEditorViewport::OpenPreviewSceneSettings));
-
-	//CommandList.MapAction(
-	//	MenuActions.SetCPUSkinning,
-	//	FExecuteAction::CreateSP(EditorViewportClientRef, &FCustomizableObjectEditorViewportClient::ToggleCPUSkinning),
-	//	FCanExecuteAction(),
-	//	FIsActionChecked::CreateSP(EditorViewportClientRef, &FCustomizableObjectEditorViewportClient::IsSetCPUSkinningChecked));
-
-	//CommandList.MapAction(
-	//	MenuActions.SetShowNormals,
-	//	FExecuteAction::CreateSP(EditorViewportClientRef, &FCustomizableObjectEditorViewportClient::ToggleShowNormals),
-	//	FCanExecuteAction(),
-	//	FIsActionChecked::CreateSP(EditorViewportClientRef, &FCustomizableObjectEditorViewportClient::IsSetShowNormalsChecked));
-
-	//CommandList.MapAction(
-	//	MenuActions.SetShowTangents,
-	//	FExecuteAction::CreateSP(EditorViewportClientRef, &FCustomizableObjectEditorViewportClient::ToggleShowTangents),
-	//	FCanExecuteAction(),
-	//	FIsActionChecked::CreateSP(EditorViewportClientRef, &FCustomizableObjectEditorViewportClient::IsSetShowTangentsChecked));
-
-	//CommandList.MapAction(
-	//	MenuActions.SetShowBinormals,
-	//	FExecuteAction::CreateSP(EditorViewportClientRef, &FCustomizableObjectEditorViewportClient::ToggleShowBinormals),
-	//	FCanExecuteAction(),
-	//	FIsActionChecked::CreateSP(EditorViewportClientRef, &FCustomizableObjectEditorViewportClient::IsSetShowBinormalsChecked));
-
-	//CommandList.MapAction(
-	//	MenuActions.AnimSetDrawUVs,
-	//	FExecuteAction::CreateSP(EditorViewportClientRef, &FCustomizableObjectEditorViewportClient::ToggleDrawUVOverlay),
-	//	FCanExecuteAction(),
-	//	FIsActionChecked::CreateSP(EditorViewportClientRef, &FCustomizableObjectEditorViewportClient::IsSetDrawUVOverlayChecked));
-
-
 	//Bind LOD preview menu commands
 	const FCustomizableObjectEditorViewportLODCommands& ViewportLODMenuCommands = FCustomizableObjectEditorViewportLODCommands::Get();
 
@@ -529,19 +403,19 @@ void SCustomizableObjectEditorViewportTabBody::BindCommands()
 		ViewportLODMenuCommands.TranslateMode,
 		FExecuteAction::CreateSP(this, &SCustomizableObjectEditorViewportTabBody::ProjectorCheckboxStateChanged, UE::Widget::WM_Translate),
 		FCanExecuteAction(),
-		FIsActionChecked());
+		FIsActionChecked::CreateSP(this, &SCustomizableObjectEditorViewportTabBody::IsProjectorCheckboxState, UE::Widget::WM_Translate));
 
 	CommandList.MapAction(
 		ViewportLODMenuCommands.RotateMode,
 		FExecuteAction::CreateSP(this, &SCustomizableObjectEditorViewportTabBody::ProjectorCheckboxStateChanged, UE::Widget::WM_Rotate),
 		FCanExecuteAction(),
-		FIsActionChecked());
+		FIsActionChecked::CreateSP(this, &SCustomizableObjectEditorViewportTabBody::IsProjectorCheckboxState, UE::Widget::WM_Rotate));
 
 	CommandList.MapAction(
 		ViewportLODMenuCommands.ScaleMode,
 		FExecuteAction::CreateSP(this, &SCustomizableObjectEditorViewportTabBody::ProjectorCheckboxStateChanged, UE::Widget::WM_Scale),
 		FCanExecuteAction(),
-		FIsActionChecked());
+		FIsActionChecked::CreateSP(this, &SCustomizableObjectEditorViewportTabBody::IsProjectorCheckboxState, UE::Widget::WM_Scale));
 
 	CommandList.MapAction(
 		ViewportLODMenuCommands.RotationGridSnap,
@@ -636,12 +510,6 @@ void SCustomizableObjectEditorViewportTabBody::OnTakeHighResScreenshot()
 }
 
 
-bool SCustomizableObjectEditorViewportTabBody::GetIsManipulatingGizmo()
-{
-	return LevelViewportClient->GetIsManipulating();
-}
-
-
 void SCustomizableObjectEditorViewportTabBody::SetCameraMode(bool Value)
 {
 	LevelViewportClient->SetCameraMode(Value);
@@ -733,21 +601,10 @@ int32 SCustomizableObjectEditorViewportTabBody::GetViewportCameraSpeed()
 
 TSharedRef<SWidget> SCustomizableObjectEditorViewportTabBody::BuildToolBar()
 {
-	//FToolBarBuilder SaveThumbnailToolbarBuilder(GetToolkitCommands());
-	//{
-	//	SaveThumbnailToolbarBuilder.AddToolBarButton(FCustomizableObjectEditorViewportCommands::Get().SaveThumbnail);
-	//}
-
-
 	FToolBarBuilder CommandToolbarBuilder(UICommandList, FMultiBoxCustomization::None);
 	{
 		CommandToolbarBuilder.AddToolBarButton(FCustomizableObjectEditorViewportCommands::Get().SetShowGrid);
 		CommandToolbarBuilder.AddToolBarButton(FCustomizableObjectEditorViewportCommands::Get().SetShowSky);
-		//CommandToolbarBuilder.AddToolBarButton(FCustomizableObjectEditorViewportCommands::Get().SetShowBounds);
-		//CommandToolbarBuilder.AddToolBarButton(FCustomizableObjectEditorViewportCommands::Get().SetShowPivot);
-		//CommandToolbarBuilder.AddToolBarButton(FCustomizableObjectEditorViewportCommands::Get().SetShowNormals);
-		//CommandToolbarBuilder.AddToolBarButton(FCustomizableObjectEditorViewportCommands::Get().SetShowTangents);
-		//CommandToolbarBuilder.AddToolBarButton(FCustomizableObjectEditorViewportCommands::Get().SetShowBinormals);
 	}
 	CommandToolbarBuilder.BeginSection("Material UVs");
 	{
@@ -1067,29 +924,15 @@ bool SCustomizableObjectEditorViewportTabBody::IsLODModelSelected(int32 LODSelec
 }
 
 
-void SCustomizableObjectEditorViewportTabBody::ProjectorCheckboxStateChanged(UE::Widget::EWidgetMode InNewMode)
+void SCustomizableObjectEditorViewportTabBody::ProjectorCheckboxStateChanged(const UE::Widget::EWidgetMode Mode)
 {
-	switch (InNewMode)
-	{
-		case UE::Widget::WM_Translate:
-		{
-			GetViewportClient()->SetWidgetMode(UE::Widget::WM_Translate);
-			GetViewportClient()->SetProjectorWidgetMode(UE::Widget::WM_Translate);
-			break;
-		}
-		case UE::Widget::WM_Rotate:
-		{
-			GetViewportClient()->SetWidgetMode(UE::Widget::WM_Rotate);
-			GetViewportClient()->SetProjectorWidgetMode(UE::Widget::WM_Rotate);
-			break;
-		}
-		case UE::Widget::WM_Scale:
-		{
-			GetViewportClient()->SetWidgetMode(UE::Widget::WM_Scale);
-			GetViewportClient()->SetProjectorWidgetMode(UE::Widget::WM_Scale);
-			break;
-		}
-	}
+	GetViewportClient()->SetWidgetMode(Mode);
+}
+
+
+bool SCustomizableObjectEditorViewportTabBody::IsProjectorCheckboxState(const UE::Widget::EWidgetMode Mode) const
+{
+	return GetViewportClient()->GetWidgetMode() == Mode;	
 }
 
 
@@ -1127,79 +970,26 @@ void SCustomizableObjectEditorViewportTabBody::OnSetLODModel(int32 LODSelectionT
 		if (PreviewComponent)
 		{
 			PreviewComponent->SetForcedLOD(LODSelection);
-			//PopulateUVChoices();
 		}
 	}
 }
 
-/*
-void SCustomizableObjectEditorViewportTabBody::SetEditorTransformViewportToolbar(TWeakPtr<class SCustomizableObjectEditorTransformViewportToolbar> EditorTransformViewportToolbarParam)
-{
-	EditorTransformViewportToolbar = EditorTransformViewportToolbarParam;
-}*/
-
-void SCustomizableObjectEditorViewportTabBody::SetViewportToolbarTransformWidget(TWeakPtr<class SWidget> InTransformWidget)
-{
-	ViewportToolbarTransformWidget = InTransformWidget;
-}
-
-void SCustomizableObjectEditorViewportTabBody::UpdateGizmoDataToOrigin()
-{
-	LevelViewportClient->UpdateGizmoDataToOrigin();
-}
-
-
-void SCustomizableObjectEditorViewportTabBody::SetGizmoCallUpdateSkeletalMesh(bool Value)
-{
-	LevelViewportClient->SetGizmoCallUpdateSkeletalMesh(Value);
-}
-
-
-FString SCustomizableObjectEditorViewportTabBody::GetGizmoProjectorParameterName()
-{
-	return LevelViewportClient->GetGizmoProjectorParameterName();
-}
-
-
-FString SCustomizableObjectEditorViewportTabBody::GetGizmoProjectorParameterNameWithIndex()
-{
-	return LevelViewportClient->GetGizmoProjectorParameterNameWithIndex();
-}
-
-
-bool SCustomizableObjectEditorViewportTabBody::GetGizmoHasAssignedData()
-{
-	return LevelViewportClient->GetGizmoHasAssignedData();
-}
-
 
-bool SCustomizableObjectEditorViewportTabBody::GetGizmoAssignedDataIsFromNode()
+TSharedPtr<FCustomizableObjectPreviewScene> SCustomizableObjectEditorViewportTabBody::GetPreviewScene() const
 {
-	return LevelViewportClient->GetGizmoAssignedDataIsFromNode();
+	return PreviewScenePtr;
 }
 
 
-void SCustomizableObjectEditorViewportTabBody::CopyTransformFromOriginData()
+TSharedPtr<FCustomizableObjectEditorViewportClient> SCustomizableObjectEditorViewportTabBody::GetViewportClient() const
 {
-    LevelViewportClient->CopyTransformFromOriginData();
+	return LevelViewportClient;
 }
 
 
-void SCustomizableObjectEditorViewportTabBody::ProjectorParameterChanged(UCustomizableObjectNodeProjectorParameter* Node)
-{
-	LevelViewportClient->ProjectorParameterChanged(Node);
-}
-
-
-void SCustomizableObjectEditorViewportTabBody::ProjectorParameterChanged(UCustomizableObjectNodeProjectorConstant* Node)
-{
-	LevelViewportClient->ProjectorParameterChanged(Node);
-}
-
-
-bool SCustomizableObjectEditorViewportTabBody::AnyProjectorNodeSelected()
+void SCustomizableObjectEditorViewportTabBody::SetViewportToolbarTransformWidget(TWeakPtr<class SWidget> InTransformWidget)
 {
-	return LevelViewportClient->AnyProjectorNodeSelected();
+	ViewportToolbarTransformWidget = InTransformWidget;
 }
 
 
