[Mutable] Projector Gizmo code refactor.

diff --git a/Engine/Plugins/Experimental/Mutable/Source/CustomizableObjectEditor/Private/MuCOE/SCustomizableObjectEditorViewport.h b/Engine/Plugins/Experimental/Mutable/Source/CustomizableObjectEditor/Private/MuCOE/SCustomizableObjectEditorViewport.h
--- a/Engine/Plugins/Experimental/Mutable/Source/CustomizableObjectEditor/Private/MuCOE/SCustomizableObjectEditorViewport.h
+++ b/Engine/Plugins/Experimental/Mutable/Source/CustomizableObjectEditor/Private/MuCOE/SCustomizableObjectEditorViewport.h
@@ -3,6 +3,7 @@
 #pragma once
 
 #include "SEditorViewport.h"
+#include "MuCOE/CustomizableObjectEditorViewportClient.h"
 #include "UObject/GCObject.h"
 
 enum ERotationGridMode : int;
@@ -22,35 +23,27 @@ class UCustomizableObjectNodeProjectorConstant;
 class UCustomizableObjectNodeProjectorParameter;
 class UDebugSkelMeshComponent;
 class UStaticMesh;
+class SCustomizableObjectEditorToolBar;
+class SCustomizableObjectViewportToolBar;
+class FCustomizableObjectEditorViewportClient;
 struct FCustomizableObjectProjector;
 struct FGeometry;
+enum class ECustomizableObjectProjectorType : uint8;
 
 
 struct FCustomizableObjectEditorViewportRequiredArgs
 {
 	FCustomizableObjectEditorViewportRequiredArgs(
-		//const TSharedRef<class ISkeletonTree>& InSkeletonTree, 
 		const TSharedRef<class FCustomizableObjectPreviewScene>& InPreviewScene,
 		TSharedRef<class SCustomizableObjectEditorViewportTabBody> InTabBody 
-		//TSharedRef<class FAssetEditorToolkit> InAssetEditorToolkit, 
-		//FSimpleMulticastDelegate& InOnPostUndo
 	)
-//		: SkeletonTree(InSkeletonTree)
 		: PreviewScene(InPreviewScene)
 		, TabBody(InTabBody)
-//		, AssetEditorToolkit(InAssetEditorToolkit)
-//		, OnPostUndo(InOnPostUndo)
 	{}
 
-	//TSharedRef<class ISkeletonTree> SkeletonTree;
+	TSharedRef<FCustomizableObjectPreviewScene> PreviewScene;
 
-	TSharedRef<class FCustomizableObjectPreviewScene> PreviewScene;
-
-	TSharedRef<class SCustomizableObjectEditorViewportTabBody> TabBody;
-
-	//TSharedRef<class FAssetEditorToolkit> AssetEditorToolkit;
-
-	//FSimpleMulticastDelegate& OnPostUndo;
+	TSharedRef<SCustomizableObjectEditorViewportTabBody> TabBody;
 };
 
 
@@ -60,12 +53,7 @@ struct FCustomizableObjectEditorViewportRequiredArgs
 class SCustomizableObjectEditorViewport : public SEditorViewport
 {
 public:
-	SLATE_BEGIN_ARGS(SCustomizableObjectEditorViewport)
-		: _ShowStats(true)
-	{}
-
-		SLATE_ARGUMENT(bool, ShowStats)
-
+	SLATE_BEGIN_ARGS(SCustomizableObjectEditorViewport)	{}
 	SLATE_END_ARGS()
 
 	void Construct(const FArguments& InArgs, const FCustomizableObjectEditorViewportRequiredArgs& InRequiredArgs);
@@ -81,31 +69,21 @@ protected:
 	/**  Handle undo/redo by refreshing the viewport */
 	void OnUndoRedo();
 
-protected:
 	// Viewport client
-	TSharedPtr<class FCustomizableObjectEditorViewportClient> LevelViewportClient;
+	TSharedPtr<FCustomizableObjectEditorViewportClient> LevelViewportClient;
 
 	// Pointer to the compound widget that owns this viewport widget
-	TWeakPtr<class SCustomizableObjectEditorViewportTabBody> TabBodyPtr;
+	TWeakPtr<SCustomizableObjectEditorViewportTabBody> TabBodyPtr;
 
 	// The preview scene that we are viewing
-	TWeakPtr<class FCustomizableObjectPreviewScene> PreviewScenePtr;
-
-	TWeakPtr<class SCustomizableObjectViewportToolBar> ToolbarPtr;
+	TWeakPtr<FCustomizableObjectPreviewScene> PreviewScenePtr;
 
-	// The skeleton tree we are editing
-	//TWeakPtr<class ISkeletonTree> SkeletonTreePtr;
-
-	// The asset editor we are embedded in
-	//TWeakPtr<class FAssetEditorToolkit> AssetEditorToolkitPtr;
-
-	/** Whether we should show stats for this viewport */
-	bool bShowStats;
+	TWeakPtr<SCustomizableObjectViewportToolBar> ToolbarPtr;
 };
 
 
 /**
- * CustomizableObject Editor Preview viewport widget wwith toolbars, etc.
+ * CustomizableObject Editor Preview viewport widget with toolbars, etc.
  */
 class SCustomizableObjectEditorViewportTabBody : public SCompoundWidget, public FGCObject
 {
@@ -115,7 +93,7 @@ public:
 	SLATE_END_ARGS()
 
 	void Construct(const FArguments& InArgs);
-	virtual ~SCustomizableObjectEditorViewportTabBody();
+	virtual ~SCustomizableObjectEditorViewportTabBody() override;
 
 	void SetAnimation(class UAnimationAsset* Animation, EAnimationMode::Type AnimationType);
 
@@ -138,23 +116,45 @@ public:
 	*/
 	const TSharedPtr<FUICommandList>& GetCommandList() const { return UICommandList; }
 
+	/** Do not call directly. Use ICustomizableObjectEditor functions instead. */
+	void ShowGizmoClipMorph(UCustomizableObjectNodeMeshClipMorph& ClipPlainNode) const;
+
+	/** Do not call directly. Use ICustomizableObjectEditor functions instead. */
+	void HideGizmoClipMorph() const;
+
+	/** Do not call directly. Use ICustomizableObjectEditor functions instead. */
+	void ShowGizmoClipMesh(UCustomizableObjectNodeMeshClipWithMesh& ClipMeshNode, UStaticMesh& ClipMesh) const;
 
-	void SetClipMorphPlaneVisibility(bool bVisible, const FVector& Origin = FVector::ZeroVector, const FVector& Normal = FVector::ZeroVector, float MorphLength = 0.f, const FBoxSphereBounds& Bounds = FBoxSphereBounds(), float Radius1 = 3.f, float Radius2 = 2.f, float RotationAngle = 0.f);
-	void SetClipMorphPlaneVisibility(bool bVisible, class UCustomizableObjectNodeMeshClipMorph* NodeMorphClipMesh);
-	void SetClipMeshVisibility(bool bVisible, UStaticMesh* ClipMesh, class UCustomizableObjectNodeMeshClipWithMesh* ClipMeshNode);
-	void SetProjectorVisibility(bool bVisible, FString ProjectorParameterName, FString ProjectorParameterNameWithIndex, int32 RangeIndex, const FCustomizableObjectProjector& Data, int32 ProjectorParameterIndex);
-	void SetProjectorVisibility(bool bVisible, class UCustomizableObjectNodeProjectorConstant* Projector = nullptr);
-	void SetProjectorParameterVisibility(bool bVisible, class UCustomizableObjectNodeProjectorParameter* ProjectorParameter = nullptr);
-	void ResetProjectorVisibility(bool OnlyNonNodeProjector);
-	void SetProjectorType(bool bVisible, FString ProjectorParameterName, FString ProjectorParameterNameWithIndex, int32 RangeIndex, const FCustomizableObjectProjector& Data, int32 ProjectorParameterIndex);
+	/** Do not call directly. Use ICustomizableObjectEditor functions instead. */
+	void HideGizmoClipMesh() const;
 
+	/** Do not call directly. Use ICustomizableObjectEditor functions instead. */
+	void ShowGizmoProjector(
+		const FWidgetLocationDelegate& WidgetLocationDelegate, const FOnWidgetLocationChangedDelegate& OnWidgetLocationChangedDelegate,
+		const FWidgetDirectionDelegate& WidgetDirectionDelegate, const FOnWidgetDirectionChangedDelegate& OnWidgetDirectionChangedDelegate,
+		const FWidgetUpDelegate& WidgetUpDelegate, const FOnWidgetUpChangedDelegate& OnWidgetUpChangedDelegate,
+		const FWidgetScaleDelegate& WidgetScaleDelegate, const FOnWidgetScaleChangedDelegate& OnWidgetScaleChangedDelegate,
+		const FWidgetAngleDelegate& WidgetAngleDelegate,
+	    const FProjectorTypeDelegate& ProjectorTypeDelegate,
+	    const FWidgetColorDelegate& WidgetColorDelegate,
+	    const FWidgetTrackingStartedDelegate& WidgetTrackingStartedDelegate) const;
+
+	/** Do not call directly. Use ICustomizableObjectEditor functions instead. */
+	void HideGizmoProjector() const;
+
+	/** Do not call directly. Use ICustomizableObjectEditor functions instead. */
+	void ShowGizmoLight(ULightComponent& SelectedLight) const;
+
+	/** Do not call directly. Use ICustomizableObjectEditor functions instead. */
+	void HideGizmoLight() const;
+	
 	/** Binds commands associated with the viewport client. */
 	void BindCommands();
 	TSharedRef<SWidget> BuildToolBar();
 
 	TSharedRef<SWidget> GenerateUVMaterialOptionsMenuContent();
 
-	/** Callback to show / hide the state and rutime parameter test in the viewport */
+	/** Callback to show / hide the state and runtime parameter test in the viewport */
 	TSharedRef<SWidget> ShowStateTestData();
 
 	/** Retrieves the skeletal mesh component. */
@@ -173,8 +173,10 @@ public:
 	int32 GetLODSelection() const { return LODSelection; };
 
 	/** Callback for when the projector checkbox changes */
-	void ProjectorCheckboxStateChanged(UE::Widget::EWidgetMode InNewMode);
+	void ProjectorCheckboxStateChanged(UE::Widget::EWidgetMode Mode);
 
+	bool IsProjectorCheckboxState(UE::Widget::EWidgetMode Mode) const;
+	
 	/** Called when the rotation grid snap is toggled off and on */
 	void RotationGridSnapClicked();
 
@@ -191,48 +193,16 @@ public:
 	void OnSetLODModel(int32 LODSelectionType);
 
 	// Getter of PreviewScenePtr
-	TSharedPtr<class FCustomizableObjectPreviewScene> GetPreviewScene() { return PreviewScenePtr; }
+	TSharedPtr<FCustomizableObjectPreviewScene> GetPreviewScene() const;
 
-	//
-	TSharedPtr<class FCustomizableObjectEditorViewportClient> GetViewportClient() { return LevelViewportClient; }
+	TSharedPtr<FCustomizableObjectEditorViewportClient> GetViewportClient() const;
 
 	// Pointer back to the editor tool that owns us.
-	TWeakPtr<class ICustomizableObjectInstanceEditor> CustomizableObjectEditorPtr;
+	TWeakPtr<ICustomizableObjectInstanceEditor> CustomizableObjectEditorPtr;
 
 	//void SetEditorTransformViewportToolbar(TWeakPtr<class SCustomizableObjectEditorTransformViewportToolbar> EditorTransformViewportToolbarParam);
 	void SetViewportToolbarTransformWidget(TWeakPtr<class SWidget> TransformWidget);
 
-	/** Update the current gizmo data to its corresponding data origin */
-	void UpdateGizmoDataToOrigin();
-
-	/** Setter for LevelViewportClient::GizmoProxy::CallUpdateSkeletalMesh */
-	void SetGizmoCallUpdateSkeletalMesh(bool Value);
-
-	/** Getter for LevelViewportClient::GizmoProxy::ProjectorParameterName */
-	FString GetGizmoProjectorParameterName();
-
-	/** Getter for LevelViewportClient::GizmoProxy::ProjectorParameterNameWithIndex */
-	FString GetGizmoProjectorParameterNameWithIndex();
-
-	/** Getter for LevelViewportClient::GizmoProxy::HasAssignedData */
-	bool GetGizmoHasAssignedData();
-
-	/** Getter for LevelViewportClient::GizmoProxy::AssignedDataIsFromNode */
-	bool GetGizmoAssignedDataIsFromNode();
-
-	void CopyTransformFromOriginData();
-
-	/** Returns true if any projector node is currently selected (in the case of the Customizable Object Editor) */
-	bool AnyProjectorNodeSelected();
-
-	/** To notify of changes in any parameter projector node property in case is the one being used, currently
-	* just the projector type parameter is taken into account */
-	void ProjectorParameterChanged(UCustomizableObjectNodeProjectorParameter* Node);
-
-	/** To notify of changes in any parameter projector node property in case is the one being used, currently
-	* just the projector type parameter is taken into account */
-	void ProjectorParameterChanged(UCustomizableObjectNodeProjectorConstant* Node);
-	
 	/** Setter of CustomizableObject */
 	void SetCustomizableObject(UCustomizableObject* CustomizableObjectParameter);
 
@@ -242,9 +212,6 @@ public:
 	/** Called to bring up the screenshot UI */
 	void OnTakeHighResScreenshot();
 
-	/** Getter for LevelViewportClient::bManipulating */
-	bool GetIsManipulatingGizmo();
-
 	/** Sets camera mode of the LevelViewportClient::bActivateOrbitalCamera */
 	void SetCameraMode(bool Value);
 
@@ -270,28 +237,25 @@ public:
 private:
 	/** Determines the visibility of the viewport. */
 	bool IsVisible() const;
-
-	EActiveTimerReturnType HandleActiveTimer(double InCurrentTime, float InDeltaTime);
-
 	
 	// Components for the preview mesh.
 	TArray<TObjectPtr<UDebugSkelMeshComponent>> PreviewSkeletalMeshComponents;
 
 	// The scene for this viewport.
-	TSharedPtr<class FCustomizableObjectPreviewScene> PreviewScenePtr;
+	TSharedPtr<FCustomizableObjectPreviewScene> PreviewScenePtr;
 
 	// Level viewport client
-	TSharedPtr<class FCustomizableObjectEditorViewportClient> LevelViewportClient;
+	TSharedPtr<FCustomizableObjectEditorViewportClient> LevelViewportClient;
 
 	// Viewport widget
-	TSharedPtr<class SCustomizableObjectEditorViewport> ViewportWidget;
-	TSharedPtr<class SCustomizableObjectEditorToolBar> ViewportToolbar;
+	TSharedPtr<SCustomizableObjectEditorViewport> ViewportWidget;
+	TSharedPtr<SCustomizableObjectEditorToolBar> ViewportToolbar;
 
 	/** Commands that are bound to delegates*/
 	TSharedPtr<FUICommandList> UICommandList;
 
 	// The currently selected view mode.
-	EViewModeIndex CurrentViewMode;
+	EViewModeIndex CurrentViewMode = VMI_Lit;
 
 	// Selection for the material to filter when showing UVs
 	TArray<TSharedPtr<FString>> ArrayUVMaterialOptionString;
@@ -326,26 +290,14 @@ private:
 	void OnUVChannelChanged(TSharedPtr<FString> Selected, ESelectInfo::Type SelectInfo);
 
 	/** Current LOD selection*/
-	int32 LODSelection;
+	int32 LODSelection = 0;
 
 	/** To be able to show and hide the constant projector RTS controls */
-	//TWeakPtr<class SCustomizableObjectEditorTransformViewportToolbar> EditorTransformViewportToolbar;
-	TWeakPtr<class SWidget> ViewportToolbarTransformWidget;
-
-	union ViewportToolbarTransformWidgetVisibility
-	{
-		uint32 Value;
-		struct 
-		{
-			uint32 bProjectorVisible      : 1;
-			uint32 bClipMeshVisible       : 1;
-			uint32 bClipMorphPlaneVisible : 1;
-		} State;
-	} TransformWidgetVisibility {0};
+	TWeakPtr<SWidget> ViewportToolbarTransformWidget;
 
 	/** Flag to know when the asset registry initial loading has completed, value set by the Customizable Object / Customizable Object Instance editor */
-	bool AssetRegistryLoaded;
+	bool AssetRegistryLoaded = false;
 
-	/** Weak pointer to the highres screenshot dialog if it's open. Will become invalid if UI is closed whilst the viewport is open */
-	TWeakPtr<class SWindow> CustomizableObjectHighresScreenshot;
+	/** Weak pointer to the high resolution screenshot dialog if it's open. Will become invalid if UI is closed whilst the viewport is open */
+	TWeakPtr<SWindow> CustomizableObjectHighresScreenshot;
 };
diff --git a/Engine/Plugins/Experimental/Mutable/Source/CustomizableObjectEditor/Private/MuCOE/SCustomizableObjectEditorViewportToolBar.cpp b/Engine/Plugins/Experimental/Mutable/Source/CustomizableObjectEditor/Private/MuCOE/SCustomizableObjectEditorViewportToolBar.cpp
--- a/Engine/Plugins/Experimental/Mutable/Source/CustomizableObjectEditor/Private/MuCOE/SCustomizableObjectEditorViewportToolBar.cpp
+++ b/Engine/Plugins/Experimental/Mutable/Source/CustomizableObjectEditor/Private/MuCOE/SCustomizableObjectEditorViewportToolBar.cpp
@@ -34,274 +34,13 @@ struct FSlateBrush;
 #define LOCTEXT_NAMESPACE "CustomizableObjectEditorViewportToolBar"
 
 
-//Class definition which represents widget to modify viewport's background color.
-//class SBackgroundColorSettings : public SCompoundWidget
-//{
-//public:
-//
-//	SLATE_BEGIN_ARGS(SBackgroundColorSettings)
-//	{}
-//	SLATE_ARGUMENT(TWeakPtr<SCustomizableObjectEditorViewport>, AnimEditorViewport)
-//		SLATE_END_ARGS()
-//
-//		/** Constructs this widget from its declaration */
-//		void Construct(const FArguments& InArgs)
-//	{
-//		AnimViewportPtr = InArgs._AnimEditorViewport;
-//
-//		TSharedPtr<SWidget> ExtraWidget =
-//			SNew(SBorder)
-//			.BorderImage(UE_MUTABLE_GET_BRUSH("FilledBorder"))
-//			[
-//				SNew(SHorizontalBox)
-//				+ SHorizontalBox::Slot()
-//			.AutoWidth()
-//			.Padding(1)
-//			[
-//				SNew(SColorBlock)
-//				.Color(AnimViewportPtr.Pin().ToSharedRef(), &SCustomizableObjectEditorViewport::GetViewportBackgroundColor)
-//			.IgnoreAlpha(true)
-//			.ToolTipText(LOCTEXT("ColorBlock_ToolTip", "Select background color"))
-//			.OnMouseButtonDown(this, &SBackgroundColorSettings::OnColorBoxClicked)
-//			]
-//			];
-//
-//		this->ChildSlot
-//			[
-//				SNew(SAnimPlusMinusSlider)
-//				.Label(LOCTEXT("BrightNess", "Brightness:"))
-//			.IsEnabled(this, &SBackgroundColorSettings::IsBrightnessSliderEnabled)
-//			.OnMinusClicked(this, &SBackgroundColorSettings::OnDecreaseBrightness)
-//			.MinusTooltip(LOCTEXT("DecreaseBrightness_ToolTip", "Decrease brightness"))
-//			.SliderValue(AnimViewportPtr.Pin().ToSharedRef(), &SCustomizableObjectEditorViewport::GetBackgroundBrightness)
-//			.OnSliderValueChanged(AnimViewportPtr.Pin().ToSharedRef(), &SCustomizableObjectEditorViewport::SetBackgroundBrightness)
-//			.SliderTooltip(LOCTEXT("BackgroundBrightness_ToolTip", "Change background brightness"))
-//			.OnPlusClicked(this, &SBackgroundColorSettings::OnIncreaseBrightness)
-//			.PlusTooltip(LOCTEXT("IncreaseBrightness_ToolTip", "Increase brightness"))
-//			.ExtraWidget(ExtraWidget)
-//			];
-//	}
-//
-//protected:
-//
-//	/** Function to open color picker window when selected from context menu */
-//	FReply OnColorBoxClicked(const FGeometry& MyGeometry, const FPointerEvent& MouseEvent)
-//	{
-//		if (MouseEvent.GetEffectingButton() != EKeys::LeftMouseButton)
-//		{
-//			return FReply::Unhandled();
-//		}
-//
-//		FSlateApplication::Get().DismissAllMenus();
-//
-//		FLinearColor BackgroundColor = AnimViewportPtr.Pin()->GetViewportBackgroundColor();
-//		TArray<FLinearColor*> LinearColorArray;
-//		LinearColorArray.Add(&BackgroundColor);
-//
-//		FColorPickerArgs PickerArgs;
-//		PickerArgs.bIsModal = true;
-//		PickerArgs.ParentWidget = AnimViewportPtr.Pin();
-//		PickerArgs.bOnlyRefreshOnOk = true;
-//		PickerArgs.DisplayGamma = TAttribute<float>::Create(TAttribute<float>::FGetter::CreateUObject(GEngine, &UEngine::GetDisplayGamma));
-//		PickerArgs.LinearColorArray = &LinearColorArray;
-//		PickerArgs.OnColorCommitted = FOnLinearColorValueChanged::CreateSP(AnimViewportPtr.Pin().ToSharedRef(), &SCustomizableObjectEditorViewport::SetViewportBackgroundColor);
-//
-//		if (OpenColorPicker(PickerArgs))
-//		{
-//			AnimViewportPtr.Pin()->RefreshViewport();
-//		}
-//
-//		return FReply::Handled();
-//	}
-//
-//	/** Callback function for decreasing color brightness */
-//	FReply OnDecreaseBrightness()
-//	{
-//		const float DeltaValue = 0.05f;
-//		AnimViewportPtr.Pin()->SetBackgroundBrightness(AnimViewportPtr.Pin()->GetBackgroundBrightness() - DeltaValue);
-//		return FReply::Handled();
-//	}
-//
-//	/** Callback function for increasing color brightness */
-//	FReply OnIncreaseBrightness()
-//	{
-//		const float DeltaValue = 0.05f;
-//		AnimViewportPtr.Pin()->SetBackgroundBrightness(AnimViewportPtr.Pin()->GetBackgroundBrightness() + DeltaValue);
-//		return FReply::Handled();
-//	}
-//
-//	/** Callback function which determines whether this widget is enabled */
-//	bool IsBrightnessSliderEnabled() const
-//	{
-//		const UAssetViewerSettings* Settings = UAssetViewerSettings::Get();
-//		const UEditorPerProjectUserSettings* PerProjectUserSettings = GetDefault<UEditorPerProjectUserSettings>();
-//		const int32 ProfileIndex = Settings->Profiles.IsValidIndex(PerProjectUserSettings->AssetViewerProfileIndex) ? PerProjectUserSettings->AssetViewerProfileIndex : 0;
-//
-//		ensureMsgf(Settings->Profiles.IsValidIndex(PerProjectUserSettings->AssetViewerProfileIndex), TEXT("Invalid default settings pointer or current profile index"));
-//
-//		return !Settings->Profiles[ProfileIndex].bShowEnvironment;
-//	}
-//
-//protected:
-//
-//	/** The viewport hosting this widget */
-//	TWeakPtr<SCustomizableObjectEditorViewport> AnimViewportPtr;
-//};
-
-//Class definition which represents widget to modify strength of wind for clothing
-//class SClothWindSettings : public SCompoundWidget
-//{
-//public:
-//
-//	SLATE_BEGIN_ARGS(SClothWindSettings)
-//	{}
-//
-//	SLATE_ARGUMENT(TWeakPtr<SCustomizableObjectEditorViewport>, AnimEditorViewport)
-//		SLATE_END_ARGS()
-//
-//		/** Constructs this widget from its declaration */
-//		void Construct(const FArguments& InArgs)
-//	{
-//		AnimViewportPtr = InArgs._AnimEditorViewport;
-//
-//		TSharedPtr<SWidget> ExtraWidget = SNew(STextBlock)
-//			.Text(AnimViewportPtr.Pin().ToSharedRef(), &SCustomizableObjectEditorViewport::GetWindStrengthLabel)
-//			.Font(UE_MUTABLE_GET_FONTSTYLE(TEXT("MenuItem.Font")));
-//
-//		this->ChildSlot
-//			[
-//				SNew(SAnimPlusMinusSlider)
-//				.IsEnabled(this, &SClothWindSettings::IsWindEnabled)
-//			.Label(LOCTEXT("WindStrength", "Wind Strength:"))
-//			.OnMinusClicked(this, &SClothWindSettings::OnDecreaseWindStrength)
-//			.MinusTooltip(LOCTEXT("DecreaseWindStrength_ToolTip", "Decrease Wind Strength"))
-//			.SliderValue(AnimViewportPtr.Pin().ToSharedRef(), &SCustomizableObjectEditorViewport::GetWindStrengthSliderValue)
-//			.OnSliderValueChanged(AnimViewportPtr.Pin().ToSharedRef(), &SCustomizableObjectEditorViewport::SetWindStrength)
-//			.SliderTooltip(LOCTEXT("WindStrength_ToolTip", "Change wind strength"))
-//			.OnPlusClicked(this, &SClothWindSettings::OnIncreaseWindStrength)
-//			.PlusTooltip(LOCTEXT("IncreasetWindStrength_ToolTip", "Increase Wind Strength"))
-//			.ExtraWidget(ExtraWidget)
-//			];
-//	}
-//
-//protected:
-//	/** Callback function for decreasing size */
-//	FReply OnDecreaseWindStrength()
-//	{
-//		const float DeltaValue = 0.1f;
-//		AnimViewportPtr.Pin()->SetWindStrength(AnimViewportPtr.Pin()->GetWindStrengthSliderValue() - DeltaValue);
-//		return FReply::Handled();
-//	}
-//
-//	/** Callback function for increasing size */
-//	FReply OnIncreaseWindStrength()
-//	{
-//		const float DeltaValue = 0.1f;
-//		AnimViewportPtr.Pin()->SetWindStrength(AnimViewportPtr.Pin()->GetWindStrengthSliderValue() + DeltaValue);
-//		return FReply::Handled();
-//	}
-//
-//	/** Callback function which determines whether this widget is enabled */
-//	bool IsWindEnabled() const
-//	{
-//		return AnimViewportPtr.Pin()->IsApplyingClothWind();
-//	}
-//
-//protected:
-//	/** The viewport hosting this widget */
-//	TWeakPtr<SCustomizableObjectEditorViewport> AnimViewportPtr;
-//};
-
-//Class definition which represents widget to modify gravity for preview
-//class SGravitySettings : public SCompoundWidget
-//{
-//public:
-//
-//	SLATE_BEGIN_ARGS(SGravitySettings)
-//	{}
-//
-//	SLATE_ARGUMENT(TWeakPtr<SCustomizableObjectEditorViewport>, AnimEditorViewport)
-//		SLATE_END_ARGS()
-//
-//		void Construct(const FArguments& InArgs)
-//	{
-//		AnimViewportPtr = InArgs._AnimEditorViewport;
-//
-//		TSharedPtr<SWidget> ExtraWidget = SNew(STextBlock)
-//			.Text(AnimViewportPtr.Pin().ToSharedRef(), &SCustomizableObjectEditorViewport::GetGravityScaleLabel)
-//			.Font(UE_MUTABLE_GET_FONTSTYLE(TEXT("MenuItem.Font")));
-//
-//		this->ChildSlot
-//			[
-//				SNew(SAnimPlusMinusSlider)
-//				.Label(LOCTEXT("Gravity Scale", "Gravity Scale Preview:"))
-//			.OnMinusClicked(this, &SGravitySettings::OnDecreaseGravityScale)
-//			.MinusTooltip(LOCTEXT("DecreaseGravitySize_ToolTip", "Decrease Gravity Scale"))
-//			.SliderValue(AnimViewportPtr.Pin().ToSharedRef(), &SCustomizableObjectEditorViewport::GetGravityScaleSliderValue)
-//			.OnSliderValueChanged(AnimViewportPtr.Pin().ToSharedRef(), &SCustomizableObjectEditorViewport::SetGravityScale)
-//			.SliderTooltip(LOCTEXT("GravityScale_ToolTip", "Change Gravity Scale"))
-//			.OnPlusClicked(this, &SGravitySettings::OnIncreaseGravityScale)
-//			.PlusTooltip(LOCTEXT("IncreaseGravityScale_ToolTip", "Increase Gravity Scale"))
-//			.ExtraWidget(ExtraWidget)
-//			];
-//	}
-//
-//protected:
-//	FReply OnDecreaseGravityScale()
-//	{
-//		const float DeltaValue = 0.025f;
-//		AnimViewportPtr.Pin()->SetGravityScale(AnimViewportPtr.Pin()->GetGravityScaleSliderValue() - DeltaValue);
-//		return FReply::Handled();
-//	}
-//
-//	FReply OnIncreaseGravityScale()
-//	{
-//		const float DeltaValue = 0.025f;
-//		AnimViewportPtr.Pin()->SetGravityScale(AnimViewportPtr.Pin()->GetGravityScaleSliderValue() + DeltaValue);
-//		return FReply::Handled();
-//	}
-//
-//protected:
-//	TWeakPtr<SCustomizableObjectEditorViewport> AnimViewportPtr;
-//};
-
-///////////////////////////////////////////////////////////
-// SCustomizableObjectEditorViewportToolBar
-
-
 void SCustomizableObjectEditorViewportToolBar::Construct(const FArguments& InArgs, TSharedPtr<class SCustomizableObjectEditorViewportTabBody> InViewport, TSharedPtr<class SEditorViewport> InRealViewport)
 {
 	Viewport = InViewport;
 
 	TSharedRef<SCustomizableObjectEditorViewportTabBody> ViewportRef = Viewport.Pin().ToSharedRef();
 
-	//TSharedRef<SCustomizableObjectEditorTransformViewportToolbar> Toolbar = SNew(SCustomizableObjectEditorTransformViewportToolbar)
-	//	.Viewport(InRealViewport.ToSharedRef())
-	//	.ViewportTapBody(ViewportRef);
-
-	//TSharedRef<SHorizontalBox> RightToolbar = SNew(SHorizontalBox)
-	//	+ SHorizontalBox::Slot()
-	//	.AutoWidth()
-	//	.Padding(2.0f, 2.0f)
-	//	[
-	//		Toolbar
-	//	];
-
-	//ViewportRef->SetEditorTransformViewportToolbar(Toolbar);
-
-
 	TSharedRef<SHorizontalBox> LeftToolbar = SNew(SHorizontalBox)
-	//	// Generic viewport options
-	//	+ SHorizontalBox::Slot()
-	//	.AutoWidth()
-	//	.Padding(2.0f, 2.0f)
-	//	[
-	//		//Menu
-	//		SNew(SEditorViewportToolbarMenu)
-	//		.ParentToolBar(SharedThis(this))
-	//	.Image("EditorViewportToolBar.MenuDropdown")
-	//	.OnGetMenuContent(this, &SCustomizableObjectEditorViewportToolBar::GenerateViewMenu)
-	//	]
 
 	//// Camera Type (Perspective/Top/etc...)
 	+ SHorizontalBox::Slot()
@@ -335,17 +74,6 @@ void SCustomizableObjectEditorViewportToolBar::Construct(const FArguments& InArg
 			SNew(SEditorViewportViewMenu, InRealViewport.ToSharedRef(), SharedThis(this))
 		]
 
-	// Show flags menu
-	//+ SHorizontalBox::Slot()
-	//	.AutoWidth()
-	//	.Padding(2.0f, 2.0f)
-	//	[
-	//		SNew(SEditorViewportToolbarMenu)
-	//		.ParentToolBar(SharedThis(this))
-	//	.Label(LOCTEXT("ShowMenu", "Show"))
-	//	.OnGetMenuContent(this, &SCustomizableObjectEditorViewportToolBar::GenerateShowMenu)
-	//	]
-
 	// LOD menu
 	+ SHorizontalBox::Slot()
 		.AutoWidth()
@@ -380,29 +108,6 @@ void SCustomizableObjectEditorViewportToolBar::Construct(const FArguments& InArg
 			RTSButtons
 		];
 
-	// Playback speed menu
-	//+ SHorizontalBox::Slot()
-	//	.AutoWidth()
-	//	.Padding(2.0f, 2.0f)
-	//	[
-	//		SNew(SEditorViewportToolbarMenu)
-	//		.ParentToolBar(SharedThis(this))
-	//	.Label(this, &SCustomizableObjectEditorViewportToolBar::GetPlaybackMenuLabel)
-	//	.LabelIcon(UE_MUTABLE_GET_BRUSH("AnimViewportMenu.PlayBackSpeed"))
-	//	.OnGetMenuContent(this, &SCustomizableObjectEditorViewportToolBar::GeneratePlaybackMenu)
-	//	]
-
-	//+ SHorizontalBox::Slot()
-	//	.Padding(3.0f, 1.0f)
-	//	.HAlign(HAlign_Right)
-	//	[
-	//		SNew(STransformViewportToolBar)
-	//		.Viewport(InRealViewport)
-	//	.CommandList(InRealViewport->GetCommandList())
-	//	.Visibility(this, &SCustomizableObjectEditorViewportToolBar::GetTransformToolbarVisibility)
-	//	];
-	//@TODO: Need clipping horizontal box: LeftToolbar->AddWrapButton();
-
 	static const FName DefaultForegroundName("DefaultForeground");
 
 	FLinearColor ButtonColor1 = FLinearColor(0.1f, 0.1f, 0.1f, 1.0f);
@@ -464,82 +169,11 @@ void SCustomizableObjectEditorViewportToolBar::Construct(const FArguments& InArg
 		];
 
 	SViewportToolBar::Construct(SViewportToolBar::FArguments());
-
-	//CompileErrorLayout.Get()->SetVisibility(EVisibility::Hidden);
 }
 
 EVisibility SCustomizableObjectEditorViewportToolBar::GetShowCompileErrorOverlay() const
 {
-	if ((Viewport.IsValid()) && (Viewport.Pin().Get() != nullptr) &&
-		Viewport.Pin()->CustomizableObjectEditorPtr.IsValid() &&
-		Viewport.Pin().Get()->GetViewportClient().IsValid() &&
-		(Viewport.Pin().Get()->GetViewportClient().Get() != nullptr))
-	{
-		if (Viewport.Pin()->CustomizableObjectEditorPtr.Pin()->GetAssetRegistryLoaded())
-		{
-			UCustomizableObjectInstance* Instance = Viewport.Pin()->CustomizableObjectEditorPtr.Pin()->GetPreviewInstance();
-
-			if ((Instance != nullptr) && !Instance->bEditorPropertyChanged && !Viewport.Pin().Get()->GetViewportClient()->GetGizmoIsProjectorParameterSelected())
-			{
-				if ((Instance->SkeletalMeshStatus == ESkeletalMeshState::UpdateError) ||
-					(Instance->SkeletalMeshStatus == ESkeletalMeshState::PostUpdateError) ||
-					(Instance->SkeletalMeshStatus == ESkeletalMeshState::AsyncUpdatePending))
-				{
-					return EVisibility::Visible;
-				}
-				else
-				{
-					return EVisibility::Hidden;
-				}
-			}
-
-			if ((Instance == nullptr) &&
-				!Viewport.Pin().Get()->GetViewportClient()->GetGizmoIsProjectorParameterSelected())
-			{
-				return EVisibility::Visible;
-			}
-		}
-		else
-		{
-			return EVisibility::Visible;
-		}
-	}
-
-	return EVisibility::Hidden;
-}
-
-
-EVisibility SCustomizableObjectEditorViewportToolBar::GetShowCompileInfoOverlay() const
-{
-	if ((Viewport.IsValid()) && (Viewport.Pin().Get() != nullptr) &&
-		(Viewport.Pin()->CustomizableObjectEditorPtr.IsValid()) &&
-		Viewport.Pin()->CustomizableObjectEditorPtr.Pin()->GetAssetRegistryLoaded())
-	{
-		if (Viewport.Pin()->CustomizableObjectEditorPtr.Pin()->GetAssetRegistryLoaded())
-		{
-			UCustomizableObjectInstance* Instance = Viewport.Pin()->CustomizableObjectEditorPtr.Pin()->GetPreviewInstance();
-
-			if (Instance != nullptr)
-			{
-				if ((Instance->SkeletalMeshStatus == ESkeletalMeshState::UpdateError) ||
-					(Instance->SkeletalMeshStatus == ESkeletalMeshState::PostUpdateError) ||
-					(Instance->SkeletalMeshStatus == ESkeletalMeshState::AsyncUpdatePending))
-				{
-					return EVisibility::Visible;
-				}
-				else
-				{
-					return EVisibility::Hidden;
-				}
-			}
-		}
-		else
-		{
-			return EVisibility::Visible;
-		}
-	}
-
-	return EVisibility::Hidden;
+	return GetCompileErrorOverlayText().IsEmpty() ? EVisibility::Hidden : EVisibility::Visible;
 }
 
 
@@ -564,22 +198,6 @@ FText SCustomizableObjectEditorViewportToolBar::GetCompileErrorOverlayText() con
 				}
 				else if (Instance->SkeletalMeshStatus == ESkeletalMeshState::AsyncUpdatePending)
 				{
-					TSharedPtr<FCustomizableObjectEditorViewportClient> viewport = Viewport.Pin()->GetViewportClient();
-					if (viewport.IsValid())
-					{
-						const GizmoRTSProxy& gizmo = viewport->GetGizmoProxy();
-						if (gizmo.AnyGizmoSelected)
-						{
-							if (gizmo.AssignedDataIsFromNode)
-							{
-								return FText::FromString("Changing default projector");
-							}
-							else
-							{
-								return FText::FromString("Changing current projector instance");
-							}
-						}
-					}
 					return FText::FromString("Updating skeletal mesh");
 				}
 			}
@@ -593,14 +211,13 @@ FText SCustomizableObjectEditorViewportToolBar::GetCompileErrorOverlayText() con
 			return FText::FromString("Loading assets");
 		}
 	}
-	return FText::FromString("");
+	return {};
 }
 
 
 EVisibility SCustomizableObjectEditorViewportToolBar::GetTransformToolbarVisibility() const
 {
 	return EVisibility::Hidden;
-	//return Viewport.Pin()->CanUseGizmos() ? EVisibility::Visible : EVisibility::Hidden;
 }
 
 
@@ -610,257 +227,10 @@ TSharedRef<SWidget> SCustomizableObjectEditorViewportToolBar::GenerateViewMenu()
 
 	const bool bInShouldCloseWindowAfterMenuSelection = true;
 	FMenuBuilder ViewMenuBuilder(bInShouldCloseWindowAfterMenuSelection, Viewport.Pin()->GetCommandList());
-	//{
-	//	ViewMenuBuilder.BeginSection("AnimViewportSceneSetup", LOCTEXT("ViewMenu_SceneSetupLabel", "Scene Setup"));
-	//	{
-	//		ViewMenuBuilder.AddMenuEntry(FCustomizableObjectEditorViewportMenuCommands::Get().PreviewSceneSettings);
-
-	//		// We cant allow animation configuration via the viewport menu in 'old regular' Persona
-	//		if (GetDefault<UPersonaOptions>()->bUseStandaloneAnimationEditors)
-	//		{
-	//			ViewMenuBuilder.AddSubMenu(
-	//				LOCTEXT("SceneSetupLabel", "Scene Setup"),
-	//				LOCTEXT("SceneSetupTooltip", "Set up preview meshes, animations etc."),
-	//				FNewMenuDelegate::CreateRaw(this, &SCustomizableObjectEditorViewportToolBar::GenerateSceneSetupMenu),
-	//				false,
-	//				FSlateIcon(FAppStyle::GetAppStyleSetName(), "AnimViewportMenu.SceneSetup")
-	//			);
-	//		}
-
-	//		ViewMenuBuilder.AddSubMenu(
-	//			LOCTEXT("TurnTableLabel", "Turn Table"),
-	//			LOCTEXT("TurnTableTooltip", "Set up auto-rotation of preview."),
-	//			FNewMenuDelegate::CreateRaw(this, &SCustomizableObjectEditorViewportToolBar::GenerateTurnTableMenu),
-	//			false,
-	//			FSlateIcon(FAppStyle::GetAppStyleSetName(), "AnimViewportMenu.TurnTableSpeed")
-	//		);
-	//	}
-	//	ViewMenuBuilder.EndSection();
-
-	//	ViewMenuBuilder.BeginSection("AnimViewportCamera", LOCTEXT("ViewMenu_CameraLabel", "Camera"));
-	//	{
-	//		ViewMenuBuilder.AddMenuEntry(FCustomizableObjectEditorViewportMenuCommands::Get().CameraFollow);
-	//	}
-	//	ViewMenuBuilder.EndSection();
-	//}
 
 	return ViewMenuBuilder.MakeWidget();
 }
 
-//TSharedRef<SWidget> SCustomizableObjectEditorViewportToolBar::GenerateShowMenu() const
-//{
-//	const FAnimViewportShowCommands& Actions = FAnimViewportShowCommands::Get();
-//
-//	const bool bInShouldCloseWindowAfterMenuSelection = true;
-//	FMenuBuilder ShowMenuBuilder(bInShouldCloseWindowAfterMenuSelection, Viewport.Pin()->GetCommandList());
-//	{
-//		ShowMenuBuilder.BeginSection("AnimViewportFOV", LOCTEXT("Viewport_FOVLabel", "Field Of View"));
-//		{
-//			const float FOVMin = 5.f;
-//			const float FOVMax = 170.f;
-//
-//			TSharedPtr<SWidget> FOVWidget = SNew(SSpinBox<float>)
-//				.Font(UE_MUTABLE_GET_FONTSTYLE(TEXT("MenuItem.Font")))
-//				.MinValue(FOVMin)
-//				.MaxValue(FOVMax)
-//				.Value(this, &SCustomizableObjectEditorViewportToolBar::OnGetFOVValue)
-//				.OnValueChanged(this, &SCustomizableObjectEditorViewportToolBar::OnFOVValueChanged)
-//				.OnValueCommitted(this, &SCustomizableObjectEditorViewportToolBar::OnFOVValueCommitted);
-//
-//			ShowMenuBuilder.AddWidget(FOVWidget.ToSharedRef(), FText());
-//		}
-//		ShowMenuBuilder.EndSection();
-//
-//		ShowMenuBuilder.BeginSection("AnimViewportAudio", LOCTEXT("Viewport_AudioLabel", "Audio"));
-//		{
-//			ShowMenuBuilder.AddMenuEntry(Actions.MuteAudio);
-//		}
-//		ShowMenuBuilder.EndSection();
-//
-//		ShowMenuBuilder.BeginSection("AnimViewportRootMotion", LOCTEXT("Viewport_RootMotionLabel", "Root Motion"));
-//		{
-//			ShowMenuBuilder.AddMenuEntry(Actions.ProcessRootMotion);
-//		}
-//		ShowMenuBuilder.EndSection();
-//
-//		ShowMenuBuilder.BeginSection("AnimViewportMesh", LOCTEXT("ShowMenu_Actions_Mesh", "Mesh"));
-//		{
-//			ShowMenuBuilder.AddMenuEntry(Actions.ShowRetargetBasePose);
-//			ShowMenuBuilder.AddMenuEntry(Actions.ShowBound);
-//			ShowMenuBuilder.AddMenuEntry(Actions.UseInGameBound);
-//			ShowMenuBuilder.AddMenuEntry(Actions.ShowPreviewMesh);
-//			ShowMenuBuilder.AddMenuEntry(Actions.ShowMorphTargets);
-//		}
-//		ShowMenuBuilder.EndSection();
-//
-//		ShowMenuBuilder.BeginSection("AnimViewportAnimation", LOCTEXT("ShowMenu_Actions_Asset", "Asset"));
-//		{
-//			ShowMenuBuilder.AddMenuEntry(Actions.ShowRawAnimation);
-//			ShowMenuBuilder.AddMenuEntry(Actions.ShowNonRetargetedAnimation);
-//			ShowMenuBuilder.AddMenuEntry(Actions.ShowAdditiveBaseBones);
-//			ShowMenuBuilder.AddMenuEntry(Actions.ShowSourceRawAnimation);
-//			ShowMenuBuilder.AddMenuEntry(Actions.ShowBakedAnimation);
-//		}
-//		ShowMenuBuilder.EndSection();
-//
-//		ShowMenuBuilder.BeginSection("AnimViewportPreviewBones", LOCTEXT("ShowMenu_Actions_Bones", "Hierarchy"));
-//		{
-//			ShowMenuBuilder.AddSubMenu(
-//				LOCTEXT("AnimViewportBoneDrawSubMenu", "Bone"),
-//				LOCTEXT("AnimViewportBoneDrawSubMenuToolTip", "Bone Draw Option"),
-//				FNewMenuDelegate::CreateRaw(this, &SCustomizableObjectEditorViewportToolBar::FillShowBoneDrawMenu)
-//			);
-//
-//			ShowMenuBuilder.AddMenuEntry(Actions.ShowSockets);
-//			ShowMenuBuilder.AddMenuEntry(Actions.ShowBoneNames);
-//
-//			ShowMenuBuilder.AddSubMenu(
-//				LOCTEXT("AnimViewportOverlayDrawSubMenu", "Mesh Overlay"),
-//				LOCTEXT("AnimViewportOverlayDrawSubMenuToolTip", "Options"),
-//				FNewMenuDelegate::CreateRaw(this, &SCustomizableObjectEditorViewportToolBar::FillShowOverlayDrawMenu)
-//			);
-//		}
-//		ShowMenuBuilder.EndSection();
-//
-//		ShowMenuBuilder.AddMenuSeparator();
-//		ShowMenuBuilder.AddSubMenu(
-//			LOCTEXT("AnimviewportInfo", "Display Info"),
-//			LOCTEXT("AnimviewportInfoSubMenuToolTip", "Display Mesh Info in Viewport"),
-//			FNewMenuDelegate::CreateRaw(this, &SCustomizableObjectEditorViewportToolBar::FillShowDisplayInfoMenu));
-//
-//
-//		ShowMenuBuilder.AddMenuSeparator();
-//
-//		ShowMenuBuilder.AddSubMenu(
-//			LOCTEXT("AnimViewportSceneSubMenu", "Scene Setup"),
-//			LOCTEXT("AnimViewportSceneSubMenuToolTip", "Options relating to the preview scene"),
-//			FNewMenuDelegate::CreateRaw(this, &SCustomizableObjectEditorViewportToolBar::FillShowSceneMenu));
-//
-//		ShowMenuBuilder.AddSubMenu(
-//			LOCTEXT("AnimViewportAdvancedSubMenu", "Advanced"),
-//			LOCTEXT("AnimViewportAdvancedSubMenuToolTip", "Advanced options"),
-//			FNewMenuDelegate::CreateRaw(this, &SCustomizableObjectEditorViewportToolBar::FillShowAdvancedMenu));
-//	}
-//
-//	return ShowMenuBuilder.MakeWidget();
-//
-//}
-
-//void SCustomizableObjectEditorViewportToolBar::FillShowSceneMenu(FMenuBuilder& MenuBuilder) const
-//{
-//	const FAnimViewportShowCommands& Actions = FAnimViewportShowCommands::Get();
-//
-//	MenuBuilder.BeginSection("AnimViewportAccessory", LOCTEXT("Viewport_AccessoryLabel", "Accessory"));
-//	{
-//		MenuBuilder.AddMenuEntry(Actions.AutoAlignFloorToMesh);
-//	}
-//	MenuBuilder.EndSection();
-//
-//	MenuBuilder.BeginSection("AnimViewportFloorOffset", LOCTEXT("Viewport_FloorOffsetLabel", "Floor Height Offset"));
-//	{
-//		TSharedPtr<SWidget> FloorOffsetWidget = SNew(SNumericEntryBox<float>)
-//			.Font(UE_MUTABLE_GET_FONTSTYLE(TEXT("MenuItem.Font")))
-//			.Value(this, &SCustomizableObjectEditorViewportToolBar::OnGetFloorOffset)
-//			.OnValueChanged(this, &SCustomizableObjectEditorViewportToolBar::OnFloorOffsetChanged)
-//			.ToolTipText(LOCTEXT("FloorOffsetToolTip", "Height offset for the floor mesh (stored per-mesh)"));
-//
-//		MenuBuilder.AddWidget(FloorOffsetWidget.ToSharedRef(), FText());
-//	}
-//	MenuBuilder.EndSection();
-//
-//	MenuBuilder.BeginSection("AnimViewportGrid", LOCTEXT("Viewport_GridLabel", "Grid"));
-//	{
-//		MenuBuilder.AddMenuEntry(Actions.ToggleGrid);
-//	}
-//	MenuBuilder.EndSection();
-//
-//	MenuBuilder.BeginSection("AnimViewportBackground", LOCTEXT("Viewport_BackgroundLabel", "Background"));
-//	{
-//		TSharedPtr<SWidget> BackgroundColorWidget = SNew(SBackgroundColorSettings).AnimEditorViewport(Viewport);
-//		MenuBuilder.AddWidget(BackgroundColorWidget.ToSharedRef(), FText());
-//	}
-//	MenuBuilder.EndSection();
-//}
-
-//void SCustomizableObjectEditorViewportToolBar::FillShowBoneDrawMenu(FMenuBuilder& MenuBuilder) const
-//{
-//	const FAnimViewportShowCommands& Actions = FAnimViewportShowCommands::Get();
-//
-//	MenuBuilder.BeginSection("AnimViewportPreviewHierarchyBoneDraw", LOCTEXT("ShowMenu_Actions_HierarchyAxes", "Hierarchy Local Axes"));
-//	{
-//		MenuBuilder.AddMenuEntry(Actions.ShowBoneDrawAll);
-//		MenuBuilder.AddMenuEntry(Actions.ShowBoneDrawSelected);
-//		MenuBuilder.AddMenuEntry(Actions.ShowBoneDrawNone);
-//	}
-//	MenuBuilder.EndSection();
-//}
-
-//void SCustomizableObjectEditorViewportToolBar::FillShowOverlayDrawMenu(FMenuBuilder& MenuBuilder) const
-//{
-//	const FAnimViewportShowCommands& Actions = FAnimViewportShowCommands::Get();
-//
-//	MenuBuilder.BeginSection("AnimViewportPreviewOverlayDraw", LOCTEXT("ShowMenu_Actions_Overlay", "Overlay Draw"));
-//	{
-//		MenuBuilder.AddMenuEntry(Actions.ShowOverlayNone);
-//		MenuBuilder.AddMenuEntry(Actions.ShowBoneWeight);
-//		MenuBuilder.AddMenuEntry(Actions.ShowMorphTargetVerts);
-//	}
-//	MenuBuilder.EndSection();
-//}
-//
-//void SCustomizableObjectEditorViewportToolBar::FillShowAdvancedMenu(FMenuBuilder& MenuBuilder) const
-//{
-//	const FAnimViewportShowCommands& Actions = FAnimViewportShowCommands::Get();
-//
-//	// Draw UVs
-//	MenuBuilder.BeginSection("UVVisualization", LOCTEXT("UVVisualization_Label", "UV Visualization"));
-//	{
-//		MenuBuilder.AddMenuEntry(FAnimViewportMenuCommands::Get().AnimSetDrawUVs);
-//		MenuBuilder.AddWidget(Viewport.Pin()->UVChannelCombo.ToSharedRef(), FText());
-//	}
-//	MenuBuilder.EndSection();
-//
-//	MenuBuilder.BeginSection("Skinning", LOCTEXT("Skinning_Label", "Skinning"));
-//	{
-//		MenuBuilder.AddMenuEntry(FAnimViewportMenuCommands::Get().SetCPUSkinning);
-//	}
-//	MenuBuilder.EndSection();
-//
-//
-//	MenuBuilder.BeginSection("ShowVertex", LOCTEXT("ShowVertex_Label", "Vertex Normal Visualization"));
-//	{
-//		// Vertex debug flags
-//		MenuBuilder.AddMenuEntry(FAnimViewportMenuCommands::Get().SetShowNormals);
-//		MenuBuilder.AddMenuEntry(FAnimViewportMenuCommands::Get().SetShowTangents);
-//		MenuBuilder.AddMenuEntry(FAnimViewportMenuCommands::Get().SetShowBinormals);
-//	}
-//	MenuBuilder.EndSection();
-//
-//	MenuBuilder.BeginSection("AnimViewportPreviewHierarchyLocalAxes", LOCTEXT("ShowMenu_Actions_HierarchyAxes", "Hierarchy Local Axes"));
-//	{
-//		MenuBuilder.AddMenuEntry(Actions.ShowLocalAxesAll);
-//		MenuBuilder.AddMenuEntry(Actions.ShowLocalAxesSelected);
-//		MenuBuilder.AddMenuEntry(Actions.ShowLocalAxesNone);
-//	}
-//	MenuBuilder.EndSection();
-//}
-//
-//void SCustomizableObjectEditorViewportToolBar::FillShowClothingMenu(FMenuBuilder& MenuBuilder) const
-//{
-//}
-//
-//void SCustomizableObjectEditorViewportToolBar::FillShowDisplayInfoMenu(FMenuBuilder& MenuBuilder) const
-//{
-//	const FAnimViewportShowCommands& Actions = FAnimViewportShowCommands::Get();
-//
-//	// display info levels
-//	{
-//		MenuBuilder.AddMenuEntry(Actions.ShowDisplayInfoBasic);
-//		MenuBuilder.AddMenuEntry(Actions.ShowDisplayInfoDetailed);
-//		MenuBuilder.AddMenuEntry(Actions.ShowDisplayInfoSkelControls);
-//		MenuBuilder.AddMenuEntry(Actions.HideDisplayInfo);
-//	}
-//}
 
 FText SCustomizableObjectEditorViewportToolBar::GetLODMenuLabel() const
 {
@@ -878,6 +248,7 @@ FText SCustomizableObjectEditorViewportToolBar::GetLODMenuLabel() const
 	return Label;
 }
 
+
 TSharedRef<SWidget> SCustomizableObjectEditorViewportToolBar::GenerateRTSButtons()
 {
 	const FCustomizableObjectEditorViewportLODCommands& Actions = FCustomizableObjectEditorViewportLODCommands::Get();
@@ -993,68 +364,14 @@ TSharedRef<SWidget> SCustomizableObjectEditorViewportToolBar::GenerateViewportTy
 
 TSharedRef<SWidget> SCustomizableObjectEditorViewportToolBar::GeneratePlaybackMenu() const
 {
-	//const FCustomizableObjectEditorViewportPlaybackCommands& Actions = FCustomizableObjectEditorViewportPlaybackCommands::Get();
-
 	const bool bInShouldCloseWindowAfterMenuSelection = true;
 
 	FMenuBuilder PlaybackMenuBuilder(bInShouldCloseWindowAfterMenuSelection, Viewport.Pin()->GetCommandList());
-	//{
-	//	// View modes
-	//	{
-	//		PlaybackMenuBuilder.BeginSection("AnimViewportPlaybackSpeed", LOCTEXT("PlaybackMenu_SpeedLabel", "Playback Speed"));
-	//		{
-	//			for (int i = 0; i < EAnimationPlaybackSpeeds::NumPlaybackSpeeds; ++i)
-	//			{
-	//				PlaybackMenuBuilder.AddMenuEntry(Actions.PlaybackSpeedCommands[i]);
-	//			}
-	//		}
-	//		PlaybackMenuBuilder.EndSection();
-	//	}
-	//}
-
+	
 	return PlaybackMenuBuilder.MakeWidget();
 
 }
 
-void SCustomizableObjectEditorViewportToolBar::GenerateTurnTableMenu(FMenuBuilder& MenuBuilder) const
-{
-	//const FCustomizableObjectEditorViewportPlaybackCommands& Actions = FCustomizableObjectEditorViewportPlaybackCommands::Get();
-
-	//const bool bInShouldCloseWindowAfterMenuSelection = true;
-
-	//MenuBuilder.BeginSection("AnimViewportTurnTableMode", LOCTEXT("TurnTableMenu_ModeLabel", "Turn Table Mode"));
-	//{
-	//	MenuBuilder.AddMenuEntry(Actions.PersonaTurnTablePlay);
-	//	MenuBuilder.AddMenuEntry(Actions.PersonaTurnTablePause);
-	//	MenuBuilder.AddMenuEntry(Actions.PersonaTurnTableStop);
-	//}
-	//MenuBuilder.EndSection();
-
-	//MenuBuilder.BeginSection("AnimViewportTurnTableSpeed", LOCTEXT("TurnTableMenu_SpeedLabel", "Turn Table Speed"));
-	//{
-	//	for (int i = 0; i < EAnimationPlaybackSpeeds::NumPlaybackSpeeds; ++i)
-	//	{
-	//		MenuBuilder.AddMenuEntry(Actions.TurnTableSpeeds[i]);
-	//	}
-	//}
-	//MenuBuilder.EndSection();
-}
-
-void SCustomizableObjectEditorViewportToolBar::GenerateSceneSetupMenu(FMenuBuilder& MenuBuilder)
-{
-	//FPropertyEditorModule& PropertyEditorModule = FModuleManager::LoadModuleChecked<FPropertyEditorModule>("PropertyEditor");
-	//FDetailsViewArgs Args(false, false, false, FDetailsViewArgs::HideNameArea, true, nullptr, false, "PersonaPreviewSceneDescription");
-	//Args.bShowScrollBar = false;
-
-	//TSharedRef<IDetailsView> DetailsView = PropertyEditorModule.CreateDetailView(Args);
-
-	//DetailsView->RegisterInstancedCustomPropertyLayout(UPersonaPreviewSceneDescription::StaticClass(), FOnGetDetailCustomizationInstance::CreateSP(this, &SCustomizableObjectEditorViewportToolBar::CustomizePreviewSceneDescription));
-	//PropertyEditorModule.RegisterCustomPropertyTypeLayout(FPreviewMeshCollectionEntry::StaticStruct()->GetFName(), FOnGetPropertyTypeCustomizationInstance::CreateSP(this, &SCustomizableObjectEditorViewportToolBar::CustomizePreviewMeshCollectionEntry), nullptr, DetailsView);
-
-	//DetailsView->SetObject(StaticCastSharedRef<FAnimationEditorPreviewScene>(Viewport.Pin()->GetPreviewScene())->GetPreviewSceneDescription());
-
-	//MenuBuilder.AddWidget(DetailsView, FText(), true);
-}
 
 TSharedRef<SWidget> SCustomizableObjectEditorViewportToolBar::GenerateViewportOptionsMenu() const
 {
@@ -1123,16 +440,6 @@ TSharedRef<SWidget> SCustomizableObjectEditorViewportToolBar::GenerateViewportOp
 }
 
 
-//TSharedRef<class IDetailCustomization> SCustomizableObjectEditorViewportToolBar::CustomizePreviewSceneDescription()
-//{
-//	return MakeShareable(new FPreviewSceneDescriptionCustomization(FAssetData(&Viewport.Pin()->GetSkeletonTree()->GetEditableSkeleton()->GetSkeleton()).GetExportTextName(), Viewport.Pin()->GetPreviewScene()->GetPersonaToolkit()));
-//}
-//
-//TSharedRef<class IPropertyTypeCustomization> SCustomizableObjectEditorViewportToolBar::CustomizePreviewMeshCollectionEntry()
-//{
-//	return MakeShareable(new FPreviewMeshCollectionEntryCustomization(Viewport.Pin()->GetPreviewScene()));
-//}
-
 FSlateColor SCustomizableObjectEditorViewportToolBar::GetFontColor() const
 {
 	const UAssetViewerSettings* Settings = UAssetViewerSettings::Get();
@@ -1166,22 +473,7 @@ FSlateColor SCustomizableObjectEditorViewportToolBar::GetFontColor() const
 
 FText SCustomizableObjectEditorViewportToolBar::GetPlaybackMenuLabel() const
 {
-	FText Label = LOCTEXT("PlaybackError", "Error");
-	//if (Viewport.IsValid())
-	//{
-	//	for (int i = 0; i < EAnimationPlaybackSpeeds::NumPlaybackSpeeds; ++i)
-	//	{
-	//		if (Viewport.Pin()->IsPlaybackSpeedSelected(i))
-	//		{
-	//			Label = FText::FromString(FString::Printf(
-	//				(i == EAnimationPlaybackSpeeds::Quarter || i == EAnimationPlaybackSpeeds::ThreeQuarters) ? TEXT("x%.2f") : TEXT("x%.1f"),
-	//				EAnimationPlaybackSpeeds::Values[i]
-	//			));
-	//			break;
-	//		}
-	//	}
-	//}
-	return Label;
+	return LOCTEXT("PlaybackError", "Error");
 }
 
 FText SCustomizableObjectEditorViewportToolBar::GetCameraMenuLabel() const
@@ -1275,6 +567,7 @@ float SCustomizableObjectEditorViewportToolBar::OnGetFOVValue() const
 	return Viewport.Pin()->GetViewportClient()->ViewFOV;
 }
 
+
 void SCustomizableObjectEditorViewportToolBar::OnFOVValueChanged(float NewValue) const
 {
 	TSharedPtr<FCustomizableObjectEditorViewportClient> ViewportClient = Viewport.Pin()->GetViewportClient();
@@ -1289,10 +582,6 @@ void SCustomizableObjectEditorViewportToolBar::OnFOVValueChanged(float NewValue)
 	ViewportClient->Invalidate();
 }
 
-void SCustomizableObjectEditorViewportToolBar::OnFOVValueCommitted(float NewValue, ETextCommit::Type CommitInfo)
-{
-	//OnFOVValueChanged will be called... nothing needed here.
-}
 
 TOptional<float> SCustomizableObjectEditorViewportToolBar::OnGetFloorOffset() const
 {
@@ -1301,6 +590,7 @@ TOptional<float> SCustomizableObjectEditorViewportToolBar::OnGetFloorOffset() co
 	return ViewportClient->GetFloorOffset();
 }
 
+
 void SCustomizableObjectEditorViewportToolBar::OnFloorOffsetChanged(float NewValue)
 {
 	TSharedPtr<FCustomizableObjectEditorViewportClient> ViewportClient = Viewport.Pin()->GetViewportClient();
diff --git a/Engine/Plugins/Experimental/Mutable/Source/CustomizableObjectEditor/Public/MuCOE/CustomizableObjectEditorActions.h b/Engine/Plugins/Experimental/Mutable/Source/CustomizableObjectEditor/Public/MuCOE/CustomizableObjectEditorActions.h
--- a/Engine/Plugins/Experimental/Mutable/Source/CustomizableObjectEditor/Public/MuCOE/CustomizableObjectEditorActions.h
+++ b/Engine/Plugins/Experimental/Mutable/Source/CustomizableObjectEditor/Public/MuCOE/CustomizableObjectEditorActions.h
@@ -47,7 +47,6 @@ class FCustomizableObjectEditorViewportCommands : public TCommands<FCustomizable
 public:
 	FCustomizableObjectEditorViewportCommands();
 	
-	/**  */
 	TSharedPtr< FUICommandInfo > SetDrawUVs;
 	TSharedPtr< FUICommandInfo > SetShowGrid;
 	TSharedPtr< FUICommandInfo > SetShowSky;
@@ -56,14 +55,7 @@ public:
 	TSharedPtr< FUICommandInfo > SetCameraLock;
 	TSharedPtr< FUICommandInfo > SaveThumbnail;
 
-	// View Menu Commands
-	TSharedPtr< FUICommandInfo > SetShowNormals;
-	TSharedPtr< FUICommandInfo > SetShowTangents;
-	TSharedPtr< FUICommandInfo > SetShowBinormals;
-	TSharedPtr< FUICommandInfo > SetShowPivot;
-
 	TSharedPtr< FUICommandInfo > BakeInstance;
-	TSharedPtr< FUICommandInfo > StateChangeTest;
 	TSharedPtr< FUICommandInfo > StateChangeShowData;
 	TSharedPtr< FUICommandInfo > StateChangeShowGeometryData;
 
diff --git a/Engine/Plugins/Experimental/Mutable/Source/CustomizableObjectEditor/Public/MuCOE/CustomizableObjectEditorViewportClient.h b/Engine/Plugins/Experimental/Mutable/Source/CustomizableObjectEditor/Public/MuCOE/CustomizableObjectEditorViewportClient.h
--- a/Engine/Plugins/Experimental/Mutable/Source/CustomizableObjectEditor/Public/MuCOE/CustomizableObjectEditorViewportClient.h
+++ b/Engine/Plugins/Experimental/Mutable/Source/CustomizableObjectEditor/Public/MuCOE/CustomizableObjectEditorViewportClient.h
@@ -8,7 +8,6 @@
 namespace EAnimationMode { enum Type : int; }
 
 class FCanvas;
-class FCustomizableObjectWidget;
 class FMaterialRenderProxy;
 class FPreviewScene;
 class FPrimitiveDrawInterface;
@@ -31,127 +30,45 @@ class UStaticMeshComponent;
 struct FInputKeyEventArgs;
 
 
-extern void RemoveRestrictedChars(FString& String);
-
-/**
-* Class used to store all the possible gizmo edition options for projectors
-* (constant projector node, parameter projector node, and projector
-* parameter from a Customizable Object Instance
-*/
-class GizmoRTSProxy
-{
-public:
-	GizmoRTSProxy();
-
-	/** Updates the origin data of this gizmo proxy with the information
-    * present in GizmoRTSProxy::Value */
-    void CopyTransformFromOriginData();
-
-	/** Updates the origin data of this gizmo proxy with the information
-	* present in GizmoRTSProxy::Value */
-	bool UpdateOriginData();
-
-	/** Call to the corresponding Modify method when DataOriginParameter
-	* or DataOriginConstant is not nullptr */
-	bool Modify();
-
-	/** Call to the corresponding graph Modify method when DataOriginParameter
-	* or DataOriginConstant is not nullptr */
-	bool ModifyGraph();
-
-	/** Reinitializes the values of DataOriginParameter, DataOriginConstant
-	* and ProjectorParameterName and ProjectorRangeIndex */
-	void CleanOriginData();
-
-	/** Returns true if the values in the variable Parameter given has the ones assigned
-	* in the FCustomizableObjectProjector constructor (to identify uninitialized
-	* parameter proyector data outside the graph editor and initialize it) */
-	bool ProjectorHasInitialValues(FCustomizableObjectProjector& Parameter);
-
-	/** Returns the initial transform data for the projector */
-	FCustomizableObjectProjector SetProjectorInitialValue(TArray<TWeakObjectPtr<class UDebugSkelMeshComponent>>& SkeletalMeshComponents, float TotalLength);
-
-	/** Returns true if a projector parameter is selected in the viewport */
-	bool IsProjectorParameterSelected();
-
-	/** To notify of changes in any parameter projector node property in case is the one being used, currently
-	* just the projector type parameter is taken into account */
-	void ProjectorParameterChanged(UCustomizableObjectNodeProjectorParameter* Node);
-
-	/** To notify of changes in any parameter projector node property in case is the one being used, currently
-	* just the projector type parameter is taken into account */
-	void ProjectorParameterChanged(UCustomizableObjectNodeProjectorConstant* Node);
-
-	/** Setter for CallUpdateSkeletalMesh */
-	void SetCallUpdateSkeletalMesh(bool Value);
-
-	/** Set Projector updated in viewport to skip rebuilding the Instance Properties panel */
-	void SetProjectorUpdatedInViewport(bool Value);
-
-	/** Getter for ProjectorParameterName */
-	FString GetProjectorParameterName();
-
-	/** Getter for ProjectorParameterNameWithIndex */
-	FString GetProjectorParameterNameWithIndex();
-
-	/** Getter for HasAssignedData */
-	bool GetHasAssignedData();
-
-	/** Getter for AssignedDataIsFromNode */
-	bool GetAssignedDataIsFromNode();
+DECLARE_DELEGATE_RetVal(FVector, FWidgetLocationDelegate);
+DECLARE_DELEGATE_OneParam(FOnWidgetLocationChangedDelegate, const FVector&)
 
-	// Projector information for edition
-	FCustomizableObjectProjector Value;
+DECLARE_DELEGATE_RetVal(FVector, FWidgetDirectionDelegate);
+DECLARE_DELEGATE_OneParam(FOnWidgetDirectionChangedDelegate, const FVector&)
 
-	// True if there's a gizmo to be shown, false otherwise
-	bool AnyGizmoSelected;
+DECLARE_DELEGATE_RetVal(FVector, FWidgetUpDelegate);
+DECLARE_DELEGATE_OneParam(FOnWidgetUpChangedDelegate, const FVector&)
 
-	// Pointer to the Customizable Object Node Projector Parameter
-	// in case this proxy is using gizmo data from an instance of this type
-	UCustomizableObjectNodeProjectorParameter* DataOriginParameter;
+DECLARE_DELEGATE_RetVal(FVector, FWidgetScaleDelegate);
+DECLARE_DELEGATE_OneParam(FOnWidgetScaleChangedDelegate, const FVector&)
 
-	// Pointer to the Customizable Object Node Projector Constant
-	// in case this proxy is using gizmo data from an instance of this type
-	UCustomizableObjectNodeProjectorConstant* DataOriginConstant;
+DECLARE_DELEGATE_RetVal(float, FWidgetAngleDelegate);
 
-	// Name of the projecto parameter in case this proxy is using gizmo
-	// data from a Customizable Object Instance Projector parameter data
-	// (an element of CustomizableObjectInstance::ProjectorParameters)
-	FString ProjectorParameterName;
-	FString ProjectorParameterNameWithIndex;
-	int32 ProjectorRangeIndex;
+DECLARE_DELEGATE_RetVal(ECustomizableObjectProjectorType, FProjectorTypeDelegate);
 
-	/** Pointer back to the editor tool that owns the FCustomizableObjectEditorViewportClient instance */
-	TWeakPtr<ICustomizableObjectInstanceEditor> CustomizableObjectEditorPtr;
-
-	/** Simple flag to cache when to update the original data source of the gizmo */
-	bool HasAssignedData;
-
-	/** Simple flag to when the gizmo assigned data is from a node */
-	bool AssignedDataIsFromNode;
+DECLARE_DELEGATE_RetVal(FColor, FWidgetColorDelegate);
 
-	/** Index of the projector parameter in the Customizable Object Instance projector
-	* parameter array in the case of editing a projector parameter */
-	int32 ProjectorParameterIndex;
+DECLARE_DELEGATE(FWidgetTrackingStartedDelegate);
 
-	/** To track edition of the projector gizmo */
-	bool ProjectorGizmoEdited;
 
-	/** In the case of selecting another parameter projector, avoid re-generating the skeletal mesh since no change is done in the selection event */
-	bool CallUpdateSkeletalMesh;
+extern void RemoveRestrictedChars(FString& String);
 
-	bool bManipulatingGizmo;
 
-	// Projection type for the associated texture.
-	ECustomizableObjectProjectorType ProjectionType;
+enum class EWidgetType : uint8
+{
+	Hidden,
+	Projector,
+	ClipMorph,
+	ClipMesh,
+	Light,
 };
 
 
-/** Viewport Client for the preview viewport */
+/** Viewport which shows the scene with the generated Instance. */
 class FCustomizableObjectEditorViewportClient : public FEditorViewportClient, public TSharedFromThis<FCustomizableObjectEditorViewportClient>
 {
 public:
-	FCustomizableObjectEditorViewportClient(TWeakPtr<class ICustomizableObjectInstanceEditor> InCustomizableObjectEditor, FPreviewScene* InPreviewScene);
+	FCustomizableObjectEditorViewportClient(TWeakPtr<ICustomizableObjectInstanceEditor> InCustomizableObjectEditor, FPreviewScene* InPreviewScene);
 	~FCustomizableObjectEditorViewportClient();
 
 	/** persona config options **/
@@ -166,9 +83,8 @@ public:
 
 	virtual bool InputKey(const FInputKeyEventArgs& EventArgs) override;
 	virtual bool InputWidgetDelta(FViewport* Viewport, EAxisList::Type CurrentAxis, FVector& Drag, FRotator& Rot, FVector& Scale) override;
-	virtual void TrackingStarted(const struct FInputEventState& InInputState, bool bIsDragging, bool bNudge) override;
+	virtual void TrackingStarted(const FInputEventState& InInputState, bool bIsDragging, bool bNudge) override;
 	virtual void TrackingStopped() override;
-	virtual UE::Widget::EWidgetMode GetWidgetMode() const override;
 	virtual bool CanSetWidgetMode(UE::Widget::EWidgetMode NewMode) const override;
 	virtual FVector GetWidgetLocation() const override;
 	virtual FMatrix GetWidgetCoordSystem() const override;
@@ -232,22 +148,7 @@ public:
 
 	/** Callback for toggling the bounds show flag. */
 	void SetShowBounds();
-	
-	/** Callback for checking the bounds show flag. */
-	bool IsSetShowBoundsChecked() const;
-
-	/** Callback for toggling the collision geometry show flag. */
-	void SetShowCollision();
-	
-	/** Callback for checking the collision geometry show flag. */
-	bool IsSetShowCollisionChecked() const;
-
-	/** Callback for toggling the pivot show flag. */
-	void SetShowPivot();
-
-	/** Callback for checking the pivot show flag. */
-	bool IsSetShowPivotChecked() const;
-	
+		
 	/** Returns the desired target of the camera */
 	FSphere GetCameraTarget();
 
@@ -259,69 +160,58 @@ public:
 
 	/* Sets the floor height offset, saves it to config and invalidates the viewport so it shows up immediately */
 	void SetFloorOffset(float NewValue);
+	
+	/** Do not call directly. Use ICustomizableObjectEditor functions instead. */
+	void ShowGizmoClipMorph(UCustomizableObjectNodeMeshClipMorph& ClipPlainNode);
+
+	/** Do not call directly. Use ICustomizableObjectEditor functions instead. */
+	void HideGizmoClipMorph();
+	
+	/** Do not call directly. Use ICustomizableObjectEditor functions instead. */
+	void ShowGizmoClipMesh(UCustomizableObjectNodeMeshClipWithMesh& ClipMeshNode, UStaticMesh& ClipMesh);
 
-	void SetClipMorphPlaneVisibility(bool bVisible, const FVector& Origin, const FVector& Normal, float MorphLength, const FBoxSphereBounds& Bounds, float InRadius1, float InRadius2, float InRotationAngle);
-	void SetClipMorphPlaneVisibility(bool bVisible, UCustomizableObjectNodeMeshClipMorph* ClipPlainNode);
-	void SetClipMeshVisibility(bool bVisible, UStaticMesh* ClipMesh, UCustomizableObjectNodeMeshClipWithMesh* ClipMeshNode);
-	void SetProjectorVisibility(bool bVisible, FString ProjectorParameterName, FString ProjectorParameterNameWithIndex, int32 ProjectorRangeIndex, const FCustomizableObjectProjector& Data, int32 ProjectorParameterIndex);
-	void SetProjectorType(bool bVisible, FString ProjectorParameterName, FString ProjectorParameterNameWithIndex, int32 ProjectorRangeIndex, const FCustomizableObjectProjector& Data, int32 ProjectorParameterIndex);
-	void SetProjectorVisibility(bool bVisible, UCustomizableObjectNodeProjectorConstant* Projector);
-	void SetProjectorParameterVisibility(bool bVisible, UCustomizableObjectNodeProjectorParameter* InProjectorParameter);
-	void ResetProjectorVisibility();
+	/** Do not call directly. Use ICustomizableObjectEditor functions instead. */
+	void HideGizmoClipMesh();
+	
+	/** Do not call directly. Use ICustomizableObjectEditor functions instead. */
+	void ShowGizmoProjector(
+		const FWidgetLocationDelegate& WidgetLocationDelegate, const FOnWidgetLocationChangedDelegate& OnWidgetLocationChangedDelegate,
+		const FWidgetDirectionDelegate& WidgetDirectionDelegate, const FOnWidgetDirectionChangedDelegate& OnWidgetDirectionChangedDelegate,
+		const FWidgetUpDelegate& WidgetUpDelegate, const FOnWidgetUpChangedDelegate& OnWidgetUpChangedDelegate,
+		const FWidgetScaleDelegate& WidgetScaleDelegate, const FOnWidgetScaleChangedDelegate& OnWidgetScaleChangedDelegate,
+		const FWidgetAngleDelegate& WidgetAngleDelegate,
+		const FProjectorTypeDelegate& ProjectorTypeDelegate,
+		const FWidgetColorDelegate& WidgetColorDelegate,
+		const FWidgetTrackingStartedDelegate& WidgetTrackingStartedDelegate);
+
+	/** Do not call directly. Use ICustomizableObjectEditor functions instead. */
+	void HideGizmoProjector();
+	
+	/** Do not call directly. Use ICustomizableObjectEditor functions instead. */
+	void ShowGizmoLight(ULightComponent& Light);
+
+	/** Do not call directly. Use ICustomizableObjectEditor functions instead. */
+	void HideGizmoLight();
+	
 	void SetVisibilityForWireframeMode(bool bIsWireframeMode);
 
-	void SetAnimation(class UAnimationAsset* Animation, EAnimationMode::Type AnimationType);
+	void SetAnimation(UAnimationAsset* Animation, EAnimationMode::Type AnimationType);
 
 	/** Set again the animation saved in AnimationBeingPlayed (if any) */
 	void ReSetAnimation();
 
-	/** Add Light Component to the scene */
-	void AddLightToScene(class ULightComponent* AddedLight);
-
-	/** Remove Light component from the scene. use destroy component instead? */
-	void RemoveLightFromScene(class ULightComponent* RemovedLight);
+	/** Add Light Component to the scene. */
+	void AddLightToScene(ULightComponent* AddedLight);
 
-	/** Select/Unselect lights, selected lights can be edited using gizmos */
-	void SetSelectedLight(class ULightComponent* SelectedLight);
+	/** Remove Light component from the scene. */
+	void RemoveLightFromScene(ULightComponent* RemovedLight);
 
-	void SetProjectorWidgetMode(UE::Widget::EWidgetMode InMode);
-
-	const GizmoRTSProxy& GetGizmoProxy();
-
-	/** Getter of GizmoRTSProxy::IsProjectorParameterSelected */
-	bool GetGizmoIsProjectorParameterSelected();
-
-	/** Getter for bManipulating */
-	bool GetIsManipulating();
-
-	/** Getter for WidgetVisibility */
-	bool GetWidgetVisibility();
-
-	/** Update the current gizmo data to its corresponding data origin */
-	void UpdateGizmoDataToOrigin();
-
-	/** Updates the origin data of this gizmo proxy with the information
-    * present in GizmoRTSProxy::Value */
-    void CopyTransformFromOriginData();
-
-	/** Returns true if any projector node is currently selected (in the case of the Customizable Object Editor) */
-	bool AnyProjectorNodeSelected();
-
-	/** Returns in ArrayBoxSize and ArrayBoxColor the data for drawing the visual clue of the performance for the state enter and update test */
-	void BuildMeanTimesBoxSizes(float MaxWidth, TArray<float>& ArrayData, TArray<float>& ArrayBoxSize, TArray<FLinearColor>& ArrayBoxColor);
-
-	/** Returns in BoxSize and BoxColor the data for drawing the visual clue of the performance test for state enter and parameter update */
-	void BuildMeanTimesBoxSize(const float MinValue, const float MaxValue, const float MeanValue, const float MaxWidth, const float Data, float& BoxSize, FLinearColor& BoxColor) const;
+	/** Remove all Light components from the scene. */
+	void RemoveAllLightsFromScene();
 
 	/** Setter of CustomizableObject */
 	void SetCustomizableObject(UCustomizableObject* CustomizableObjectParameter);
 
-	/** Setter for StateChangeShowDataFlag */
-	void SetStateChangeShowDataFlag(bool Value);
-
-	/** Getter for StateChangeShowDataFlag */
-	bool GetStateChangeShowDataFlag();
-
 	/** Helper method to draw shadowed string on viewport */
 	void DrawShadowedString(FCanvas* Canvas, float StartX, float StartY, const FLinearColor& Color, float TextScale, FString String);
 
@@ -345,32 +235,6 @@ public:
 	/** Delegate for preview profile is changed (used for updating show flags) */
 	void OnAssetViewerSettingsChanged(const FName& InPropertyName);
 
-	/** To notify of changes in any parameter projector node property in case is the one being used, currently
-	* just the projector type parameter is taken into account */
-	void ProjectorParameterChanged(UCustomizableObjectNodeProjectorParameter* Node);
-
-	/** To notify of changes in any parameter projector node property in case is the one being used, currently
-	* just the projector type parameter is taken into account */
-	void ProjectorParameterChanged(UCustomizableObjectNodeProjectorConstant* Node);
-
-	/** Setter for LevelViewportClient::GizmoProxy::CallUpdateSkeletalMesh */
-	void SetGizmoCallUpdateSkeletalMesh(bool Value);
-
-	/** Getter for LevelViewportClient::GizmoProxy::ProjectorParameterName */
-	FString GetGizmoProjectorParameterName();
-
-	/** Getter for LevelViewportClient::GizmoProxy::ProjectorParameterNameWithIndex */
-	FString GetGizmoProjectorParameterNameWithIndex();
-
-	/** Getter for LevelViewportClient::GizmoProxy::HasAssignedData */
-	bool GetGizmoHasAssignedData();
-
-	/** Getter for LevelViewportClient::GizmoProxy::AssignedDataIsFromNode */
-	bool GetGizmoAssignedDataIsFromNode();
-
-	/** */
-	bool IsNodeMeshClipMorphSelected();
-
 	/** Debug draw a partial cylinder, given by MaxAngle in [0, 2PI] */
 	void DrawCylinderArc(FPrimitiveDrawInterface* PDI, const FMatrix& CylToWorld, const FVector& Base, const FVector& XAxis, const FVector& YAxis, const FVector& ZAxis, float Radius, float HalfHeight, uint32 Sides, const FMaterialRenderProxy* MaterialRenderProxy, uint8 DepthPriority, FColor Color, float MaxAngle);
 
@@ -401,41 +265,29 @@ public:
 	/** Sets the skeletal mesh bones visibility */
 	void SetShowBones();
 
-	/** Returns true if bones are vibile in viewport */
+	/** Returns true if bones are visible in viewport */
 	bool IsShowingBones() const;
 
+	const TArray<ULightComponent*>& GetLightComponents() const;
+	
+private:
 	/** Draws Mesh Bones in foreground (From: FAnimationViewportClient) */
 	void DrawMeshBones(UDebugSkelMeshComponent* MeshComponent, FPrimitiveDrawInterface* PDI);
 
-private:
+	void SetWidgetType(EWidgetType Type);
+	
 	/** Component for the static/skeletal mesh. */
-	TWeakObjectPtr<class UStaticMeshComponent> StaticMeshComponent;
-	TArray<TWeakObjectPtr<class UDebugSkelMeshComponent>> SkeletalMeshComponents;
-		
-	TMap< int, TSharedPtr<FCustomizableObjectWidget> > ProjectorWidgets;
+	TWeakObjectPtr<UStaticMeshComponent> StaticMeshComponent;
+	TArray<TWeakObjectPtr<UDebugSkelMeshComponent>> SkeletalMeshComponents;
 
-	/** Last hit proxy data. */
-	bool bManipulating;
-	EAxis::Type LastHitProxyAxis;
-	FCustomizableObjectWidget* LastHitProxyWidget;
-
-	/** Drag values for handling moving and rotating widgets. */
-	FVector LocalManDir;
-	FVector WorldManDir;
-	float DragDirX;
-	float DragDirY;
+	/** True if the widget is being dragged. */
+	bool bManipulating = false;
 
 	/** Pointer back to the editor tool that owns us */
 	TWeakPtr<ICustomizableObjectInstanceEditor> CustomizableObjectEditorPtr;
 
 	/** Flags for various options in the editor. */
-	bool bCameraMove;
 	bool bDrawUVs;
-	bool bShowSockets;
-	bool bDrawNormals;
-	bool bDrawTangents;
-	bool bDrawBinormals;
-	bool bShowPivot;
 	bool bDrawSky;
 
 	// Material index to draw in the uv preview.
@@ -449,7 +301,6 @@ private:
 
 	UCustomizableObjectNodeMeshClipMorph* ClipMorphNode;
 	TObjectPtr<UMaterial> ClipMorphMaterial;
-	bool bClipMorphVisible;
 	bool bClipMorphLocalStartOffset;
 	FVector ClipMorphOrigin;
 	FVector ClipMorphOffset;
@@ -465,21 +316,14 @@ private:
 	float Radius2;
 	float RotationAngle;
 
-	UE::Widget::EWidgetMode WidgetMode;
 	FSphere BoundSphere;
 
-	/** Light beeing edited */
-	class ULightComponent* SelectedLightComponent;
-
-	bool bIsEditingLightEnabled;
-
-	/** Proxy class to store the gizmo information needed for gizmo edition for a projector
-	* (constant node, parameter node or Customizable Object Instance projector parameter */
-	GizmoRTSProxy GizmoProxy;
-
-	/** To track Widget::bDefaultVisibility value (there's no getter and FWidget::bDefaultVisibility is private) */
-	bool WidgetVisibility;
+	/** Light being edited */
+	ULightComponent* SelectedLightComponent;
 
+	/** Spawned light components. */
+	TArray<ULightComponent*> LightComponents;
+	
 	/** To know if an animation ia being played by the Customizable Object and restre it after compilation */
 	bool IsPlayingAnimation;
 
@@ -489,7 +333,7 @@ private:
 	/** Customizable object being used */
 	UCustomizableObject* CustomizableObject;
 
-	/** Flag to control whther to show / hide the instance geometry information data */
+	/** Flag to control whether to show / hide the instance geometry information data */
 	bool StateChangeShowGeometryDataFlag;
 
 	/** To know if the user has given permission to overwrite files when baking if already present */
@@ -513,6 +357,29 @@ private:
 	// Temp Instance used in the bake process if a new instance is needed because mutable texture streaming is enabled so the viewport 
 	// instance does not have the high quality mips in the texture's platform data
 	TObjectPtr<UCustomizableObjectInstance> BakeTempInstance = nullptr;
+
+	// The following delegates currently are only used by the EWidgetType::Projector
+	FWidgetLocationDelegate WidgetLocationDelegate;
+	FOnWidgetLocationChangedDelegate OnWidgetLocationChangedDelegate;
+
+	FWidgetLocationDelegate WidgetDirectionDelegate;
+	FOnWidgetLocationChangedDelegate OnWidgetDirectionChangedDelegate;
+
+	FWidgetLocationDelegate WidgetUpDelegate;
+	FOnWidgetLocationChangedDelegate OnWidgetUpChangedDelegate;
+
+	FWidgetLocationDelegate WidgetScaleDelegate;
+	FOnWidgetLocationChangedDelegate OnWidgetScaleChangedDelegate;
+
+	FWidgetAngleDelegate WidgetAngleDelegate; 
+
+	FProjectorTypeDelegate ProjectorTypeDelegate;
+
+	FWidgetColorDelegate WidgetColorDelegate;
+
+	FWidgetTrackingStartedDelegate WidgetTrackingStartedDelegate;
+	
+	EWidgetType WidgetType = EWidgetType::Hidden;
 };
 
 #if UE_ENABLE_INCLUDE_ORDER_DEPRECATED_IN_5_2
diff --git a/Engine/Plugins/Experimental/Mutable/Source/CustomizableObjectEditor/Public/MuCOE/CustomizableObjectInstanceEditorActions.h b/Engine/Plugins/Experimental/Mutable/Source/CustomizableObjectEditor/Public/MuCOE/CustomizableObjectInstanceEditorActions.h
--- a/Engine/Plugins/Experimental/Mutable/Source/CustomizableObjectEditor/Public/MuCOE/CustomizableObjectInstanceEditorActions.h
+++ b/Engine/Plugins/Experimental/Mutable/Source/CustomizableObjectEditor/Public/MuCOE/CustomizableObjectInstanceEditorActions.h
@@ -14,9 +14,7 @@ class FCustomizableObjectInstanceEditorCommands : public TCommands<FCustomizable
 
 public:
 	FCustomizableObjectInstanceEditorCommands();
-	
-	
-	/**  */
+
 	TSharedPtr< FUICommandInfo > SetDrawUVs;
 	TSharedPtr< FUICommandInfo > SetShowGrid;
 	TSharedPtr< FUICommandInfo > SetShowSky;
@@ -25,15 +23,6 @@ public:
 	TSharedPtr< FUICommandInfo > SetCameraLock;
 	TSharedPtr< FUICommandInfo > SaveThumbnail;
 
-	// View Menu Commands
-	TSharedPtr< FUICommandInfo > SetShowNormals;
-	TSharedPtr< FUICommandInfo > SetShowTangents;
-	TSharedPtr< FUICommandInfo > SetShowBinormals;
-	TSharedPtr< FUICommandInfo > SetShowPivot;
-
-	TSharedPtr< FUICommandInfo > BakeInstanceFromInstance;
-	TSharedPtr< FUICommandInfo > StateChangeTest;
-
 	//Toolbar Commands
 	TSharedPtr< FUICommandInfo > ShowParentCO;
 	TSharedPtr< FUICommandInfo > EditParentCO;
diff --git a/Engine/Plugins/Experimental/Mutable/Source/CustomizableObjectEditor/Public/MuCOE/ICustomizableObjectInstanceEditor.h b/Engine/Plugins/Experimental/Mutable/Source/CustomizableObjectEditor/Public/MuCOE/ICustomizableObjectInstanceEditor.h
--- a/Engine/Plugins/Experimental/Mutable/Source/CustomizableObjectEditor/Public/MuCOE/ICustomizableObjectInstanceEditor.h
+++ b/Engine/Plugins/Experimental/Mutable/Source/CustomizableObjectEditor/Public/MuCOE/ICustomizableObjectInstanceEditor.h
@@ -4,8 +4,19 @@
 
 #include "Toolkits/AssetEditorToolkit.h"
 
+class UPointLightComponent;
+class UCustomizableObjectNodeMeshClipWithMesh;
+class UCustomizableObjectNodeProjectorParameter;
+class UCustomizableObjectNodeProjectorConstant;
+class SCustomizableObjectEditorViewportTabBody;
 class UCustomizableObject;
 class UCustomizableObjectInstance;
+class UProjectorParameter;
+class UCustomSettings;
+class UCustomizableObjectNodeMeshClipMorph;
+class UCustomizableObjectNodeMeshClipWithMesh;
+class ULightComponent;
+class UPoseAsset;
 
 /**
  * Public interface to Customizable Object Instance Editor
@@ -16,15 +27,72 @@ public:
 	virtual UCustomizableObjectInstance* GetPreviewInstance() = 0;
 
 	/** Refreshes the Customizable Object Instance Editor's viewport. */
-	virtual void RefreshViewport() = 0;
+	virtual TSharedPtr<SCustomizableObjectEditorViewportTabBody> GetViewport() = 0;
 
 	/** Refreshes everything in the Customizable Object Instance Editor. */
 	virtual void RefreshTool() = 0;
 
 	/** Getter of AssetRegistryLoaded */
 	virtual bool GetAssetRegistryLoaded() = 0;
+	
+	virtual void SetPoseAsset(UPoseAsset* PoseAssetParameter) {}
 
-	virtual void SetPoseAsset(class UPoseAsset* PoseAssetParameter) {};
+	/** Return the selected projector in the viewport/editor. */
+	virtual UProjectorParameter* GetProjectorParameter() = 0;
+
+	virtual UCustomSettings* GetCustomSettings() = 0;
+
+	/** Hide the currently selected gizmo. Notice that only a single gizmo can be shown at the same time.
+	 * Synchronizes all editor widgets so that the selection and the widgets is always consistent. */
+	virtual void HideGizmo() = 0;
+
+	/** Show the default projector value gizmo of the given NodeProjectorConstant.
+	 * Synchronizes all editor widgets so that the selection and the widgets is always consistent. */
+	virtual void ShowGizmoProjectorNodeProjectorConstant(UCustomizableObjectNodeProjectorConstant& Node) {}
+
+	/** Hide the default projector value gizmo of the given NodeProjectorConstant.
+	 * Synchronizes all editor widgets so that the selection and the widgets is always consistent. */
+	virtual void HideGizmoProjectorNodeProjectorConstant() {}
+	
+	/** Show the default projector value gizmo of the given NodeProjectorParameter.
+	 * Synchronizes all editor widgets so that the selection and the widgets is always consistent. */
+	virtual void ShowGizmoProjectorNodeProjectorParameter(UCustomizableObjectNodeProjectorParameter& Node) {}
+
+	/** Hide the default projector value gizmo of the given NodeProjectorParameter.
+	 * Synchronizes all editor widgets so that the selection and the widgets is always consistent. */
+	virtual void HideGizmoProjectorNodeProjectorParameter() {}
+	
+	/** Show the projector value gizmo of the given parameter.
+	 * Synchronizes all editor widgets so that the selection and the widgets is always consistent. */
+	virtual void ShowGizmoProjectorParameter(const FString& ParamName, int32 RangeIndex = -1) = 0;
+
+	/** Hide the projector value gizmo of the given parameter.
+	 * Synchronizes all editor widgets so that the selection and the widgets is always consistent. */
+	virtual void HideGizmoProjectorParameter() = 0;
+	
+	/** Show the clip morph plane gizmo of the NodeClipMorph.
+	 * Synchronizes all editor widgets so that the selection and the widgets is always consistent. */
+	virtual void ShowGizmoClipMorph(UCustomizableObjectNodeMeshClipMorph& ClipPlainNode) {}
+
+	/** Hide the clip morph plane gizmo of the NodeClipMorph.
+	 * Synchronizes all editor widgets so that the selection and the widgets is always consistent. */
+	virtual void HideGizmoClipMorph() {}
+	
+	/** Show the clipping mesh gizmo from the NodeMeshClipWithMesh.
+	 * Synchronizes all editor widgets so that the selection and the widgets is always consistent. */
+	virtual void ShowGizmoClipMesh(UCustomizableObjectNodeMeshClipWithMesh& ClipMeshNode) {}
+
+	/** Hide the clipping mesh gizmo from the NodeMeshClipWithMesh.
+	 * Synchronizes all editor widgets so that the selection and the widgets is always consistent. */
+	virtual void HideGizmoClipMesh() {}
+	
+	/** Show the light gizmo of the light.
+	 * Synchronizes all editor widgets so that the selection and the widgets is always consistent. */
+	virtual void ShowGizmoLight(ULightComponent& SelectedLight) {}
+
+	/** Hide the light gizmo of the light.
+	 * Synchronizes all editor widgets so that the selection and the widgets is always consistent. */
+	virtual void HideGizmoLight() {}
 };
 
 
diff --git a/Engine/Plugins/Experimental/Mutable/Source/CustomizableObjectEditor/Public/MuCOE/SCustomizableInstanceProperties.h b/Engine/Plugins/Experimental/Mutable/Source/CustomizableObjectEditor/Public/MuCOE/SCustomizableInstanceProperties.h
--- a/Engine/Plugins/Experimental/Mutable/Source/CustomizableObjectEditor/Public/MuCOE/SCustomizableInstanceProperties.h
+++ b/Engine/Plugins/Experimental/Mutable/Source/CustomizableObjectEditor/Public/MuCOE/SCustomizableInstanceProperties.h
@@ -2,6 +2,8 @@
 
 #pragma once
 
+#include "CustomizableObjectEditorViewportClient.h"
+#include "ICustomizableObjectEditor.h"
 #include "MuR/Types.h"
 #include "Widgets/Input/SSpinBox.h"
 
@@ -14,6 +16,7 @@ class SSearchBox;
 class STextBlock;
 class SVerticalBox;
 class UCustomizableObjectInstance;
+class FCustomizableObjectEditor;
 struct FGeometry;
 struct FPointerEvent;
 
@@ -33,7 +36,7 @@ public:
 	friend class SCreateProfileParameters;
 	
 	SLATE_BEGIN_ARGS(SCustomizableInstanceProperties) {}
-	SLATE_ARGUMENT(TWeakObjectPtr<UCustomizableObjectInstance>, CustomInstance)
+	SLATE_ARGUMENT(UCustomizableObjectInstance*, CustomInstance) 
 	SLATE_ARGUMENT(TWeakPtr<FCustomizableInstanceDetails>, InstanceDetails)
 	SLATE_END_ARGS()
 	
@@ -51,9 +54,6 @@ public:
 	const UCustomizableObjectInstance* GetInstance() const;
 	void SetNoInstanceMessage( const FText& Message );
 
-	/** Set SimplifiedProjectorInterface bool value - Deprecated */
-	void SetSimplifiedProjectorInterface( bool Enabled );
-
 	/** Redraws the Instance Properties Widget */
 	void ResetParamBox();
 
@@ -61,15 +61,14 @@ public:
 	bool bShouldResetParamBox = true;
 
 private:
-	/** Bool that decides the type of interface of projectors - Deprecated */
-	bool SimplifiedProjectorInterface;
-
 	/** Message to to indicate that there is no instance and the CO needs to be compiled*/
 	FText NoInstanceMessage;
 
 	/** Pointer to the Instance Details. Owner of this widget. */
 	TWeakPtr<FCustomizableInstanceDetails> InstanceDetails;
 
+	TWeakPtr<ICustomizableObjectInstanceEditor> WeakEditor;
+
 	/** Pointer to the Customizable Object Instance */
 	TWeakObjectPtr<UCustomizableObjectInstance> CustomInstance;
 
@@ -171,37 +170,26 @@ private:
 
 	/** Gets the name value of a int parameter */
 	FString GetIntParameterValue(FString ParamName, int RangeIndex = -1) const;
-
-	/** CheckBox IsChecked Callbacks for projectors management */
-	ECheckBoxState IsProjectorHidden(FString ParamNameWithIndex) const;
-	ECheckBoxState IsProjectorTranslate(FString ParamNameWithIndex) const;
-	ECheckBoxState IsProjectorRotate(FString ParamNameWithIndex) const;
-	ECheckBoxState IsProjectorScale(FString ParamNameWithIndex) const;
-
-	/** CheckBox OnCheckStateChanged Callbacks for projectors management */
-	void OnProjectorTranslateStateChanged( ECheckBoxState InCheckboxState, FString ParamNameWithIndex) const;
-	void OnProjectorRotateStateChanged( ECheckBoxState InCheckboxState, FString ParamNameWithIndex) const;
-	void OnProjectorScaleStateChanged( ECheckBoxState InCheckboxState, FString ParamNameWithIndex) const;
 	
 	/** Button OnClicked Callback to select or unselect a projector */
-	FReply OnProjectorSelectStateChanged(FString ParamNameWithIndex);
+	FReply OnProjectorSelectChanged(FString ParamName, int32 RangeIndex) const;
 	
 	/** SearchBox OnTextCommitted Callback to search a parameter texture by name */
-	void OnProjectorTextureParameterComboBoxChanged(TSharedPtr<FString> Selection, ESelectInfo::Type SelectInfo, FString ParamNameWithIndex) const;
+	void OnProjectorTextureParameterComboBoxChanged(TSharedPtr<FString> Selection, ESelectInfo::Type SelectInfo, FString ParamName, int32 RangeIndex) const;
 
 	/** Slider OnValueChanged Callbacks for float parameters */
-	void OnProjectorFloatParameterChanged(float Value, int SliderIndex, FString ParamNameWithIndex);
+	void OnProjectorFloatParameterChanged(float Value, int SliderIndex);
 	void OnProjectorFloatParameterSliderBegin();
 	void OnProjectorFloatParameterSliderEnd(float Value);
 
 	/** OnClicked Button Callbacks - Manage Projector layers */
 	FReply OnProjectorLayerAdded(FString ParamName) const;
-	FReply OnProjectorLayerRemoved(FString ParamNameWithIndex) const;
+	FReply OnProjectorLayerRemoved(FString ParamName, int32 RangeIndex) const;
 
 	/** OnClicked Button Callbacks - Manage Projectors Transform */
-    FReply OnProjectorCopyTransform(FString ParamName) const;
-    FReply OnProjectorPasteTransform(FString ParamName) const;
-    FReply OnProjectorResetTransform(FString ParamName) const;
+    FReply OnProjectorCopyTransform(FString ParamName, int32 RangeIndex) const;
+    FReply OnProjectorPasteTransform(FString ParamName, int32 RangeIndex);
+    FReply OnProjectorResetTransform(FString ParamName, int32 RangeIndex);
 
 	/** Gets the value of a vector parameter */
 	FLinearColor GetColorParameterValue(FString ParamName) const;
@@ -215,16 +203,8 @@ private:
 	/** ComboBox OnSelectionChanged Callback - Changes the texture of a projector */
 	void OnTextureParameterComboBoxSelectionChanged(TSharedPtr<FString> Selection, ESelectInfo::Type SelectInfo, FString ParamName);
 
-	/** Array with the names of the current projector parameter names with their multidimensional (range) index */
-	TArray<FString> ArrayProjectorParameterNameWithIndex;
-
 	bool bSlidersChanged = false;
 
-	// Map where to store the instances of the select / unselect projector SButton
-	TMap<FString, TSharedPtr<SButton>> MapProjectorButton;
-	TMap<FString, TSharedPtr<STextBlock>> MapProjectorButtonTextBlock;
-
-
 	/** Called when the filter text changes.  This filters specific property nodes out of view */
 	void OnFilterTextChanged(const FText& InFilterText);
 
@@ -258,6 +238,8 @@ private:
 	bool HasAnyParameters() const;
 
 	void InstanceUpdated(UCustomizableObjectInstance* Instance) const;
+
+	TSharedPtr<ICustomizableObjectInstanceEditor> GetEditorChecked() const;
 	
 	/*Check if SaveDescriptor has been called before LoadDescriptor*/
 	bool ParametersHaveBeenRead;
diff --git a/Engine/Plugins/Experimental/Mutable/Source/CustomizableObjectEditor/Public/MuCOE/SCustomizableObjectEditorViewportToolBar.h b/Engine/Plugins/Experimental/Mutable/Source/CustomizableObjectEditor/Public/MuCOE/SCustomizableObjectEditorViewportToolBar.h
--- a/Engine/Plugins/Experimental/Mutable/Source/CustomizableObjectEditor/Public/MuCOE/SCustomizableObjectEditorViewportToolBar.h
+++ b/Engine/Plugins/Experimental/Mutable/Source/CustomizableObjectEditor/Public/MuCOE/SCustomizableObjectEditorViewportToolBar.h
@@ -11,6 +11,7 @@ namespace ETextCommit { enum Type : int; }
 class FMenuBuilder;
 class SMenuAnchor;
 class SWidget;
+class SButton;
 struct FSlateBrush;
 
 
@@ -32,43 +33,7 @@ private:
 	* @return The widget containing the view menu content
 	*/
 	TSharedRef<SWidget> GenerateViewMenu() const;
-
-	/**
-	* Generates the toolbar show menu content
-	*
-	* @return The widget containing the show menu content
-	*/
-	//TSharedRef<SWidget> GenerateShowMenu() const;
-
-	/**
-	* Generates the Show -> Scene sub menu content
-	*/
-	//void FillShowSceneMenu(FMenuBuilder& MenuBuilder) const;
-
-	/**
-	* Generates the Show -> Advanced sub menu content
-	*/
-	//void FillShowAdvancedMenu(FMenuBuilder& MenuBuilder) const;
-
-	/**
-	* Generates the Show -> Bone sub menu content
-	*/
-	//void FillShowBoneDrawMenu(FMenuBuilder& MenuBuilder) const;
-
-	/**
-	* Generates the Show -> Overlay sub menu content
-	*/
-	//void FillShowOverlayDrawMenu(FMenuBuilder& MenuBuilder) const;
-
-	/**
-	* Generates the Show -> Clothing sub menu content
-	*/
-	//void FillShowClothingMenu(FMenuBuilder& MenuBuilder) const;
-
-	/**
-	* Generates the Show -> Display Info sub menu content
-	*/
-	//void FillShowDisplayInfoMenu(FMenuBuilder& MenuBuilder) const;
+	
 
 	/** Add the projector Rotation, traslation and scale buttons to viewport toolbar */
 	TSharedRef<SWidget> GenerateRTSButtons();
@@ -101,12 +66,6 @@ private:
 	*/
 	TSharedRef<SWidget> GeneratePlaybackMenu() const;
 
-	/** Generate the turntable menu entries */
-	void GenerateTurnTableMenu(FMenuBuilder& MenuBuilder) const;
-
-	/** Generate the scene setup menu */
-	void GenerateSceneSetupMenu(FMenuBuilder& MenuBuilder);
-
 	/**
 	* Generates the toolbar Options Mode menu content
 	*
@@ -114,13 +73,6 @@ private:
 	*/
 	TSharedRef<SWidget> GenerateViewportOptionsMenu() const;
 
-
-	/** Customize the details of the scene setup object */
-	//TSharedRef<class IDetailCustomization> CustomizePreviewSceneDescription();
-
-	/** Customize a preview mesh collection entry */
-	//TSharedRef<class IPropertyTypeCustomization> CustomizePreviewMeshCollectionEntry();
-
 	/**
 	* Generate color of the text on the top
 	*/
@@ -145,8 +97,6 @@ private:
 	float OnGetFOVValue() const;
 	/** Called when the FOV slider is adjusted in the perspective viewport */
 	void OnFOVValueChanged(float NewValue) const;
-	/** Called when a value is entered into the FOV slider/box in the perspective viewport */
-	void OnFOVValueCommitted(float NewValue, ETextCommit::Type CommitInfo);
 
 	/** Called by the floor offset slider in the perspective viewport to get the offset value */
 	TOptional<float> OnGetFloorOffset() const;
@@ -156,15 +106,12 @@ private:
 	// Called to determine if the gizmos can be used in the current preview
 	EVisibility GetTransformToolbarVisibility() const;
 
-	// Called to show / hide Customizalbe Object compile error
+	// Called to show / hide Customizable Object compile error
 	EVisibility GetShowCompileErrorOverlay() const;
 
 	// Called to modify the instance compile status overlay information text
 	FText GetCompileErrorOverlayText() const;
 
-	// Called to show / hide Customizalbe Object info layout
-	EVisibility GetShowCompileInfoOverlay() const;
-
 	ECheckBoxState IsRotationGridSnapChecked() const;
 	void HandleToggleRotationGridSnap(ECheckBoxState InState);
 	FText GetRotationGridLabel() const;
@@ -180,17 +127,11 @@ private:
 	/** Generates widgets for viewport camera FOV control */
 	TSharedRef<SWidget> GenerateFOVMenu() const;
 
-private:
 	/** The viewport that we are in */
-	TWeakPtr<class SCustomizableObjectEditorViewportTabBody> Viewport;
+	TWeakPtr<SCustomizableObjectEditorViewportTabBody> Viewport;
 
 	// Layout to show information about instance skeletal mesh update / CO asset data
-	TSharedPtr<class SButton> CompileErrorLayout;
-
-	// Layout to show if the skeletal mesh of a functional CO is being updated
-	TSharedPtr<class SButton> UpdateInfoLayout;
-
-	SViewportToolBar* CastedViewportToolbar;
+	TSharedPtr<SButton> CompileErrorLayout;
 
 	TSharedPtr<SMenuAnchor> MenuAnchor;
 };
