[Mutable] Fix incorrect Update flow result code.

#rnx

diff --git a/Engine/Plugins/Experimental/Mutable/Source/CustomizableObject/Private/MuCO/CustomizableObjectSystem.cpp b/Engine/Plugins/Experimental/Mutable/Source/CustomizableObject/Private/MuCO/CustomizableObjectSystem.cpp
--- a/Engine/Plugins/Experimental/Mutable/Source/CustomizableObject/Private/MuCO/CustomizableObjectSystem.cpp
+++ b/Engine/Plugins/Experimental/Mutable/Source/CustomizableObject/Private/MuCO/CustomizableObjectSystem.cpp
@@ -2533,7 +2533,7 @@ namespace impl
 			FinishUpdateGlobal(nullptr, EUpdateResult::Error, &Operation->UpdateCallback);
 			return;
 		}
-
+		
 		CandidateInstancePrivateData->InstanceUpdateFlags(*CandidateInstance);
 		
 		if (CandidateInstancePrivateData->HasCOInstanceFlags(PendingLODsUpdate))
@@ -2543,6 +2543,17 @@ namespace impl
 			//Operation->CustomizableObjectInstance->ReleaseMutableInstanceId(); // To make mutable regenerate the LODs even if the instance parameters have not changed
 		}
 
+		// Skip update, the requested update is equal to the running update.
+		if (Operation->InstanceDescriptorRuntimeHash.IsSubset(CandidateInstance->GetDescriptorRuntimeHash()))
+		{
+			CandidateInstance->SkeletalMeshStatus = ESkeletalMeshState::Correct;
+
+			CandidateInstancePrivateData->ClearCOInstanceFlags(Updating);
+			System->ClearCurrentMutableOperation();
+			UpdateSkeletalMesh(*CandidateInstance, CandidateInstance->GetDescriptorRuntimeHash(), EUpdateResult::ErrorOptimized, &Operation->UpdateCallback);
+			return;
+		}
+
 		bool bCancel = false;
 
 		// If the object is locked (for instance, compiling) we skip any instance update.
@@ -2570,14 +2581,6 @@ namespace impl
 			bCancel = true;
 		}
 
-		// Skip update, the requested update is equal to the running update.
-		if (Operation->InstanceDescriptorRuntimeHash.IsSubset(CandidateInstance->GetDescriptorRuntimeHash()))
-		{
-			CandidateInstance->SkeletalMeshStatus = ESkeletalMeshState::Correct;
-			UpdateSkeletalMesh(*CandidateInstance, CandidateInstance->GetDescriptorRuntimeHash(), EUpdateResult::Success, &Operation->UpdateCallback);
-			bCancel = true;
-		}
-
 		mu::Ptr<const mu::Parameters> Parameters = Operation->GetParameters();
 		if (!Parameters)
 		{
