[mutable] Now, table color parameters set to "none" will use the default value of material instance of Customizable Object Instance.

+Fixed "Add none option"  where a true node set the all the rest of nodes (these using the same data table) to true. Now this option is stored on each column nodes and not in the mutable table.
+Changed some code to follow the coding standards of UE.

#rnx

diff --git a/Engine/Plugins/Experimental/Mutable/Source/CustomizableObject/Private/MuCO/CustomizableObjectInstance.cpp b/Engine/Plugins/Experimental/Mutable/Source/CustomizableObject/Private/MuCO/CustomizableObjectInstance.cpp
--- a/Engine/Plugins/Experimental/Mutable/Source/CustomizableObject/Private/MuCO/CustomizableObjectInstance.cpp
+++ b/Engine/Plugins/Experimental/Mutable/Source/CustomizableObject/Private/MuCO/CustomizableObjectInstance.cpp
@@ -6008,7 +6008,15 @@ void UCustomizableInstancePrivateData::BuildMaterials(const TSharedPtr<FMutableO
 							case FMutableMaterialPlaceholder::EPlaceHolderParamType::Vector:
 								if (Param.LayerIndex < 0)
 								{
-									MaterialInstance->SetVectorParameterValue(Param.ParamName, Param.Vector);
+									FLinearColor Color = Param.Vector;
+
+									if (FVector4f(Color) == mu::DefaultMutableColorValue)
+									{
+										FMaterialParameterInfo ParameterInfo(Param.ParamName);
+										MaterialTemplate->GetVectorParameterValue(ParameterInfo, Color);
+									}
+
+									MaterialInstance->SetVectorParameterValue(Param.ParamName, Color);
 								}
 								else
 								{
diff --git a/Engine/Plugins/Experimental/Mutable/Source/CustomizableObjectEditor/Private/MuCOE/GenerateMutableSource/GenerateMutableSourceColor.cpp b/Engine/Plugins/Experimental/Mutable/Source/CustomizableObjectEditor/Private/MuCOE/GenerateMutableSource/GenerateMutableSourceColor.cpp
--- a/Engine/Plugins/Experimental/Mutable/Source/CustomizableObjectEditor/Private/MuCOE/GenerateMutableSource/GenerateMutableSourceColor.cpp
+++ b/Engine/Plugins/Experimental/Mutable/Source/CustomizableObjectEditor/Private/MuCOE/GenerateMutableSource/GenerateMutableSourceColor.cpp
@@ -322,6 +322,7 @@ mu::NodeColourPtr GenerateMutableSourceColor(const UEdGraphPin* Pin, FMutableGra
 						ColorTableNode->SetTable(Table);
 						ColorTableNode->SetColumn(StringCast<ANSICHAR>(*ColumnName).Get());
 						ColorTableNode->SetParameterName(StringCast<ANSICHAR>(*TypedNodeTable->ParameterName).Get());
+						ColorTableNode->SetNoneOption(TypedNodeTable->bAddNoneOption);
 
 						GenerationContext.AddParameterNameUnique(Node, TypedNodeTable->ParameterName);
 					}
diff --git a/Engine/Plugins/Experimental/Mutable/Source/CustomizableObjectEditor/Private/MuCOE/GenerateMutableSource/GenerateMutableSourceFloat.cpp b/Engine/Plugins/Experimental/Mutable/Source/CustomizableObjectEditor/Private/MuCOE/GenerateMutableSource/GenerateMutableSourceFloat.cpp
--- a/Engine/Plugins/Experimental/Mutable/Source/CustomizableObjectEditor/Private/MuCOE/GenerateMutableSource/GenerateMutableSourceFloat.cpp
+++ b/Engine/Plugins/Experimental/Mutable/Source/CustomizableObjectEditor/Private/MuCOE/GenerateMutableSource/GenerateMutableSourceFloat.cpp
@@ -362,6 +362,7 @@ mu::NodeScalarPtr GenerateMutableSourceFloat(const UEdGraphPin* Pin, FMutableGra
 						ScalarTableNode->SetTable(Table);
 						ScalarTableNode->SetColumn(StringCast<ANSICHAR>(*ColumnName).Get());
 						ScalarTableNode->SetParameterName(StringCast<ANSICHAR>(*TypedNodeTable->ParameterName).Get());
+						ScalarTableNode->SetNoneOption(TypedNodeTable->bAddNoneOption);
 
 						GenerationContext.AddParameterNameUnique(Node, TypedNodeTable->ParameterName);
 					}
diff --git a/Engine/Plugins/Experimental/Mutable/Source/CustomizableObjectEditor/Private/MuCOE/GenerateMutableSource/GenerateMutableSourceImage.cpp b/Engine/Plugins/Experimental/Mutable/Source/CustomizableObjectEditor/Private/MuCOE/GenerateMutableSource/GenerateMutableSourceImage.cpp
--- a/Engine/Plugins/Experimental/Mutable/Source/CustomizableObjectEditor/Private/MuCOE/GenerateMutableSource/GenerateMutableSourceImage.cpp
+++ b/Engine/Plugins/Experimental/Mutable/Source/CustomizableObjectEditor/Private/MuCOE/GenerateMutableSource/GenerateMutableSourceImage.cpp
@@ -1047,6 +1047,7 @@ mu::NodeImagePtr GenerateMutableSourceImage(const UEdGraphPin* Pin, FMutableGrap
 							ImageTableNode->SetTable(Table);
 							ImageTableNode->SetColumn(StringCast<ANSICHAR>(*ColumnName).Get());
 							ImageTableNode->SetParameterName(StringCast<ANSICHAR>(*TypedNodeTable->ParameterName).Get());
+							ImageTableNode->SetNoneOption(TypedNodeTable->bAddNoneOption);
 
 							GenerationContext.AddParameterNameUnique(Node, TypedNodeTable->ParameterName);
 						}
diff --git a/Engine/Plugins/Experimental/Mutable/Source/CustomizableObjectEditor/Private/MuCOE/GenerateMutableSource/GenerateMutableSourceMesh.cpp b/Engine/Plugins/Experimental/Mutable/Source/CustomizableObjectEditor/Private/MuCOE/GenerateMutableSource/GenerateMutableSourceMesh.cpp
--- a/Engine/Plugins/Experimental/Mutable/Source/CustomizableObjectEditor/Private/MuCOE/GenerateMutableSource/GenerateMutableSourceMesh.cpp
+++ b/Engine/Plugins/Experimental/Mutable/Source/CustomizableObjectEditor/Private/MuCOE/GenerateMutableSource/GenerateMutableSourceMesh.cpp
@@ -3670,6 +3670,7 @@ mu::NodeMeshPtr GenerateMutableSourceMesh(const UEdGraphPin* Pin,
 						MeshTableNode->SetTable(Table);
 						MeshTableNode->SetColumn(StringCast<ANSICHAR>(*MutableColumnName).Get());
 						MeshTableNode->SetParameterName(StringCast<ANSICHAR>(*TypedNodeTable->ParameterName).Get());
+						MeshTableNode->SetNoneOption(TypedNodeTable->bAddNoneOption);
 
 						GenerationContext.AddParameterNameUnique(Node, TypedNodeTable->ParameterName);
 
diff --git a/Engine/Plugins/Experimental/Mutable/Source/CustomizableObjectEditor/Private/MuCOE/GenerateMutableSource/GenerateMutableSourceTable.cpp b/Engine/Plugins/Experimental/Mutable/Source/CustomizableObjectEditor/Private/MuCOE/GenerateMutableSource/GenerateMutableSourceTable.cpp
--- a/Engine/Plugins/Experimental/Mutable/Source/CustomizableObjectEditor/Private/MuCOE/GenerateMutableSource/GenerateMutableSourceTable.cpp
+++ b/Engine/Plugins/Experimental/Mutable/Source/CustomizableObjectEditor/Private/MuCOE/GenerateMutableSource/GenerateMutableSourceTable.cpp
@@ -695,9 +695,6 @@ mu::TablePtr GenerateMutableSourceTable(const FString& TableName, const UEdGraph
 			
 			ParameterUIData.IntegerParameterGroupType = TypedTable->bAddNoneOption ? ECustomizableObjectGroupType::COGT_ONE_OR_NONE : ECustomizableObjectGroupType::COGT_ONE;
 
-			// Adding a None Option
-			MutableTable->SetNoneOption(TypedTable->bAddNoneOption);
-
 			for (int32 i = 0; i < RowNames.Num(); ++i)
 			{
 				MutableTable->AddRow(i);
diff --git a/Engine/Plugins/Experimental/Mutable/Source/MutableRuntime/Public/MuR/Parameters.h b/Engine/Plugins/Experimental/Mutable/Source/MutableRuntime/Public/MuR/Parameters.h
--- a/Engine/Plugins/Experimental/Mutable/Source/MutableRuntime/Public/MuR/Parameters.h
+++ b/Engine/Plugins/Experimental/Mutable/Source/MutableRuntime/Public/MuR/Parameters.h
@@ -7,6 +7,7 @@
 #include "MuR/RefCounted.h"
 #include "MuR/Types.h"
 #include "Templates/SharedPointer.h"
+#include "VectorTypes.h"
 
 
 namespace mu
@@ -27,6 +28,9 @@ namespace mu
 
 	class Model;
 
+	//! Default parameter values definition
+	static const FVector4f DefaultMutableColorValue(-UE_BIG_NUMBER);
+
 
     //! \brief Model parameter types.
 	//! \ingroup runtime
diff --git a/Engine/Plugins/Experimental/Mutable/Source/MutableTools/Private/MuT/CodeGenerator.cpp b/Engine/Plugins/Experimental/Mutable/Source/MutableTools/Private/MuT/CodeGenerator.cpp
--- a/Engine/Plugins/Experimental/Mutable/Source/MutableTools/Private/MuT/CodeGenerator.cpp
+++ b/Engine/Plugins/Experimental/Mutable/Source/MutableTools/Private/MuT/CodeGenerator.cpp
@@ -296,7 +296,7 @@ namespace mu
 
 
 	//---------------------------------------------------------------------------------------------
-	Ptr<ASTOp> CodeGenerator::GenerateTableVariable(TablePtr pTable, const FString& strName)
+	Ptr<ASTOp> CodeGenerator::GenerateTableVariable(TablePtr pTable, const FString& strName, bool bAddNoneOption)
 	{
 		Ptr<ASTOp> result;
 
@@ -323,7 +323,7 @@ namespace mu
 				}
 			}
 
-			if (pTable->GetPrivate()->bNoneOption)
+			if (bAddNoneOption)
 			{
 				FParameterDesc::FIntValueDesc nullValue;
 				nullValue.m_value = -1;
diff --git a/Engine/Plugins/Experimental/Mutable/Source/MutableTools/Private/MuT/CodeGenerator.h b/Engine/Plugins/Experimental/Mutable/Source/MutableTools/Private/MuT/CodeGenerator.h
--- a/Engine/Plugins/Experimental/Mutable/Source/MutableTools/Private/MuT/CodeGenerator.h
+++ b/Engine/Plugins/Experimental/Mutable/Source/MutableTools/Private/MuT/CodeGenerator.h
@@ -304,14 +304,14 @@ namespace mu
 
         //-----------------------------------------------------------------------------------------
         //!
-        Ptr<ASTOp> GenerateTableVariable( TablePtr pTable, const FString& strName );
+        Ptr<ASTOp> GenerateTableVariable( TablePtr pTable, const FString& strName, bool bAddNoneOption );
 
         //!
         Ptr<ASTOp> GenerateMissingBoolCode(const TCHAR* strWhere, bool value, const void* errorContext );
 
         //!
 		template<class NODE_TABLE_PRIVATE, ETableColumnType TYPE, OP_TYPE OPTYPE, typename F>
-		Ptr<ASTOp> GenerateTableSwitch( const NODE_TABLE_PRIVATE& node, F&& GenerateOption );
+		Ptr<ASTOp> GenerateTableSwitch( const NODE_TABLE_PRIVATE& node, F&& GenerateOption, Ptr<ASTOp> DefaultValue );
 
 
     private:
@@ -758,40 +758,40 @@ namespace mu
     template<class NODE_TABLE_PRIVATE, ETableColumnType TYPE, OP_TYPE OPTYPE, typename F>
     Ptr<ASTOp> CodeGenerator::GenerateTableSwitch
         (
-            const NODE_TABLE_PRIVATE& node, F&& GenerateOption
+            const NODE_TABLE_PRIVATE& node, F&& GenerateOption, Ptr<ASTOp> DefaultValue
         )
     {
-        TablePtr pTable;
-        Ptr<ASTOp> variable;
+        TablePtr Table;
+        Ptr<ASTOp> Variable;
 
         map< std::pair<TablePtr,FString>, std::pair<TablePtr,Ptr<ASTOp>> >::iterator it
                 = m_generatedTables.find
-                ( std::pair<TablePtr,FString>(node.m_pTable,node.m_parameterName) );
+                ( std::pair<TablePtr,FString>(node.Table,node.ParameterName) );
         if ( it!=m_generatedTables.end() )
         {
-            pTable = it->second.first;
-            variable = it->second.second;
+            Table = it->second.first;
+            Variable = it->second.second;
         }
 
-        if ( !pTable )
+        if ( !Table )
         {
             // Create the table variable expression
-            pTable = node.m_pTable;
-            variable = GenerateTableVariable( pTable, node.m_parameterName );
+            Table = node.Table;
+            Variable = GenerateTableVariable( Table, node.ParameterName, node.bNoneOption);
 
-            m_generatedTables[ std::pair<TablePtr, FString>(node.m_pTable,node.m_parameterName) ] =
-                    std::pair<TablePtr,Ptr<ASTOp>>( pTable, variable );
+            m_generatedTables[ std::pair<TablePtr,FString>(node.Table,node.ParameterName) ] =
+                    std::pair<TablePtr,Ptr<ASTOp>>( Table, Variable );
         }
 
         // Verify that the table column is the right type
-        int colIndex = pTable->FindColumn( node.m_columnName );
-        if ( colIndex<0 )
+        int ColIndex = Table->FindColumn( node.ColumnName );
+        if ( ColIndex<0 )
         {
             m_pErrorLog->GetPrivate()->Add("Table column not found.", ELMT_ERROR, node.m_errorContext);
             return nullptr;
         }
 
-        if ( pTable->GetPrivate()->Columns[ colIndex ].Type != TYPE )
+        if ( Table->GetPrivate()->Columns[ ColIndex ].Type != TYPE )
         {
             m_pErrorLog->GetPrivate()->Add("Table column type is not the right type.",
                                            ELMT_ERROR, node.m_errorContext);
@@ -800,22 +800,22 @@ namespace mu
 
         // Create the switch to cover all the options
         Ptr<ASTOp> lastSwitch;
-        int32 rows = pTable->GetPrivate()->Rows.Num();
+        int32 Rows = Table->GetPrivate()->Rows.Num();
 
         Ptr<ASTOpSwitch> SwitchOp = new ASTOpSwitch();
 		SwitchOp->type = OPTYPE;
-		SwitchOp->variable = variable;
-		SwitchOp->def = nullptr;
+		SwitchOp->variable = Variable;
+		SwitchOp->def = DefaultValue;
 
-		for (int32 i = 0; i < rows; ++i)
+		for (int32 i = 0; i < Rows; ++i)
         {
-            check( pTable->GetPrivate()->Rows[i].Id <= 0xFFFF);
-            auto condition = (uint16)pTable->GetPrivate()->Rows[i].Id;
-            Ptr<ASTOp> Branch = GenerateOption( node, colIndex, (int)i, m_pErrorLog.get() );
+            check( Table->GetPrivate()->Rows[i].Id <= 0xFFFF);
+            auto Condition = (uint16)Table->GetPrivate()->Rows[i].Id;
+            Ptr<ASTOp> Branch = GenerateOption( node, ColIndex, (int)i, m_pErrorLog.get() );
 
 			if (Branch || TYPE != ETableColumnType::Mesh)
 			{
-				SwitchOp->cases.Add(ASTOpSwitch::FCase(condition, SwitchOp, Branch));
+				SwitchOp->cases.Add(ASTOpSwitch::FCase(Condition, SwitchOp, Branch));
 			}
         }
 
diff --git a/Engine/Plugins/Experimental/Mutable/Source/MutableTools/Private/MuT/CodeGenerator_Color.cpp b/Engine/Plugins/Experimental/Mutable/Source/MutableTools/Private/MuT/CodeGenerator_Color.cpp
--- a/Engine/Plugins/Experimental/Mutable/Source/MutableTools/Private/MuT/CodeGenerator_Color.cpp
+++ b/Engine/Plugins/Experimental/Mutable/Source/MutableTools/Private/MuT/CodeGenerator_Color.cpp
@@ -471,14 +471,23 @@ namespace mu
 	{
 		const NodeColourTable::Private& node = *Typed->GetPrivate();
 
+		FColorGenerationResult DefaultValue;
+		// Creating a default colour node
+		{
+			NodeColourConstantPtr pNode = new NodeColourConstant();
+			pNode->SetValue(mu::DefaultMutableColorValue);
+
+			GenerateColor(DefaultValue, pNode);
+		}
+
 		result.op = GenerateTableSwitch<NodeColourTable::Private, ETableColumnType::Color, OP_TYPE::CO_SWITCH>(node,
 			[this](const NodeColourTable::Private& node, int colIndex, int row, ErrorLog* pErrorLog)
 			{
 				NodeColourConstantPtr CellData = new NodeColourConstant();
-				FVector4f Colour = node.m_pTable->GetPrivate()->Rows[row].Values[colIndex].Color;
+				FVector4f Colour = node.Table->GetPrivate()->Rows[row].Values[colIndex].Color;
 				CellData->SetValue(Colour);
 				return Generate(CellData);
-			});
+			}, DefaultValue.op);
 	}
 
 
diff --git a/Engine/Plugins/Experimental/Mutable/Source/MutableTools/Private/MuT/CodeGenerator_Image.cpp b/Engine/Plugins/Experimental/Mutable/Source/MutableTools/Private/MuT/CodeGenerator_Image.cpp
--- a/Engine/Plugins/Experimental/Mutable/Source/MutableTools/Private/MuT/CodeGenerator_Image.cpp
+++ b/Engine/Plugins/Experimental/Mutable/Source/MutableTools/Private/MuT/CodeGenerator_Image.cpp
@@ -1703,7 +1703,7 @@ namespace mu
 		Result.op = GenerateTableSwitch<NodeImageTable::Private, ETableColumnType::Image, OP_TYPE::IM_SWITCH>(node,
 			[this,Options](const NodeImageTable::Private& node, int colIndex, int row, ErrorLog* pErrorLog)
 			{
-				const FTableValue& CellData = node.m_pTable->GetPrivate()->Rows[row].Values[colIndex];
+				const FTableValue& CellData = node.Table->GetPrivate()->Rows[row].Values[colIndex];
 				ImagePtrConst pImage = nullptr;
 				NodeImagePtr CellImage = nullptr;
 
@@ -1738,7 +1738,7 @@ namespace mu
 				FImageGenerationResult Result;
 				GenerateImage(Options, Result, CellImage);
 				return Result.op;
-			});
+			}, nullptr /* TODO(max): MTBL-1660 */);
 	}
 
 
diff --git a/Engine/Plugins/Experimental/Mutable/Source/MutableTools/Private/MuT/CodeGenerator_Mesh.cpp b/Engine/Plugins/Experimental/Mutable/Source/MutableTools/Private/MuT/CodeGenerator_Mesh.cpp
--- a/Engine/Plugins/Experimental/Mutable/Source/MutableTools/Private/MuT/CodeGenerator_Mesh.cpp
+++ b/Engine/Plugins/Experimental/Mutable/Source/MutableTools/Private/MuT/CodeGenerator_Mesh.cpp
@@ -1120,7 +1120,7 @@ class Node;
 		Ptr<ASTOp> Op = GenerateTableSwitch<NodeMeshTable::Private, ETableColumnType::Mesh, OP_TYPE::ME_SWITCH>(*TableNode->GetPrivate(),
 			[this, &NewResult, &t, &InOptions] (const NodeMeshTable::Private& node, int colIndex, int row, ErrorLog* pErrorLog)
 			{
-				mu::Ptr<mu::Mesh> pMesh = node.m_pTable->GetPrivate()->Rows[row].Values[colIndex].Mesh;
+				mu::Ptr<mu::Mesh> pMesh = node.Table->GetPrivate()->Rows[row].Values[colIndex].Mesh;
 				FMeshGenerationResult BranchResults;
 
 				if (pMesh)
@@ -1129,11 +1129,11 @@ class Node;
 					pCell->SetValue(pMesh);
 
 					// TODO Take into account layout strategy
-					int numLayouts = node.m_layouts.Num();
+					int numLayouts = node.Layouts.Num();
 					pCell->SetLayoutCount(numLayouts);
 					for (int i = 0; i < numLayouts; ++i)
 					{
-						pCell->SetLayout(i, node.m_layouts[i]);
+						pCell->SetLayout(i, node.Layouts[i]);
 					}
 
 					FMeshGenerationOptions TargetOptions = InOptions;
@@ -1154,7 +1154,7 @@ class Node;
 				}
 
 				return BranchResults.meshOp;
-			});
+			}, nullptr /* TODO(max): MTBL-1660 */ );
 
 		NewResult.meshOp = Op;
 
diff --git a/Engine/Plugins/Experimental/Mutable/Source/MutableTools/Private/MuT/CodeGenerator_Scalar.cpp b/Engine/Plugins/Experimental/Mutable/Source/MutableTools/Private/MuT/CodeGenerator_Scalar.cpp
--- a/Engine/Plugins/Experimental/Mutable/Source/MutableTools/Private/MuT/CodeGenerator_Scalar.cpp
+++ b/Engine/Plugins/Experimental/Mutable/Source/MutableTools/Private/MuT/CodeGenerator_Scalar.cpp
@@ -418,10 +418,10 @@ namespace mu
 			[this](const NodeScalarTable::Private& node, int colIndex, int row, ErrorLog* pErrorLog)
 			{
 				NodeScalarConstantPtr pCell = new NodeScalarConstant();
-				float scalar = node.m_pTable->GetPrivate()->Rows[row].Values[colIndex].Scalar;
+				float scalar = node.Table->GetPrivate()->Rows[row].Values[colIndex].Scalar;
 				pCell->SetValue(scalar);
 				return Generate(pCell);
-			});
+			}, nullptr /* TODO(max): MTBL-1660 */ );
 
 		result.op = Op;
 	}
diff --git a/Engine/Plugins/Experimental/Mutable/Source/MutableTools/Private/MuT/ImageDescGenerator.cpp b/Engine/Plugins/Experimental/Mutable/Source/MutableTools/Private/MuT/ImageDescGenerator.cpp
--- a/Engine/Plugins/Experimental/Mutable/Source/MutableTools/Private/MuT/ImageDescGenerator.cpp
+++ b/Engine/Plugins/Experimental/Mutable/Source/MutableTools/Private/MuT/ImageDescGenerator.cpp
@@ -120,7 +120,7 @@ namespace mu
         m_desc = MUTABLE_MISSING_IMAGE_DESC;
 
         // Verify that the table column is the right type
-        int colIndex = node.m_pTable->FindColumn( node.m_columnName );
+        int colIndex = node.Table->FindColumn( node.ColumnName );
         if ( colIndex<0 )
         {
             check( false );
@@ -130,15 +130,15 @@ namespace mu
         Ptr<const Image> pImage;
 
         std::size_t i = 0;
-        while ( !pImage && i<node.m_pTable->GetPrivate()->Rows.Num() )
+        while ( !pImage && i<node.Table->GetPrivate()->Rows.Num() )
         {
-			if (node.m_pTable->GetPrivate()->Rows[i].Values[colIndex].ProxyImage->Get()->IsReference())
+			if (node.Table->GetPrivate()->Rows[i].Values[colIndex].ProxyImage->Get()->IsReference())
 			{
 				// Image References do not need an image desc
 				break;
 			}
 
-            pImage = node.m_pTable->GetPrivate()->Rows[i].Values[ colIndex ].ProxyImage->Get();
+            pImage = node.Table->GetPrivate()->Rows[i].Values[ colIndex ].ProxyImage->Get();
             ++i;
         }
 
diff --git a/Engine/Plugins/Experimental/Mutable/Source/MutableTools/Private/MuT/NodeColourTable.cpp b/Engine/Plugins/Experimental/Mutable/Source/MutableTools/Private/MuT/NodeColourTable.cpp
--- a/Engine/Plugins/Experimental/Mutable/Source/MutableTools/Private/MuT/NodeColourTable.cpp
+++ b/Engine/Plugins/Experimental/Mutable/Source/MutableTools/Private/MuT/NodeColourTable.cpp
@@ -60,28 +60,35 @@ namespace mu
 	//---------------------------------------------------------------------------------------------
 	void NodeColourTable::SetColumn( const FString& strName )
 	{
-		m_pD->m_columnName = strName;
+		m_pD->ColumnName = strName;
 	}
 
 
 	//---------------------------------------------------------------------------------------------
 	void NodeColourTable::SetParameterName( const FString& strName )
 	{
-		m_pD->m_parameterName = strName;
+		m_pD->ParameterName = strName;
 	}
 
 
 	//---------------------------------------------------------------------------------------------
 	void NodeColourTable::SetTable( TablePtr pTable )
 	{
-		m_pD->m_pTable = pTable;
+		m_pD->Table = pTable;
 	}
 
 
 	//---------------------------------------------------------------------------------------------
 	TablePtr NodeColourTable::GetTable() const
 	{
-		return m_pD->m_pTable;
+		return m_pD->Table;
+	}
+
+
+	//---------------------------------------------------------------------------------------------
+	void NodeColourTable::SetNoneOption(bool bAddNoneOption)
+	{
+		m_pD->bNoneOption = bAddNoneOption;
 	}
 
 
diff --git a/Engine/Plugins/Experimental/Mutable/Source/MutableTools/Private/MuT/NodeColourTablePrivate.h b/Engine/Plugins/Experimental/Mutable/Source/MutableTools/Private/MuT/NodeColourTablePrivate.h
--- a/Engine/Plugins/Experimental/Mutable/Source/MutableTools/Private/MuT/NodeColourTablePrivate.h
+++ b/Engine/Plugins/Experimental/Mutable/Source/MutableTools/Private/MuT/NodeColourTablePrivate.h
@@ -23,19 +23,21 @@ namespace mu
 
 		static NODE_TYPE s_type;
 
-		FString m_parameterName;
-		TablePtr m_pTable;
-		FString m_columnName;
+		FString ParameterName;
+		TablePtr Table;
+		FString ColumnName;
+		bool bNoneOption = false;
 
 		//!
 		void Serialise( OutputArchive& arch ) const
 		{
-            uint32 ver = 1;
+            uint32_t ver = 2;
 			arch << ver;
 
-			arch << m_parameterName;
-			arch << m_pTable;
-			arch << m_columnName;
+			arch << ParameterName;
+			arch << Table;
+			arch << ColumnName;
+			arch << bNoneOption;
 		}
 
 		//!
@@ -43,30 +45,35 @@ namespace mu
 		{
             uint32 ver;
 			arch >> ver;
-			check(ver <= 1);
+			check(ver <= 2);
 
 			if (ver == 0)
 			{
 				std::string Temp;
 				arch >> Temp;
-				m_parameterName = Temp.c_str();
+				ParameterName = Temp.c_str();
 			}
 			else
 			{
-				arch >> m_parameterName;
+				arch >> ParameterName;
 			}
 
-			arch >> m_pTable;
+			arch >> Table;
 
 			if (ver == 0)
 			{
 				std::string Temp;
 				arch >> Temp;
-				m_columnName = Temp.c_str();
+				ColumnName = Temp.c_str();
 			}
 			else
 			{
-				arch >> m_columnName;
+				arch >> ColumnName;
+			}
+
+			if (ver >= 2)
+			{
+				arch >> bNoneOption;
 			}
 		}
 
diff --git a/Engine/Plugins/Experimental/Mutable/Source/MutableTools/Private/MuT/NodeImageTable.cpp b/Engine/Plugins/Experimental/Mutable/Source/MutableTools/Private/MuT/NodeImageTable.cpp
--- a/Engine/Plugins/Experimental/Mutable/Source/MutableTools/Private/MuT/NodeImageTable.cpp
+++ b/Engine/Plugins/Experimental/Mutable/Source/MutableTools/Private/MuT/NodeImageTable.cpp
@@ -59,31 +59,37 @@ namespace mu
 	//---------------------------------------------------------------------------------------------
 	void NodeImageTable::SetColumn( const FString& strName )
 	{
-		m_pD->m_columnName = strName;
+		m_pD->ColumnName = strName;
 	}
 
 
 	//---------------------------------------------------------------------------------------------
 	void NodeImageTable::SetParameterName( const FString& strName )
 	{
-		m_pD->m_parameterName = strName;
+		m_pD->ParameterName = strName;
 	}
 
 
 	//---------------------------------------------------------------------------------------------
 	void NodeImageTable::SetTable( TablePtr pTable )
 	{
-		m_pD->m_pTable = pTable;
+		m_pD->Table = pTable;
 	}
 
 
 	//---------------------------------------------------------------------------------------------
 	TablePtr NodeImageTable::GetTable() const
 	{
-		return m_pD->m_pTable;
+		return m_pD->Table;
 	}
 
 
+	//---------------------------------------------------------------------------------------------
+	void NodeImageTable::SetNoneOption(bool bAddNoneOption)
+	{
+		m_pD->bNoneOption = bAddNoneOption;
+	}
+
 }
 
 
diff --git a/Engine/Plugins/Experimental/Mutable/Source/MutableTools/Private/MuT/NodeImageTablePrivate.h b/Engine/Plugins/Experimental/Mutable/Source/MutableTools/Private/MuT/NodeImageTablePrivate.h
--- a/Engine/Plugins/Experimental/Mutable/Source/MutableTools/Private/MuT/NodeImageTablePrivate.h
+++ b/Engine/Plugins/Experimental/Mutable/Source/MutableTools/Private/MuT/NodeImageTablePrivate.h
@@ -22,19 +22,21 @@ namespace mu
 
 		static NODE_TYPE s_type;
 
-		FString m_parameterName;
-		TablePtr m_pTable;
-		FString m_columnName;
+		FString ParameterName;
+		TablePtr Table;
+		FString ColumnName;
+		bool bNoneOption = false;
 
 		//!
 		void Serialise( OutputArchive& arch ) const
 		{
-            uint32_t ver = 2;
+            uint32_t ver = 3;
 			arch << ver;
 
-			arch << m_parameterName;
-			arch << m_pTable;
-			arch << m_columnName;
+			arch << ParameterName;
+			arch << Table;
+			arch << ColumnName;
+			arch << bNoneOption;
 		}
 
 		//!
@@ -42,30 +44,35 @@ namespace mu
 		{
             uint32_t ver;
 			arch >> ver;
-			check(ver>=1 && ver<= 2);
+			check(ver>=1 && ver<= 3);
 
 			if (ver == 1)
 			{
 				std::string Temp;
 				arch >> Temp;
-				m_parameterName = Temp.c_str();
+				ParameterName = Temp.c_str();
 			}
 			else
 			{
-				arch >> m_parameterName;
+				arch >> ParameterName;
 			}
 
-			arch >> m_pTable;
+			arch >> Table;
 
 			if (ver == 1)
 			{
 				std::string Temp;
 				arch >> Temp;
-				m_columnName = Temp.c_str();
+				ColumnName = Temp.c_str();
 			}
 			else
 			{
-				arch >> m_columnName;
+				arch >> ColumnName;
+			}
+			
+			if (ver >= 3)
+			{
+				arch >> bNoneOption;
 			}
 		}
 
diff --git a/Engine/Plugins/Experimental/Mutable/Source/MutableTools/Private/MuT/NodeMeshTable.cpp b/Engine/Plugins/Experimental/Mutable/Source/MutableTools/Private/MuT/NodeMeshTable.cpp
--- a/Engine/Plugins/Experimental/Mutable/Source/MutableTools/Private/MuT/NodeMeshTable.cpp
+++ b/Engine/Plugins/Experimental/Mutable/Source/MutableTools/Private/MuT/NodeMeshTable.cpp
@@ -60,42 +60,42 @@ namespace mu
 	//---------------------------------------------------------------------------------------------
 	void NodeMeshTable::SetColumn( const FString& strName )
 	{
-		m_pD->m_columnName = strName;
+		m_pD->ColumnName = strName;
 	}
 
 
 	//---------------------------------------------------------------------------------------------
 	void NodeMeshTable::SetParameterName( const FString& strName )
 	{
-		m_pD->m_parameterName = strName;
+		m_pD->ParameterName = strName;
 	}
 
 
 	//---------------------------------------------------------------------------------------------
 	void NodeMeshTable::SetTable( TablePtr pTable )
 	{
-		m_pD->m_pTable = pTable;
+		m_pD->Table = pTable;
 	}
 
 
 	//---------------------------------------------------------------------------------------------
 	TablePtr NodeMeshTable::GetTable() const
 	{
-		return m_pD->m_pTable;
+		return m_pD->Table;
 	}
 
 
 	//---------------------------------------------------------------------------------------------
 	int NodeMeshTable::GetLayoutCount() const
 	{
-		return m_pD->m_layouts.Num();
+		return m_pD->Layouts.Num();
 	}
 
 
 	//---------------------------------------------------------------------------------------------
 	void NodeMeshTable::SetLayoutCount( int i )
 	{
-		m_pD->m_layouts.SetNum( i );
+		m_pD->Layouts.SetNum( i );
 	}
 
 
@@ -111,7 +111,7 @@ namespace mu
 	void NodeMeshTable::SetLayout( int i, NodeLayoutPtr pLayout )
 	{
 		check( i>=0 && i<GetLayoutCount() );
-		m_pD->m_layouts[i] = pLayout;
+		m_pD->Layouts[i] = pLayout;
 	}
 
 
@@ -120,14 +120,21 @@ namespace mu
 	{
 		NodeLayoutPtr pResult;
 
-		if ( i>=0 && i<m_layouts.Num() )
+		if ( i>=0 && i< Layouts.Num() )
 		{
-			pResult = m_layouts[i];
+			pResult = Layouts[i];
 		}
 
 		return pResult;
 	}
 
+
+	//---------------------------------------------------------------------------------------------
+	void NodeMeshTable::SetNoneOption(bool bAddNoneOption)
+	{
+		m_pD->bNoneOption = bAddNoneOption;
+	}
+
 }
 
 
diff --git a/Engine/Plugins/Experimental/Mutable/Source/MutableTools/Private/MuT/NodeMeshTablePrivate.h b/Engine/Plugins/Experimental/Mutable/Source/MutableTools/Private/MuT/NodeMeshTablePrivate.h
--- a/Engine/Plugins/Experimental/Mutable/Source/MutableTools/Private/MuT/NodeMeshTablePrivate.h
+++ b/Engine/Plugins/Experimental/Mutable/Source/MutableTools/Private/MuT/NodeMeshTablePrivate.h
@@ -23,22 +23,24 @@ namespace mu
 
 		static NODE_TYPE s_type;
 
-		FString m_parameterName;
-		TablePtr m_pTable;
-		FString m_columnName;
+		FString ParameterName;
+		TablePtr Table;
+		FString ColumnName;
+		bool bNoneOption = false;
 
-		TArray<NodeLayoutPtr> m_layouts;
+		TArray<NodeLayoutPtr> Layouts;
 
 		//!
 		void Serialise( OutputArchive& arch ) const
 		{
-            uint32 ver = 2;
+            uint32 ver = 3;
 			arch << ver;
 
-			arch << m_parameterName;
-			arch << m_pTable;
-			arch << m_columnName;
-			arch << m_layouts;
+			arch << ParameterName;
+			arch << Table;
+			arch << ColumnName;
+			arch << Layouts;
+			arch << bNoneOption;
 		}
 
 		//!
@@ -46,32 +48,38 @@ namespace mu
 		{
             uint32 ver;
 			arch >> ver;
-            check(ver>=1 && ver<=2);
+            check(ver>=1 && ver<=3);
 
 			if (ver == 1)
 			{
 				std::string Temp;
 				arch >> Temp;
-				m_parameterName = Temp.c_str();
+				ParameterName = Temp.c_str();
 			}
 			else
 			{
-				arch >> m_parameterName;
+				arch >> ParameterName;
 			}
 
-			arch >> m_pTable;
+			arch >> Table;
 
 			if (ver == 1)
 			{
 				std::string Temp;
 				arch >> Temp;
-				m_columnName = Temp.c_str();
+				ColumnName = Temp.c_str();
 			}
 			else
 			{
-				arch >> m_columnName;
+				arch >> ColumnName;
+			}
+			
+			arch >> Layouts;
+
+			if(ver >= 3)
+			{
+				arch >> bNoneOption;
 			}
-			arch >> m_layouts;
 		}
 
 		// NodeMesh::Private interface
diff --git a/Engine/Plugins/Experimental/Mutable/Source/MutableTools/Private/MuT/NodeScalarTable.cpp b/Engine/Plugins/Experimental/Mutable/Source/MutableTools/Private/MuT/NodeScalarTable.cpp
--- a/Engine/Plugins/Experimental/Mutable/Source/MutableTools/Private/MuT/NodeScalarTable.cpp
+++ b/Engine/Plugins/Experimental/Mutable/Source/MutableTools/Private/MuT/NodeScalarTable.cpp
@@ -56,31 +56,37 @@ namespace mu
 	//---------------------------------------------------------------------------------------------
 	void NodeScalarTable::SetColumn( const FString& strName )
 	{
-		m_pD->m_columnName = strName;
+		m_pD->ColumnName = strName;
 	}
 
 
 	//---------------------------------------------------------------------------------------------
 	void NodeScalarTable::SetParameterName( const FString& strName )
 	{
-		m_pD->m_parameterName = strName;
+		m_pD->ParameterName = strName;
 	}
 
 
 	//---------------------------------------------------------------------------------------------
 	void NodeScalarTable::SetTable( TablePtr pTable )
 	{
-		m_pD->m_pTable = pTable;
+		m_pD->Table = pTable;
 	}
 
 
 	//---------------------------------------------------------------------------------------------
 	TablePtr NodeScalarTable::GetTable() const
 	{
-		return m_pD->m_pTable;
+		return m_pD->Table;
 	}
 
 
+	//---------------------------------------------------------------------------------------------
+	void NodeScalarTable::SetNoneOption(bool bAddNoneOption)
+	{
+		m_pD->bNoneOption = bAddNoneOption;
+	}
+
 }
 
 
diff --git a/Engine/Plugins/Experimental/Mutable/Source/MutableTools/Private/MuT/NodeScalarTablePrivate.h b/Engine/Plugins/Experimental/Mutable/Source/MutableTools/Private/MuT/NodeScalarTablePrivate.h
--- a/Engine/Plugins/Experimental/Mutable/Source/MutableTools/Private/MuT/NodeScalarTablePrivate.h
+++ b/Engine/Plugins/Experimental/Mutable/Source/MutableTools/Private/MuT/NodeScalarTablePrivate.h
@@ -22,19 +22,21 @@ namespace mu
 
 		static NODE_TYPE s_type;
 
-		FString m_parameterName;
-		Ptr<Table> m_pTable;
-		FString m_columnName;
+		FString ParameterName;
+		TablePtr Table;
+		FString ColumnName;
+		bool bNoneOption = false;
 
 		//!
 		void Serialise( OutputArchive& arch ) const
 		{
-            uint32 ver = 1;
+            uint32 ver = 2;
 			arch << ver;
 
-			arch << m_parameterName;
-			arch << m_pTable;
-			arch << m_columnName;
+			arch << ParameterName;
+			arch << Table;
+			arch << ColumnName;
+			arch << bNoneOption;
 		}
 
 		//!
@@ -42,30 +44,35 @@ namespace mu
 		{
             uint32 ver;
 			arch >> ver;
-            check(ver<=1);
+            check(ver <= 2);
 
 			if (ver==0)
 			{
 				std::string Temp;
 				arch >> Temp;
-				m_parameterName = Temp.c_str();
+				ParameterName = Temp.c_str();
 			}
 			else
 			{
-				arch >> m_parameterName;
+				arch >> ParameterName;
 			}
 
-			arch >> m_pTable;
+			arch >> Table;
 
 			if (ver == 0)
 			{
 				std::string Temp;
 				arch >> Temp;
-				m_columnName = Temp.c_str();
+				ColumnName = Temp.c_str();
 			}
 			else
 			{
-				arch >> m_columnName;
+				arch >> ColumnName;
+			}
+
+			if(ver >= 2)
+			{
+				arch >> bNoneOption;
 			}
 		}
 
diff --git a/Engine/Plugins/Experimental/Mutable/Source/MutableTools/Private/MuT/Table.cpp b/Engine/Plugins/Experimental/Mutable/Source/MutableTools/Private/MuT/Table.cpp
--- a/Engine/Plugins/Experimental/Mutable/Source/MutableTools/Private/MuT/Table.cpp
+++ b/Engine/Plugins/Experimental/Mutable/Source/MutableTools/Private/MuT/Table.cpp
@@ -120,12 +120,6 @@ namespace mu
 	}
 
 
-	void Table::SetNoneOption(bool bAddOption)
-	{
-		m_pD->bNoneOption = bAddOption;
-	}
-
-
     void Table::SetCell( int32 Column, uint32 RowId, float Value)
 	{
 		int32 Row = m_pD->FindRow(RowId);
diff --git a/Engine/Plugins/Experimental/Mutable/Source/MutableTools/Private/MuT/TablePrivate.h b/Engine/Plugins/Experimental/Mutable/Source/MutableTools/Private/MuT/TablePrivate.h
--- a/Engine/Plugins/Experimental/Mutable/Source/MutableTools/Private/MuT/TablePrivate.h
+++ b/Engine/Plugins/Experimental/Mutable/Source/MutableTools/Private/MuT/TablePrivate.h
@@ -50,12 +50,11 @@ namespace mu
 		FString Name;
 		TArray<FTableColumn> Columns;
 		TArray<FTableRow> Rows;
-		bool bNoneOption = false;
 
 		//!
 		void Serialise( OutputArchive& arch ) const
 		{
-            uint32 ver = 3;
+            uint32 ver = 4;
 			arch << ver;
 
             arch << Columns.Num();
@@ -94,8 +93,6 @@ namespace mu
 					}
 				}
 			}
-
-			arch << bNoneOption;
 		}
 
 		//!
@@ -103,7 +100,7 @@ namespace mu
 		{
             uint32 ver;
 			arch >> ver;
-			check(ver<=3);
+			check(ver<=4);
 
             uint32 ColumnCount;
 			arch >> ColumnCount;
@@ -194,9 +191,10 @@ namespace mu
 				}
 			}
 
-			if (ver >= 1)
+			if (ver >= 1 && ver <= 3)
 			{
-				arch >> bNoneOption;
+				bool bNoneOption_DEPRECATED;
+				arch >> bNoneOption_DEPRECATED;
 			}
 		}
 
diff --git a/Engine/Plugins/Experimental/Mutable/Source/MutableTools/Public/MuT/NodeColourTable.h b/Engine/Plugins/Experimental/Mutable/Source/MutableTools/Public/MuT/NodeColourTable.h
--- a/Engine/Plugins/Experimental/Mutable/Source/MutableTools/Public/MuT/NodeColourTable.h
+++ b/Engine/Plugins/Experimental/Mutable/Source/MutableTools/Public/MuT/NodeColourTable.h
@@ -62,6 +62,10 @@ namespace mu
 		//!
 		void SetColumn( const FString& strName );
 
+		//! Adds the "None" option to the parameter that represents this table column
+		void SetNoneOption(bool bAddOption);
+
+
 		//-----------------------------------------------------------------------------------------
 		// Interface pattern
 		//-----------------------------------------------------------------------------------------
diff --git a/Engine/Plugins/Experimental/Mutable/Source/MutableTools/Public/MuT/NodeImageTable.h b/Engine/Plugins/Experimental/Mutable/Source/MutableTools/Public/MuT/NodeImageTable.h
--- a/Engine/Plugins/Experimental/Mutable/Source/MutableTools/Public/MuT/NodeImageTable.h
+++ b/Engine/Plugins/Experimental/Mutable/Source/MutableTools/Public/MuT/NodeImageTable.h
@@ -45,8 +45,6 @@ namespace mu
 		// Node Interface
 		//-----------------------------------------------------------------------------------------
 
-        
-
         const NODE_TYPE* GetType() const override;
 		static const NODE_TYPE* GetStaticType();
 
@@ -68,6 +66,9 @@ namespace mu
 		//!
 		void SetColumn( const FString& strName );
 
+		//! Adds the "None" option to the parameter that represents this table column
+		void SetNoneOption(bool bAddOption);
+
 		//-----------------------------------------------------------------------------------------
 		// Interface pattern
 		//-----------------------------------------------------------------------------------------
diff --git a/Engine/Plugins/Experimental/Mutable/Source/MutableTools/Public/MuT/NodeMeshTable.h b/Engine/Plugins/Experimental/Mutable/Source/MutableTools/Public/MuT/NodeMeshTable.h
--- a/Engine/Plugins/Experimental/Mutable/Source/MutableTools/Public/MuT/NodeMeshTable.h
+++ b/Engine/Plugins/Experimental/Mutable/Source/MutableTools/Public/MuT/NodeMeshTable.h
@@ -46,8 +46,6 @@ namespace mu
 		// Node Interface
 		//-----------------------------------------------------------------------------------------
 
-        
-
         const NODE_TYPE* GetType() const override;
 		static const NODE_TYPE* GetStaticType();
 
@@ -77,6 +75,9 @@ namespace mu
 		NodeLayoutPtr GetLayout( int index ) const;
 		void SetLayout( int index, NodeLayoutPtr );
 
+		//! Adds the "None" option to the parameter that represents this table column
+		void SetNoneOption(bool bAddOption);
+
 		//-----------------------------------------------------------------------------------------
 		// Interface pattern
 		//-----------------------------------------------------------------------------------------
diff --git a/Engine/Plugins/Experimental/Mutable/Source/MutableTools/Public/MuT/NodeScalarTable.h b/Engine/Plugins/Experimental/Mutable/Source/MutableTools/Public/MuT/NodeScalarTable.h
--- a/Engine/Plugins/Experimental/Mutable/Source/MutableTools/Public/MuT/NodeScalarTable.h
+++ b/Engine/Plugins/Experimental/Mutable/Source/MutableTools/Public/MuT/NodeScalarTable.h
@@ -65,6 +65,9 @@ namespace mu
 		//!
 		void SetColumn( const FString& strName );
 
+		//! Adds the "None" option to the parameter that represents this table column
+		void SetNoneOption(bool bAddOption);
+
 		//-----------------------------------------------------------------------------------------
 		// Interface pattern
 		//-----------------------------------------------------------------------------------------
diff --git a/Engine/Plugins/Experimental/Mutable/Source/MutableTools/Public/MuT/Table.h b/Engine/Plugins/Experimental/Mutable/Source/MutableTools/Public/MuT/Table.h
--- a/Engine/Plugins/Experimental/Mutable/Source/MutableTools/Public/MuT/Table.h
+++ b/Engine/Plugins/Experimental/Mutable/Source/MutableTools/Public/MuT/Table.h
@@ -64,9 +64,6 @@ namespace mu
 		//!
         void AddRow( uint32 id );
 
-		//! Adds a "NONE" option
-		void SetNoneOption(bool bAddOption);
-
 		//!
         void SetCell( int32 Column, uint32 RowId, float );
         void SetCell( int32 Column, uint32 RowId, const FVector4f& );
