[mutable] Added nullptr guard to GEditor dereferencing for factory registration.
- Merging MutableValidation module into CustomizableObjectEditor module is not plausible.

diff --git a/Engine/Plugins/Experimental/Mutable/Source/CustomizableObjectEditor/Private/MuCOE/CustomizableObjectEditorModule.cpp b/Engine/Plugins/Experimental/Mutable/Source/CustomizableObjectEditor/Private/MuCOE/CustomizableObjectEditorModule.cpp
--- a/Engine/Plugins/Experimental/Mutable/Source/CustomizableObjectEditor/Private/MuCOE/CustomizableObjectEditorModule.cpp
+++ b/Engine/Plugins/Experimental/Mutable/Source/CustomizableObjectEditor/Private/MuCOE/CustomizableObjectEditorModule.cpp
@@ -191,10 +191,13 @@ void FCustomizableObjectEditorModule::StartupModule()
 	PropertyModule.NotifyCustomizationModuleChanged();
 
 	// Register factory
-	GEditor->ActorFactories.Add(NewObject<UCustomizableObjectInstanceFactory>());
-	if (UPlacementSubsystem* PlacementSubsystem = GEditor->GetEditorSubsystem<UPlacementSubsystem>())
+	if (GEditor)
 	{
-		PlacementSubsystem->RegisterAssetFactory(NewObject<UCustomizableObjectInstanceFactory>());
+		GEditor->ActorFactories.Add(NewObject<UCustomizableObjectInstanceFactory>());
+		if (UPlacementSubsystem* PlacementSubsystem = GEditor->GetEditorSubsystem<UPlacementSubsystem>())
+		{
+			PlacementSubsystem->RegisterAssetFactory(NewObject<UCustomizableObjectInstanceFactory>());
+		}
 	}
 	
 
