#rnx

diff --git a/Engine/Plugins/Experimental/Mutable/Source/CustomizableObject/Classes/MuCO/CustomizableObject.h b/Engine/Plugins/Experimental/Mutable/Source/CustomizableObject/Classes/MuCO/CustomizableObject.h
--- a/Engine/Plugins/Experimental/Mutable/Source/CustomizableObject/Classes/MuCO/CustomizableObject.h
+++ b/Engine/Plugins/Experimental/Mutable/Source/CustomizableObject/Classes/MuCO/CustomizableObject.h
@@ -303,6 +303,7 @@ struct FMutableModelImageProperties
 		: Filter(TF_Default)
 		, SRGB(0)
 		, FlipGreenChannel(0)
+		, IsPassThrough(0)
 		, LODBias(0)
 		, LODGroup(TEXTUREGROUP_World)
 		, AddressX(TA_Clamp)
@@ -310,12 +311,13 @@ struct FMutableModelImageProperties
 	{}
 
 	FMutableModelImageProperties(const FString& InTextureParameterName, TextureFilter InFilter, uint32 InSRGB, 
-		uint32 InFlipGreenChannel, int32 InLODBias, TEnumAsByte<enum TextureGroup> InLODGroup,
+		uint32 InFlipGreenChannel, uint32 bInIsPassThrough, int32 InLODBias, TEnumAsByte<enum TextureGroup> InLODGroup,
 		TEnumAsByte<enum TextureAddress> InAddressX, TEnumAsByte<enum TextureAddress> InAddressY)
 		: TextureParameterName(InTextureParameterName)
 		, Filter(InFilter)
 		, SRGB(InSRGB)
 		, FlipGreenChannel(InFlipGreenChannel)
+		, IsPassThrough(bInIsPassThrough)
 		, LODBias(InLODBias)
 		, LODGroup(InLODGroup)
 		, AddressX(InAddressX)
@@ -335,6 +337,9 @@ struct FMutableModelImageProperties
 	UPROPERTY()
 	uint32 FlipGreenChannel : 1;
 
+	UPROPERTY()
+	uint32 IsPassThrough : 1;
+
 	UPROPERTY()
 	int32 LODBias;
 
@@ -354,6 +359,7 @@ struct FMutableModelImageProperties
 			Filter != rhs.Filter ||
 			SRGB != rhs.SRGB ||
 			FlipGreenChannel != rhs.FlipGreenChannel ||
+			IsPassThrough != rhs.IsPassThrough ||
 			LODBias != rhs.LODBias ||
 			LODGroup != rhs.LODGroup ||
 			AddressX != rhs.AddressX ||
@@ -375,6 +381,10 @@ struct FMutableModelImageProperties
 			Aux = 0;
 			Ar << Aux;
 			ImageProps.FlipGreenChannel = Aux;
+
+			Aux = 0;
+			Ar << Aux;
+			ImageProps.IsPassThrough = Aux;
 		}
 		else
 		{
@@ -383,6 +393,9 @@ struct FMutableModelImageProperties
 
 			Aux = ImageProps.FlipGreenChannel;
 			Ar << Aux;
+
+			Aux = ImageProps.IsPassThrough;
+			Ar << Aux;
 		}
 
 		Ar << ImageProps.LODBias;
@@ -1459,7 +1472,7 @@ private:
 	// This is a manual version number for the binary blobs in this asset.
 	// Increasing it invalidates all the previously compiled models.
 	// Warning: If while merging code both versions have changed, take the highest+1.
-	static const int32 CurrentSupportedVersion = 405;
+	static const int32 CurrentSupportedVersion = 406;
 
 public:
 
diff --git a/Engine/Plugins/Experimental/Mutable/Source/CustomizableObject/Private/MuCO/CustomizableInstancePrivateData.h b/Engine/Plugins/Experimental/Mutable/Source/CustomizableObject/Private/MuCO/CustomizableInstancePrivateData.h
--- a/Engine/Plugins/Experimental/Mutable/Source/CustomizableObject/Private/MuCO/CustomizableInstancePrivateData.h
+++ b/Engine/Plugins/Experimental/Mutable/Source/CustomizableObject/Private/MuCO/CustomizableInstancePrivateData.h
@@ -355,6 +355,14 @@ public:
 	UPROPERTY(Transient, Category = CustomizableObjectInstance, editfixedsize, VisibleAnywhere)
 	TMap<TSubclassOf<UAnimInstance>, FAnimBpGeneratedPhysicsAssets> AnimBpPhysicsAssets;
 
+	// The pass-through textures that will be loaded during an update
+	TArray<TSoftObjectPtr<UTexture>> PassThroughTexturesToLoad;
+
+	// Used during an update to prevent the pass-through textures loaded by LoadAdditionalAssetsAsync() from being unloaded by GC
+	// between AdditionalAssetsAsyncLoaded() and their setting into the generated materials in BuildMaterials()
+	UPROPERTY(Transient)
+	TArray<TObjectPtr<UTexture>> LoadedPassThroughTexturesPendingSetMaterial;
+
 	// Struct used during an update to avoid generating resources that can be reused.
 	FInstanceGeneratedData LastUpdateData;
 
diff --git a/Engine/Plugins/Experimental/Mutable/Source/CustomizableObject/Private/MuCO/CustomizableObjectInstance.cpp b/Engine/Plugins/Experimental/Mutable/Source/CustomizableObject/Private/MuCO/CustomizableObjectInstance.cpp
--- a/Engine/Plugins/Experimental/Mutable/Source/CustomizableObject/Private/MuCO/CustomizableObjectInstance.cpp
+++ b/Engine/Plugins/Experimental/Mutable/Source/CustomizableObject/Private/MuCO/CustomizableObjectInstance.cpp
@@ -5225,6 +5225,11 @@ FGraphEventRef UCustomizableInstancePrivateData::LoadAdditionalAssetsAsync(const
 		}
 	}
 
+	for (TSoftObjectPtr<UTexture>& TextureRef : PassThroughTexturesToLoad)
+	{
+		AssetsToStream.Add(TextureRef.ToSoftObjectPath());
+	}
+
 	if (AssetsToStream.Num() > 0)
 	{
 		check(!StreamingHandle);
@@ -5399,6 +5404,16 @@ void UCustomizableInstancePrivateData::AdditionalAssetsAsyncLoaded(UCustomizable
 		}
 		ComponentData.PhysicsAssets.AdditionalPhysicsAssetsToLoad.Empty();
 	}
+
+	LoadedPassThroughTexturesPendingSetMaterial.Empty(PassThroughTexturesToLoad.Num());
+
+	for (TSoftObjectPtr<UTexture>& TextureRef : PassThroughTexturesToLoad)
+	{
+		ensure(TextureRef.IsValid());
+		LoadedPassThroughTexturesPendingSetMaterial.Add(TextureRef.Get());
+	}
+
+	PassThroughTexturesToLoad.Empty();
 }
 
 
@@ -5723,20 +5738,32 @@ void UCustomizableInstancePrivateData::BuildMaterials(const TSharedPtr<FMutableO
 						}
 
 						// Check if the image is a reference to an engine texture
-						if (MutableImage && MutableImage->IsReference())
+						if (MutableImage && Image.bIsPassThrough)
 						{
+							check(MutableImage->IsReference());
+
 							uint32 ReferenceID = MutableImage->GetReferencedTexture();
 							if (CustomizableObject->ReferencedPassThroughTextures.IsValidIndex(ReferenceID))
 							{
 								TSoftObjectPtr<UTexture> Ref = CustomizableObject->ReferencedPassThroughTextures[ReferenceID];
 
-								// \TODO: This will force the load of the reference texture, potentially causing a hich. 
-								PassThroughTexture = Ref.LoadSynchronous();
+								// The texture should have been loaded by now by LoadAdditionalAssetsAsync()
+								PassThroughTexture = Ref.Get();
+
+								if (!PassThroughTexture)
+								{
+									// The texture should be loaded, something went wrong, possibly a bug in LoadAdditionalAssetsAsync()
+									UE_LOG(LogMutable, Error,
+										TEXT("Pass-through texture with name %s hasn't been loaded yet in BuildMaterials(). Forcing sync load."),
+										*Ref.ToSoftObjectPath().ToString());
+									ensure(false);
+									PassThroughTexture = Ref.LoadSynchronous();
+								}
 							}
 							else
 							{
 								// internal error.
-								UE_LOG(LogMutable, Error, TEXT("Referenced image [%d] was not dtored in the resource array."), ReferenceID);
+								UE_LOG(LogMutable, Error, TEXT("Referenced image [%d] was not stored in the resource array."), ReferenceID);
 							}
 						}
 
@@ -6122,6 +6149,9 @@ void UCustomizableInstancePrivateData::BuildMaterials(const TSharedPtr<FMutableO
 		}
 
 		Exchange(GeneratedTextures, NewGeneratedTextures);
+
+		// All pass-through textures have been set, no need to keep referencing them from the instance
+		LoadedPassThroughTexturesPendingSetMaterial.Empty();
 	}
 }
 
diff --git a/Engine/Plugins/Experimental/Mutable/Source/CustomizableObject/Private/MuCO/CustomizableObjectSystem.cpp b/Engine/Plugins/Experimental/Mutable/Source/CustomizableObject/Private/MuCO/CustomizableObjectSystem.cpp
--- a/Engine/Plugins/Experimental/Mutable/Source/CustomizableObject/Private/MuCO/CustomizableObjectSystem.cpp
+++ b/Engine/Plugins/Experimental/Mutable/Source/CustomizableObject/Private/MuCO/CustomizableObjectSystem.cpp
@@ -1457,6 +1457,11 @@ namespace impl
 		mu::System* System = UCustomizableObjectSystem::GetInstance()->GetPrivate()->MutableSystem.get();
 		check(System != nullptr);
 
+		const UCustomizableObject* CustomizableObject = OperationData->Instance->GetCustomizableObject();
+		UCustomizableInstancePrivateData* CustomizableObjectInstancePrivateData = OperationData->Instance->GetPrivate();
+
+		CustomizableObjectInstancePrivateData->PassThroughTexturesToLoad.Empty();
+
 		if (OperationData->bLiveUpdateMode)
 		{
 			if (OperationData->InstanceID == 0)
@@ -1654,6 +1659,40 @@ namespace impl
 							Image.FullImageSizeY = 0;
 							Image.BaseLOD = BaseLODIndex;
 							Image.BaseMip = 0;
+
+							FString KeyName = Image.Name.ToString();
+							int32 ImageKey = FCString::Atoi(*KeyName);
+
+							if (ImageKey >= 0 && ImageKey < CustomizableObject->ImageProperties.Num())
+							{
+								const FMutableModelImageProperties& Props = CustomizableObject->ImageProperties[ImageKey];
+
+								if (Props.IsPassThrough)
+								{
+									// Since it's known it's a pass-through texture, the GetImage will be inexpensive enough to be run in the game thread
+									Image.Image = System->GetImage(OperationData->InstanceID, Image.ImageID, 0, 0);
+									check(Image.Image->IsReference());
+
+									uint32 ReferenceID = Image.Image->GetReferencedTexture();
+
+									if (CustomizableObject->ReferencedPassThroughTextures.IsValidIndex(ReferenceID))
+									{
+										TSoftObjectPtr<UTexture> Ref = CustomizableObject->ReferencedPassThroughTextures[ReferenceID];
+										CustomizableObjectInstancePrivateData->PassThroughTexturesToLoad.Add(Ref);
+										Image.bIsPassThrough = true;
+									}
+									else
+									{
+										// internal error.
+										UE_LOG(LogMutable, Error, TEXT("Referenced image [%d] was not stored in the resource array."), ReferenceID);
+									}
+								}
+							}
+							else
+							{
+								// This means the compiled model (maybe coming from derived data) has images that the asset doesn't know about.
+								UE_LOG(LogMutable, Error, TEXT("CustomizableObject derived data out of sync with asset for [%s]. Try recompiling it."), *CustomizableObject->GetName());
+							}
 						}
 
 						// Vectors
@@ -1732,6 +1771,12 @@ namespace impl
 		{
 			MUTABLE_CPUPROFILER_SCOPE(GetImage);
 
+			// If the image is a reference to an engine texture, we are done.
+			if (Image.bIsPassThrough)
+			{
+				continue;
+			}
+
 			mu::FImageDesc ImageDesc;
 
 			// This should only be done when using progressive images, since GetImageDesc does some actual processing.
@@ -1789,13 +1834,6 @@ namespace impl
 
 				check(Image.Image);
 
-				// If the image is a reference to an engine texture, we are done.
-				if (Image.Image->IsReference())
-				{
-					continue;
-				}
-
-
 				// We should have generated exactly this size.
 				const bool bSizeMissmatch = Image.Image->GetSizeX() != MipSizeX || Image.Image->GetSizeY() != MipSizeY;
 				if (bSizeMissmatch)
diff --git a/Engine/Plugins/Experimental/Mutable/Source/CustomizableObject/Private/MuCO/CustomizableObjectSystemPrivate.h b/Engine/Plugins/Experimental/Mutable/Source/CustomizableObject/Private/MuCO/CustomizableObjectSystemPrivate.h
--- a/Engine/Plugins/Experimental/Mutable/Source/CustomizableObject/Private/MuCO/CustomizableObjectSystemPrivate.h
+++ b/Engine/Plugins/Experimental/Mutable/Source/CustomizableObject/Private/MuCO/CustomizableObjectSystemPrivate.h
@@ -406,6 +406,8 @@ struct FInstanceUpdateData
 		uint16 FullImageSizeX, FullImageSizeY;
 		mu::Ptr<const mu::Image> Image;
 		TWeakObjectPtr<UTexture2D> Cached;
+
+		bool bIsPassThrough = false;
 	};
 
 	struct FVector
diff --git a/Engine/Plugins/Experimental/Mutable/Source/CustomizableObjectEditor/Private/MuCOE/CustomizableObjectCompiler.cpp b/Engine/Plugins/Experimental/Mutable/Source/CustomizableObjectEditor/Private/MuCOE/CustomizableObjectCompiler.cpp
--- a/Engine/Plugins/Experimental/Mutable/Source/CustomizableObjectEditor/Private/MuCOE/CustomizableObjectCompiler.cpp
+++ b/Engine/Plugins/Experimental/Mutable/Source/CustomizableObjectEditor/Private/MuCOE/CustomizableObjectCompiler.cpp
@@ -1158,6 +1158,7 @@ void FCustomizableObjectCompiler::CompileInternal(UCustomizableObject* Object, c
 										ImageProp.Filter,
 										ImageProp.SRGB,
 										ImageProp.bFlipGreenChannel,
+										ImageProp.bIsPassThrough,
 										ImageProp.LODBias,
 										ImageProp.LODGroup,
 										ImageProp.AddressX, ImageProp.AddressY });
diff --git a/Engine/Plugins/Experimental/Mutable/Source/CustomizableObjectEditor/Private/MuCOE/GenerateMutableSource/GenerateMutableSource.h b/Engine/Plugins/Experimental/Mutable/Source/CustomizableObjectEditor/Private/MuCOE/GenerateMutableSource/GenerateMutableSource.h
--- a/Engine/Plugins/Experimental/Mutable/Source/CustomizableObjectEditor/Private/MuCOE/GenerateMutableSource/GenerateMutableSource.h
+++ b/Engine/Plugins/Experimental/Mutable/Source/CustomizableObjectEditor/Private/MuCOE/GenerateMutableSource/GenerateMutableSource.h
@@ -65,6 +65,8 @@ struct FGeneratedImageProperties
 
 	TEnumAsByte<TextureAddress> AddressX = TA_Clamp;
 	TEnumAsByte<TextureAddress> AddressY = TA_Clamp;
+
+	bool bIsPassThrough = false;
 };
 
 
diff --git a/Engine/Plugins/Experimental/Mutable/Source/CustomizableObjectEditor/Private/MuCOE/GenerateMutableSource/GenerateMutableSourceSurface.cpp b/Engine/Plugins/Experimental/Mutable/Source/CustomizableObjectEditor/Private/MuCOE/GenerateMutableSource/GenerateMutableSourceSurface.cpp
--- a/Engine/Plugins/Experimental/Mutable/Source/CustomizableObjectEditor/Private/MuCOE/GenerateMutableSource/GenerateMutableSourceSurface.cpp
+++ b/Engine/Plugins/Experimental/Mutable/Source/CustomizableObjectEditor/Private/MuCOE/GenerateMutableSource/GenerateMutableSourceSurface.cpp
@@ -491,6 +491,7 @@ mu::NodeSurfacePtr GenerateMutableSourceSurface(const UEdGraphPin * Pin, FMutabl
 					// We don't need a reference texture or props here, but we do need the parameter name.
 					FGeneratedImageProperties Props;
 					Props.TextureParameterName = ImageName;
+					Props.bIsPassThrough = true;
 					GenerationContext.ImageProperties.Add(Props);
 					SurfaceData.ImageProperties = Props;
 				}
