[mutable] Optimize compilation for objects with many textures.
#rnx
#preflgiht 6513ceb8e85968229586c532

diff --git a/Engine/Plugins/Experimental/Mutable/Source/MutableTools/Private/MuT/CodeOptimiser.cpp b/Engine/Plugins/Experimental/Mutable/Source/MutableTools/Private/MuT/CodeOptimiser.cpp
--- a/Engine/Plugins/Experimental/Mutable/Source/MutableTools/Private/MuT/CodeOptimiser.cpp
+++ b/Engine/Plugins/Experimental/Mutable/Source/MutableTools/Private/MuT/CodeOptimiser.cpp
@@ -26,47 +26,47 @@
 namespace mu
 {
 
-
 	namespace
 	{
 
-		struct MESH_ENTRY
+		struct FMeshEntry
 		{
-			MeshPtrConst mesh;
-			mu::Ptr<ASTOpConstantResource> op;
+			Ptr<const Mesh> Mesh;
+			Ptr<ASTOpConstantResource> Op;
 
-			bool operator==(const MESH_ENTRY& o) const
+			bool operator==(const FMeshEntry& o) const
 			{
-				return mesh==o.mesh || *mesh==*o.mesh;
+				return Mesh == o.Mesh || *Mesh ==*o.Mesh;
 			}
 		};
 
-		struct custom_mesh_hash
+		struct FImageEntry
 		{
-			uint64 operator()( const MeshPtrConst& k ) const
+			Ptr< const Image> Image;
+			Ptr<ASTOpConstantResource> Op;
+
+			bool operator==(const FImageEntry& o) const
 			{
-				uint64 h =  std::hash<uint64>()(k->GetVertexCount());
-				hash_combine(h, k->GetFaceCount());
-				return h;
+				return Image == o.Image || *Image == *o.Image;
 			}
 		};
 
-		struct custom_mesh_equal
+		struct FLayoutEntry
 		{
-			bool operator()( const MeshPtrConst& a, const MeshPtrConst& b ) const
+			Ptr< const Layout> Layout;
+			Ptr<ASTOpConstantResource> Op;
+
+			bool operator==(const FLayoutEntry& o) const
 			{
-				return a==b || *a==*b;
+				return Layout == o.Layout || *Layout == *o.Layout;
 			}
 		};
 
-		struct custom_image_hash
+		struct custom_mesh_equal
 		{
-			uint64 operator()( const ImagePtrConst& k ) const
+			bool operator()( const MeshPtrConst& a, const MeshPtrConst& b ) const
 			{
-				uint64 h =  std::hash<uint64>()(k->GetSizeX());
-				hash_combine(h, k->GetSizeY());
-				hash_combine(h, k->GetLODCount());
-				return h;
+				return a==b || *a==*b;
 			}
 		};
 
@@ -78,18 +78,6 @@ namespace mu
 			}
 		};
 
-		struct custom_layout_hash
-		{
-			uint64 operator()( const LayoutPtrConst& k ) const
-			{
-				uint64 h =  std::hash<uint64>()(k->GetBlockCount());
-				FIntPoint s = k->GetGridSize();
-				hash_combine(h, s[0]);
-				hash_combine(h, s[1]);
-				return h;
-			}
-		};
-
 		struct custom_layout_equal
 		{
 			bool operator()( const LayoutPtrConst& a, const LayoutPtrConst& b ) const
@@ -105,13 +93,13 @@ namespace mu
 	{
 		MUTABLE_CPUPROFILER_SCOPE(DuplicatedDataRemoverAST);
 
-		std::vector<Ptr<ASTOpConstantResource>> allMeshOps;
-		std::vector<Ptr<ASTOpConstantResource>> allImageOps;
-		std::vector<Ptr<ASTOpConstantResource>> allLayoutOps;
+		TArray<Ptr<ASTOpConstantResource>> AllMeshOps;
+		TArray<Ptr<ASTOpConstantResource>> AllImageOps;
+		TArray<Ptr<ASTOpConstantResource>> AllLayoutOps;
 
 		// Gather constants
 		{
-			//MUTABLE_CPUPROFILER_SCOPE(Gather);
+			MUTABLE_CPUPROFILER_SCOPE(Gather);
 
 			ASTOp::Traverse_TopRandom_Unique_NonReentrant( roots, [&](Ptr<ASTOp> n)
 			{
@@ -125,7 +113,7 @@ namespace mu
 
 					if (typedNode)
 					{
-						allMeshOps.push_back(typedNode);
+						AllMeshOps.Add(typedNode);
 					}
 					break;
 				}
@@ -137,7 +125,7 @@ namespace mu
 
 					if (typedNode)
 					{
-						allImageOps.push_back(typedNode);
+						AllImageOps.Add(typedNode);
 					}
 					break;
 				}
@@ -149,7 +137,7 @@ namespace mu
 
 					if (typedNode)
 					{
-						allLayoutOps.push_back(typedNode);
+						AllLayoutOps.Add(typedNode);
 					}
 					break;
 				}
@@ -162,12 +150,6 @@ namespace mu
 				//            case OP_TYPE::IN_ADDSCALAR:
 				//            case OP_TYPE::IN_ADDCOMPONENT:
 				//            case OP_TYPE::IN_ADDSURFACE:
-				//            {
-				//                OP::ADDRESS value = program.m_code[at].args.InstanceAdd.name;
-				//                program.m_code[at].args.InstanceAdd.name = m_oldToNewStrings[value];
-				//                break;
-				//            }
-
 
 				default:
 					break;
@@ -180,91 +162,143 @@ namespace mu
 
 		// Compare meshes
 		{
-			//MUTABLE_CPUPROFILER_SCOPE(CompareMeshes);
+			MUTABLE_CPUPROFILER_SCOPE(CompareMeshes);
 
-			std::unordered_multimap< size_t, MESH_ENTRY > meshes;
+			TMultiMap< SIZE_T, FMeshEntry > Meshes;
 
-			for (auto& typedNode: allMeshOps)
+			for (Ptr<ASTOpConstantResource>& typedNode : AllMeshOps)
 			{
-				size_t key = typedNode->GetValueHash();
+				SIZE_T Key = typedNode->GetValueHash();
+
+				Ptr<ASTOp> Found;
 
-				Ptr<ASTOp> found;
+				TArray<FMeshEntry*, TInlineAllocator<4>> Candidates;
+				Meshes.MultiFindPointer(Key, Candidates, false);
 
-				auto r = meshes.equal_range(key);
-				if (r.first!=r.second)
+				if (!Candidates.IsEmpty())
 				{
-					MeshPtrConst mesh = static_cast<const Mesh*>( typedNode->GetValue().get() );
+					Ptr<const Mesh> mesh = static_cast<const Mesh*>(typedNode->GetValue().get());
 
-					for ( auto it=r.first; it!=r.second; ++it )
+					for (FMeshEntry* It : Candidates)
 					{
-						if (!it->second.mesh)
+						if (!It->Mesh)
 						{
-							it->second.mesh = static_cast<const Mesh*>( it->second.op->GetValue().get() );
+							It->Mesh = static_cast<const Mesh*>(It->Op->GetValue().get());
 						}
 
-						if ( custom_mesh_equal()( mesh, it->second.mesh ) )
+						if (custom_mesh_equal()(mesh, It->Mesh))
 						{
-							found = it->second.op;
+							Found = It->Op;
 							break;
 						}
 					}
 				}
 
-				if (found)
+				if (Found)
 				{
-					ASTOp::Replace(typedNode,found);
+					ASTOp::Replace(typedNode, Found);
 				}
 				else
 				{
-					MESH_ENTRY e;
-					e.op = typedNode;
-					meshes.insert( std::make_pair<>(key, e) );
+					// The mesh will be loaded only if it needs to be compared
+					FMeshEntry e;
+					e.Op = typedNode;
+					Meshes.Add(Key, e);
 				}
 			}
 		}
 
 		// Compare images
 		{
-			//MUTABLE_CPUPROFILER_SCOPE(CompareImages);
+			MUTABLE_CPUPROFILER_SCOPE(CompareImages);
 
-			// TODO Optimise like the mesh compare above
-			std::unordered_map< ImagePtrConst, Ptr<ASTOp>, custom_image_hash, custom_image_equal > images;
+			TMultiMap< SIZE_T, FImageEntry > Images;
 
-			for (auto& typedNode: allImageOps)
+			for (Ptr<ASTOpConstantResource>& typedNode : AllImageOps)
 			{
-				ImagePtrConst image = static_cast<const Image*>( typedNode->GetValue().get() );
+				SIZE_T Key = typedNode->GetValueHash();
+
+				Ptr<ASTOp> Found;
 
-				auto it = images.find(image);
-				if (it!=images.end())
+				TArray<FImageEntry*,TInlineAllocator<4>> Candidates;
+				Images.MultiFindPointer(Key, Candidates, false);
+				
+				if (!Candidates.IsEmpty())
 				{
-					ASTOp::Replace(typedNode,it->second);
+					Ptr<const Image> image = static_cast<const Image*>(typedNode->GetValue().get());
+
+					for (FImageEntry* It: Candidates)
+					{
+						if (!It->Image)
+						{
+							It->Image = static_cast<const Image*>(It->Op->GetValue().get());
+						}
+
+						if (custom_image_equal()(image, It->Image))
+						{
+							Found = It->Op;
+							break;
+						}
+					}
+				}
+
+				if (Found)
+				{
+					ASTOp::Replace(typedNode, Found);
 				}
 				else
 				{
-					images.insert( std::make_pair<>(image, typedNode) );
+					// The image will be loaded only if it needs to be compared
+					FImageEntry e;
+					e.Op = typedNode;
+					Images.Add(Key, e);
 				}
 			}
 		}
 
 		// Compare layouts
 		{
-			//MUTABLE_CPUPROFILER_SCOPE(CompareLayouts);
+			MUTABLE_CPUPROFILER_SCOPE(CompareLayouts);
 
-			// TODO Optimise like the mesh compare above
-			std::unordered_map< LayoutPtrConst, Ptr<ASTOp>, custom_layout_hash, custom_layout_equal > constants;
+			TMultiMap< SIZE_T, FLayoutEntry > Layouts;
 
-			for (auto& typedNode: allLayoutOps)
+			for (Ptr<ASTOpConstantResource>& typedNode : AllLayoutOps)
 			{
-				LayoutPtrConst r = static_cast<const Layout*>( typedNode->GetValue().get() );
+				SIZE_T Key = typedNode->GetValueHash();
+
+				Ptr<ASTOp> Found;
+
+				TArray<FLayoutEntry*, TInlineAllocator<4>> Candidates;
+				Layouts.MultiFindPointer(Key, Candidates, false);
+
+				if (!Candidates.IsEmpty())
+				{
+					Ptr<const Layout> layout = static_cast<const Layout*>(typedNode->GetValue().get());
+
+					for (FLayoutEntry* It : Candidates)
+					{
+						if (!It->Layout)
+						{
+							It->Layout = static_cast<const Layout*>(It->Op->GetValue().get());
+						}
+
+						if (custom_layout_equal()(layout, It->Layout))
+						{
+							Found = It->Op;
+							break;
+						}
+					}
+				}
 
-				auto it = constants.find(r);
-				if (it!=constants.end())
+				if (Found)
 				{
-					ASTOp::Replace(typedNode,it->second);
+					ASTOp::Replace(typedNode, Found);
 				}
 				else
 				{
-					constants.insert( std::make_pair<>(r, typedNode) );
+					FLayoutEntry e;
+					e.Op = typedNode;
+					Layouts.Add(Key, e);
 				}
 			}
 		}
@@ -1271,7 +1305,7 @@ namespace mu
 			ASTOp::LogHistogram(roots);
 		}
 
-	//        // Minimal optimisation of constant subtrees
+		// Minimal optimisation of constant subtrees
 		else if ( m_options->GetPrivate()->OptimisationOptions.bConstReduction )
 		{
 			// The first duplicated data remover has the special mission of removing
