[mutable] Fixed duplicated node ids problem of node objects. Forced to regenerate the its guid in the BackwardsCompatibleFixup() method.

#rnx

diff --git a/Engine/Plugins/Experimental/Mutable/Source/CustomizableObject/Public/MuCO/CustomizableObjectCustomVersion.h b/Engine/Plugins/Experimental/Mutable/Source/CustomizableObject/Public/MuCO/CustomizableObjectCustomVersion.h
--- a/Engine/Plugins/Experimental/Mutable/Source/CustomizableObject/Public/MuCO/CustomizableObjectCustomVersion.h
+++ b/Engine/Plugins/Experimental/Mutable/Source/CustomizableObject/Public/MuCO/CustomizableObjectCustomVersion.h
@@ -100,6 +100,8 @@ struct CUSTOMIZABLEOBJECT_API FCustomizableObjectCustomVersion
 
 		NodeVariationSerializationIssue, // The parent class and child classes had a member with the same name. Unreal was serialized the member from the child class but deserialized it to the parent class.
 
+		RegenerateNodeObjectsIds,
+
 		// -----<new versions can be added above this line>-------------------------------------------------
 		VersionPlusOne,
 		LatestVersion = VersionPlusOne - 1
diff --git a/Engine/Plugins/Experimental/Mutable/Source/CustomizableObjectEditor/Private/MuCOE/Nodes/CustomizableObjectNodeObject.cpp b/Engine/Plugins/Experimental/Mutable/Source/CustomizableObjectEditor/Private/MuCOE/Nodes/CustomizableObjectNodeObject.cpp
--- a/Engine/Plugins/Experimental/Mutable/Source/CustomizableObjectEditor/Private/MuCOE/Nodes/CustomizableObjectNodeObject.cpp
+++ b/Engine/Plugins/Experimental/Mutable/Source/CustomizableObjectEditor/Private/MuCOE/Nodes/CustomizableObjectNodeObject.cpp
@@ -58,6 +58,13 @@ void UCustomizableObjectNodeObject::BackwardsCompatibleFixup()
 			}
 		}
 	}
+	
+	if (CustomizableObjectCustomVersion < FCustomizableObjectCustomVersion::RegenerateNodeObjectsIds)
+	{
+		// This will regenerate all the Node Objets Guids to finally remove the duplicated Guids warning.
+		// It is safe to do this here as Node Object do not use its node guid to link themeselves to other nodes.
+		CreateNewGuid();
+	}
 }
 
 
