[Mutable] Fix SkeletalMeshStatus flags.

- Removed Updating flag.
- Added Error16BitBoneIndex Update error.
- Created Update starting point StartUpdateSkeletalMesh.
- Removed unused functions.

#rnx

diff --git a/Engine/Plugins/Experimental/Mutable/Source/CustomizableObject/Classes/MuCO/CustomizableObjectInstance.h b/Engine/Plugins/Experimental/Mutable/Source/CustomizableObject/Classes/MuCO/CustomizableObjectInstance.h
--- a/Engine/Plugins/Experimental/Mutable/Source/CustomizableObject/Classes/MuCO/CustomizableObjectInstance.h
+++ b/Engine/Plugins/Experimental/Mutable/Source/CustomizableObject/Classes/MuCO/CustomizableObjectInstance.h
@@ -27,16 +27,6 @@ struct FPropertyChangedEvent;
 #define MUTABLE_VERTEXBUFFER_TEXCOORDS	2
 
 
-/** FString with the possible errors from skeletal mesh update */
-namespace ESkeletalMeshState
-{
-	const FString Correct = "Correct";
-	const FString UpdateError = "Update error";
-	const FString PostUpdateError = "Post Update Error";
-	const FString AsyncUpdatePending = "Async update pending";
-};
-
-
 // Priority for the mutable update queue, Low is the normal distance-based priority, High is normally used for discards and Mid for LOD downgrades
 enum class EQueuePriorityType : uint8 { High, Med, Med_Low, Low };
 
@@ -58,7 +48,17 @@ enum class EUpdateResult : uint8
 	Error, // Generic error.
 	ErrorOptimized, // The update was skipped since its result would have been the same as the current customization.
 	ErrorReplaced, // The update was replaced by a newer update request.
-	ErrorDiscarded // The update was not finished since due to the LOD management discarding the data.
+	ErrorDiscarded, // The update was not finished since due to the LOD management discarding the data.
+	Error16BitBoneIndex // The update finish unsuccessfully due to Instance not supporting 16 Bit Bone Indexing required by the Engine.
+};
+
+
+/** Indicates the status of the generated Skeletal Mesh. */
+enum class ESkeletalMeshStatus : uint8
+{
+	NotGenerated,
+	Success,
+	Error
 };
 
 
@@ -137,15 +137,7 @@ public:
 	UPROPERTY(Transient, VisibleAnywhere, Category = CustomizableSkeletalMesh)
 	TArray< TObjectPtr<USkeletalMesh> > SkeletalMeshes;
 
-
-	// Will store status description of current skeletal mesh generation (for instance, "EmptyLOD0" or "EmptyMesh"
-	UPROPERTY()
-	FString SkeletalMeshStatus;
-
 #if WITH_EDITOR
-	// Will store the previous status description to avoid losing notifications with partial updates.
-	FString PreUpdateSkeletalMeshStatus;
-
 	/** During editor, always remember the duration of the last update in the mutable runtime, for profiling. */
 	int32 LastUpdateMutableRuntimeCycles = 0;
 #endif
@@ -637,9 +629,6 @@ public:
 
 	const TArray<uint16>& GetRequestedLODsPerComponent() const;
 
-	/** Common end point of all updates. Even those which failed. */
-	void FinishUpdate(EUpdateResult UpdateResult, const FDescriptorRuntimeHash& UpdatedHash);
-	
 	/** Return the UCustomizableObjectInstance::Descriptor hash on the last update request. */
 	FDescriptorRuntimeHash GetDescriptorRuntimeHash() const;
 
diff --git a/Engine/Plugins/Experimental/Mutable/Source/CustomizableObject/Classes/MuCO/CustomizableObjectSystem.h b/Engine/Plugins/Experimental/Mutable/Source/CustomizableObject/Classes/MuCO/CustomizableObjectSystem.h
--- a/Engine/Plugins/Experimental/Mutable/Source/CustomizableObject/Classes/MuCO/CustomizableObjectSystem.h
+++ b/Engine/Plugins/Experimental/Mutable/Source/CustomizableObject/Classes/MuCO/CustomizableObjectSystem.h
@@ -385,9 +385,14 @@ private:
 	void DiscardInstances();
 	void ReleaseInstanceIDs();
 
+public:
+	/** Return true if the instance is being updated. */
+	UFUNCTION(BlueprintCallable, BlueprintPure, Category = CustomizableObjectSystem)
+	bool IsUpdating(const UCustomizableObjectInstance* Instance) const;
+
+private:
 	// TODO: Can we move this to the editor module?
 #if WITH_EDITOR
-
 	// Used to ask the user if they want to recompile uncompiled PIE COs
 	void OnPreBeginPIE(const bool bIsSimulatingInEditor);
 
diff --git a/Engine/Plugins/Experimental/Mutable/Source/CustomizableObject/Private/MuCO/CustomizableInstancePrivateData.h b/Engine/Plugins/Experimental/Mutable/Source/CustomizableObject/Private/MuCO/CustomizableInstancePrivateData.h
--- a/Engine/Plugins/Experimental/Mutable/Source/CustomizableObject/Private/MuCO/CustomizableInstancePrivateData.h
+++ b/Engine/Plugins/Experimental/Mutable/Source/CustomizableObject/Private/MuCO/CustomizableInstancePrivateData.h
@@ -31,7 +31,6 @@ enum ECOInstanceFlags
 	ECONone							= 0,  // Should not use the name None here.. it collides with other enum in global namespace
 
 	// Update process
-	Updating						= 1 << 0,	//
 	CreatingSkeletalMesh			= 1 << 1,	//
 	Generated						= 1 << 2,	//
 	ReuseTextures					= 1 << 3, 	// 
@@ -315,7 +314,10 @@ private:
 	void SetLastMeshId(int32 ComponentIndex, int32 LODIndex, mu::FResourceID MeshId);
 
 public:
+	CUSTOMIZABLEOBJECT_API ESkeletalMeshStatus GetSkeletalMeshStatus() const;
 
+	void SetSkeletalMeshStatus(ESkeletalMeshStatus Status);
+	
 	// If any components are using this instance, they will store the min of their distances to the player here every frame for LOD purposes
 	float MinSquareDistFromComponentToPlayer;
 	float LastMinSquareDistFromComponentToPlayer; // The same as the previous dist for last frame
@@ -378,6 +380,9 @@ public:
 	FDescriptorRuntimeHash DescriptorRuntimeHash;
 
 private:
+	/** Status of the generated Skeletal Mesh. Not to be confused with the Update Result. */
+	ESkeletalMeshStatus SkeletalMeshStatus = ESkeletalMeshStatus::NotGenerated;
+	
 	// Struct used by BuildMaterials() to identify common materials between LODs
 	struct FMutableMaterialPlaceholder
 	{
diff --git a/Engine/Plugins/Experimental/Mutable/Source/CustomizableObject/Private/MuCO/CustomizableObjectInstance.cpp b/Engine/Plugins/Experimental/Mutable/Source/CustomizableObject/Private/MuCO/CustomizableObjectInstance.cpp
--- a/Engine/Plugins/Experimental/Mutable/Source/CustomizableObject/Private/MuCO/CustomizableObjectInstance.cpp
+++ b/Engine/Plugins/Experimental/Mutable/Source/CustomizableObject/Private/MuCO/CustomizableObjectInstance.cpp
@@ -91,6 +91,18 @@ void UCustomizableInstancePrivateData::SetLastMeshId(int32 ComponentIndex, int32
 }
 
 
+ESkeletalMeshStatus UCustomizableInstancePrivateData::GetSkeletalMeshStatus() const
+{
+	return SkeletalMeshStatus;
+}
+
+
+void UCustomizableInstancePrivateData::SetSkeletalMeshStatus(ESkeletalMeshStatus Status)
+{
+	SkeletalMeshStatus = Status;
+}
+
+
 void UCustomizableInstancePrivateData::InitLastUpdateData(const TSharedPtr<FMutableOperationData>& OperationData)
 {
 	LastUpdateData.LODs.Init(FInstanceGeneratedData::FLOD(), OperationData->NumLODsAvailable);
@@ -1774,7 +1786,6 @@ bool UCustomizableInstancePrivateData::UpdateSkeletalMesh_PostBeginUpdate0(UCust
 	{
 		UE_LOG(LogMutable, Warning, TEXT("Failed to generate SkeletalMesh for CO Instance %s. CO [%s]"), *Public->GetName(), *GetNameSafe(Public->GetCustomizableObject()));
 
-		Public->SkeletalMeshStatus = ESkeletalMeshState::UpdateError;
 		Public->SkeletalMeshes.Reset(); // What about all the references
 
 		SetCOInstanceFlags(Generated);
@@ -1784,8 +1795,6 @@ bool UCustomizableInstancePrivateData::UpdateSkeletalMesh_PostBeginUpdate0(UCust
 		return false;
 	}
 
-	Public->SkeletalMeshStatus = ESkeletalMeshState::Correct;
-
 	// None of the current meshes requires a mesh update. Continue to BuildMaterials
 	if (!bUpdateMeshes)
 	{
@@ -2090,10 +2099,6 @@ bool UCustomizableInstancePrivateData::UpdateSkeletalMesh_PostBeginUpdate0(UCust
 
 	if (!bSuccess)
 	{
-#if WITH_EDITOR
-		Public->SkeletalMeshStatus = Public->PreUpdateSkeletalMeshStatus; // The mesh won't be updated, set the previous SkeletalMeshStatus to avoid losing error notifications
-#endif
-
 		Public->SkeletalMeshes.Reset();
 
 		InvalidateGeneratedData();
@@ -2150,7 +2155,6 @@ UCustomizableObjectInstance* UCustomizableObjectInstance::CloneStatic(UObject* O
 void UCustomizableObjectInstance::CopyParametersFromInstance(UCustomizableObjectInstance* Instance)
 {
 	SetDescriptor(Instance->GetDescriptor());
-	SkeletalMeshStatus = Instance->SkeletalMeshStatus;
 }
 
 
@@ -2761,6 +2765,7 @@ void UCustomizableInstancePrivateData::DiscardResourcesAndSetReferenceSkeletalMe
 	}
 
 	ClearCOInstanceFlags(Generated);
+	SkeletalMeshStatus = ESkeletalMeshStatus::NotGenerated;
 	
 	InvalidateGeneratedData();
 	
@@ -2835,22 +2840,6 @@ void SetTexturePropertiesFromMutableImageProps(UTexture2D* Texture, const FMutab
 }
 
 
-void UCustomizableObjectInstance::FinishUpdate(EUpdateResult UpdateResult, const FDescriptorRuntimeHash& InUpdatedHash)
-{
-	if (UpdateResult == EUpdateResult::Success)
-	{
-		GetPrivate()->DescriptorRuntimeHash = InUpdatedHash;
-	}
-
-	if (UpdateResult == EUpdateResult::Success)
-	{
-		// Call Instance updated (this) callbacks.
-		UpdatedDelegate.Broadcast(this);
-		UpdatedNativeDelegate.Broadcast(this);
-	}
-}
-
-
 FDescriptorRuntimeHash UCustomizableObjectInstance::GetDescriptorRuntimeHash() const
 {
 	return GetPrivate()->DescriptorRuntimeHash;
@@ -4860,7 +4849,7 @@ bool UCustomizableInstancePrivateData::BuildOrCopyRenderData(const TSharedPtr<FM
 
 		if (LODModel.DoesVertexBufferUse16BitBoneIndex() && !UCustomizableObjectSystem::GetInstance()->IsSupport16BitBoneIndexEnabled())
 		{
-			Public->SkeletalMeshStatus = ESkeletalMeshState::PostUpdateError;
+			OperationData->UpdateResult = EUpdateResult::Error16BitBoneIndex;
 
 			const FString Msg = FString::Printf(TEXT("Customizable Object [%s] requires of Skinning - 'Support 16 Bit Bone Index' to be enabled. Please, update the Project Settings."),
 				*CustomizableObject->GetName());
diff --git a/Engine/Plugins/Experimental/Mutable/Source/CustomizableObject/Private/MuCO/CustomizableObjectSystem.cpp b/Engine/Plugins/Experimental/Mutable/Source/CustomizableObject/Private/MuCO/CustomizableObjectSystem.cpp
--- a/Engine/Plugins/Experimental/Mutable/Source/CustomizableObject/Private/MuCO/CustomizableObjectSystem.cpp
+++ b/Engine/Plugins/Experimental/Mutable/Source/CustomizableObject/Private/MuCO/CustomizableObjectSystem.cpp
@@ -150,7 +150,7 @@ void FMutablePendingInstanceWork::RemoveUpdate(const TWeakObjectPtr<UCustomizabl
 }
 
 
-const FMutablePendingInstanceUpdate* FMutablePendingInstanceWork::GetUpdate(const TWeakObjectPtr<UCustomizableObjectInstance>& Instance) const
+const FMutablePendingInstanceUpdate* FMutablePendingInstanceWork::GetUpdate(const TWeakObjectPtr<const UCustomizableObjectInstance>& Instance) const
 {
 	return PendingInstanceUpdates.Find(Instance);
 }
@@ -165,8 +165,6 @@ void FMutablePendingInstanceWork::AddDiscard(const FMutablePendingInstanceDiscar
 	}
 
 	PendingInstanceDiscards.Add(TaskToEnqueue);
-
-	TaskToEnqueue.CustomizableObjectInstance->GetPrivate()->SetCOInstanceFlags(Updating);
 }
 
 
@@ -687,11 +685,40 @@ static FAutoConsoleVariableRef CVarApplyFixPrepareSkeletons(
 void FinishUpdateGlobal(UCustomizableObjectInstance* Instance, EUpdateResult UpdateResult, const FInstanceUpdateDelegate* UpdateCallback, const FDescriptorRuntimeHash InUpdatedHash)
 {
 	check(IsInGameThread())
-	
-	// Callbacks. Must be done at the end.
+
 	if (Instance)
 	{
-		Instance->FinishUpdate(UpdateResult, InUpdatedHash);			
+		UCustomizableInstancePrivateData* PrivateInstance = Instance->GetPrivate();
+		
+		switch (UpdateResult)
+		{
+		case EUpdateResult::Success:
+			PrivateInstance->SetSkeletalMeshStatus(ESkeletalMeshStatus::Success);
+
+			PrivateInstance->DescriptorRuntimeHash = InUpdatedHash;
+
+			// Delegates must be called only after updating the Instance flags.
+			Instance->UpdatedDelegate.Broadcast(Instance);
+			Instance->UpdatedNativeDelegate.Broadcast(Instance);
+			break;
+
+		case EUpdateResult::ErrorOptimized:
+			break; // Skeletal Mesh not changed.
+			
+		case EUpdateResult::ErrorDiscarded:
+			break; // Status will be updated once the discard is performed.
+
+		case EUpdateResult::Error: 
+		case EUpdateResult::Error16BitBoneIndex:
+			PrivateInstance->SetSkeletalMeshStatus(ESkeletalMeshStatus::Error);
+			break;
+			
+		case EUpdateResult::ErrorReplaced:
+			break; // Skeletal Mesh not changed.
+			
+		default:
+			unimplemented();
+		}
 	}
 
 	if (UpdateResult == EUpdateResult::Success)
@@ -913,7 +940,7 @@ EUpdateRequired FCustomizableObjectSystemPrivate::IsUpdateRequired(const UCustom
 	if (Private->HasCOInstanceFlags(DiscardedByNumInstancesLimit) ||
 		bLODManagementDiscard)
 	{
-		if (bIsGenerated && !Private->HasCOInstanceFlags(Updating))
+		if (bIsGenerated)
 		{
 			return EUpdateRequired::Discard;		
 		}
@@ -978,22 +1005,23 @@ void FCustomizableObjectSystemPrivate::EnqueueUpdateSkeletalMesh(UCustomizableOb
 	MUTABLE_CPUPROFILER_SCOPE(FCustomizableObjectSystemPrivate::EnqueueUpdateSkeletalMesh);
 	check(IsInGameThread());
 
+	UCustomizableInstancePrivateData* InstancePrivate = Instance.GetPrivate();
+	
 	if (!Instance.CanUpdateInstance())
 	{
-		FinishUpdateGlobal(&Instance, EUpdateResult::ErrorDiscarded, UpdateCallback);
+		FinishUpdateGlobal(&Instance, EUpdateResult::Error, UpdateCallback);
 		return;
 	}
 
 	UCustomizableObjectSystem* System = UCustomizableObjectSystem::GetInstance();
 
-	UCustomizableInstancePrivateData* InstancePrivate = Instance.GetPrivate();
 	
 	const EUpdateRequired UpdateRequired = IsUpdateRequired(Instance, bOnlyUpdateIfNotGenerated, bIgnoreCloseDist);
 	switch (UpdateRequired)
 	{
 	case EUpdateRequired::NoUpdate:
 	{	
-		FinishUpdateGlobal(&Instance, EUpdateResult::ErrorDiscarded, UpdateCallback);
+		FinishUpdateGlobal(&Instance, EUpdateResult::Error, UpdateCallback);
 		break;
 	}		
 	case EUpdateRequired::Update:
@@ -1005,16 +1033,6 @@ void FCustomizableObjectSystemPrivate::EnqueueUpdateSkeletalMesh(UCustomizableOb
 		const bool bIsPlayerOrNearIt = InstancePrivate->HasCOInstanceFlags(UsedByPlayerOrNearIt);
 		UE_LOG(LogMutable, Log, TEXT("Enqueued UpdateSkeletalMesh Async. of Instance %d with priority %d at dist %f bIsPlayerOrNearIt=%d, frame=%d"), InstanceId, static_cast<int32>(Priority), Distance, bIsPlayerOrNearIt, GFrameNumber);				
 
-#if WITH_EDITOR
-		if (!UCustomizableObjectSystem::GetInstance()->IsCompilationDisabled())
-		{
-			if (Instance.SkeletalMeshStatus != ESkeletalMeshState::AsyncUpdatePending)
-			{
-				Instance.PreUpdateSkeletalMeshStatus = Instance.SkeletalMeshStatus;
-			}
-		}
-#endif
-
 		if (InstancePrivate->HasCOInstanceFlags(PendingLODsUpdate))
 		{
 			UE_LOG(LogMutable, Verbose, TEXT("LOD change: %d, %d -> %d, %d"), Instance.GetCurrentMinLOD(), Instance.GetCurrentMaxLOD(), Instance.GetMinLODToLoad(), Instance.GetMaxLODToLoad());
@@ -1043,7 +1061,6 @@ void FCustomizableObjectSystemPrivate::EnqueueUpdateSkeletalMesh(UCustomizableOb
 			!(CurrentMutableOperation &&
 			&Instance == CurrentMutableOperation->CustomizableObjectInstance)) // This condition is necessary because even if the descriptor is a subset, it will be replaced by the CurrentMutableOperation
 		{
-			Instance.SkeletalMeshStatus = ESkeletalMeshState::Correct; // TODO FutureGMT MTBL-1033 should not be here. Move to UCustomizableObjectInstance::Updated
 			UpdateSkeletalMesh(Instance, Instance.GetDescriptorRuntimeHash(), EUpdateResult::Success, UpdateCallback);
 		}
 		else
@@ -1080,8 +1097,6 @@ void FCustomizableObjectSystemPrivate::EnqueueUpdateSkeletalMesh(UCustomizableOb
 				}
 			}
 
-			Instance.SkeletalMeshStatus = ESkeletalMeshState::AsyncUpdatePending;
-
 			const FMutablePendingInstanceUpdate InstanceUpdate(&Instance, Priority, UpdateCallback);
 			MutablePendingInstanceWork.AddUpdate(InstanceUpdate);
 		}
@@ -2584,9 +2599,6 @@ namespace impl
 		// Skip update, the requested update is equal to the running update.
 		if (Operation->InstanceDescriptorRuntimeHash.IsSubset(CandidateInstance->GetDescriptorRuntimeHash()))
 		{
-			CandidateInstance->SkeletalMeshStatus = ESkeletalMeshState::Correct;
-
-			CandidateInstancePrivateData->ClearCOInstanceFlags(Updating);
 			System->ClearCurrentMutableOperation();
 			UpdateSkeletalMesh(*CandidateInstance, CandidateInstance->GetDescriptorRuntimeHash(), EUpdateResult::ErrorOptimized, &Operation->UpdateCallback);
 			return;
@@ -2627,8 +2639,6 @@ namespace impl
 
 		if (bCancel)
 		{
-			CandidateInstancePrivateData->ClearCOInstanceFlags(Updating);
-
 			System->ClearCurrentMutableOperation();
 
 			FinishUpdateGlobal(CandidateInstance, EUpdateResult::Error, &Operation->UpdateCallback);
@@ -3014,8 +3024,9 @@ bool UCustomizableObjectSystem::Tick(float DeltaTime)
 					LODUpdateWithSameInstance->ApplyLODUpdateParamsToInstance();
 				}
 
-				Private->CurrentMutableOperation = MakeShared<FMutableOperation>(FMutableOperation::CreateInstanceUpdate(*PendingInstance, PendingInstanceUpdateFound->Callback));
-
+				const TSharedPtr<FMutableOperation> Operation = MakeShared<FMutableOperation>(FMutableOperation::CreateInstanceUpdate(*PendingInstance, PendingInstanceUpdateFound->Callback));
+				Private->StartUpdateSkeletalMesh(Operation);
+				
 				Private->MutablePendingInstanceWork.RemoveUpdate(PendingInstanceUpdateFound->CustomizableObjectInstance);
 			}
 			else if (LODUpdateCandidateFound)
@@ -3023,7 +3034,8 @@ bool UCustomizableObjectSystem::Tick(float DeltaTime)
 				// Commit the LOD changes
 				LODUpdateCandidateFound->ApplyLODUpdateParamsToInstance();
 
-				Private->CurrentMutableOperation = MakeShared<FMutableOperation>(FMutableOperation::CreateInstanceUpdate(*LODUpdateCandidateFound->CustomizableObjectInstance, nullptr));
+				const TSharedRef<FMutableOperation> Operation = MakeShared<FMutableOperation>(FMutableOperation::CreateInstanceUpdate(*LODUpdateCandidateFound->CustomizableObjectInstance, nullptr));
+				Private->StartUpdateSkeletalMesh(Operation);
 			}
 		}
 
@@ -3086,7 +3098,8 @@ void UCustomizableObjectSystem::DiscardInstances()
 
 		UCustomizableObjectInstance* COI = Iterator->CustomizableObjectInstance.Get();
 
-		if (COI)
+		const bool bUpdating = Private->CurrentMutableOperation && Private->CurrentMutableOperation->CustomizableObjectInstance != Iterator->CustomizableObjectInstance;
+		if (COI && !bUpdating)
 		{
 			UCustomizableInstancePrivateData* COIPrivateData = COI ? COI->GetPrivate() : nullptr;
 
@@ -3102,22 +3115,8 @@ void UCustomizableObjectSystem::DiscardInstances()
 				{
 					check(COIPrivateData != nullptr);
 					COIPrivateData->DiscardResourcesAndSetReferenceSkeletalMesh(COI);
-					COIPrivateData->ClearCOInstanceFlags(Updating);
-					COI->SkeletalMeshStatus = ESkeletalMeshState::Correct;
 				}
 			}
-			else
-			{
-				check(COIPrivateData != nullptr);
-				COIPrivateData->ClearCOInstanceFlags(Updating);
-			}
-
-			if (COI && !COI->HasAnySkeletalMesh())
-			{
-				// To solve the problem in the Mutable demo where PIE just after editor start made all instances appear as reference mesh until editor restart
-				check(COIPrivateData != nullptr);
-				COIPrivateData->ClearCOInstanceFlags(Generated);
-			}
 		}
 
 		Iterator.RemoveCurrent();
@@ -3150,6 +3149,17 @@ void UCustomizableObjectSystem::ReleaseInstanceIDs()
 }
 
 
+bool UCustomizableObjectSystem::IsUpdating(const UCustomizableObjectInstance* Instance) const
+{
+	if (!Instance)
+	{
+		return false;
+	}
+	
+	return GetPrivateChecked()->IsUpdating(*Instance);
+}
+
+
 TArray<FCustomizableObjectExternalTexture> UCustomizableObjectSystem::GetTextureParameterValues()
 {
 	TArray<FCustomizableObjectExternalTexture> Result;
@@ -3666,6 +3676,31 @@ FUnrealMutableImageProvider* FCustomizableObjectSystemPrivate::GetImageProviderC
 }
 
 
+void FCustomizableObjectSystemPrivate::StartUpdateSkeletalMesh(const TSharedPtr<FMutableOperation>& Operation)
+{
+	check(!CurrentMutableOperation); // Can not start an update if there is already another in progress
+	check(Operation->CustomizableObjectInstance.IsValid()) // The instance has to be alive to start the update
+		
+	CurrentMutableOperation = Operation;
+}
+
+
+bool FCustomizableObjectSystemPrivate::IsUpdating(const UCustomizableObjectInstance& Instance) const
+{
+	if (CurrentMutableOperation && CurrentMutableOperation->CustomizableObjectInstance.Get() == &Instance)
+	{
+		return true;
+	}
+
+	if (MutablePendingInstanceWork.GetUpdate(TWeakObjectPtr<const UCustomizableObjectInstance>(&Instance)))
+	{
+		return true;
+	}
+	
+	return false;
+}
+
+
 bool UCustomizableObjectSystem::IsMutableAnimInfoDebuggingEnabled() const
 {
 #if WITH_EDITOR
diff --git a/Engine/Plugins/Experimental/Mutable/Source/CustomizableObject/Private/MuCO/CustomizableObjectSystemPrivate.h b/Engine/Plugins/Experimental/Mutable/Source/CustomizableObject/Private/MuCO/CustomizableObjectSystemPrivate.h
--- a/Engine/Plugins/Experimental/Mutable/Source/CustomizableObject/Private/MuCO/CustomizableObjectSystemPrivate.h
+++ b/Engine/Plugins/Experimental/Mutable/Source/CustomizableObject/Private/MuCO/CustomizableObjectSystemPrivate.h
@@ -146,19 +146,19 @@ inline uint32 GetTypeHash(const FMutablePendingInstanceUpdate& Update)
 }
 
 
-struct FPendingInstanceUpdateKeyFuncs : BaseKeyFuncs<FMutablePendingInstanceUpdate, TWeakObjectPtr<UCustomizableObjectInstance>>
+struct FPendingInstanceUpdateKeyFuncs : BaseKeyFuncs<FMutablePendingInstanceUpdate, TWeakObjectPtr<const UCustomizableObjectInstance>>
 {
-	FORCEINLINE static const TWeakObjectPtr<UCustomizableObjectInstance>& GetSetKey(const FMutablePendingInstanceUpdate& PendingUpdate)
+	FORCEINLINE static TWeakObjectPtr<const UCustomizableObjectInstance> GetSetKey(const FMutablePendingInstanceUpdate& PendingUpdate)
 	{
 		return PendingUpdate.CustomizableObjectInstance;
 	}
 
-	FORCEINLINE static bool Matches(const TWeakObjectPtr<UCustomizableObjectInstance>& A, const TWeakObjectPtr<UCustomizableObjectInstance>& B)
+	FORCEINLINE static bool Matches(const TWeakObjectPtr<const UCustomizableObjectInstance>& A, const TWeakObjectPtr<const UCustomizableObjectInstance>& B)
 	{
 		return A.HasSameIndexAndSerialNumber(B);
 	}
 
-	FORCEINLINE static uint32 GetKeyHash(const TWeakObjectPtr<UCustomizableObjectInstance>& Identifier)
+	FORCEINLINE static uint32 GetKeyHash(const TWeakObjectPtr<const UCustomizableObjectInstance>& Identifier)
 	{
 		return GetTypeHash(Identifier.GetWeakPtrTypeHash());
 	}
@@ -235,7 +235,7 @@ public:
 	// Removes an instance update
 	void RemoveUpdate(const TWeakObjectPtr<UCustomizableObjectInstance>& Instance);
 
-	const FMutablePendingInstanceUpdate* GetUpdate(const TWeakObjectPtr<UCustomizableObjectInstance>& Instance) const;
+	const FMutablePendingInstanceUpdate* GetUpdate(const TWeakObjectPtr<const UCustomizableObjectInstance>& Instance) const;
 
 	TSet<FMutablePendingInstanceUpdate, FPendingInstanceUpdateKeyFuncs>::TIterator GetUpdateIterator()
 	{
@@ -752,6 +752,12 @@ public:
 	bool IsMutableAnimInfoDebuggingEnabled() const;
 
 	FUnrealMutableImageProvider* GetImageProviderChecked() const;
+
+	/** Start the actual work of Update Skeletal Mesh process (Update Skeletal Mesh without the queue). */
+	void StartUpdateSkeletalMesh(const TSharedPtr<FMutableOperation>& Operation);
+
+	/** See UCustomizableObjectInstance::IsUpdating. */
+	bool IsUpdating(const UCustomizableObjectInstance& Instance) const;
 	
 	/** Mutable TaskGraph system (Mutable Thread). */
 	FMutableTaskGraph MutableTaskGraph;
diff --git a/Engine/Plugins/Experimental/Mutable/Source/CustomizableObject/Private/MuCO/LogInformationUtil.cpp b/Engine/Plugins/Experimental/Mutable/Source/CustomizableObject/Private/MuCO/LogInformationUtil.cpp
--- a/Engine/Plugins/Experimental/Mutable/Source/CustomizableObject/Private/MuCO/LogInformationUtil.cpp
+++ b/Engine/Plugins/Experimental/Mutable/Source/CustomizableObject/Private/MuCO/LogInformationUtil.cpp
@@ -107,7 +107,6 @@ void LogInformationUtil::LogShowInstanceData(const UCustomizableObjectInstance*
 	LogData += FString::Printf(TEXT("Priority=%f "), FMath::Sqrt(CustomizableObjectInstance->GetPrivate()->LastMinSquareDistFromComponentToPlayer));
 
 	const uint32 InstanceFlags = CustomizableObjectInstance->GetPrivate()->GetCOInstanceFlags();
-	LogData += FString::Printf(TEXT("bIsUpdating=%d "), InstanceFlags & ECOInstanceFlags::Updating);
 	LogData += FString::Printf(TEXT("bShouldUpdateLODs=%d "), InstanceFlags & ECOInstanceFlags::PendingLODsUpdate);
 
 	LogData += FString::Printf(TEXT("CurrentMinLOD=%d "), CustomizableObjectInstance->GetCurrentMinLOD());
@@ -166,7 +165,6 @@ void LogInformationUtil::LogShowInstanceDataFull(const UCustomizableObjectInstan
 	LogData += MessageChunk;
 
 	MessageChunk += FString::Printf(TEXT("        bIsDowngradeLODUpdate = %d\n"), InstanceFlags & ECOInstanceFlags::PendingLODsDowngrade);
-	MessageChunk += FString::Printf(TEXT("        bIsUpdating = %d\n"), InstanceFlags & ECOInstanceFlags::Updating);
 	MessageChunk += FString::Printf(TEXT("        bIsCreatingSkeletalMesh = %d\n"), InstanceFlags & ECOInstanceFlags::CreatingSkeletalMesh);
 	MessageChunk += FString::Printf(TEXT("        bIsGenerated = %d\n"), InstanceFlags & ECOInstanceFlags::Generated);
 	LogData += MessageChunk;
diff --git a/Engine/Plugins/Experimental/Mutable/Source/CustomizableObjectEditor/Private/MuCOE/CustomizableObjectEditor.cpp b/Engine/Plugins/Experimental/Mutable/Source/CustomizableObjectEditor/Private/MuCOE/CustomizableObjectEditor.cpp
--- a/Engine/Plugins/Experimental/Mutable/Source/CustomizableObjectEditor/Private/MuCOE/CustomizableObjectEditor.cpp
+++ b/Engine/Plugins/Experimental/Mutable/Source/CustomizableObjectEditor/Private/MuCOE/CustomizableObjectEditor.cpp
@@ -52,6 +52,7 @@
 #include "PropertyEditorModule.h"
 #include "ScopedTransaction.h"
 #include "MuCOE/SCustomizableInstanceProperties.h"
+#include "MuCO/CustomizableInstancePrivateData.h"
 #include "UObject/EnumProperty.h"
 #include "Widgets/Docking/SDockTab.h"
 #include "Widgets/Input/SSearchBox.h"
@@ -2489,7 +2490,7 @@ void FCustomizableObjectEditor::UpdatePreviewVisibility()
 	{
 		if (PreviewSkeletalMeshComponent)
 		{
-			if (PreviewInstance->SkeletalMeshStatus != ESkeletalMeshState::UpdateError)
+			if (PreviewInstance->GetPrivate()->GetSkeletalMeshStatus() == ESkeletalMeshStatus::Success)
 			{
 				PreviewSkeletalMeshComponent->SetVisibility(true, true);
 			}
diff --git a/Engine/Plugins/Experimental/Mutable/Source/CustomizableObjectEditor/Private/MuCOE/CustomizableObjectInstanceEditor.cpp b/Engine/Plugins/Experimental/Mutable/Source/CustomizableObjectEditor/Private/MuCOE/CustomizableObjectInstanceEditor.cpp
--- a/Engine/Plugins/Experimental/Mutable/Source/CustomizableObjectEditor/Private/MuCOE/CustomizableObjectInstanceEditor.cpp
+++ b/Engine/Plugins/Experimental/Mutable/Source/CustomizableObjectEditor/Private/MuCOE/CustomizableObjectInstanceEditor.cpp
@@ -34,6 +34,7 @@
 #include "Widgets/Notifications/SNotificationList.h"
 #include "WorkspaceMenuStructure.h"
 #include "WorkspaceMenuStructureModule.h"
+#include "MuCO/CustomizableInstancePrivateData.h"
 
 class FAdvancedPreviewScene;
 class FProperty;
@@ -436,7 +437,7 @@ void FCustomizableObjectInstanceEditor::OnAssetRegistryLoadComplete()
 
 void FCustomizableObjectInstanceEditor::UpdatePreviewVisibility()
 {
-	const bool bEnableVisibility = CustomizableObjectInstance->SkeletalMeshStatus == ESkeletalMeshState::Correct;
+	const bool bEnableVisibility = CustomizableObjectInstance->GetPrivate()->GetSkeletalMeshStatus() == ESkeletalMeshStatus::Success;
 	for (UDebugSkelMeshComponent* PreviewSkeletalMeshComponent : PreviewSkeletalMeshComponents)
 	{
 		PreviewSkeletalMeshComponent->SetVisibility(bEnableVisibility, true);
diff --git a/Engine/Plugins/Experimental/Mutable/Source/CustomizableObjectEditor/Private/MuCOE/SCustomizableObjectEditorViewportToolBar.cpp b/Engine/Plugins/Experimental/Mutable/Source/CustomizableObjectEditor/Private/MuCOE/SCustomizableObjectEditorViewportToolBar.cpp
--- a/Engine/Plugins/Experimental/Mutable/Source/CustomizableObjectEditor/Private/MuCOE/SCustomizableObjectEditorViewportToolBar.cpp
+++ b/Engine/Plugins/Experimental/Mutable/Source/CustomizableObjectEditor/Private/MuCOE/SCustomizableObjectEditorViewportToolBar.cpp
@@ -21,6 +21,8 @@
 #include "SViewportToolBarComboMenu.h"
 #include "Settings/LevelEditorViewportSettings.h"
 #include "UnrealEdGlobals.h"
+#include "MuCO/CustomizableInstancePrivateData.h"
+#include "MuCO/CustomizableObjectSystem.h"
 #include "Widgets/Input/SButton.h"
 #include "Widgets/Input/SMenuAnchor.h"
 #include "Widgets/Input/SSpinBox.h"
@@ -34,12 +36,14 @@ struct FSlateBrush;
 #define LOCTEXT_NAMESPACE "CustomizableObjectEditorViewportToolBar"
 
 
-void SCustomizableObjectEditorViewportToolBar::Construct(const FArguments& InArgs, TSharedPtr<class SCustomizableObjectEditorViewportTabBody> InViewport, TSharedPtr<class SEditorViewport> InRealViewport)
+void SCustomizableObjectEditorViewportToolBar::Construct(const FArguments& InArgs, TSharedPtr<SCustomizableObjectEditorViewportTabBody> InViewport, TSharedPtr<SEditorViewport> InRealViewport)
 {
 	Viewport = InViewport;
 
 	TSharedRef<SCustomizableObjectEditorViewportTabBody> ViewportRef = Viewport.Pin().ToSharedRef();
 
+	WeakEditor = ViewportRef->CustomizableObjectEditorPtr;
+	
 	TSharedRef<SHorizontalBox> LeftToolbar = SNew(SHorizontalBox)
 
 	//// Camera Type (Perspective/Top/etc...)
@@ -171,6 +175,7 @@ void SCustomizableObjectEditorViewportToolBar::Construct(const FArguments& InArg
 	SViewportToolBar::Construct(SViewportToolBar::FArguments());
 }
 
+
 EVisibility SCustomizableObjectEditorViewportToolBar::GetShowCompileErrorOverlay() const
 {
 	return GetCompileErrorOverlayText().IsEmpty() ? EVisibility::Hidden : EVisibility::Visible;
@@ -179,45 +184,47 @@ EVisibility SCustomizableObjectEditorViewportToolBar::GetShowCompileErrorOverlay
 
 FText SCustomizableObjectEditorViewportToolBar::GetCompileErrorOverlayText() const
 {
-	if ((Viewport.IsValid()) && (Viewport.Pin().Get() != nullptr) &&
-		(Viewport.Pin()->CustomizableObjectEditorPtr.IsValid()))
+	const TSharedPtr<ICustomizableObjectInstanceEditor> Editor = WeakEditor.Pin();
+	if (!Editor)
 	{
-		if (Viewport.Pin()->CustomizableObjectEditorPtr.Pin()->GetAssetRegistryLoaded())
-		{
-			UCustomizableObjectInstance* Instance = Viewport.Pin()->CustomizableObjectEditorPtr.Pin()->GetPreviewInstance();
+		return {};
+	}
 
-			if (Instance != nullptr)
-			{
-				if (Instance->SkeletalMeshStatus == ESkeletalMeshState::UpdateError)
-				{
-					return FText::FromString("Error updating skeletal mesh");
-				}
-				else if (Instance->SkeletalMeshStatus == ESkeletalMeshState::PostUpdateError)
-				{
-					return FText::FromString("Post Update Error: check the output log for more information.");
-				}
-				else if (Instance->SkeletalMeshStatus == ESkeletalMeshState::AsyncUpdatePending)
-				{
-					return FText::FromString("Updating skeletal mesh");
-				}
-			}
-			else
-			{
-				return FText::FromString("No Skeletal Mesh generated");
-			}
-		}
-		else
-		{
-			return FText::FromString("Loading assets");
-		}
+	if (!Editor->GetAssetRegistryLoaded())
+	{
+		return LOCTEXT("LoadingAssets", "Loading Assets");
 	}
-	return {};
-}
 
+	const UCustomizableObjectInstance* Instance = Editor->GetPreviewInstance();
+	if (!Instance)
+	{
+		return LOCTEXT("NoPreviewInstance", "No Preview Instance");
+	}
 
-EVisibility SCustomizableObjectEditorViewportToolBar::GetTransformToolbarVisibility() const
-{
-	return EVisibility::Hidden;
+	UCustomizableObjectSystem* System = UCustomizableObjectSystem::GetInstanceChecked();
+	
+	if (System->IsUpdating(Instance))
+	{
+		return LOCTEXT("UpdatingSkeletalMesh", "Updating Skeletal Mesh");
+	}
+
+	const UCustomizableInstancePrivateData* PrivateInstance = Instance->GetPrivate();
+	
+	switch (PrivateInstance->GetSkeletalMeshStatus())
+	{
+	case ESkeletalMeshStatus::NotGenerated:
+		return LOCTEXT("NoSkeletalMeshGenerated", "No Skeletal Mesh Generated");
+				
+	case ESkeletalMeshStatus::Error:
+		return LOCTEXT("ErrorUpdatingSkeletalMesh", "Error Updating Skeletal Mesh");
+
+	case ESkeletalMeshStatus::Success:
+		return {};
+
+	default:
+		unimplemented();
+		return {};
+	}
 }
 
 
diff --git a/Engine/Plugins/Experimental/Mutable/Source/CustomizableObjectEditor/Public/MuCOE/SCustomizableObjectEditorViewportToolBar.h b/Engine/Plugins/Experimental/Mutable/Source/CustomizableObjectEditor/Public/MuCOE/SCustomizableObjectEditorViewportToolBar.h
--- a/Engine/Plugins/Experimental/Mutable/Source/CustomizableObjectEditor/Public/MuCOE/SCustomizableObjectEditorViewportToolBar.h
+++ b/Engine/Plugins/Experimental/Mutable/Source/CustomizableObjectEditor/Public/MuCOE/SCustomizableObjectEditorViewportToolBar.h
@@ -4,6 +4,7 @@
 
 #include "SViewportToolBar.h"
 
+class ICustomizableObjectInstanceEditor;
 enum ERotationGridMode : int;
 enum class ECheckBoxState : uint8;
 namespace ETextCommit { enum Type : int; }
@@ -103,9 +104,6 @@ private:
 	/** Called when the floor offset slider is adjusted in the perspective viewport */
 	void OnFloorOffsetChanged(float NewValue);
 
-	// Called to determine if the gizmos can be used in the current preview
-	EVisibility GetTransformToolbarVisibility() const;
-
 	// Called to show / hide Customizable Object compile error
 	EVisibility GetShowCompileErrorOverlay() const;
 
@@ -130,6 +128,8 @@ private:
 	/** The viewport that we are in */
 	TWeakPtr<SCustomizableObjectEditorViewportTabBody> Viewport;
 
+	TWeakPtr<ICustomizableObjectInstanceEditor> WeakEditor;
+
 	// Layout to show information about instance skeletal mesh update / CO asset data
 	TSharedPtr<SButton> CompileErrorLayout;
 
