[mutable] Fixed mesh layout details editor scroll focus.

#rnx

diff --git a/Engine/Plugins/Experimental/Mutable/Source/CustomizableObjectEditor/Private/MuCOE/SCustomizableObjectLayoutGrid.cpp b/Engine/Plugins/Experimental/Mutable/Source/CustomizableObjectEditor/Private/MuCOE/SCustomizableObjectLayoutGrid.cpp
--- a/Engine/Plugins/Experimental/Mutable/Source/CustomizableObjectEditor/Private/MuCOE/SCustomizableObjectLayoutGrid.cpp
+++ b/Engine/Plugins/Experimental/Mutable/Source/CustomizableObjectEditor/Private/MuCOE/SCustomizableObjectLayoutGrid.cpp
@@ -844,6 +844,8 @@ FReply SCustomizableObjectLayoutGrid::OnMouseWheel(const FGeometry& MyGeometry,
 				Zoom--;
 			}
 		}
+
+		return FReply::Handled().SetUserFocus(SharedThis(this), EFocusCause::Mouse, true);
 	}
 
 	return SCompoundWidget::OnMouseWheel(MyGeometry, MouseEvent);
