#rnx

diff --git a/Engine/Plugins/Experimental/Mutable/Source/CustomizableObject/Private/MuCO/CustomizableObjectSystem.cpp b/Engine/Plugins/Experimental/Mutable/Source/CustomizableObject/Private/MuCO/CustomizableObjectSystem.cpp
--- a/Engine/Plugins/Experimental/Mutable/Source/CustomizableObject/Private/MuCO/CustomizableObjectSystem.cpp
+++ b/Engine/Plugins/Experimental/Mutable/Source/CustomizableObject/Private/MuCO/CustomizableObjectSystem.cpp
@@ -2333,7 +2333,11 @@ namespace impl
 
 		for (int32 ComponentIndex = 0; ComponentIndex < CustomizableObjectInstance->GetPrivate()->ComponentsData.Num(); ++ComponentIndex)
 		{
-			ensure(CustomizableObjectInstance->SkeletalMeshes.IsValidIndex(ComponentIndex));
+			// TODO PRP: Fix properly by moving the whole AssetUserData block to UCustomizableInstancePrivateData::InitSkeletalMeshData
+			if (!CustomizableObjectInstance->SkeletalMeshes.IsValidIndex(ComponentIndex))
+			{
+				break;
+			}
 
 			FCustomizableInstanceComponentData& ComponentData = CustomizableObjectInstance->GetPrivate()->ComponentsData[ComponentIndex];
 			
