[Mutable] Fix cpu profile name tag for ImagePixelFormatTask and ImageMipmapTask prepare stages.

#rnx

diff --git a/Engine/Plugins/Experimental/Mutable/Source/MutableRuntime/Private/MuR/CodeRunnerTasked.cpp b/Engine/Plugins/Experimental/Mutable/Source/MutableRuntime/Private/MuR/CodeRunnerTasked.cpp
--- a/Engine/Plugins/Experimental/Mutable/Source/MutableRuntime/Private/MuR/CodeRunnerTasked.cpp
+++ b/Engine/Plugins/Experimental/Mutable/Source/MutableRuntime/Private/MuR/CodeRunnerTasked.cpp
@@ -1191,7 +1191,7 @@ namespace mu
 	bool FImagePixelFormatTask::Prepare(CodeRunner* Runner, bool& bOutFailed)
 	{
 		// This runs in the mutable Runner thread
-		MUTABLE_CPUPROFILER_SCOPE(FImageLayerColourTask_Prepare);
+		MUTABLE_CPUPROFILER_SCOPE(FImageLayerPixelFormatTask_Prepare);
 		bOutFailed = false;
 
 		ImageCompressionQuality = Runner->m_pSettings->ImageCompressionQuality;
@@ -1302,7 +1302,7 @@ namespace mu
 	bool FImageMipmapTask::Prepare(CodeRunner* Runner, bool& bOutFailed)
 	{
 		// This runs in the mutable Runner thread
-		MUTABLE_CPUPROFILER_SCOPE(FImageLayerColourTask_Prepare);
+		MUTABLE_CPUPROFILER_SCOPE(FImageMipmapTask_Prepare);
 		bOutFailed = false;
 
 		ImageCompressionQuality = Runner->m_pSettings->ImageCompressionQuality;
