[Mutable] Fix possible crash in InputFileStream when opening a non existing file path.

#rnx

diff --git a/Engine/Plugins/Experimental/Mutable/Source/MutableTools/Private/MuT/Streams.cpp b/Engine/Plugins/Experimental/Mutable/Source/MutableTools/Private/MuT/Streams.cpp
--- a/Engine/Plugins/Experimental/Mutable/Source/MutableTools/Private/MuT/Streams.cpp
+++ b/Engine/Plugins/Experimental/Mutable/Source/MutableTools/Private/MuT/Streams.cpp
@@ -25,16 +25,19 @@ namespace mu
 
 		TSharedPtr<IFileHandle> FilePtr(FPlatformFileManager::Get().GetPlatformFile().OpenRead(*File));
 		m_pD->File = FilePtr;
-
-		m_pD->File->SeekFromEnd();
-		m_pD->FileSize = m_pD->File->Tell();
-		m_pD->File->Seek(0);
-
+		
+		if (m_pD->File)
+		{
+			m_pD->File->SeekFromEnd();
+			m_pD->FileSize = m_pD->File->Tell();
+			m_pD->File->Seek(0);
+		}
+		
 		m_pD->BytesInBuffer = 0;
 		m_pD->BufferPosition = 0;
-        m_pD->FilePosition = 0;
+		m_pD->FilePosition = 0;
 
-        m_pD->Buffer.SetNum(MUTABLE_STREAM_BUFFER_SIZE);
+		m_pD->Buffer.SetNum(MUTABLE_STREAM_BUFFER_SIZE);
 	}
 
 
@@ -49,13 +52,15 @@ namespace mu
 	//---------------------------------------------------------------------------------------------
 	bool InputFileStream::IsOpen() const
     {
-        return m_pD->File!=0;
+        return m_pD->File != nullptr;
     }
     
 
     //---------------------------------------------------------------------------------------------
 	uint64 InputFileStream::Tell() const
     {
+		check(IsOpen());
+
         size_t bufferOffset = m_pD->BytesInBuffer - m_pD->BufferPosition;
         int64 filePos = m_pD->File->Tell();
         check(filePos>0);
@@ -67,6 +72,8 @@ namespace mu
     //---------------------------------------------------------------------------------------------
     void InputFileStream::Seek( uint64 NewPosition )
     {
+		check(IsOpen());
+
         bool bSuccess = m_pD->File->Seek( int64(NewPosition) );
         check(bSuccess);
 
@@ -79,6 +86,8 @@ namespace mu
 	//---------------------------------------------------------------------------------------------
     void InputFileStream::Read( void* pData, uint64 size )
 	{
+		check(IsOpen());
+
         uint8* pDest = (uint8*) pData;
         while (size && m_pD->File )
 		{			
