[Mutable] Attach a callback to the PostCompileDelegate to update the state of the CustomizableObjectInstance after compiling the CO.
On execution, It will invalidate all generated data and re-initialize the ComponentsData to match the state of the compiled model.
* Remove InstanceUpdateFlags and Private::ReloadParameters methods.
* Move the Code to set the AssetUserData from the COSystem to the CustomizableInstancePrivateData
* Add PostDuplicate override to the CustomizableInstancePrivateData to invalidate generated data after duplicating an instance in the editor

#rnx

diff --git a/Engine/Plugins/Experimental/Mutable/Source/CustomizableObject/Classes/MuCO/CustomizableObjectInstance.h b/Engine/Plugins/Experimental/Mutable/Source/CustomizableObject/Classes/MuCO/CustomizableObjectInstance.h
--- a/Engine/Plugins/Experimental/Mutable/Source/CustomizableObject/Classes/MuCO/CustomizableObjectInstance.h
+++ b/Engine/Plugins/Experimental/Mutable/Source/CustomizableObject/Classes/MuCO/CustomizableObjectInstance.h
@@ -656,6 +656,12 @@ public:
 	bool bSelectedProfileDirty = false;
 #endif 
 
+private:
+
+#if WITH_EDITOR
+	void BindPostCompileDelegate(UCustomizableObject* CustomizableObject);
+#endif
+
 private:
 	UPROPERTY()
 	FCustomizableObjectInstanceDescriptor Descriptor;
diff --git a/Engine/Plugins/Experimental/Mutable/Source/CustomizableObject/Private/MuCO/CustomizableInstancePrivateData.h b/Engine/Plugins/Experimental/Mutable/Source/CustomizableObject/Private/MuCO/CustomizableInstancePrivateData.h
--- a/Engine/Plugins/Experimental/Mutable/Source/CustomizableObject/Private/MuCO/CustomizableInstancePrivateData.h
+++ b/Engine/Plugins/Experimental/Mutable/Source/CustomizableObject/Private/MuCO/CustomizableInstancePrivateData.h
@@ -128,8 +128,6 @@ struct FCustomizableInstanceComponentData
 {
 	GENERATED_USTRUCT_BODY();
 
-	uint16 ComponentIndex = 0;
-
 	// AnimBP data gathered for a component from its constituent meshes
 	UPROPERTY(Transient, Category = CustomizableObjectInstance, editfixedsize, VisibleAnywhere)
 	TMap<FName, TSoftClassPtr<UAnimInstance>> AnimSlotToBP;
@@ -161,8 +159,6 @@ struct FCustomizableInstanceComponentData
 	 *  Size == NumLODsAvailable
 	 *  LODs without mesh will be set to the maximum value of FResourceID (Max_uint64). */
 	TArray<mu::FResourceID> LastMeshIdPerLOD;
-
-	bool operator==(const FCustomizableInstanceComponentData& Other) const { return ComponentIndex == Other.ComponentIndex; }
 };
 
 USTRUCT()
@@ -219,6 +215,18 @@ public:
 	 * This cache is required since the Instance can have a LOD Update at any time.
 	 * So we need to make sure that the initially provided Texture Parameters by the user will be available until the user decides to change them. */
 	TArray<FName> UpdateTextureParameters;
+
+#if WITH_EDITOR
+
+	virtual void PostDuplicate(bool bDuplicateForPIE) override;
+
+	void OnPostCompile();
+
+#endif
+
+	/** Invalidates the previously generated data and retrieves information from the CObject after specific actions.
+	 *  It'll be called in the PostLoad, after Compiling the CO, and after changing the CO of the Instance. */
+	void InitCustomizableObjectData(const UCustomizableObject* InCustomizableObject);
 	
 	FCustomizableInstanceComponentData* GetComponentData(int32 ComponentIndex);
 	const FCustomizableInstanceComponentData* GetComponentData(int32 ComponentIndex) const;
@@ -236,12 +244,6 @@ public:
 	FGraphEventRef LoadAdditionalAssetsAsync(const TSharedPtr<FMutableOperationData>& OperationData, UCustomizableObjectInstance* Public, struct FStreamableManager &StreamableManager);
 	void AdditionalAssetsAsyncLoaded(UCustomizableObjectInstance* Public);
 
-	/** Temporal function. Used by GetParameters and ReloadParameters but its body should be moved somewhere else (probably on some other Instance update function). */
-	void InstanceUpdateFlags(const UCustomizableObjectInstance& Public);
-
-	/** See FCustomizableObjectInstanceDescriptor::ReloadParameters(...). */
-	void ReloadParameters(UCustomizableObjectInstance* Public, bool bInvalidatePreviousData = false);
-
 	void TickUpdateCloseCustomizableObjects(UCustomizableObjectInstance& Publics, FMutableInstanceUpdateMap& InOutRequestedUpdates);
 	void UpdateInstanceIfNotGenerated(UCustomizableObjectInstance& Public, FMutableInstanceUpdateMap& InOutRequestedUpdates);
 
diff --git a/Engine/Plugins/Experimental/Mutable/Source/CustomizableObject/Private/MuCO/CustomizableObjectInstance.cpp b/Engine/Plugins/Experimental/Mutable/Source/CustomizableObject/Private/MuCO/CustomizableObjectInstance.cpp
--- a/Engine/Plugins/Experimental/Mutable/Source/CustomizableObject/Private/MuCO/CustomizableObjectInstance.cpp
+++ b/Engine/Plugins/Experimental/Mutable/Source/CustomizableObject/Private/MuCO/CustomizableObjectInstance.cpp
@@ -131,7 +131,7 @@ void UCustomizableInstancePrivateData::InitLastUpdateData(const TSharedPtr<FMuta
 	const int32 NumSurfaces = OperationData->InstanceUpdateData.Surfaces.Num();
 	LastUpdateData.SurfaceIds.Init(0, NumSurfaces);
 	for (int32 SurfaceIndex = 0; SurfaceIndex < NumSurfaces; ++SurfaceIndex)
-{
+	{
 		LastUpdateData.SurfaceIds[SurfaceIndex] = OperationData->InstanceUpdateData.Surfaces[SurfaceIndex].SurfaceId;
 	}
 }
@@ -139,26 +139,47 @@ void UCustomizableInstancePrivateData::InitLastUpdateData(const TSharedPtr<FMuta
 
 void UCustomizableInstancePrivateData::InvalidateGeneratedData()
 {
-	for (FCustomizableInstanceComponentData& ComponentData : ComponentsData)
+	// Init Component Data
+	FCustomizableInstanceComponentData TemplateComponentData;
+	TemplateComponentData.LastMeshIdPerLOD.Init(MAX_uint64, MAX_MESH_LOD_COUNT);
+	ComponentsData.Init(TemplateComponentData, ComponentsData.Num());
+
+	GeneratedMaterials.Empty();
+	LastUpdateData.Clear();
+	ClearCOInstanceFlags(Generated);
+}
+
+
+void UCustomizableInstancePrivateData::InitCustomizableObjectData(const UCustomizableObject* InCustomizableObject)
+{
+	InvalidateGeneratedData();
+
+	if (!InCustomizableObject || !InCustomizableObject->IsCompiled())
 	{
-		ComponentData.LastMeshIdPerLOD.Init(MAX_uint64, MAX_MESH_LOD_COUNT);
+		return;
 	}
 
-	GeneratedMaterials.Empty();
+	// Init LOD Data
+	NumLODsAvailable = InCustomizableObject->GetNumLODs();
+	FirstLODAvailable = InCustomizableObject->LODSettings.FirstLODAvailable;
+	NumMaxLODsToStream = InCustomizableObject->LODSettings.bLODStreamingEnabled ? InCustomizableObject->LODSettings.NumLODsToStream : 0;
 
-	LastUpdateData.Clear();
+	// Init Component Data
+	FCustomizableInstanceComponentData TemplateComponentData;
+	TemplateComponentData.LastMeshIdPerLOD.Init(MAX_uint64, MAX_MESH_LOD_COUNT);
+	ComponentsData.Init(TemplateComponentData, InCustomizableObject->GetComponentCount());
 }
 
 
 FCustomizableInstanceComponentData* UCustomizableInstancePrivateData::GetComponentData(int32 ComponentIndex)
 {
-	return ComponentsData.FindByPredicate([&ComponentIndex](FCustomizableInstanceComponentData& C) { return C.ComponentIndex == ComponentIndex; });
+	return ComponentsData.IsValidIndex(ComponentIndex) ? &ComponentsData[ComponentIndex] : nullptr;
 }
 
 
 const FCustomizableInstanceComponentData* UCustomizableInstancePrivateData::GetComponentData(int32 ComponentIndex) const
 {
-	return ComponentsData.FindByPredicate([&ComponentIndex](FCustomizableInstanceComponentData& C) { return C.ComponentIndex == ComponentIndex; });
+	return ComponentsData.IsValidIndex(ComponentIndex) ? &ComponentsData[ComponentIndex] : nullptr;
 }
 
 
@@ -184,9 +205,20 @@ const FCustomizableObjectInstanceDescriptor& UCustomizableObjectInstance::GetDes
 
 void UCustomizableObjectInstance::SetDescriptor(const FCustomizableObjectInstanceDescriptor& InDescriptor)
 {
-	const bool bInvalidatePreviousData = GetCustomizableObject() != InDescriptor.CustomizableObject;
+	UCustomizableObject* InCustomizableObject = InDescriptor.GetCustomizableObject();
+	const bool bCustomizableObjectChanged = Descriptor.GetCustomizableObject() != InCustomizableObject;
+
+#if WITH_EDITOR
+	// Bind a lambda to the PostCompileDelegate and unbind from the previous object if any.
+	BindPostCompileDelegate(InCustomizableObject);
+#endif
+
 	Descriptor = InDescriptor;
-	PrivateData->ReloadParameters(this, bInvalidatePreviousData);
+
+	if (bCustomizableObjectChanged)
+	{
+		PrivateData->InitCustomizableObjectData(InCustomizableObject);
+	}
 }
 
 
@@ -195,70 +227,68 @@ void UCustomizableInstancePrivateData::PrepareForUpdate(const TSharedPtr<FMutabl
 	// Prepare LastUpdateData to allow reuse of meshes and surfaces
 	InitLastUpdateData(OperationData);
 
-	TArray<FInstanceUpdateData::FLOD>& LODs = OperationData->InstanceUpdateData.LODs;
-	TArray<FInstanceUpdateData::FComponent>& Components = OperationData->InstanceUpdateData.Components;
-
-	// Clear, reinit or create ComponentData for each component 
-	TSet<uint16> ComponentIds;
-	for (const FInstanceUpdateData::FLOD& LOD : LODs)
+	// Clear the ComponentData from previous updates
+	for (FCustomizableInstanceComponentData& ComponentData : ComponentsData)
 	{
-		check(LOD.FirstComponent + LOD.ComponentCount <= Components.Num());
+		ComponentData.AnimSlotToBP.Empty();
+		ComponentData.AssetUserDataArray.Empty();
+		ComponentData.AssetUserDataToStream.Empty();
+		ComponentData.Skeletons.Skeleton = nullptr;
+		ComponentData.Skeletons.SkeletonIds.Empty();
+		ComponentData.Skeletons.SkeletonsToMerge.Empty();
+		ComponentData.PhysicsAssets.PhysicsAssetToLoad.Empty();
+		ComponentData.PhysicsAssets.PhysicsAssetsToMerge.Empty();
+		ComponentData.ClothingPhysicsAssetsToStream.Empty();
 
-		for (uint16 ComponentIndex = 0; ComponentIndex < LOD.ComponentCount; ++ComponentIndex)
-		{
-			const FInstanceUpdateData::FComponent& Component = Components[LOD.FirstComponent + ComponentIndex];
-			ComponentIds.Add(Component.Id);
+#if WITH_EDITORONLY_DATA
+		ComponentData.MeshPartPaths.Empty();
+#endif
+	}
+}
 
-			if (FCustomizableInstanceComponentData* ComponentData = GetComponentData(Component.Id))
-			{
-				if (!Component.bGenerated || Component.bReuseMesh)
-				{
-					continue;
-				}
 
-				ComponentData->AnimSlotToBP.Empty();
-				ComponentData->AssetUserDataArray.Empty();
-				ComponentData->AssetUserDataToStream.Empty();
+#if WITH_EDITOR
 
-#if WITH_EDITORONLY_DATA
-				ComponentData->MeshPartPaths.Empty();
-#endif
-				ComponentData->Skeletons.Skeleton = nullptr;
-				ComponentData->Skeletons.SkeletonIds.Empty();
-				ComponentData->Skeletons.SkeletonsToMerge.Empty();
-				ComponentData->PhysicsAssets.PhysicsAssetToLoad.Empty();
-				ComponentData->PhysicsAssets.PhysicsAssetsToMerge.Empty();
-				ComponentData->ClothingPhysicsAssetsToStream.Empty();
 
-				continue;
-			}
+void UCustomizableInstancePrivateData::PostDuplicate(bool bDuplicateForPIE)
+{
+	Super::PostDuplicate(bDuplicateForPIE);
 
-			FCustomizableInstanceComponentData& NewComponentData = ComponentsData.AddDefaulted_GetRef();
-			NewComponentData.ComponentIndex = Component.Id;
-			NewComponentData.LastMeshIdPerLOD.Init(MAX_uint64, MAX_MESH_LOD_COUNT);
-		}
-	}
+	// Invalidate all generated data to avoid modifying resources shared between CO instances.
+	InvalidateGeneratedData();
 
-	// Check if a component have been removed.
-	if (ComponentsData.Num() != ComponentIds.Num())
-	{
-		for (uint16 ComponentIndex = 0; ComponentIndex < ComponentsData.Num();)
-		{
-			FCustomizableInstanceComponentData& ComponentData = ComponentsData[ComponentIndex];
+	// Empty after duplicating or ReleasingMutableResources may free textures used by the other CO instance.
+	GeneratedTextures.Empty();
+}
 
-			if (!ComponentIds.Find(ComponentData.ComponentIndex))
-			{
-				ComponentsData.RemoveSingleSwap(ComponentData);
-				continue;
-			}
 
-			ComponentIndex++;
-		}
+void UCustomizableInstancePrivateData::OnPostCompile()
+{
+	UCustomizableObjectInstance* Instance = Cast<UCustomizableObjectInstance>(GetOuter());
+	if (!Instance)
+	{
+		return;
 	}
+
+	Instance->GetDescriptor().ReloadParameters();
+	InitCustomizableObjectData(Instance->GetCustomizableObject());
 }
 
 
-#if WITH_EDITOR
+void UCustomizableObjectInstance::BindPostCompileDelegate(UCustomizableObject* InCustomizableObject)
+{
+	// Unbind callback from the previous CO
+	if (UCustomizableObject* CurrentObject = GetCustomizableObject())
+	{
+		CurrentObject->PostCompileDelegate.RemoveAll(GetPrivate());
+	}
+
+	// Bind callback to the new CO
+	if (InCustomizableObject)
+	{
+		InCustomizableObject->PostCompileDelegate.AddUObject(GetPrivate(), &UCustomizableInstancePrivateData::OnPostCompile);
+	}
+}
 
 
 void UCustomizableObjectInstance::PostEditChangeProperty(FPropertyChangedEvent& PropertyChangedEvent)
@@ -310,6 +340,11 @@ bool UCustomizableObjectInstance::IsEditorOnly() const
 
 void UCustomizableObjectInstance::BeginDestroy()
 {
+#if WITH_EDITOR
+	// Unbind PostCompileDelegate
+	BindPostCompileDelegate(nullptr);
+#endif
+
 	BeginDestroyDelegate.Broadcast(this);
 	BeginDestroyNativeDelegate.Broadcast(this);
 
@@ -472,7 +507,12 @@ void UCustomizableObjectInstance::PostLoad()
 
 	Super::PostLoad();
 
+#if WITH_EDITOR
+	BindPostCompileDelegate(GetCustomizableObject());
+#endif
+
 	Descriptor.ReloadParameters();
+	PrivateData->InitCustomizableObjectData(GetCustomizableObject());
 }
 
 
@@ -536,59 +576,15 @@ void UCustomizableInstancePrivateData::ReleaseMutableTexture(const FMutableImage
 }
 
 
-void UCustomizableInstancePrivateData::InstanceUpdateFlags(const UCustomizableObjectInstance& Public)
-{
-	const UCustomizableObject* CustomizableObject = Public.GetCustomizableObject();
-	
-	FirstLODAvailable = CustomizableObject->LODSettings.FirstLODAvailable;
-	NumLODsAvailable = CustomizableObject->GetNumLODs();
-
-	if (CustomizableObject->LODSettings.bLODStreamingEnabled)
-	{
-		NumMaxLODsToStream = CustomizableObject->LODSettings.NumLODsToStream;
-	}
-	else
-	{
-		NumMaxLODsToStream = 0;
-	}
-}
-
-
-void UCustomizableInstancePrivateData::ReloadParameters(UCustomizableObjectInstance* Public, bool bInvalidatePreviousData)
+void UCustomizableObjectInstance::SetObject(UCustomizableObject* InObject)
 {
-	const UCustomizableObject* CustomizableObject = Public->GetCustomizableObject();
-	if (!CustomizableObject)
-	{
-		// May happen when deleting assets in editor, while they are still open.
-		return;
-	}
-
-	if (!CustomizableObject->IsCompiled())
-	{	
-		return;
-	}
-
-	InstanceUpdateFlags(*Public); // TODO Move somewhere else.
-
-	if (bInvalidatePreviousData)
-	{
-		InvalidateGeneratedData();
-
 #if WITH_EDITOR
-		// Clear the Generated flag to trigger a new update after changing or compiling the CO,
-		ClearCOInstanceFlags(Generated);
+	// Bind a lambda to the PostCompileDelegate and unbind from the previous object if any.
+	BindPostCompileDelegate(InObject);
 #endif
-	}
-
-	Public->GetDescriptor().ReloadParameters();
-}
-
 
-void UCustomizableObjectInstance::SetObject(UCustomizableObject* InObject)
-{
 	Descriptor.SetCustomizableObject(*InObject);
-	PrivateData->ReloadParameters(this, true);
-	//SetRequestedLODs(Descriptor.MinLOD, Descriptor.MaxLOD, Descriptor.RequestedLODLevels);
+	PrivateData->InitCustomizableObjectData(InObject);
 }
 
 
@@ -1788,10 +1784,10 @@ bool UCustomizableInstancePrivateData::UpdateSkeletalMesh_PostBeginUpdate0(UCust
 
 		Public->SkeletalMeshes.Reset(); // What about all the references
 
-		SetCOInstanceFlags(Generated);
-
 		InvalidateGeneratedData();
 
+		SetCOInstanceFlags(Generated);
+
 		return false;
 	}
 
@@ -3279,11 +3275,19 @@ void UCustomizableInstancePrivateData::InitSkeletalMeshData(const TSharedPtr<FMu
 	SkeletalMesh->SetEnablePerPolyCollision(RefSkeletalMeshData->Settings.bEnablePerPolyCollision);
 
 	// Asset User Data
-	for (const FMutableRefAssetUserData& MutAssetUserData : RefSkeletalMeshData->AssetUserData)
 	{
-		if (MutAssetUserData.AssetUserData)
+		for (const FMutableRefAssetUserData& MutAssetUserData : RefSkeletalMeshData->AssetUserData)
+		{
+			if (MutAssetUserData.AssetUserData)
+			{
+				SkeletalMesh->AddAssetUserData(MutAssetUserData.AssetUserData);
+			}
+		}
+
+		const FCustomizableInstanceComponentData& ComponentData = ComponentsData[ComponentIndex];
+		for (TObjectPtr<UAssetUserData> AssetUserData : ComponentData.AssetUserDataArray)
 		{
-			SkeletalMesh->AddAssetUserData(MutAssetUserData.AssetUserData);
+			SkeletalMesh->AddAssetUserData(AssetUserData);
 		}
 	}
 
@@ -4996,31 +5000,32 @@ FGraphEventRef UCustomizableInstancePrivateData::LoadAdditionalAssetsAsync(const
 	// Clear invalid skeletons from the MergedSkeletons cache
 	CustomizableObject->UnCacheInvalidSkeletons();
 
-	const int32 ComponentCount = OperationData->InstanceUpdateData.Skeletons.Num();
-
 	// Load Skeletons required by the SubMeshes of the newly generated Mesh, will be merged later
-	for (FCustomizableInstanceComponentData& ComponentData : ComponentsData)
+	for (FInstanceUpdateData::FSkeletonData& SkeletonData : OperationData->InstanceUpdateData.Skeletons)
 	{
-		const uint16 ComponentIndex = ComponentData.ComponentIndex;
-
-		FInstanceUpdateData::FSkeletonData* SkeletonData = OperationData->InstanceUpdateData.Skeletons.FindByPredicate(
-			[&ComponentIndex](FInstanceUpdateData::FSkeletonData& S) { return S.ComponentIndex == ComponentIndex; });
-		check(SkeletonData);
+		FCustomizableInstanceComponentData* ComponentData = GetComponentData(SkeletonData.ComponentIndex);
+		if (!ComponentData)
+		{
+			check(false);
+			continue;
+		}
 
-		ComponentData.Skeletons.Skeleton = CustomizableObject->GetCachedMergedSkeleton(ComponentIndex, SkeletonData->SkeletonIds);
-		if (ComponentData.Skeletons.Skeleton)
+		// Reuse merged Skeleton if cached
+		ComponentData->Skeletons.Skeleton = CustomizableObject->GetCachedMergedSkeleton(SkeletonData.ComponentIndex, SkeletonData.SkeletonIds);
+		if (ComponentData->Skeletons.Skeleton)
 		{
-			ComponentData.Skeletons.SkeletonIds.Empty();
-			ComponentData.Skeletons.SkeletonsToMerge.Empty();
+			ComponentData->Skeletons.SkeletonIds.Empty();
+			ComponentData->Skeletons.SkeletonsToMerge.Empty();
 			continue;
 		}
 
-		for (const uint32 SkeletonId : SkeletonData->SkeletonIds)
+		// Add Skeletons to merge
+		for (const uint32 SkeletonId : SkeletonData.SkeletonIds)
 		{
 			TSoftObjectPtr<USkeleton> AssetPtr;
 			if (SkeletonId == 0)
 			{
-				FMutableRefSkeletalMeshData* RefSkeletalMeshData = CustomizableObject->GetRefSkeletalMeshData(ComponentIndex);
+				FMutableRefSkeletalMeshData* RefSkeletalMeshData = CustomizableObject->GetRefSkeletalMeshData(SkeletonData.ComponentIndex);
 				check(RefSkeletalMeshData);
 
 				AssetPtr = RefSkeletalMeshData->Skeleton;
@@ -5036,11 +5041,11 @@ FGraphEventRef UCustomizableInstancePrivateData::LoadAdditionalAssetsAsync(const
 			}
 
 			// Add referenced skeletons to the assets to stream
-			ComponentData.Skeletons.SkeletonIds.Add(SkeletonId);
+			ComponentData->Skeletons.SkeletonIds.Add(SkeletonId);
 
 			if (USkeleton* Skeleton = AssetPtr.Get())
 			{
-				ComponentData.Skeletons.SkeletonsToMerge.Add(Skeleton);
+				ComponentData->Skeletons.SkeletonsToMerge.Add(Skeleton);
 			}
 			else
 			{
diff --git a/Engine/Plugins/Experimental/Mutable/Source/CustomizableObject/Private/MuCO/CustomizableObjectSystem.cpp b/Engine/Plugins/Experimental/Mutable/Source/CustomizableObject/Private/MuCO/CustomizableObjectSystem.cpp
--- a/Engine/Plugins/Experimental/Mutable/Source/CustomizableObject/Private/MuCO/CustomizableObjectSystem.cpp
+++ b/Engine/Plugins/Experimental/Mutable/Source/CustomizableObject/Private/MuCO/CustomizableObjectSystem.cpp
@@ -2331,22 +2331,6 @@ namespace impl
 
 		FCustomizableObjectSystemPrivate* CustomizableObjectSystemPrivateData = System->GetPrivateChecked();
 
-		for (int32 ComponentIndex = 0; ComponentIndex < CustomizableObjectInstance->GetPrivate()->ComponentsData.Num(); ++ComponentIndex)
-		{
-			// TODO PRP: Fix properly by moving the whole AssetUserData block to UCustomizableInstancePrivateData::InitSkeletalMeshData
-			if (!CustomizableObjectInstance->SkeletalMeshes.IsValidIndex(ComponentIndex))
-			{
-				break;
-			}
-
-			FCustomizableInstanceComponentData& ComponentData = CustomizableObjectInstance->GetPrivate()->ComponentsData[ComponentIndex];
-			
-			for (TObjectPtr<UAssetUserData> AssetUserData : ComponentData.AssetUserDataArray)
-			{
-				CustomizableObjectInstance->SkeletalMeshes[ComponentIndex]->AddAssetUserData(AssetUserData);
-			}
-		}
-
 		// Next Task: Release Mutable. We need this regardless if we cancel or not
 		//-------------------------------------------------------------		
 		{
@@ -2590,9 +2574,7 @@ namespace impl
 			FinishUpdateGlobal(nullptr, EUpdateResult::Error, &Operation->UpdateCallback);
 			return;
 		}
-		
-		CandidateInstancePrivateData->InstanceUpdateFlags(*CandidateInstance);
-		
+
 		if (CandidateInstancePrivateData->HasCOInstanceFlags(PendingLODsUpdate))
 		{
 			CandidateInstancePrivateData->ClearCOInstanceFlags(PendingLODsUpdate);
