[Mutable] Remove IsUVNormalized since now it's done when processing surfaces in the code generator.

#rnx

diff --git a/Engine/Plugins/Experimental/Mutable/Source/CustomizableObjectEditor/Private/MuCOE/GenerateMutableSource/GenerateMutableSourceSurface.cpp b/Engine/Plugins/Experimental/Mutable/Source/CustomizableObjectEditor/Private/MuCOE/GenerateMutableSource/GenerateMutableSourceSurface.cpp
--- a/Engine/Plugins/Experimental/Mutable/Source/CustomizableObjectEditor/Private/MuCOE/GenerateMutableSource/GenerateMutableSourceSurface.cpp
+++ b/Engine/Plugins/Experimental/Mutable/Source/CustomizableObjectEditor/Private/MuCOE/GenerateMutableSource/GenerateMutableSourceSurface.cpp
@@ -53,44 +53,6 @@
 
 #define LOCTEXT_NAMESPACE "CustomizableObjectEditor"
 
-/**
- * Show a compilation warning when any of the meshes behind the pin InMeshPin has a UV that is not normalized.
- * 
- * @param InMeshPin Meshes connected to this pin (directly or indirectly through switch/variation nodes).
- * @param OperationNode Node which is performing the block operation. 
- */
-void LayoutOperationUVNormalizedWarning(FMutableGraphGenerationContext& GenerationContext, const UEdGraphPin* InMeshPin, const UCustomizableObjectNode* OperationNode, const int32 UVIndex)
-{
-	if (const UEdGraphPin* ConnectedPin = FollowInputPin(*InMeshPin))
-	{
-		FPinDataValue* PinData = GenerationContext.PinData.Find(ConnectedPin);
-		if (PinData)
-		{
-			for (const FMeshData& MeshData : PinData->MeshesData)
-			{
-				if (const UCustomizableObjectNodeMesh* Node = Cast<UCustomizableObjectNodeMesh>(MeshData.Node))
-				{
-					bool bNormalized = true;
-					if (const USkeletalMesh* SkeletalMesh = Cast<USkeletalMesh>(MeshData.Mesh))
-					{
-						bNormalized = IsUVNormalized(*SkeletalMesh, MeshData.LOD, MeshData.MaterialIndex, UVIndex);
-					}
-					else if (const UStaticMesh* StaticMesh = Cast<UStaticMesh>(MeshData.Mesh))
-					{
-						bNormalized = IsUVNormalized(*StaticMesh, MeshData.LOD, MeshData.MaterialIndex, UVIndex);
-					}
-
-					if (!bNormalized)
-					{
-						FText Text = FText::Format(LOCTEXT("UVNotNormalized", "UV from mesh {0} not normalized. Required to perform texture layout operations."), FText::FromString(*MeshData.Mesh->GetName()));
-						GenerationContext.Compiler->CompilerLog(Text, OperationNode, EMessageSeverity::Type::Info);
-					}
-				}
-			}
-		}
-	}
-}
-
 
 void SetSurfaceFormat( FMutableGraphGenerationContext& GenerationContext,
 					   mu::FMeshBufferSet& OutVertexBufferFormat, mu::FMeshBufferSet& OutIndexBufferFormat, const FMutableGraphMeshGenerationData& MeshData, 
@@ -1181,10 +1143,6 @@ mu::NodeSurfacePtr GenerateMutableSourceSurface(const UEdGraphPin * Pin, FMutabl
 					SurfNode->SetImage(ImageIndex, ImageNode);
 				}
 
-				const int32 UVIndex = ParentMaterialNode->GetImageUVLayout(ImageIndex);
-				LayoutOperationUVNormalizedWarning(GenerationContext, ParentMaterialNode->GetMeshPin(), Node, UVIndex);
-				LayoutOperationUVNormalizedWarning(GenerationContext, TypedNodeExt->AddMeshPin(), Node, UVIndex);
-
 				// Validate if the ParentMaterialNode can be shared between LODs.
 				UpdateSharedSurfaceId(GenerationContext, ParentMaterialNode, nullptr);
 			}
@@ -1314,8 +1272,6 @@ mu::NodeSurfacePtr GenerateMutableSourceSurface(const UEdGraphPin * Pin, FMutabl
 				MeshPatch->SetMessageContext(Node);
 			}
 
-			LayoutOperationUVNormalizedWarning(GenerationContext, ParentMaterialNode->GetMeshPin(), Node, TypedNodeRemBlocks->ParentLayoutIndex);
-
 			// Validate if the ParentMaterialNode can be shared between LODs.
 			UpdateSharedSurfaceId(GenerationContext, ParentMaterialNode, nullptr);
 		}
@@ -1424,8 +1380,6 @@ mu::NodeSurfacePtr GenerateMutableSourceSurface(const UEdGraphPin * Pin, FMutabl
 
 					SurfNode->SetPatch(ImageIndex, ImagePatchNode);
 				}
-				
-				LayoutOperationUVNormalizedWarning(GenerationContext, ParentMaterialNode->GetMeshPin(), Node, ParentMaterialNode->GetImageUVLayout(ImageIndex));
 
 				// Validate if the ParentMaterialNode can be shared between LODs.
 				UpdateSharedSurfaceId(GenerationContext, ParentMaterialNode, nullptr);
diff --git a/Engine/Plugins/Experimental/Mutable/Source/CustomizableObjectEditor/Private/MuCOE/MutableUtils.cpp b/Engine/Plugins/Experimental/Mutable/Source/CustomizableObjectEditor/Private/MuCOE/MutableUtils.cpp
--- a/Engine/Plugins/Experimental/Mutable/Source/CustomizableObjectEditor/Private/MuCOE/MutableUtils.cpp
+++ b/Engine/Plugins/Experimental/Mutable/Source/CustomizableObjectEditor/Private/MuCOE/MutableUtils.cpp
@@ -39,6 +39,7 @@ TArray<FVector2f> GetUV(const USkeletalMesh& SkeletalMesh, const int32 LODIndex,
 	return Result;
 }
 
+
 TArray<FVector2f> GetUV(const UStaticMesh& StaticMesh, const int32 LODIndex, const int32 SectionIndex, const int32 UVIndex)
 {
 	TArray<FVector2f> Result;
@@ -67,58 +68,13 @@ TArray<FVector2f> GetUV(const UStaticMesh& StaticMesh, const int32 LODIndex, con
 	return Result;
 }
 
+
 bool HasNormalizedBounds(const FVector2f& Point)
 {
 	return Point.X >= 0.0f && Point.X <= 1.0 && Point.Y >= 0.0f && Point.Y <= 1.0;
 }
 
 
-bool IsUVNormalized(const USkeletalMesh& SkeletalMesh, const int32 LODIndex, const int32 SectionIndex, const int32 UVIndex)
-{
-	const FSkeletalMeshModel* ImportedModel = SkeletalMesh.GetImportedModel();
-	const FSkeletalMeshLODModel& LOD = ImportedModel->LODModels[LODIndex];
-	const FSkelMeshSection& Section = LOD.Sections[SectionIndex];
-	
-	TArray<FSoftSkinVertex> Vertices;
-	LOD.GetVertices(Vertices);
-
-	TArray<uint32> Indices;
-	SkeletalMesh.GetResourceForRendering()->LODRenderData[LODIndex].MultiSizeIndexContainer.GetIndexBuffer(Indices);
-	
-	for (int32 VertexIndex = Section.BaseVertexIndex; VertexIndex < Section.NumVertices; ++VertexIndex)
-	{
-		if (!HasNormalizedBounds(Vertices[Indices[VertexIndex]].UVs[UVIndex]))
-		{
-			return false;
-		}
-	}
-
-	return true;
-}
-
-
-bool IsUVNormalized(const UStaticMesh& StaticMesh, const int32 LODIndex, const int32 SectionIndex, const int32 UVIndex)
-{
-	const FStaticMeshRenderData* RenderData = StaticMesh.GetRenderData();
-	const FStaticMeshLODResources& LOD = RenderData->LODResources[LODIndex];
-	const FStaticMeshSection& Section = LOD.Sections[SectionIndex];
-	
-	const FStaticMeshVertexBuffer& VertexBuffer = LOD.VertexBuffers.StaticMeshVertexBuffer;
-	
-	FIndexArrayView Indices = LOD.IndexBuffer.GetArrayView();
-
-	for (uint32 VertexIndex = Section.MinVertexIndex; VertexIndex < Section.MaxVertexIndex; ++VertexIndex)
-	{
-		if (!HasNormalizedBounds(VertexBuffer.GetVertexUV(Indices[VertexIndex + 0], UVIndex)))
-		{
-			return false;
-		}
-	}
-	
-	return true;
-}
-
-
 bool IsMeshClosed(const USkeletalMesh* Mesh, int LOD, int MaterialIndex)
 {
 	check(Mesh->GetImportedModel() &&
diff --git a/Engine/Plugins/Experimental/Mutable/Source/CustomizableObjectEditor/Public/MuCOE/MutableUtils.h b/Engine/Plugins/Experimental/Mutable/Source/CustomizableObjectEditor/Public/MuCOE/MutableUtils.h
--- a/Engine/Plugins/Experimental/Mutable/Source/CustomizableObjectEditor/Public/MuCOE/MutableUtils.h
+++ b/Engine/Plugins/Experimental/Mutable/Source/CustomizableObjectEditor/Public/MuCOE/MutableUtils.h
@@ -20,14 +20,6 @@ TArray<FVector2f> GetUV(const UStaticMesh& StaticMesh, const int32 LODIndex, con
 bool HasNormalizedBounds(const FVector2f& Point);
 
 
-/** Returns if the given UV is normalized. */
-bool IsUVNormalized(const USkeletalMesh& SkeletalMesh, const int32 LODIndex, const int32 SectionIndex, const int32 UVIndex);
-
-
-/** Returns if the given UV is normalized. */
-bool IsUVNormalized(const UStaticMesh& SkeletalMesh, const int32 LODIndex, const int32 SectionIndex, const int32 UVIndex);
-
-
 /** Check if a given LOD and material in a skeletal mesh is closed. */
 bool IsMeshClosed(const USkeletalMesh* Mesh, int LOD, int MaterialIndex);
 
