[Mutable] Invalidate the compiled data of the Customizable Object after modifying it.

#rnx

diff --git a/Engine/Plugins/Experimental/Mutable/Source/CustomizableObject/Private/MuCO/CustomizableObject.cpp b/Engine/Plugins/Experimental/Mutable/Source/CustomizableObject/Private/MuCO/CustomizableObject.cpp
--- a/Engine/Plugins/Experimental/Mutable/Source/CustomizableObject/Private/MuCO/CustomizableObject.cpp
+++ b/Engine/Plugins/Experimental/Mutable/Source/CustomizableObject/Private/MuCO/CustomizableObject.cpp
@@ -159,7 +159,13 @@ void UCustomizableObject::PostLoad()
 	// Register to dirty delegate so we update derived data version ID each time that the package is marked as dirty.
 	if (UPackage* Package = GetOutermost())
 	{
-		Package->PackageMarkedDirtyEvent.AddWeakLambda(this, [this](UPackage* Pkg, bool bWasDirty){UpdateVersionId();});
+		Package->PackageMarkedDirtyEvent.AddWeakLambda(this, [this](UPackage* Pkg, bool bWasDirty)
+			{
+				if (GetPackage() == Pkg)
+				{
+					UpdateVersionId();
+				}
+			});
 	}
 
 	const int32 CustomizableObjectCustomVersion = GetLinkerCustomVersion(FCustomizableObjectCustomVersion::GUID);
