[Mutable] Force compilation when any of the participating objects change.

We consider a participating object any external object which is used directly or indirectly in a compilation.

This quick implementation adds objects manually on the graph traversal. This can be improved by adding a virtual function, which will be required for non MuCOE nodes.

#rnx

diff --git a/Engine/Plugins/Experimental/Mutable/Source/CustomizableObject/Classes/MuCO/CustomizableObject.h b/Engine/Plugins/Experimental/Mutable/Source/CustomizableObject/Classes/MuCO/CustomizableObject.h
--- a/Engine/Plugins/Experimental/Mutable/Source/CustomizableObject/Classes/MuCO/CustomizableObject.h
+++ b/Engine/Plugins/Experimental/Mutable/Source/CustomizableObject/Classes/MuCO/CustomizableObject.h
@@ -1110,6 +1110,8 @@ private:
 UCLASS( BlueprintType, config=Engine )
 class CUSTOMIZABLEOBJECT_API UCustomizableObject : public UObject
 {
+	friend FCustomizableObjectPrivateData;
+	
 public:
 	GENERATED_BODY()
 
@@ -1472,7 +1474,7 @@ private:
 	// This is a manual version number for the binary blobs in this asset.
 	// Increasing it invalidates all the previously compiled models.
 	// Warning: If while merging code both versions have changed, take the highest+1.
-	static const int32 CurrentSupportedVersion = 406;
+	static const int32 CurrentSupportedVersion = 407;
 
 public:
 
@@ -1728,6 +1730,13 @@ private:
 	UPROPERTY()
 	TObjectPtr<UCustomizableObjectBulk> BulkData;
 
+#if WITH_EDITORONLY_DATA
+	/** List of external objects that if changed, a compilation is required.
+	 * The FGuid is the the UPackage::Guid, which is regenerated each time the packages is saved. */
+	UPROPERTY()
+	TMap<TObjectPtr<const UObject>, FGuid> ParticipatingObjects;
+#endif
+	
 	TSharedPtr<FCustomizableObjectPrivateData> PrivateData;
 };
 
diff --git a/Engine/Plugins/Experimental/Mutable/Source/CustomizableObject/Private/MuCO/CustomizableObject.cpp b/Engine/Plugins/Experimental/Mutable/Source/CustomizableObject/Private/MuCO/CustomizableObject.cpp
--- a/Engine/Plugins/Experimental/Mutable/Source/CustomizableObject/Private/MuCO/CustomizableObject.cpp
+++ b/Engine/Plugins/Experimental/Mutable/Source/CustomizableObject/Private/MuCO/CustomizableObject.cpp
@@ -179,6 +179,23 @@ void UCustomizableObject::PostLoad()
 			;
 	}
 #endif
+
+#if WITH_EDITORONLY_DATA
+	for (TTuple<TObjectPtr<const UObject>, FGuid>& ParticipatingObject : ParticipatingObjects)
+	{
+		PRAGMA_DISABLE_DEPRECATION_WARNINGS
+		const FGuid PackageGuid = ParticipatingObject.Key.GetPackage()->GetGuid();
+		PRAGMA_ENABLE_DEPRECATION_WARNINGS
+		
+		if (PackageGuid != ParticipatingObject.Value)
+		{
+			SetModel(nullptr);
+			
+			UE_LOG(LogMutable, Display, TEXT("Forcing recompilation due to changes in %s."), *ParticipatingObject.Key->GetFullName());
+			break;
+		}
+	}
+#endif
 }
 
 
@@ -1852,6 +1869,14 @@ TSharedPtr<const mu::Model, ESPMode::ThreadSafe> FCustomizableObjectPrivateData:
 }
 
 
+#if WITH_EDITORONLY_DATA
+TMap<TObjectPtr<const UObject>, FGuid>& FCustomizableObjectPrivateData::GetParticipatingObjects(UCustomizableObject& Public)
+{
+	return Public.ParticipatingObjects;
+}
+#endif
+
+
 //-------------------------------------------------------------------------------------------------
 //-------------------------------------------------------------------------------------------------
 //-------------------------------------------------------------------------------------------------
diff --git a/Engine/Plugins/Experimental/Mutable/Source/CustomizableObject/Private/MuCO/CustomizableObjectPrivate.h b/Engine/Plugins/Experimental/Mutable/Source/CustomizableObject/Private/MuCO/CustomizableObjectPrivate.h
--- a/Engine/Plugins/Experimental/Mutable/Source/CustomizableObject/Private/MuCO/CustomizableObjectPrivate.h
+++ b/Engine/Plugins/Experimental/Mutable/Source/CustomizableObject/Private/MuCO/CustomizableObjectPrivate.h
@@ -9,21 +9,23 @@
 #endif
 
 namespace mu { class Model; }
-
+class UCustomizableObject;
 
 
 class FCustomizableObjectPrivateData
 {
-private:
-
 	TSharedPtr<mu::Model, ESPMode::ThreadSafe> MutableModel;
 
 public:
-
 	void SetModel(const TSharedPtr<mu::Model, ESPMode::ThreadSafe>& Model, const FGuid Identifier);
 	const TSharedPtr<mu::Model, ESPMode::ThreadSafe>& GetModel();
 	TSharedPtr<const mu::Model, ESPMode::ThreadSafe> GetModel() const;
 
+#if WITH_EDITORONLY_DATA
+	/** See UCustomizableObject::ParticipatingObjects. */
+	CUSTOMIZABLEOBJECT_API TMap<TObjectPtr<const UObject>, FGuid>& GetParticipatingObjects(UCustomizableObject& Public);
+#endif
+	
 	// See UCustomizableObjectSystem::LockObject. Must only be modified from the game thread
 	bool bLocked = false;
 
@@ -33,6 +35,5 @@ public:
 	bool bModelCompiledForCook = false;
 	TArray<FString> CachedPlatformNames;
 #endif
-
 };
 
diff --git a/Engine/Plugins/Experimental/Mutable/Source/CustomizableObjectEditor/Private/MuCOE/CustomizableObjectCompiler.cpp b/Engine/Plugins/Experimental/Mutable/Source/CustomizableObjectEditor/Private/MuCOE/CustomizableObjectCompiler.cpp
--- a/Engine/Plugins/Experimental/Mutable/Source/CustomizableObjectEditor/Private/MuCOE/CustomizableObjectCompiler.cpp
+++ b/Engine/Plugins/Experimental/Mutable/Source/CustomizableObjectEditor/Private/MuCOE/CustomizableObjectCompiler.cpp
@@ -23,6 +23,7 @@
 #include "Widgets/Notifications/SNotificationList.h"
 #include "Interfaces/ITargetPlatform.h"
 #include "Misc/App.h"
+#include "MuCO/CustomizableObjectPrivate.h"
 
 class UTexture2D;
 
@@ -335,7 +336,7 @@ void FCustomizableObjectCompiler::UpdateArrayGCProtect()
 }
 
 
-void FCustomizableObjectCompiler::ProcessChildObjectsRecursively(UCustomizableObject* Object, FAssetRegistryModule& AssetRegistryModule, FMutableGraphGenerationContext &GenerationContext)
+void FCustomizableObjectCompiler::ProcessChildObjectsRecursively(UCustomizableObject* Object, FMutableGraphGenerationContext &GenerationContext)
 {
 	TArray<FName> ArrayReferenceNames;
 	AddCachedReferencers(*Object->GetOuter()->GetPathName(), ArrayReferenceNames);
@@ -343,60 +344,62 @@ void FCustomizableObjectCompiler::ProcessChildObjectsRecursively(UCustomizableOb
 
 	bool bMultipleBaseObjectsFound = false;
 
-	FAssetData* AssetData;
-	UCustomizableObject* ChildObject;
-
 	for (const FName& ReferenceName : ArrayReferenceNames)
 	{
-		AssetData = GetCachedAssetData(ReferenceName.ToString());
+		if (ArrayAlreadyProcessedChild.Contains(ReferenceName))
+		{
+			continue;
+		}
+		
+		ArrayAlreadyProcessedChild.Add(ReferenceName);
 
-		if (AssetData != nullptr) // Elements in ArrayAssetData are already of static class UCustomizableObject
+		const FAssetData* AssetData = GetCachedAssetData(ReferenceName.ToString());
+		if (!AssetData) // Elements in ArrayAssetData are already of static class UCustomizableObject
 		{
-			ChildObject = Cast<UCustomizableObject>(AssetData->GetAsset());
-			
-			if (!ChildObject)
-				continue;
+			continue;
+		}
+		
+		UCustomizableObject* ChildObject = Cast<UCustomizableObject>(AssetData->GetAsset());
+		if (!ChildObject)
+		{
+			continue;
+		}
+		
+		if (ChildObject != Object && !ChildObject->HasAnyFlags(RF_Transient))
+		{
+			UCustomizableObjectNodeObject* ChildRoot = GetRootNode(ChildObject, bMultipleBaseObjectsFound);
 
-			if (ChildObject != Object && !ChildObject->HasAnyFlags(RF_Transient))
+			if (ChildRoot && !bMultipleBaseObjectsFound)
 			{
-				UCustomizableObjectNodeObject* ChildRoot = GetRootNode(ChildObject, bMultipleBaseObjectsFound);
-
-				if (ChildRoot && !bMultipleBaseObjectsFound)
+				if (ChildRoot->ParentObject == Object)
 				{
-					if (ChildRoot->ParentObject == Object)
+					if (const FGroupNodeIdsTempData* GroupGuid = GenerationContext.DuplicatedGroupNodeIds.FindPair(Object, FGroupNodeIdsTempData(ChildRoot->ParentObjectGroupId)))
 					{
-						if (const FGroupNodeIdsTempData* GroupGuid = GenerationContext.DuplicatedGroupNodeIds.FindPair(Object, FGroupNodeIdsTempData(ChildRoot->ParentObjectGroupId)))
-						{
-							ChildRoot->ParentObjectGroupId = GroupGuid->NewGroupNodeId;
-						}
-
-						GenerationContext.GroupIdToExternalNodeMap.Add(ChildRoot->ParentObjectGroupId, ChildRoot);
-						GenerationContext.CustomizableObjectGuidsInCompilation.Add(ChildObject->GetVersionId());
+						ChildRoot->ParentObjectGroupId = GroupGuid->NewGroupNodeId;
 					}
+
+					GenerationContext.GroupIdToExternalNodeMap.Add(ChildRoot->ParentObjectGroupId, ChildRoot);
+					GenerationContext.CustomizableObjectGuidsInCompilation.Add(ChildObject->GetVersionId());
 				}
 			}
+		}
 
-			TArray<UCustomizableObjectNodeObjectGroup*> GroupNodes;
-			ChildObject->Source->GetNodesOfClass<UCustomizableObjectNodeObjectGroup>(GroupNodes);
+		TArray<UCustomizableObjectNodeObjectGroup*> GroupNodes;
+		ChildObject->Source->GetNodesOfClass<UCustomizableObjectNodeObjectGroup>(GroupNodes);
 
-			if (GroupNodes.Num() > 0) // Only grafs with group nodes should have child grafs
+		if (GroupNodes.Num() > 0) // Only grafs with group nodes should have child grafs
+		{
+			for (int32 i = 0; i < GroupNodes.Num(); ++i)
 			{
-				if (ArrayAlreadyProcessedChild.Find(ReferenceName) == INDEX_NONE)
+				const FGuid NodeId = GenerationContext.GetNodeIdUnique(GroupNodes[i]);
+				if (NodeId != GroupNodes[i]->NodeGuid)
 				{
-					for (int32 i = 0; i < GroupNodes.Num(); ++i)
-					{
-						const FGuid NodeId = GenerationContext.GetNodeIdUnique(GroupNodes[i]);
-						if (NodeId != GroupNodes[i]->NodeGuid)
-						{
-							GenerationContext.DuplicatedGroupNodeIds.Add(ChildObject, FGroupNodeIdsTempData(GroupNodes[i]->NodeGuid, NodeId));
-							GroupNodes[i]->NodeGuid = NodeId;
-						}
-					}
-
-					ArrayAlreadyProcessedChild.Add(ReferenceName);
-					ProcessChildObjectsRecursively(ChildObject, AssetRegistryModule, GenerationContext);
+					GenerationContext.DuplicatedGroupNodeIds.Add(ChildObject, FGroupNodeIdsTempData(GroupNodes[i]->NodeGuid, NodeId));
+					GroupNodes[i]->NodeGuid = NodeId;
 				}
 			}
+
+			ProcessChildObjectsRecursively(ChildObject, GenerationContext);
 		}
 	}
 }
@@ -510,6 +513,8 @@ mu::NodeObjectPtr FCustomizableObjectCompiler::GenerateMutableRoot(
 	FText& ErrorMsg, 
 	bool& bOutIsRootObject)
 {
+	check(Object);
+	
 	if (!Object->Source)
 	{
 		ErrorMsg = LOCTEXT("NoSource", "Object with no valid graph found. Object not build.");
@@ -574,7 +579,7 @@ mu::NodeObjectPtr FCustomizableObjectCompiler::GenerateMutableRoot(
 
 			// The object doesn't reference a root object but is a root object, look for all the objects that reference it and get their root nodes
 			FAssetRegistryModule& AssetRegistryModule = FModuleManager::LoadModuleChecked<FAssetRegistryModule>("AssetRegistry");
-			ProcessChildObjectsRecursively(ActualRootObject, AssetRegistryModule, GenerationContext);
+			ProcessChildObjectsRecursively(ActualRootObject, GenerationContext);
 			UE_LOG(LogMutable, Verbose, TEXT("PROFILE: [ %16.8f ] End search for children."), FPlatformTime::Seconds());
 		}
 	}
@@ -618,7 +623,7 @@ mu::NodeObjectPtr FCustomizableObjectCompiler::GenerateMutableRoot(
 			if (GroupNodes.Num() > 0) // Only grafs with group nodes should have child grafs
 			{
 				FAssetRegistryModule& AssetRegistryModule = FModuleManager::LoadModuleChecked<FAssetRegistryModule>("AssetRegistry");
-				ProcessChildObjectsRecursively(Object, AssetRegistryModule, GenerationContext);
+				ProcessChildObjectsRecursively(Object, GenerationContext);
 			}
 		}
 
@@ -678,6 +683,8 @@ mu::NodeObjectPtr FCustomizableObjectCompiler::GenerateMutableRoot(
 
 	Object->ReferenceSkeletalMeshes = ActualRootObject->ReferenceSkeletalMeshes;
 
+	GenerationContext.AddParticipatingObject(Object->ReferenceSkeletalMeshes);
+
     GenerationContext.RealTimeMorphTargetsOverrides = ActualRoot->RealTimeMorphSelectionOverrides;
     GenerationContext.RealTimeMorphTargetsOverrides.Reset();
 
@@ -1282,6 +1289,26 @@ void FCustomizableObjectCompiler::CompileInternal(UCustomizableObject* Object, c
 			}
 		}
 
+		{
+			GenerationContext.ParticipatingObjects.Remove(Object); // Remove self CO reference.
+
+			checkCode
+			(
+				for (auto Pair : GenerationContext.ParticipatingObjects)
+				{
+					check(Pair.Key.GetPackage() != Object->GetPackage()); // Adding a reference to itself will always force a recompilation after the CO is saved, so avoid doing that.
+				}
+			);
+			
+			TMap<TObjectPtr<const UObject>, FGuid>& ParticipatingObjects = Object->GetPrivate()->GetParticipatingObjects(*Object);
+			if (ParticipatingObjects != GenerationContext.ParticipatingObjects)
+			{
+				Object->MarkPackageDirty();
+			}
+
+			ParticipatingObjects = MoveTemp(GenerationContext.ParticipatingObjects);
+		}
+
 		if (bIsRootObject && (GenerationContext.MaskOutMaterialCache.Num() > 0 || GenerationContext.MaskOutTextureCache.Num() > 0))
 		{
 			// Load MaskOutCache, if can't, create it.
diff --git a/Engine/Plugins/Experimental/Mutable/Source/CustomizableObjectEditor/Private/MuCOE/CustomizableObjectCompiler.h b/Engine/Plugins/Experimental/Mutable/Source/CustomizableObjectEditor/Private/MuCOE/CustomizableObjectCompiler.h
--- a/Engine/Plugins/Experimental/Mutable/Source/CustomizableObjectEditor/Private/MuCOE/CustomizableObjectCompiler.h
+++ b/Engine/Plugins/Experimental/Mutable/Source/CustomizableObjectEditor/Private/MuCOE/CustomizableObjectCompiler.h
@@ -108,7 +108,7 @@ private:
 
 	static void PreloadingReferencerAssetsCallback(UCustomizableObject* Object, FCustomizableObjectCompiler* CustomizableObjectCompiler, const FCompilationOptions Options, bool bAsync);
 	
-	void ProcessChildObjectsRecursively(UCustomizableObject* Object, class FAssetRegistryModule& AssetRegistryModule, struct FMutableGraphGenerationContext &GenerationContext);
+	void ProcessChildObjectsRecursively(UCustomizableObject* Object, FMutableGraphGenerationContext &GenerationContext);
 
 	//
 	FCompilationOptions Options;
diff --git a/Engine/Plugins/Experimental/Mutable/Source/CustomizableObjectEditor/Private/MuCOE/GenerateMutableSource/GenerateMutableSource.cpp b/Engine/Plugins/Experimental/Mutable/Source/CustomizableObjectEditor/Private/MuCOE/GenerateMutableSource/GenerateMutableSource.cpp
--- a/Engine/Plugins/Experimental/Mutable/Source/CustomizableObjectEditor/Private/MuCOE/GenerateMutableSource/GenerateMutableSource.cpp
+++ b/Engine/Plugins/Experimental/Mutable/Source/CustomizableObjectEditor/Private/MuCOE/GenerateMutableSource/GenerateMutableSource.cpp
@@ -48,6 +48,7 @@
 #include "MuT/NodeSurfaceEdit.h"
 #include "PhysicsEngine/PhysicsAsset.h"
 #include "PlatformInfo.h"
+#include "MuCO/CustomizableObjectPrivate.h"
 
 #define LOCTEXT_NAMESPACE "CustomizableObjectEditor"
 
@@ -128,6 +129,17 @@ FMutableGraphGenerationContext::FMutableGraphGenerationContext(UCustomizableObje
 
 FMutableGraphGenerationContext::~FMutableGraphGenerationContext() = default;
 
+
+void FMutableGraphGenerationContext::AddParticipatingObject(const UObject& InObject)
+{
+	PRAGMA_DISABLE_DEPRECATION_WARNINGS
+	const FGuid PackageGuid = InObject.GetPackage()->GetGuid();
+	PRAGMA_ENABLE_DEPRECATION_WARNINGS
+
+	ParticipatingObjects.Add(&InObject, PackageGuid);
+}
+
+
 mu::MeshPtr FMutableGraphGenerationContext::FindGeneratedMesh( const FGeneratedMeshData::FKey& Key )
 {
 	for (const FGeneratedMeshData& d : GeneratedMeshes)
@@ -613,6 +625,8 @@ mu::NodeObjectPtr GenerateMutableSource(const UEdGraphPin * Pin, FMutableGraphGe
 	CheckNumOutputs(*Pin, GenerationContext);
 	
 	UCustomizableObjectNode* Node = CastChecked<UCustomizableObjectNode>(Pin->GetOwningNode());
+
+	GenerationContext.AddParticipatingObject(*GetRootObject(*Node));
 	
 	const FGeneratedKey Key(reinterpret_cast<void*>(&GenerateMutableSource), *Pin, *Node, GenerationContext, true);
 	if (const FGeneratedData* Generated = GenerationContext.Generated.Find(Key))
diff --git a/Engine/Plugins/Experimental/Mutable/Source/CustomizableObjectEditor/Private/MuCOE/GenerateMutableSource/GenerateMutableSource.h b/Engine/Plugins/Experimental/Mutable/Source/CustomizableObjectEditor/Private/MuCOE/GenerateMutableSource/GenerateMutableSource.h
--- a/Engine/Plugins/Experimental/Mutable/Source/CustomizableObjectEditor/Private/MuCOE/GenerateMutableSource/GenerateMutableSource.h
+++ b/Engine/Plugins/Experimental/Mutable/Source/CustomizableObjectEditor/Private/MuCOE/GenerateMutableSource/GenerateMutableSource.h
@@ -505,6 +505,22 @@ struct FMutableGraphGenerationContext
 	FMutableGraphGenerationContext(UCustomizableObject* CustomizableObject, class FCustomizableObjectCompiler* InCompiler, const FCompilationOptions& InOptions);
 	~FMutableGraphGenerationContext();
 
+	/** See UCustomizableObject::ParticipatingObjects. */
+	void AddParticipatingObject(const UObject& Object);
+
+	/** See UCustomizableObject::ParticipatingObjects. */
+	template<typename Type>
+	void AddParticipatingObject(const TArray<Type>& InObjects)
+	{
+		for (Type InObject : InObjects)
+		{
+			if (InObject)
+			{	
+				AddParticipatingObject(*InObject);
+			}
+		}
+	}
+	
 	UCustomizableObject* Object = nullptr;
 
 	// Non-owned reference to the compiler object
@@ -755,6 +771,9 @@ struct FMutableGraphGenerationContext
 	FExtensionDataCompilerInterface ExtensionDataCompilerInterface;
 	TArray<FCustomizableObjectExtensionData> AlwaysLoadedExtensionData;
 	TArray<UCustomizableObjectExtensionDataContainer*> StreamedExtensionData;
+
+	/** See UCustomizableObject::ParticipatingObjects. */
+	TMap<TObjectPtr<const UObject>, FGuid> ParticipatingObjects;
 };
 
 /** Pin Data scope wrapper. Pops the pin data on scope exit. */
diff --git a/Engine/Plugins/Experimental/Mutable/Source/CustomizableObjectEditor/Private/MuCOE/GenerateMutableSource/GenerateMutableSourceGroupProjector.cpp b/Engine/Plugins/Experimental/Mutable/Source/CustomizableObjectEditor/Private/MuCOE/GenerateMutableSource/GenerateMutableSourceGroupProjector.cpp
--- a/Engine/Plugins/Experimental/Mutable/Source/CustomizableObjectEditor/Private/MuCOE/GenerateMutableSource/GenerateMutableSourceGroupProjector.cpp
+++ b/Engine/Plugins/Experimental/Mutable/Source/CustomizableObjectEditor/Private/MuCOE/GenerateMutableSource/GenerateMutableSourceGroupProjector.cpp
@@ -69,6 +69,11 @@ mu::NodeImagePtr GenerateMutableGroupProjection(const int32 NodeLOD, const int32
 		if (!GroupProjectionReferenceTexture)
 		{
 			GroupProjectionReferenceTexture = ProjectorTempData.CustomizableObjectNodeGroupProjectorParameter->ReferenceTexture;
+			
+			if (GroupProjectionReferenceTexture)
+			{
+				GenerationContext.AddParticipatingObject(*GroupProjectionReferenceTexture);
+			}
 		}
 
 		const bool bProjectToImage = [&]
diff --git a/Engine/Plugins/Experimental/Mutable/Source/CustomizableObjectEditor/Private/MuCOE/GenerateMutableSource/GenerateMutableSourceImage.cpp b/Engine/Plugins/Experimental/Mutable/Source/CustomizableObjectEditor/Private/MuCOE/GenerateMutableSource/GenerateMutableSourceImage.cpp
--- a/Engine/Plugins/Experimental/Mutable/Source/CustomizableObjectEditor/Private/MuCOE/GenerateMutableSource/GenerateMutableSourceImage.cpp
+++ b/Engine/Plugins/Experimental/Mutable/Source/CustomizableObjectEditor/Private/MuCOE/GenerateMutableSource/GenerateMutableSourceImage.cpp
@@ -152,6 +152,8 @@ mu::NodeImagePtr GenerateMutableSourceImage(const UEdGraphPin* Pin, FMutableGrap
 		UTexture2D* BaseTexture = TypedNodeTex->Texture;
 		if (BaseTexture)
 		{
+			GenerationContext.AddParticipatingObject(*BaseTexture);
+			
 			// Check the specific image cache
 			FGeneratedImageKey imageKey = FGeneratedImageKey(Pin);
 			mu::NodeImageConstantPtr ImageNode;
@@ -216,6 +218,8 @@ mu::NodeImagePtr GenerateMutableSourceImage(const UEdGraphPin* Pin, FMutableGrap
 		const UTexture2D* ReferenceTexture = TypedNodeParam->ReferenceValue;
 		if (ReferenceTexture)
 		{
+			GenerationContext.AddParticipatingObject(*ReferenceTexture);
+
 			ResizeNode->SetSize(FMath::Max(ReferenceTexture->GetImportedSize().X,1), FMath::Max(ReferenceTexture->GetImportedSize().X, 1));
 		}
 		else
@@ -584,6 +588,8 @@ mu::NodeImagePtr GenerateMutableSourceImage(const UEdGraphPin* Pin, FMutableGrap
 		// Calculating Texture size using Reference texture parameters
 		if (TypedNodeProject->ReferenceTexture)
 		{
+			GenerationContext.AddParticipatingObject(*TypedNodeProject->ReferenceTexture);
+
 			int32 LODBias = ComputeLODBias(GenerationContext, TypedNodeProject->ReferenceTexture, TypedNodeProject->ReferenceTexture->MaxTextureSize, nullptr, INDEX_NONE, false);
 
 			if (TextureSize.X > 0 && TextureSize.Y > 0)
@@ -817,6 +823,8 @@ mu::NodeImagePtr GenerateMutableSourceImage(const UEdGraphPin* Pin, FMutableGrap
 		// Calculating Texture size using Reference texture parameters
 		if (TypedNodeTransform->ReferenceTexture)
 		{
+			GenerationContext.AddParticipatingObject(*TypedNodeTransform->ReferenceTexture);
+			
 			const int32 LODBias = ComputeLODBias(
 				GenerationContext, TypedNodeTransform->ReferenceTexture, TypedNodeTransform->ReferenceTexture->MaxTextureSize, 
 				nullptr, INDEX_NONE, false);
diff --git a/Engine/Plugins/Experimental/Mutable/Source/CustomizableObjectEditor/Private/MuCOE/GenerateMutableSource/GenerateMutableSourceMesh.cpp b/Engine/Plugins/Experimental/Mutable/Source/CustomizableObjectEditor/Private/MuCOE/GenerateMutableSource/GenerateMutableSourceMesh.cpp
--- a/Engine/Plugins/Experimental/Mutable/Source/CustomizableObjectEditor/Private/MuCOE/GenerateMutableSource/GenerateMutableSourceMesh.cpp
+++ b/Engine/Plugins/Experimental/Mutable/Source/CustomizableObjectEditor/Private/MuCOE/GenerateMutableSource/GenerateMutableSourceMesh.cpp
@@ -601,6 +601,8 @@ mu::MeshPtr ConvertSkeletalMeshToMutable(const USkeletalMesh* InSkeletalMesh, co
 		return nullptr;
 	}
 
+	GenerationContext.AddParticipatingObject(*InSkeletalMesh);
+	
 	const FSkeletalMeshModel* ImportedModel = InSkeletalMesh->GetImportedModel();
 	if (!ImportedModel)
 	{
@@ -1784,6 +1786,8 @@ mu::MeshPtr ConvertStaticMeshToMutable(const UStaticMesh* StaticMesh, int32 LODI
 		return nullptr;
 	}
 
+	GenerationContext.AddParticipatingObject(*StaticMesh);
+
 	mu::MeshPtr MutableMesh = new mu::Mesh();
 
 	// Vertices
diff --git a/Engine/Plugins/Experimental/Mutable/Source/CustomizableObjectEditor/Private/MuCOE/GenerateMutableSource/GenerateMutableSourceSurface.cpp b/Engine/Plugins/Experimental/Mutable/Source/CustomizableObjectEditor/Private/MuCOE/GenerateMutableSource/GenerateMutableSourceSurface.cpp
--- a/Engine/Plugins/Experimental/Mutable/Source/CustomizableObjectEditor/Private/MuCOE/GenerateMutableSource/GenerateMutableSourceSurface.cpp
+++ b/Engine/Plugins/Experimental/Mutable/Source/CustomizableObjectEditor/Private/MuCOE/GenerateMutableSource/GenerateMutableSourceSurface.cpp
@@ -298,6 +298,8 @@ mu::NodeSurfacePtr GenerateMutableSourceSurface(const UEdGraphPin * Pin, FMutabl
 		int32 ReferencedMaterialsIndex = -1;
 		if (TypedNodeMat->Material)
 		{
+			GenerationContext.AddParticipatingObject(*TypedNodeMat->Material);
+			
 			const int32 lastMaterialAmount = GenerationContext.ReferencedMaterials.Num();
 			ReferencedMaterialsIndex = GenerationContext.ReferencedMaterials.AddUnique(TypedNodeMat->Material);
 			// Used ReferencedMaterialsIndex instead of TypedNodeMat->Material->GetName() to prevent material name collisions
@@ -520,6 +522,8 @@ mu::NodeSurfacePtr GenerateMutableSourceSurface(const UEdGraphPin * Pin, FMutabl
 						FGeneratedImageProperties Props;
 						if (ReferenceTexture)
 						{
+							GenerationContext.AddParticipatingObject(*ReferenceTexture);
+							
 							// Store properties for the generated images
 							Props.TextureParameterName = ImageName;
 							Props.CompressionSettings = ReferenceTexture->CompressionSettings;
@@ -1316,6 +1320,12 @@ mu::NodeSurfacePtr GenerateMutableSourceSurface(const UEdGraphPin * Pin, FMutabl
 
 					// Calculate the LODBias for this texture
 					UTexture2D* ReferenceTexture = ParentMaterialNode->GetImageReferenceTexture(ImageIndex);
+
+					if (ReferenceTexture)
+					{
+						GenerationContext.AddParticipatingObject(*ReferenceTexture);
+					}
+					
 					int32 LODBias = ComputeLODBias(GenerationContext, ReferenceTexture, ReferenceTexture ? ReferenceTexture->MaxTextureSize : 0, ParentMaterialNode, ImageIndex);
 
 					GenerationContext.CurrentTextureLODBias = LODBias;
diff --git a/Engine/Plugins/Experimental/Mutable/Source/CustomizableObjectEditor/Private/MuCOE/GenerateMutableSource/GenerateMutableSourceTable.cpp b/Engine/Plugins/Experimental/Mutable/Source/CustomizableObjectEditor/Private/MuCOE/GenerateMutableSource/GenerateMutableSourceTable.cpp
--- a/Engine/Plugins/Experimental/Mutable/Source/CustomizableObjectEditor/Private/MuCOE/GenerateMutableSource/GenerateMutableSourceTable.cpp
+++ b/Engine/Plugins/Experimental/Mutable/Source/CustomizableObjectEditor/Private/MuCOE/GenerateMutableSource/GenerateMutableSourceTable.cpp
@@ -350,7 +350,7 @@ bool FillTableColumn(const UCustomizableObjectNodeTable* TableNode,	mu::TablePtr
 				CurrentColumn = MutableTable->AddColumn(StringCast<ANSICHAR>(*MutableColumnName).Get(), mu::ETableColumnType::Mesh);
 			}
 
-			mu::MeshPtr MutableMesh = GenerateMutableMesh(StaticMesh, TSoftClassPtr<UAnimInstance>(), CurrentLOD, SectionIndex, CurrentLOD, SectionIndex, FString(), GenerationContext, TableNode); // TODO GMT
+			mu::MeshPtr MutableMesh = GenerateMutableMesh(StaticMesh, TSoftClassPtr<UAnimInstance>(), CurrentLOD, SectionIndex, CurrentLOD, SectionIndex, FString(), GenerationContext, TableNode);
 
 			if (MutableMesh)
 			{
@@ -414,9 +414,14 @@ bool FillTableColumn(const UCustomizableObjectNodeTable* TableNode,	mu::TablePtr
 				mu::Ptr<mu::ResourceProxyMemory<mu::Image>> Proxy = new mu::ResourceProxyMemory<mu::Image>(mu::Image::CreateAsReference(ImageReferenceID));
 				MutableTable->SetCell(CurrentColumn, RowIdx, Proxy.get());
 			}
-			else
+			else 
 			{
-				GenerationContext.ArrayTextureUnrealToMutableTask.Add(FTextureUnrealToMutableTask(MutableTable, Texture2D, TableNode, CurrentColumn, RowIdx));
+				if (Texture2D)
+				{
+					GenerationContext.AddParticipatingObject(*Texture2D);
+
+					GenerationContext.ArrayTextureUnrealToMutableTask.Add(FTextureUnrealToMutableTask(MutableTable, Texture2D, TableNode, CurrentColumn, RowIdx));
+				}
 			}
 		}
 
@@ -436,6 +441,8 @@ bool FillTableColumn(const UCustomizableObjectNodeTable* TableNode,	mu::TablePtr
 				return false;
 			}
 
+			GenerationContext.AddParticipatingObject(*ReferenceMaterial);
+
 			if (!Material || ReferenceMaterial->GetMaterial() != Material->GetMaterial())
 			{
 				FText Warning;
@@ -456,6 +463,8 @@ bool FillTableColumn(const UCustomizableObjectNodeTable* TableNode,	mu::TablePtr
 				GenerationContext.Compiler->CompilerLog(Warning, TableNode);
 			}
 
+			GenerationContext.AddParticipatingObject(*Material);
+			
 			FString EncodedSwitchParameterName = "__MutableMaterialId";
 			if (ColumnName.Contains(EncodedSwitchParameterName))
 			{
@@ -619,6 +628,8 @@ bool GenerateTableColumn(const UCustomizableObjectNodeTable* TableNode, const UE
 	{
 		return false;
 	}
+	
+	GenerationContext.AddParticipatingObject(*TableNode->Table);
 
 	// Getting names of the rows to access the information
 	TArray<FName> RowNames = TableNode->GetRowNames();
diff --git a/Engine/Plugins/Experimental/Mutable/Source/CustomizableObjectEditor/Private/MuCOE/GraphTraversal.cpp b/Engine/Plugins/Experimental/Mutable/Source/CustomizableObjectEditor/Private/MuCOE/GraphTraversal.cpp
--- a/Engine/Plugins/Experimental/Mutable/Source/CustomizableObjectEditor/Private/MuCOE/GraphTraversal.cpp
+++ b/Engine/Plugins/Experimental/Mutable/Source/CustomizableObjectEditor/Private/MuCOE/GraphTraversal.cpp
@@ -298,6 +298,12 @@ UCustomizableObject* GetFullGraphRootObject(UCustomizableObjectNodeObject* Node,
 }
 
 
+UCustomizableObject* GetRootObject(const UCustomizableObjectNode& Node)
+{
+	return CastChecked<UCustomizableObject>(Node.GetGraph()->GetOuter());
+}
+
+
 UCustomizableObjectNodeObject* GetFullGraphRootNodeObject(UCustomizableObjectNodeObject* Node, TArray<UCustomizableObject*>& VisitedObjects)
 {
 	if (Node->ParentObject != nullptr)
diff --git a/Engine/Plugins/Experimental/Mutable/Source/CustomizableObjectEditor/Public/MuCOE/GraphTraversal.h b/Engine/Plugins/Experimental/Mutable/Source/CustomizableObjectEditor/Public/MuCOE/GraphTraversal.h
--- a/Engine/Plugins/Experimental/Mutable/Source/CustomizableObjectEditor/Public/MuCOE/GraphTraversal.h
+++ b/Engine/Plugins/Experimental/Mutable/Source/CustomizableObjectEditor/Public/MuCOE/GraphTraversal.h
@@ -7,6 +7,7 @@
 class UCustomizableObject;
 class UCustomizableObjectNodeObject;
 class UEdGraphPin;
+class UCustomizableObjectNode;
 
 /** Follow the given input pin returning the output connected pin.
  *
@@ -45,9 +46,12 @@ bool GetParentsUntilRoot(UCustomizableObject* Object, TArray<UCustomizableObject
 /** Returns true if the Candidate is parent of the current Customizable Object */
 bool HasCandidateAsParent(UCustomizableObjectNodeObject* Node, UCustomizableObject* ParentCandidate);
 
-/** Return the full graph Customizable Object root of the node given as parameter */
+/** Return the full graph Customizable Object root of the given node. */
 UCustomizableObject* GetFullGraphRootObject(UCustomizableObjectNodeObject* Node, TArray<UCustomizableObject*>& VisitedObjects);
 
+/** Return Customizable Object of the given node. */
+UCustomizableObject* GetRootObject(const UCustomizableObjectNode& Node);
+
 /** Return the full graph Customizable Object Node root of the node given as parameter */
 UCustomizableObjectNodeObject* GetFullGraphRootNodeObject(UCustomizableObjectNodeObject* Node, TArray<UCustomizableObject*>& VisitedObjects);
 
