[mutable] Now, color parameters with "none" value use the value set in the material instance.

#rnx

diff --git a/Engine/Plugins/Experimental/Mutable/Source/CustomizableObject/Private/MuCO/CustomizableObjectInstance.cpp b/Engine/Plugins/Experimental/Mutable/Source/CustomizableObject/Private/MuCO/CustomizableObjectInstance.cpp
--- a/Engine/Plugins/Experimental/Mutable/Source/CustomizableObject/Private/MuCO/CustomizableObjectInstance.cpp
+++ b/Engine/Plugins/Experimental/Mutable/Source/CustomizableObject/Private/MuCO/CustomizableObjectInstance.cpp
@@ -6020,7 +6020,15 @@ void UCustomizableInstancePrivateData::BuildMaterials(const TSharedPtr<FMutableO
 							case FMutableMaterialPlaceholder::EPlaceHolderParamType::Vector:
 								if (Param.LayerIndex < 0)
 								{
-									MaterialInstance->SetVectorParameterValue(Param.ParamName, Param.Vector);
+									FLinearColor Color = Param.Vector;
+
+									if (FVector4f(Color) == mu::DefaultMutableColorValue)
+									{
+										FMaterialParameterInfo ParameterInfo(Param.ParamName);
+										MaterialTemplate->GetVectorParameterValue(ParameterInfo, Color);
+									}
+
+									MaterialInstance->SetVectorParameterValue(Param.ParamName, Color);
 								}
 								else
 								{
diff --git a/Engine/Plugins/Experimental/Mutable/Source/MutableRuntime/Public/MuR/Parameters.h b/Engine/Plugins/Experimental/Mutable/Source/MutableRuntime/Public/MuR/Parameters.h
--- a/Engine/Plugins/Experimental/Mutable/Source/MutableRuntime/Public/MuR/Parameters.h
+++ b/Engine/Plugins/Experimental/Mutable/Source/MutableRuntime/Public/MuR/Parameters.h
@@ -7,6 +7,7 @@
 #include "MuR/RefCounted.h"
 #include "MuR/Types.h"
 #include "Templates/SharedPointer.h"
+#include "VectorTypes.h"
 
 
 namespace mu
@@ -27,6 +28,9 @@ namespace mu
 
 	class Model;
 
+	//! Default parameter values definition
+	static const FVector4f DefaultMutableColorValue(-UE_BIG_NUMBER);
+
 
     //! \brief Model parameter types.
 	//! \ingroup runtime
diff --git a/Engine/Plugins/Experimental/Mutable/Source/MutableTools/Private/MuT/CodeGenerator.cpp b/Engine/Plugins/Experimental/Mutable/Source/MutableTools/Private/MuT/CodeGenerator.cpp
--- a/Engine/Plugins/Experimental/Mutable/Source/MutableTools/Private/MuT/CodeGenerator.cpp
+++ b/Engine/Plugins/Experimental/Mutable/Source/MutableTools/Private/MuT/CodeGenerator.cpp
@@ -43,6 +43,7 @@
 #include "MuT/CompilerPrivate.h"
 #include "MuT/ErrorLogPrivate.h"
 #include "MuT/NodeColour.h"
+#include "MuT/NodeColourConstant.h"
 #include "MuT/NodeComponent.h"
 #include "MuT/NodeComponentNewPrivate.h"
 #include "MuT/NodeImage.h"
@@ -50,6 +51,7 @@
 #include "MuT/NodeImageFormatPrivate.h"
 #include "MuT/NodeImageMipmap.h"
 #include "MuT/NodeImageMipmapPrivate.h"
+#include "MuT/NodeImageReference.h"
 #include "MuT/NodeLODPrivate.h"
 #include "MuT/NodeMesh.h"
 #include "MuT/NodeMeshClipMorphPlane.h"
@@ -2327,4 +2329,49 @@ namespace mu
 		return lastMeshOp;
 	}
 
+
+	Ptr<ASTOp> CodeGenerator::GenerateDefaultTableValue(ETableColumnType NodeType)
+	{
+		switch (NodeType)
+		{
+		case mu::ETableColumnType::Scalar:
+		{
+			//TODO(Max):MTBL-1660
+			//Ptr<NodeScalarConstant> pNode = new NodeScalarConstant();
+			//pNode->SetValue(-UE_MAX_FLT);
+			//
+			//return Generate(pNode);
+			return nullptr;
+		}
+		case mu::ETableColumnType::Color:
+		{
+			mu::Ptr<mu::NodeColourConstant> pNode = new NodeColourConstant();
+			pNode->SetValue(mu::DefaultMutableColorValue);
+
+			return Generate(pNode);
+		}
+		case mu::ETableColumnType::Image:
+		{
+			//TODO(Max):MTBL-1660
+			//FImageGenerationOptions DummyOptions;
+			//FImageGenerationResult DefaultValue;
+			//
+			//mu::Ptr<mu::NodeImageReference> ImageNode = new mu::NodeImageReference();
+			//ImageNode->SetImageReference(-1);
+			//
+			//GenerateImage(DummyOptions, DefaultValue, ImageNode);
+			//
+			//return DefaultValue.op;
+			return nullptr;
+		}
+		case mu::ETableColumnType::Mesh:
+			/* The default mesh is always null */
+			return nullptr;
+		default:
+			break;
+		}
+
+		return nullptr;
+	}
+
 }
diff --git a/Engine/Plugins/Experimental/Mutable/Source/MutableTools/Private/MuT/CodeGenerator.h b/Engine/Plugins/Experimental/Mutable/Source/MutableTools/Private/MuT/CodeGenerator.h
--- a/Engine/Plugins/Experimental/Mutable/Source/MutableTools/Private/MuT/CodeGenerator.h
+++ b/Engine/Plugins/Experimental/Mutable/Source/MutableTools/Private/MuT/CodeGenerator.h
@@ -737,6 +737,10 @@ namespace mu
         void GenerateSurface( FSurfaceGenerationResult& result,
                               NodeSurfaceNewPtrConst node,
                               const TArray<FirstPassGenerator::FSurface::FEdit>& edits );
+
+		//-----------------------------------------------------------------------------------------
+		//Default Table Parameters
+		Ptr<ASTOp> GenerateDefaultTableValue(ETableColumnType NodeType);
     };
 
 
@@ -805,7 +809,7 @@ namespace mu
         Ptr<ASTOpSwitch> SwitchOp = new ASTOpSwitch();
 		SwitchOp->type = OPTYPE;
 		SwitchOp->variable = variable;
-		SwitchOp->def = nullptr;
+		SwitchOp->def = GenerateDefaultTableValue(TYPE);
 
 		for (int32 i = 0; i < rows; ++i)
         {
