[mutable] Improve the compilation optimization options: "None" is now much faster because it doesn't even convert textures. "Medium" is gone.
"None" now generates "referenced images" instead of converting pixels.
Fix unsupported BGRA format that was being generated for some images.
#rnx
#preflgiht 651c3c027e782effa700f991

diff --git a/Engine/Plugins/Experimental/Mutable/Source/CustomizableObject/Classes/MuCO/CustomizableObject.h b/Engine/Plugins/Experimental/Mutable/Source/CustomizableObject/Classes/MuCO/CustomizableObject.h
--- a/Engine/Plugins/Experimental/Mutable/Source/CustomizableObject/Classes/MuCO/CustomizableObject.h
+++ b/Engine/Plugins/Experimental/Mutable/Source/CustomizableObject/Classes/MuCO/CustomizableObject.h
@@ -176,7 +176,7 @@ struct FCompilationOptions
 	UPROPERTY()
 	ECustomizableObjectTextureCompression TextureCompression = ECustomizableObjectTextureCompression::Fast;
 
-	// From 0 to 3
+	// From 0 to 2
 	UPROPERTY()
 	int32 OptimizationLevel = 1;
 
diff --git a/Engine/Plugins/Experimental/Mutable/Source/CustomizableObject/CustomizableObject.Build.cs b/Engine/Plugins/Experimental/Mutable/Source/CustomizableObject/CustomizableObject.Build.cs
--- a/Engine/Plugins/Experimental/Mutable/Source/CustomizableObject/CustomizableObject.Build.cs
+++ b/Engine/Plugins/Experimental/Mutable/Source/CustomizableObject/CustomizableObject.Build.cs
@@ -26,6 +26,7 @@ public class CustomizableObject : ModuleRules
 			"ClothingSystemRuntimeInterface",
 			//"ClothingSystemEditor",
 			"UMG",
+			"ImageCore",
 		});
 
         PublicDependencyModuleNames.AddRange(new string[] {
diff --git a/Engine/Plugins/Experimental/Mutable/Source/CustomizableObject/Private/MuCO/CustomizableObject.cpp b/Engine/Plugins/Experimental/Mutable/Source/CustomizableObject/Private/MuCO/CustomizableObject.cpp
--- a/Engine/Plugins/Experimental/Mutable/Source/CustomizableObject/Private/MuCO/CustomizableObject.cpp
+++ b/Engine/Plugins/Experimental/Mutable/Source/CustomizableObject/Private/MuCO/CustomizableObject.cpp
@@ -742,7 +742,7 @@ void UCustomizableObject::CompileForTargetPlatform(const ITargetPlatform* Target
 	if (bIsRootObject && bIsRelevantForThisTarget)
 	{
 		FCompilationOptions Options;
-		Options.OptimizationLevel = 3;	// max optimization when packaging.
+		Options.OptimizationLevel = 2;	// max optimization when packaging.
 		Options.TextureCompression = ECustomizableObjectTextureCompression::HighQuality;
 		Options.bIsCooking = true;
 		Options.bSaveCookedDataToDisk = bUsesOnCookStart;
diff --git a/Engine/Plugins/Experimental/Mutable/Source/CustomizableObject/Private/MuCO/UnrealMutableImageProvider.cpp b/Engine/Plugins/Experimental/Mutable/Source/CustomizableObject/Private/MuCO/UnrealMutableImageProvider.cpp
--- a/Engine/Plugins/Experimental/Mutable/Source/CustomizableObject/Private/MuCO/UnrealMutableImageProvider.cpp
+++ b/Engine/Plugins/Experimental/Mutable/Source/CustomizableObject/Private/MuCO/UnrealMutableImageProvider.cpp
@@ -4,6 +4,7 @@
 
 #include "MuCO/CustomizableObjectSystem.h"
 #include "MuCO/CustomizableObject.h"
+#include "MuCO/UnrealToMutableTextureConversionUtils.h"
 #include "TextureResource.h"
 #include "MuR/Parameters.h"
 
@@ -12,44 +13,12 @@
 namespace
 {
 
-	mu::ImagePtr ConvertTextureUnrealToMutable(UTexture2D* Texture, uint8 MipmapsToSkip)
-	{
-		mu::ImagePtr pResult;
-
+	void ConvertTextureUnrealToMutable(mu::Image* OutResult, UTexture2D* Texture, bool bIsNormalComposite, uint8 MipmapsToSkip)
+	{		
 #if WITH_EDITOR
-		int LODs = 1;
-		int SizeX = Texture->Source.GetSizeX() >> MipmapsToSkip;
-		int SizeY = Texture->Source.GetSizeY() >> MipmapsToSkip;
-		check(SizeX > 0 && SizeY > 0);
-
-		ETextureSourceFormat Format = Texture->Source.GetFormat();
-		mu::EImageFormat MutableFormat = mu::EImageFormat::IF_NONE;
-
-		switch (Format)
-		{
-		case ETextureSourceFormat::TSF_BGRA8: MutableFormat = mu::EImageFormat::IF_BGRA_UBYTE; break;
-		// This format is deprecated and using the enum fails to compile in some cases.
-		//case ETextureSourceFormat::TSF_RGBA8: MutableFormat = mu::EImageFormat::IF_RGBA_UBYTE; break;
-		case ETextureSourceFormat::TSF_G8: MutableFormat = mu::EImageFormat::IF_L_UBYTE; break;
-		default:
-			break;
-		}
-
-		// If not locked ReadOnly the Texture Source's FGuid can change, invalidating the texture's caching/shaders
-		// making shader compile and cook times increase
-		const uint8* pSource = Texture->Source.LockMipReadOnly(MipmapsToSkip);
-		if (pSource)
-		{
-			pResult = new mu::Image(SizeX, SizeY, LODs, MutableFormat, mu::EInitializationType::NotInitialized);
-			FMemory::Memcpy(pResult->GetData(), pSource, pResult->GetDataSize());
-			Texture->Source.UnlockMip(MipmapsToSkip);
-		}
-		else
-		{
-			check(false);
-			pResult = new mu::Image(SizeX, SizeY, LODs, MutableFormat, mu::EInitializationType::Black);
-		}
 
+		EUnrealToMutableConversionError Error = ConvertTextureUnrealSourceToMutable(OutResult, Texture, bIsNormalComposite, MipmapsToSkip);
+		check(Error==EUnrealToMutableConversionError::Success);
 
 #else
 		check(Texture->GetPlatformData()->Mips[MipmapsToSkip].BulkData.IsBulkDataLoaded());
@@ -78,19 +47,17 @@ namespace
 
 		if (pSource)
 		{
-			pResult = new mu::Image(SizeX, SizeY, LODs, MutableFormat, mu::EInitializationType::NotInitialized);
-			FMemory::Memcpy(pResult->GetData(), pSource, pResult->GetDataSize());
+			OutResult->Init(SizeX, SizeY, LODs, MutableFormat, mu::EInitializationType::NotInitialized);
+			FMemory::Memcpy(OutResult->GetData(), pSource, OutResult->GetDataSize());
 			Texture->GetPlatformData()->Mips[MipmapsToSkip].BulkData.Unlock();
 		}
 		else
 		{
 			check(false);
-			pResult = new mu::Image(SizeX, SizeY, LODs, MutableFormat, mu::EInitializationType::Black);
+			OutResult->Init(SizeX, SizeY, LODs, MutableFormat, mu::EInitializationType::Black);
 		}
 
 #endif
-
-		return pResult;
 	}
 }
 
@@ -117,6 +84,77 @@ mu::EImageFormat GetMutablePixelFormat(EPixelFormat InTextureFormat)
 }
 
 
+FUnrealMutableImageProvider::FUnrealMutableImageProvider()
+{
+#if WITH_EDITOR
+	TickDelegate = FTickerDelegate::CreateRaw(this, &FUnrealMutableImageProvider::Tick);
+	TickDelegateHandle = FTSTicker::GetCoreTicker().AddTicker(TickDelegate, 0.f);
+#endif 
+}
+
+
+#if WITH_EDITOR
+
+bool FUnrealMutableImageProvider::Tick(float DeltaTime)
+{
+	FReferencedImageRequest* Request=nullptr;
+
+	// Process only one request per tick, or zero.
+	if (!QueuedReferencedImageRequests.Dequeue(Request))
+	{
+		return true;
+	}
+
+	// Find the CO for this model.
+	UCustomizableObject* CO = nullptr;
+	for (TObjectIterator<UCustomizableObject> It; It; ++It)
+	{
+		if (It->GetModel().Get() == Request->ModelPtr)
+		{
+			CO = *It;
+			break;
+		}
+	}
+
+	if (!CO)
+	{
+		// The CO for this request has been unloaded!
+		check(false);
+		return true;
+	}
+
+	if (!CO->ReferencedPassThroughTextures.IsValidIndex(Request->Id))
+	{
+		// The id is not valid for this CO
+		check(false);
+		return true;
+	}
+
+	// Find the texture id
+	TSoftObjectPtr<UTexture> TexturePtr = CO->ReferencedPassThroughTextures[Request->Id];
+
+	// This can cause a stall because of loading the asset.
+	UTexture2D* Texture = Cast<UTexture2D>( TexturePtr.Get() );
+	if (!Texture)
+	{
+		// Failed to load the texture
+		check(false);
+		return true;
+	}
+
+	// In the editor the src data can be directly accessed
+	int32 MipIndex = (Request->MipmapsToSkip < Texture->GetPlatformData()->Mips.Num()) ? Request->MipmapsToSkip : Texture->GetPlatformData()->Mips.Num() - 1;
+	check(MipIndex >= 0);
+	bool bIsNormalComposite = false; // TODO?
+	ConvertTextureUnrealToMutable(Request->ResultImage.get(), Texture, bIsNormalComposite, MipIndex);
+	Request->CompletionEvent.Trigger();
+
+	return true;
+}
+
+#endif
+
+
 //-------------------------------------------------------------------------------------------------
 #ifdef MUTABLE_USE_NEW_TASKGRAPH
 	TTuple<UE::Tasks::FTask, TFunction<void()>> FUnrealMutableImageProvider::GetImageAsync(FName Id, uint8 MipmapsToSkip, TFunction<void(mu::Ptr<mu::Image>)>& ResultCallback)
@@ -127,9 +165,6 @@ mu::EImageFormat GetMutablePixelFormat(EPixelFormat InTextureFormat)
 	// Thread: worker
 	MUTABLE_CPUPROFILER_SCOPE(FUnrealMutableImageProvider::GetImage);
 
-	// Out Texture
-	mu::ImagePtr Image;
-
 	// Some data that may have to be copied from the GlobalExternalImages while it's locked
 	IBulkDataIORequest* IORequest = nullptr;
 	const int32 LODs = 1;
@@ -212,7 +247,10 @@ mu::EImageFormat GetMutablePixelFormat(EPixelFormat InTextureFormat)
 
 #if WITH_EDITOR
 			// In the editor the src data can be directly accessed
-			ResultCallback(ConvertTextureUnrealToMutable(TextureToLoad, MipIndex));
+			mu::Ptr<mu::Image> Image = new mu::Image();
+			bool bIsNormalComposite = false; // TODO?
+			ConvertTextureUnrealToMutable(Image.get(), TextureToLoad, bIsNormalComposite,  MipIndex);
+			ResultCallback(Image);
 			return Invoke(TrivialReturn);
 #else
 			// Texture format and the equivalent mutable format
@@ -230,7 +268,7 @@ mu::EImageFormat GetMutablePixelFormat(EPixelFormat InTextureFormat)
 			int SizeX = TextureToLoad->GetSizeX() >> MipIndex;
 			int SizeY = TextureToLoad->GetSizeY() >> MipIndex;
 
-			Image = new mu::Image(SizeX, SizeY, LODs, MutImageFormat, mu::EInitializationType::NotInitialized);
+			mu::Ptr<mu::Image> Image = new mu::Image(SizeX, SizeY, LODs, MutImageFormat, mu::EInitializationType::NotInitialized);
 			MutImageDataSize = Image->GetDataSize();
 
 			// In a packaged game the bulk data has to be loaded
@@ -402,6 +440,73 @@ mu::EImageFormat GetMutablePixelFormat(EPixelFormat InTextureFormat)
 }
 
 
+//-------------------------------------------------------------------------------------------------
+#ifdef MUTABLE_USE_NEW_TASKGRAPH
+TTuple<UE::Tasks::FTask, TFunction<void()>> FUnrealMutableImageProvider::GetReferencedImageAsync(const void* ModelPtr, int32 Id, uint8 MipmapsToSkip, TFunction<void(mu::Ptr<mu::Image>)>& ResultCallback)
+#else
+TTuple<FGraphEventRef, TFunction<void()>> FUnrealMutableImageProvider::GetReferencedImageAsync(const void* ModelPtr, int32 Id, uint8 MipmapsToSkip, TFunction<void(mu::Ptr<mu::Image>)>& ResultCallback)
+#endif
+{
+	// Thread: worker
+	MUTABLE_CPUPROFILER_SCOPE(FUnrealMutableImageProvider::GetReferencedImageAsync);
+
+#ifdef MUTABLE_USE_NEW_TASKGRAPH
+	auto TrivialReturn = []() -> TTuple<UE::Tasks::FTask, TFunction<void()>>
+	{
+		UE::Tasks::FTaskEvent CompletionEvent(TEXT("GetImageAsyncCompleted"));
+		CompletionEvent.Trigger();
+
+		return MakeTuple(CompletionEvent, []() -> void {});
+	};
+#else
+	auto TrivialReturn = []() -> TTuple<FGraphEventRef, TFunction<void()>>
+	{
+		FGraphEventRef CompletionEvent = FGraphEvent::CreateGraphEvent();
+		CompletionEvent->DispatchSubsequents();
+
+		return MakeTuple(CompletionEvent, []() -> void {});
+	};
+#endif
+
+
+#if WITH_EDITOR
+
+	TUniquePtr<FReferencedImageRequest> Request( new FReferencedImageRequest(UE::Tasks::FTaskEvent(TEXT("GetReferencedImageCompletion"))) );
+	Request->ModelPtr = ModelPtr;
+	Request->Id = Id;
+	Request->MipmapsToSkip = MipmapsToSkip;
+	Request->ResultImage = new mu::Image();
+
+	QueuedReferencedImageRequests.Enqueue(Request.Get());
+
+	if (IsInGameThread())
+	{
+		// This may happen in the mutable debugger.
+		while (!Request->CompletionEvent.IsCompleted())
+		{
+			Tick(0.1f);
+		}
+	}
+	else
+	{
+		Request->CompletionEvent.BusyWait();
+	}
+
+	ResultCallback(Request->ResultImage);
+	return Invoke(TrivialReturn);
+
+#else // WITH_EDITOR
+
+	// Not supported outside editor yet.
+	UE_LOG(LogMutable, Warning, TEXT("Failed to get reference image. Only supported in editor."));
+
+	ResultCallback(CreateDummy());
+	return Invoke(TrivialReturn);
+
+#endif
+}
+
+
 // This should mantain parity with the descriptor of the images generated by GetImageAsync 
 mu::FImageDesc FUnrealMutableImageProvider::GetImageDesc(FName Id, uint8 MipmapsToSkip)
 {
@@ -521,7 +626,9 @@ void FUnrealMutableImageProvider::CacheImage(FName Id, bool bUser)
 					case UCustomizableSystemImageProvider::ValueType::Unreal:
 					{
 						UTexture2D* UnrealTexture = Provider->GetTextureParameterValue(Id);
-						pResult = ConvertTextureUnrealToMutable(UnrealTexture, 0);
+						pResult = new mu::Image();
+						bool bIsNormalComposite = false;
+						ConvertTextureUnrealToMutable(pResult.get(), UnrealTexture, bIsNormalComposite, 0);
 						break;
 					}
 
diff --git a/Engine/Plugins/Experimental/Mutable/Source/CustomizableObject/Private/MuCO/UnrealMutableImageProvider.h b/Engine/Plugins/Experimental/Mutable/Source/CustomizableObject/Private/MuCO/UnrealMutableImageProvider.h
--- a/Engine/Plugins/Experimental/Mutable/Source/CustomizableObject/Private/MuCO/UnrealMutableImageProvider.h
+++ b/Engine/Plugins/Experimental/Mutable/Source/CustomizableObject/Private/MuCO/UnrealMutableImageProvider.h
@@ -20,12 +20,16 @@ class FUnrealMutableImageProvider : public mu::ImageParameterGenerator, public F
 
 public:
 
+	FUnrealMutableImageProvider();
+
 	// mu::ImageParameterGenerator interface
 	// Thread: worker
 #ifdef MUTABLE_USE_NEW_TASKGRAPH
 	virtual TTuple<UE::Tasks::FTask, TFunction<void()>> GetImageAsync(FName Id, uint8 MipmapsToSkip, TFunction<void(mu::Ptr<mu::Image>)>& ResultCallback) override;
+	virtual TTuple<UE::Tasks::FTask, TFunction<void()>> GetReferencedImageAsync(const void* ModelPtr, int32 Id, uint8 MipmapsToSkip, TFunction<void(mu::Ptr<mu::Image>)>& ResultCallback) override;
 #else
 	virtual TTuple<FGraphEventRef, TFunction<void()>> GetImageAsync(FName Id, uint8 MipmapsToSkip, TFunction<void(mu::Ptr<mu::Image>)>& ResultCallback) override;
+	virtual TTuple<FGraphEventRef, TFunction<void()>> GetReferencedImageAsync(const void* ModelPtr, int32 Id, uint8 MipmapsToSkip, TFunction<void(mu::Ptr<mu::Image>)>& ResultCallback) override;
 #endif
 
 	virtual mu::FImageDesc GetImageDesc(FName Id, uint8 MipmapsToSkip) override;
@@ -94,4 +98,28 @@ private:
 	/** Access to GlobalExternalImages must be protected with this because it may be accessed concurrently from the 
 	Game thread to modify it and from the Mutable thread to read it. */
 	FCriticalSection ExternalImagesLock;
+
+#if WITH_EDITOR
+	/** The provider may tick to support referenced images. */
+	bool Tick(float DeltaTime);
+	FTSTicker::FDelegateHandle TickDelegateHandle;
+	FTickerDelegate TickDelegate;
+	struct FReferencedImageRequest
+	{
+		FReferencedImageRequest(const UE::Tasks::FTaskEvent&& Event) : CompletionEvent(Event) {}
+
+		// Sync
+		UE::Tasks::FTaskEvent CompletionEvent;
+
+		// Input
+		int32 Id = 0;
+		uint8 MipmapsToSkip = 0;
+		const void* ModelPtr = nullptr;
+
+		// Result
+		mu::Ptr<mu::Image> ResultImage;
+	};
+	TQueue<FReferencedImageRequest*, EQueueMode::Spsc> QueuedReferencedImageRequests;
+#endif
+
 };
diff --git a/Engine/Plugins/Experimental/Mutable/Source/CustomizableObject/Private/MuCO/UnrealToMutableTextureConversionUtils.cpp b/Engine/Plugins/Experimental/Mutable/Source/CustomizableObject/Private/MuCO/UnrealToMutableTextureConversionUtils.cpp
new file mode 100644
--- /dev/null
+++ b/Engine/Plugins/Experimental/Mutable/Source/CustomizableObject/Private/MuCO/UnrealToMutableTextureConversionUtils.cpp
@@ -0,0 +1,325 @@
+// Copyright Epic Games, Inc. All Rights Reserved.
+
+#include "MuCO/UnrealToMutableTextureConversionUtils.h"
+
+#include "MuR/MutableTrace.h"
+#include "MuR/Image.h"
+#include "Engine/Texture2D.h"
+#include "ImageCoreUtils.h"
+#include "ImageUtils.h"
+#include "Async/ParallelFor.h"
+
+#if WITH_EDITOR
+
+namespace UnrealToMutableImageConversion_Interanl
+{
+	
+FORCEINLINE ERawImageFormat::Type ConvertFormatSourceToRaw(const ETextureSourceFormat SourceFormat)
+{
+	return FImageCoreUtils::ConvertToRawImageFormat(SourceFormat);
+}
+
+EUnrealToMutableConversionError ApplyCompositeTexture(
+        FImage& Image, 
+        UTexture* CompositeTexture,
+        const ECompositeTextureMode CompositeTextureMode,
+        const float CompositePower)
+{
+    const int32 SizeX = CompositeTexture->Source.GetSizeX();
+    const int32 SizeY = CompositeTexture->Source.GetSizeY();
+    const ETextureSourceFormat SourceFormat = CompositeTexture->Source.GetFormat();
+
+    const ERawImageFormat::Type RawFormat = ConvertFormatSourceToRaw(SourceFormat);
+
+    if (RawFormat == ERawImageFormat::RGBA32F)
+    {
+        return EUnrealToMutableConversionError::CompositeUnsupportedFormat;
+    }   
+
+    FImage CompositeImage(SizeX, SizeY, 1, RawFormat, EGammaSpace::Linear);
+    
+    if(!CompositeTexture->Source.GetMipData(CompositeImage.RawData, 0))
+    {
+        return EUnrealToMutableConversionError::Unknown;
+    }
+
+    // Convert Composite Image to RGBA32F format and resize so both images have
+    // the source image size.
+    const bool bHaveSimilarAspect = FMath::IsNearlyEqual(
+                float(SizeX) / float(SizeY), 
+                float(Image.SizeX) / float(Image.SizeY), 
+                KINDA_SMALL_NUMBER);
+   
+    if (SizeX < Image.SizeX || SizeY < Image.SizeY || !bHaveSimilarAspect)
+    {
+        return EUnrealToMutableConversionError::CompositeImageDimensionMismatch;
+    }
+
+    {
+        FImage TempImage;
+        CompositeImage.ResizeTo(
+            TempImage, Image.SizeX, Image.SizeY, ERawImageFormat::RGBA32F, EGammaSpace::Linear);
+        
+        Exchange(CompositeImage, TempImage);
+    }
+
+    TArrayView64<FLinearColor> ImageView = Image.AsRGBA32F();
+    TArrayView64<FLinearColor> CompositeImageView = CompositeImage.AsRGBA32F();
+   
+    const size_t OutChannelOffset = [CompositeTextureMode]() -> size_t
+        {
+            switch(CompositeTextureMode)
+            {
+                case CTM_NormalRoughnessToRed:   return offsetof(FLinearColor, R);
+                case CTM_NormalRoughnessToGreen: return offsetof(FLinearColor, G);
+                case CTM_NormalRoughnessToBlue:  return offsetof(FLinearColor, B);
+                case CTM_NormalRoughnessToAlpha: return offsetof(FLinearColor, A);
+            }
+
+            check(false);
+            return 0;
+        }(); 
+
+    const int64 NumPixels = ImageView.Num();
+    for (int64 I = 0; I < NumPixels; ++I)
+    {
+        const FVector Normal = FVector( 
+                CompositeImageView[I].R * 2.0f - 1.0f,
+                CompositeImageView[I].G * 2.0f - 1.0f,
+                CompositeImageView[I].B * 2.0f - 1.0f); 
+
+        // Is that C++ undefined behaviour?
+        float* Value = reinterpret_cast<float*>(
+                reinterpret_cast<uint8*>(&ImageView[I]) + OutChannelOffset);
+
+        // See TextureCompressorModule.cpp:1924 for details. 
+        // Toksvig estimation of variance
+        float LengthN = FMath::Min( Normal.Size(), 1.0f );
+        float Variance = ( 1.0f - LengthN ) / LengthN;
+        Variance = FMath::Max( 0.0f, Variance - 0.00004f );
+
+        Variance *= CompositePower;
+        
+        float Roughness = *Value;
+
+        float a = Roughness * Roughness;
+        float a2 = a * a;
+        float B = 2.0f * Variance * (a2 - 1.0f);
+        a2 = ( B - a2 ) / ( B - 1.0f );
+        Roughness = FMath::Pow( a2, 0.25f );
+        
+        *Value = Roughness;
+    }
+
+    return EUnrealToMutableConversionError::Success;
+}
+
+void FlipGreenChannelRGBA32F(FImage& Image)
+{
+	TArrayView64<FLinearColor> ImageDataView = Image.AsRGBA32F();
+	ParallelFor(ImageDataView.Num(),
+		[&ImageDataView](uint32 p)
+		{
+			ImageDataView[p].G = 1.0f - FMath::Clamp(ImageDataView[p].G, 0.0f, 1.0f);
+		});
+}
+
+void FlipGreenChannelBGRA8(FImage& Image)
+{
+	TArrayView64<FColor> ImageDataView = Image.AsBGRA8();
+	ParallelFor(ImageDataView.Num(),
+		[&ImageDataView](uint32 p)
+		{
+			ImageDataView[p].G = 255 - ImageDataView[p].G;
+		});
+}
+
+void Normalize(FImage& Image)
+{
+	TArrayView64<FLinearColor> ImageView = Image.AsRGBA32F();
+
+	for (FLinearColor& Color : ImageView)
+	{
+		FVector3f Normal = (FVector3f(Color.R, Color.G, Color.B) * 2.0f - 1.0f).GetUnsafeNormal();
+
+		Color.R = Normal.X * 0.5f + 0.5f;
+		Color.G = Normal.Y * 0.5f + 0.5f;
+		Color.B = Normal.Z * 0.5f + 0.5f;
+	}
+}
+
+void BlurNormalForComposite(FImage& Image)
+{
+	TArrayView64<FLinearColor> ImageView = Image.AsRGBA32F();
+
+	const int32 SizeX = Image.SizeX;
+	const int32 SizeY = Image.SizeY;
+
+	for (int32 Y = 0; Y < SizeY - 1; ++Y)
+	{
+		for (int32 X = 0; X < SizeX - 1; ++X)
+		{
+			const int64 Idx0 = Y * SizeX + X;
+			const int64 Idx1 = Y * SizeX + (X + 1);
+			const int64 Idx2 = (Y + 1) * SizeX + X;
+			const int64 Idx3 = (Y + 1) * SizeX + (X + 1);
+
+			// Simple 2x2 box filter in place to gather info about the top mip normals variance.
+			ImageView[Idx0] = (ImageView[Idx0] + ImageView[Idx1] + ImageView[Idx2] + ImageView[Idx3]) * 0.25f;
+		}
+	}
+}
+
+} //namespace UnrealToMutableImageConversion_Internal
+
+
+EUnrealToMutableConversionError ConvertTextureUnrealSourceToMutable(mu::Image* OutResult, UTexture2D* Texture, bool bIsNormalComposite, uint8 MipmapsToSkip)
+{
+	MUTABLE_CPUPROFILER_SCOPE(ConvertTextureUnrealToMutableTuple);
+
+    using namespace UnrealToMutableImageConversion_Interanl;
+
+	// Correct mips to skip to fit source data
+	MipmapsToSkip = FMath::Clamp(MipmapsToSkip, 0, Texture->Source.GetNumMips()-1);
+
+    const int32 LODs = 1;
+    const int32 SizeX = Texture->Source.GetSizeX() >> MipmapsToSkip;
+    const int32 SizeY = Texture->Source.GetSizeY() >> MipmapsToSkip;
+	check(SizeX > 0 && SizeY > 0);
+
+    ETextureSourceFormat Format = Texture->Source.GetFormat();
+ 
+    ERawImageFormat::Type RawFormat = ConvertFormatSourceToRaw(Format);
+
+	// Not true, we will convert it.
+	// \TODO: Warn?
+    //if (RawFormat == ERawImageFormat::RGBA32F)
+    //{
+    //    return MakeTuple(nullptr, EUnrealToMutableConversionError::UnsupportedFormat);
+    //}
+
+	// What if source data is not linear?
+    FImage TempImage(SizeX, SizeY, 1, RawFormat, EGammaSpace::Linear);
+	FImage TempImage2;
+
+    if (!Texture->Source.GetMipData(TempImage.RawData, MipmapsToSkip))
+    {
+        return EUnrealToMutableConversionError::Unknown;
+    }
+
+    bool bFlipGreenChannel = Texture->bFlipGreenChannel;
+
+    // If any post processes of the image is needed, convert to RGBA32F
+    if (bIsNormalComposite)
+    { 
+		MUTABLE_CPUPROFILER_SCOPE(FlipOrComposite);
+
+        TempImage.CopyTo(TempImage2, ERawImageFormat::RGBA32F, EGammaSpace::Linear);
+		RawFormat = ERawImageFormat::RGBA32F;
+
+        if (bFlipGreenChannel)
+        {
+			FlipGreenChannelRGBA32F(TempImage2);
+			// Don't flip again below.
+			bFlipGreenChannel = false;
+		}
+
+		// Prepare texture for use as normal composite.
+		Normalize(TempImage2);
+		BlurNormalForComposite(TempImage2);
+
+        // The result is needed to TempImage
+        // Swap internals so the memory allocations is potentially reused.
+		TempImage2.Swap(TempImage);
+    }
+   
+    const ERawImageFormat::Type MutableCompatibleFormat = Format == TSF_G8 
+            ? ERawImageFormat::G8 
+            : ERawImageFormat::BGRA8;
+
+	if (MutableCompatibleFormat != RawFormat)
+	{
+		MUTABLE_CPUPROFILER_SCOPE(ToCompatibleFormat);
+		TempImage.CopyTo(TempImage2, MutableCompatibleFormat, EGammaSpace::Linear);
+		TempImage2.Swap(TempImage);
+	}
+	
+	if (bFlipGreenChannel)
+	{
+		FlipGreenChannelBGRA8(TempImage);
+	}
+
+	switch (MutableCompatibleFormat)
+	{
+	case ERawImageFormat::G8:
+	{
+		MUTABLE_CPUPROFILER_SCOPE(NoConvert);
+		OutResult->Init(SizeX, SizeY, LODs, mu::EImageFormat::IF_L_UBYTE, mu::EInitializationType::NotInitialized);
+		OutResult->m_data = MoveTemp(TempImage.RawData);
+		break;
+	}
+
+	case ERawImageFormat::BGRA8:
+	{
+		// Try to find out if the texture has and actually makes use of the alpha channel
+		bool bHasAlphaChannel = 
+			Texture->AdjustMinAlpha != Texture->AdjustMaxAlpha
+			&& Texture->CompressionSettings != TextureCompressionSettings::TC_Normalmap
+			&& !Texture->CompressionNoAlpha
+			&& (Texture->CompressionForceAlpha 
+				||
+				FImageCore::DetectAlphaChannel(TempImage));
+
+		// TODO: If we ever manage to get Pixel Format data on cook compilation time, remove the code that sets bHasAlphaChannel and just use Texture->HasAlphaChannel() here. Currently unreliable, it always returns EPixelFormat::PF_Unknown when cooking, which returns always "false" to HasAlphaChannel().
+		if (bHasAlphaChannel)
+		{
+			MUTABLE_CPUPROFILER_SCOPE(ToRGBA);
+			OutResult->Init(SizeX, SizeY, LODs, mu::EImageFormat::IF_RGBA_UBYTE, mu::EInitializationType::NotInitialized);
+			uint8* DataDest = OutResult->GetData();
+
+			// Convert to RGBA8 in place
+			TArrayView64<FColor> ImageDataView = TempImage.AsBGRA8();
+			ParallelFor(ImageDataView.Num(),
+				[DataDest, &ImageDataView](uint32 p)
+				{
+					DataDest[4 * p + 0] = ImageDataView[p].R;
+					DataDest[4 * p + 1] = ImageDataView[p].G;
+					DataDest[4 * p + 2] = ImageDataView[p].B;
+					DataDest[4 * p + 3] = ImageDataView[p].A;
+				});
+		}
+		else
+		{
+			MUTABLE_CPUPROFILER_SCOPE(ToRGB);
+
+			// TODO: add support for a mu::IF_RGBX_UBYTE?			
+			OutResult->Init(SizeX, SizeY, LODs, mu::EImageFormat::IF_RGB_UBYTE, mu::EInitializationType::NotInitialized);
+			uint8* DataDest = OutResult->GetData();
+
+			// Convert to RGB8 in place
+			TArrayView64<FColor> ImageDataView = TempImage.AsBGRA8();
+			ParallelFor(ImageDataView.Num(),
+				[DataDest, &ImageDataView](uint32 p)
+				{
+					DataDest[3 * p + 0] = ImageDataView[p].R;
+					DataDest[3 * p + 1] = ImageDataView[p].G;
+					DataDest[3 * p + 2] = ImageDataView[p].B;
+				});
+		}
+
+		//FString Msg = FString::Printf(TEXT("Alpha channel is %s for %s"), bHasAlphaChannel ? TEXT("enabled") : TEXT("disabled"), *Texture->GetName());
+		//UE_LOG(LogMutable, VeryVerbose, TEXT("%s"), *Msg);
+
+		break;
+	}
+
+	default:
+		// Format not supported yet?
+		check(false);
+		break;
+	}
+
+    return EUnrealToMutableConversionError::Success;
+}
+
+#endif // WITH_EDITOR
diff --git a/Engine/Plugins/Experimental/Mutable/Source/CustomizableObject/Private/MuCO/UnrealToMutableTextureConversionUtils.h b/Engine/Plugins/Experimental/Mutable/Source/CustomizableObject/Private/MuCO/UnrealToMutableTextureConversionUtils.h
new file mode 100644
--- /dev/null
+++ b/Engine/Plugins/Experimental/Mutable/Source/CustomizableObject/Private/MuCO/UnrealToMutableTextureConversionUtils.h
@@ -0,0 +1,27 @@
+// Copyright Epic Games, Inc. All Rights Reserved.
+
+#pragma once
+
+#if WITH_EDITOR
+
+namespace mu
+{
+	class Image;
+}
+
+// Forward declarations
+class UTexture2D;
+
+enum class EUnrealToMutableConversionError 
+{
+    Success,
+    UnsupportedFormat,
+    CompositeImageDimensionMismatch,
+    CompositeUnsupportedFormat,
+    Unknown
+};
+
+
+CUSTOMIZABLEOBJECT_API EUnrealToMutableConversionError ConvertTextureUnrealSourceToMutable(mu::Image* OutResult, UTexture2D* Texture, bool bIsNormalComposite, uint8 MipmapsToSkip);
+
+#endif
diff --git a/Engine/Plugins/Experimental/Mutable/Source/CustomizableObjectEditor/Private/MuCOE/CustomizableObjectCompileRunnable.cpp b/Engine/Plugins/Experimental/Mutable/Source/CustomizableObjectEditor/Private/MuCOE/CustomizableObjectCompileRunnable.cpp
--- a/Engine/Plugins/Experimental/Mutable/Source/CustomizableObjectEditor/Private/MuCOE/CustomizableObjectCompileRunnable.cpp
+++ b/Engine/Plugins/Experimental/Mutable/Source/CustomizableObjectEditor/Private/MuCOE/CustomizableObjectCompileRunnable.cpp
@@ -48,10 +48,10 @@ uint32 FCustomizableObjectCompileRunnable::Run()
 
 	CompilerOptions->SetUseDiskCache(Options.bUseDiskCompilation);
 
-	if (Options.OptimizationLevel > 3)
+	if (Options.OptimizationLevel > 2)
 	{
 		UE_LOG(LogMutable, Log, TEXT("Mutable compile optimization level out of range. Clamping to maximum."));
-		Options.OptimizationLevel = 3;
+		Options.OptimizationLevel = 2;
 	}
 
 	switch (Options.OptimizationLevel)
@@ -69,21 +69,15 @@ uint32 FCustomizableObjectCompileRunnable::Run()
 		break;
 
 	case 2:
-		CompilerOptions->SetOptimisationEnabled(true);
-		CompilerOptions->SetConstReductionEnabled(true);
-		CompilerOptions->SetOptimisationMaxIteration(16);
-		break;
-
-	case 3:
 		CompilerOptions->SetOptimisationEnabled(true);
 		CompilerOptions->SetConstReductionEnabled(true);
 		CompilerOptions->SetOptimisationMaxIteration(0);
 		break;
 
 	default:
-		CompilerOptions->SetOptimisationEnabled(false);
+		CompilerOptions->SetOptimisationEnabled(true);
 		CompilerOptions->SetConstReductionEnabled(true);
-		CompilerOptions->SetOptimisationMaxIteration(1);
+		CompilerOptions->SetOptimisationMaxIteration(0);
 		break;
 	}
 
diff --git a/Engine/Plugins/Experimental/Mutable/Source/CustomizableObjectEditor/Private/MuCOE/CustomizableObjectCompiler.cpp b/Engine/Plugins/Experimental/Mutable/Source/CustomizableObjectEditor/Private/MuCOE/CustomizableObjectCompiler.cpp
--- a/Engine/Plugins/Experimental/Mutable/Source/CustomizableObjectEditor/Private/MuCOE/CustomizableObjectCompiler.cpp
+++ b/Engine/Plugins/Experimental/Mutable/Source/CustomizableObjectEditor/Private/MuCOE/CustomizableObjectCompiler.cpp
@@ -620,7 +620,7 @@ mu::NodeObjectPtr FCustomizableObjectCompiler::GenerateMutableRoot(
 			TArray<UCustomizableObjectNodeObjectGroup*> GroupNodes;
 			Object->Source->GetNodesOfClass<UCustomizableObjectNodeObjectGroup>(GroupNodes);
 
-			if (GroupNodes.Num() > 0) // Only grafs with group nodes should have child grafs
+			if (GroupNodes.Num() > 0) // Only graphs with group nodes should have child graphs
 			{
 				FAssetRegistryModule& AssetRegistryModule = FModuleManager::LoadModuleChecked<FAssetRegistryModule>("AssetRegistry");
 				ProcessChildObjectsRecursively(Object, GenerationContext);
@@ -718,10 +718,6 @@ mu::NodeObjectPtr FCustomizableObjectCompiler::GenerateMutableRoot(
 		return nullptr;
 	}
 
-	// Apply the platform specific transform
-	// Format and mips are done per-image now.
-	//MutableRoot = GenerateMutableTransform(MutableRoot, GenerationContext);
-
 	return MutableRoot;
 }
 
@@ -1271,9 +1267,9 @@ void FCustomizableObjectCompiler::CompileInternal(UCustomizableObject* Object, c
 		// Pass-through textures
 		TArray<TSoftObjectPtr<UTexture>> NewReferencedPassThroughTextures;
 
-		for (const TPair<TSoftObjectPtr<UTexture>, uint32>& Pair : GenerationContext.PassThroughTextureToIndexMap)
+		for (const TPair<TSoftObjectPtr<UTexture>, FMutableGraphGenerationContext::FGeneratedPassThroughTexture>& Pair : GenerationContext.PassThroughTextureMap)
 		{
-			check(Pair.Value == NewReferencedPassThroughTextures.Num());
+			check(Pair.Value.ID == NewReferencedPassThroughTextures.Num());
 			NewReferencedPassThroughTextures.Add(Pair.Key);
 		}
 
@@ -1450,7 +1446,7 @@ void FCustomizableObjectCompiler::CompileInternal(UCustomizableObject* Object, c
 }
 
 
-mu::NodePtr FCustomizableObjectCompiler::Export(UCustomizableObject* Object, const FCompilationOptions& InCompilerOptions)
+mu::NodePtr FCustomizableObjectCompiler::Export(UCustomizableObject* Object, const FCompilationOptions& InCompilerOptions, TArray<TSoftObjectPtr<UTexture>>& OutReferencedTextures)
 {
 	UE_LOG(LogMutable, Log, TEXT("Started Customizable Object Export %s."), *Object->GetName());
 
@@ -1493,6 +1489,14 @@ mu::NodePtr FCustomizableObjectCompiler::Export(UCustomizableObject* Object, con
 		UpdatePendingTextureConversion(false);
 	}
 
+	// Pass out the references textures
+	OutReferencedTextures.Empty();
+	for (const TPair<TSoftObjectPtr<UTexture>, FMutableGraphGenerationContext::FGeneratedPassThroughTexture>& Pair : GenerationContext.PassThroughTextureMap)
+	{
+		check(Pair.Value.ID == OutReferencedTextures.Num());
+		OutReferencedTextures.Add(Pair.Key);
+	}
+
 	return MutableRoot;
 }
 
diff --git a/Engine/Plugins/Experimental/Mutable/Source/CustomizableObjectEditor/Private/MuCOE/CustomizableObjectCompiler.h b/Engine/Plugins/Experimental/Mutable/Source/CustomizableObjectEditor/Private/MuCOE/CustomizableObjectCompiler.h
--- a/Engine/Plugins/Experimental/Mutable/Source/CustomizableObjectEditor/Private/MuCOE/CustomizableObjectCompiler.h
+++ b/Engine/Plugins/Experimental/Mutable/Source/CustomizableObjectEditor/Private/MuCOE/CustomizableObjectCompiler.h
@@ -43,7 +43,7 @@ public:
 	virtual UCustomizableObject* GetRootObject(UCustomizableObject* InObject) override;
 
 	/** Generate the Mutable Graph from the Unreal Graph. */
-	mu::NodePtr Export(UCustomizableObject* Object, const FCompilationOptions& Options);
+	mu::NodePtr Export(UCustomizableObject* Object, const FCompilationOptions& Options, TArray<TSoftObjectPtr<UTexture>>& OutReferencedTextures);
 
 	void CompilerLog(const FText& Message, const TArray<const UCustomizableObjectNode*>& ArrayNode, const EMessageSeverity::Type MessageSeverity = EMessageSeverity::Warning, const bool bAddBaseObjectInfo = true, const ELoggerSpamBin SpamBin = ELoggerSpamBin::ShowAll);
 	void CompilerLog(const FText& Message, const UCustomizableObjectNode* Node = nullptr, const EMessageSeverity::Type MessageSeverity = EMessageSeverity::Warning, const bool bAddBaseObjectInfo = true, const ELoggerSpamBin SpamBin = ELoggerSpamBin::ShowAll);
diff --git a/Engine/Plugins/Experimental/Mutable/Source/CustomizableObjectEditor/Private/MuCOE/CustomizableObjectEditor.cpp b/Engine/Plugins/Experimental/Mutable/Source/CustomizableObjectEditor/Private/MuCOE/CustomizableObjectEditor.cpp
--- a/Engine/Plugins/Experimental/Mutable/Source/CustomizableObjectEditor/Private/MuCOE/CustomizableObjectEditor.cpp
+++ b/Engine/Plugins/Experimental/Mutable/Source/CustomizableObjectEditor/Private/MuCOE/CustomizableObjectEditor.cpp
@@ -1319,7 +1319,6 @@ TSharedRef<SWidget> FCustomizableObjectEditor::GenerateCompileOptionsMenuContent
 		CompileOptimizationStrings.Empty();
 		CompileOptimizationStrings.Add(MakeShareable(new FString(LOCTEXT("OptimizationNone", "None").ToString())));
 		CompileOptimizationStrings.Add(MakeShareable(new FString(LOCTEXT("OptimizationMin", "Minimal").ToString())));
-		CompileOptimizationStrings.Add(MakeShareable(new FString(LOCTEXT("OptimizationMed", "Medium").ToString())));
 		CompileOptimizationStrings.Add(MakeShareable(new FString(LOCTEXT("OptimizationMax", "Maximum").ToString())));
 
 		if (CustomizableObject)
diff --git a/Engine/Plugins/Experimental/Mutable/Source/CustomizableObjectEditor/Private/MuCOE/GenerateMutableSource/GenerateMutableSource.h b/Engine/Plugins/Experimental/Mutable/Source/CustomizableObjectEditor/Private/MuCOE/GenerateMutableSource/GenerateMutableSource.h
--- a/Engine/Plugins/Experimental/Mutable/Source/CustomizableObjectEditor/Private/MuCOE/GenerateMutableSource/GenerateMutableSource.h
+++ b/Engine/Plugins/Experimental/Mutable/Source/CustomizableObjectEditor/Private/MuCOE/GenerateMutableSource/GenerateMutableSource.h
@@ -540,10 +540,16 @@ struct FMutableGraphGenerationContext
 
 	// Cache of generated images, because sometimes they are reused by LOD, we use this as a second
 	// level cache
-	TMap<FGeneratedImageKey, mu::NodeImageConstantPtr> GeneratedImages;
+	TMap<FGeneratedImageKey, mu::NodeImagePtr> GeneratedImages;
 
+	/** Data stored per-generated passthrough texture. */
+	struct FGeneratedPassThroughTexture
+	{
+		uint32 ID;
+		mu::FImageDesc ImageDesc;
+	};
 	// Cache of pass-through images and their IDs used in the core to indentify them
-	TMap<TSoftObjectPtr<UTexture>, uint32> PassThroughTextureToIndexMap;
+	TMap<TSoftObjectPtr<UTexture>, FGeneratedPassThroughTexture> PassThroughTextureMap;
 
     // Global morph selection overrides.
     TArray<FRealTimeMorphSelectionOverride> RealTimeMorphTargetsOverrides;
diff --git a/Engine/Plugins/Experimental/Mutable/Source/CustomizableObjectEditor/Private/MuCOE/GenerateMutableSource/GenerateMutableSourceImage.cpp b/Engine/Plugins/Experimental/Mutable/Source/CustomizableObjectEditor/Private/MuCOE/GenerateMutableSource/GenerateMutableSourceImage.cpp
--- a/Engine/Plugins/Experimental/Mutable/Source/CustomizableObjectEditor/Private/MuCOE/GenerateMutableSource/GenerateMutableSourceImage.cpp
+++ b/Engine/Plugins/Experimental/Mutable/Source/CustomizableObjectEditor/Private/MuCOE/GenerateMutableSource/GenerateMutableSourceImage.cpp
@@ -30,7 +30,7 @@
 #include "MuCOE/Nodes/CustomizableObjectNodeTextureSaturate.h"
 #include "MuCOE/Nodes/CustomizableObjectNodeTextureVariation.h"
 #include "MuCOE/UnrealEditorPortabilityHelpers.h"
-#include "MuCOE/UnrealToMutableTextureConversionUtils.h"
+#include "MuCO/UnrealToMutableTextureConversionUtils.h"
 #include "MuT/NodeImageBinarise.h"
 #include "MuT/NodeImageColourMap.h"
 #include "MuT/NodeImageFormat.h"
@@ -49,6 +49,7 @@
 #include "MuT/NodeImageTransform.h"
 #include "MuT/NodeImageVariation.h"
 #include "MuT/NodeImageReference.h"
+#include "MuT/NodeImageConstant.h"
 #include "MuT/NodeScalarConstant.h"
 
 #define LOCTEXT_NAMESPACE "CustomizableObjectEditor"
@@ -58,10 +59,10 @@ mu::ImagePtr ConvertTextureUnrealToMutable(UTexture2D* Texture, const UCustomiza
 {     
 	MUTABLE_CPUPROFILER_SCOPE(ConvertTextureUnrealToMutable);
 
-	mu::ImagePtr MutableImage;
+	mu::Ptr<mu::Image> MutableImage = new mu::Image;
 	EUnrealToMutableConversionError Error = EUnrealToMutableConversionError::Unknown;
 
-	Tie(MutableImage, Error) = ConvertTextureUnrealToMutable(Texture, bIsNormalComposite);
+	Error = ConvertTextureUnrealSourceToMutable(MutableImage.get(), Texture, bIsNormalComposite, 0);
 
 	switch(Error)
 	{
@@ -104,7 +105,7 @@ mu::ImagePtr ConvertTextureUnrealToMutable(UTexture2D* Texture, const UCustomiza
 }
 
 
-mu::Ptr<mu::NodeImage> ResizeToMaxTextureSize(float MaxTextureSize, const UTexture2D* BaseTexture, mu::Ptr<mu::NodeImageConstant> ImageNode)
+mu::Ptr<mu::NodeImage> ResizeToMaxTextureSize(float MaxTextureSize, const UTexture2D* BaseTexture, const mu::Ptr<mu::NodeImage>& ImageNode)
 {
 	// To scale when above MaxTextureSize if defined
 	if (MaxTextureSize > 0 && BaseTexture
@@ -156,16 +157,44 @@ mu::NodeImagePtr GenerateMutableSourceImage(const UEdGraphPin* Pin, FMutableGrap
 			
 			// Check the specific image cache
 			FGeneratedImageKey imageKey = FGeneratedImageKey(Pin);
-			mu::NodeImageConstantPtr ImageNode;
-			mu::NodeImageConstantPtr* Cached = GenerationContext.GeneratedImages.Find(imageKey);
+			mu::NodeImagePtr ImageNode;
+			mu::NodeImagePtr* Cached = GenerationContext.GeneratedImages.Find(imageKey);
 			if (Cached)
 			{
 				ImageNode = *Cached;
 			}
 			else
 			{
-				ImageNode = new mu::NodeImageConstant();
-				GenerationContext.ArrayTextureUnrealToMutableTask.Add(FTextureUnrealToMutableTask(ImageNode, BaseTexture, Node));
+				if (GenerationContext.Options.OptimizationLevel==0)
+				{
+					// In the "None" optimization level, we don't process any image, and we set them all as image references
+					// that will be loaded at instance generation time: fast compilation x slow generation.
+					FMutableGraphGenerationContext::FGeneratedPassThroughTexture* FoundIndex = GenerationContext.PassThroughTextureMap.Find(BaseTexture);
+					FMutableGraphGenerationContext::FGeneratedPassThroughTexture NewEntry;
+
+					if (!FoundIndex)
+					{
+						FoundIndex = &NewEntry;
+						NewEntry.ID = GenerationContext.PassThroughTextureMap.Num();
+						NewEntry.ImageDesc.m_size[0] = BaseTexture->Source.GetSizeX();
+						NewEntry.ImageDesc.m_size[1] = BaseTexture->Source.GetSizeY();
+						NewEntry.ImageDesc.m_lods = BaseTexture->Source.GetNumMips();
+						NewEntry.ImageDesc.m_format = mu::EImageFormat::IF_RGBA_UBYTE; //TODO: it cannot be known without actually loading the image, which we don't want to do here.
+						GenerationContext.PassThroughTextureMap.Add(BaseTexture, NewEntry);
+					}
+
+					mu::Ptr<mu::NodeImageReference> ReferenceImageNode = new mu::NodeImageReference;
+					ReferenceImageNode->SetImageReference(FoundIndex->ID,FoundIndex->ImageDesc);
+					ReferenceImageNode->SetForceLoad(true);
+					ImageNode = ReferenceImageNode;
+				}
+				else
+				{
+					mu::Ptr<mu::NodeImageConstant> ConstantImageNode = new mu::NodeImageConstant();
+					ImageNode = ConstantImageNode;
+					GenerationContext.ArrayTextureUnrealToMutableTask.Add(FTextureUnrealToMutableTask(ConstantImageNode, BaseTexture, Node));
+				}
+
 				GenerationContext.GeneratedImages.Add(imageKey, ImageNode);
 			}
 
@@ -184,8 +213,8 @@ mu::NodeImagePtr GenerateMutableSourceImage(const UEdGraphPin* Pin, FMutableGrap
 
 		GenerationContext.AddParameterNameUnique(Node, TypedNodeParam->ParameterName);
 
-		TextureNode->SetName(StringCast<ANSICHAR>(*TypedNodeParam->ParameterName).Get());
-		TextureNode->SetUid(StringCast<ANSICHAR>(*GenerationContext.GetNodeIdUnique(Node).ToString()).Get());
+		TextureNode->SetName(TypedNodeParam->ParameterName);
+		TextureNode->SetUid(GenerationContext.GetNodeIdUnique(Node).ToString());
 
 		if (TypedNodeParam->DefaultValue)
 		{
@@ -869,21 +898,22 @@ mu::NodeImagePtr GenerateMutableSourceImage(const UEdGraphPin* Pin, FMutableGrap
 		UTexture2D* BaseTexture = TypedNodePassThroughTex->Texture;
 		if (BaseTexture)
 		{
-			uint32* FoundIndex = GenerationContext.PassThroughTextureToIndexMap.Find(BaseTexture);
-			uint32 NewId;
+			FMutableGraphGenerationContext::FGeneratedPassThroughTexture* FoundIndex = GenerationContext.PassThroughTextureMap.Find(BaseTexture);
+			FMutableGraphGenerationContext::FGeneratedPassThroughTexture NewEntry;
 
 			if (!FoundIndex)
 			{
-				NewId = GenerationContext.PassThroughTextureToIndexMap.Num();
-				GenerationContext.PassThroughTextureToIndexMap.Add(BaseTexture, NewId);
-			}
-			else
-			{
-				NewId = *FoundIndex;
+				FoundIndex = &NewEntry;
+				NewEntry.ID = GenerationContext.PassThroughTextureMap.Num();
+				NewEntry.ImageDesc.m_size[0] = BaseTexture->Source.GetSizeX();
+				NewEntry.ImageDesc.m_size[1] = BaseTexture->Source.GetSizeY();
+				NewEntry.ImageDesc.m_lods = BaseTexture->Source.GetNumMips();
+				NewEntry.ImageDesc.m_format = mu::EImageFormat::IF_RGBA_UBYTE; //TODO: it cannot be known without actually loading the image, which we don't want to do here.
+				GenerationContext.PassThroughTextureMap.Add(BaseTexture, NewEntry);
 			}
 
 			mu::Ptr<mu::NodeImageReference> ImageNode = new mu::NodeImageReference;
-			ImageNode->SetImageReference(NewId);
+			ImageNode->SetImageReference(FoundIndex->ID, FoundIndex->ImageDesc);
 
 			Result = ImageNode;
 		}
diff --git a/Engine/Plugins/Experimental/Mutable/Source/CustomizableObjectEditor/Private/MuCOE/GenerateMutableSource/GenerateMutableSourceImage.h b/Engine/Plugins/Experimental/Mutable/Source/CustomizableObjectEditor/Private/MuCOE/GenerateMutableSource/GenerateMutableSourceImage.h
--- a/Engine/Plugins/Experimental/Mutable/Source/CustomizableObjectEditor/Private/MuCOE/GenerateMutableSource/GenerateMutableSourceImage.h
+++ b/Engine/Plugins/Experimental/Mutable/Source/CustomizableObjectEditor/Private/MuCOE/GenerateMutableSource/GenerateMutableSourceImage.h
@@ -22,7 +22,7 @@ namespace mu
 mu::Ptr<mu::Image> ConvertTextureUnrealToMutable(UTexture2D* Texture, const UCustomizableObjectNode* Node, FCustomizableObjectCompiler* Compiler, bool bIsNormalComposite);
 
 
-mu::Ptr<mu::NodeImage> ResizeToMaxTextureSize(float MaxTextureSize, const UTexture2D* BaseTexture, mu::Ptr<mu::NodeImageConstant> ImageNode);
+mu::Ptr<mu::NodeImage> ResizeToMaxTextureSize(float MaxTextureSize, const UTexture2D* BaseTexture, const mu::Ptr<mu::NodeImage>& ImageNode);
 
 
 /** Convert a CustomizableObject Source Graph from an Image pin into a mutable source graph. */
diff --git a/Engine/Plugins/Experimental/Mutable/Source/CustomizableObjectEditor/Private/MuCOE/GenerateMutableSource/GenerateMutableSourceSurface.cpp b/Engine/Plugins/Experimental/Mutable/Source/CustomizableObjectEditor/Private/MuCOE/GenerateMutableSource/GenerateMutableSourceSurface.cpp
--- a/Engine/Plugins/Experimental/Mutable/Source/CustomizableObjectEditor/Private/MuCOE/GenerateMutableSource/GenerateMutableSourceSurface.cpp
+++ b/Engine/Plugins/Experimental/Mutable/Source/CustomizableObjectEditor/Private/MuCOE/GenerateMutableSource/GenerateMutableSourceSurface.cpp
@@ -717,8 +717,8 @@ mu::NodeSurfacePtr GenerateMutableSourceSurface(const UEdGraphPin * Pin, FMutabl
 									EPixelFormat UnrealTargetPlatformFormatAlpha = TextureFormat->GetEncodedPixelFormat(BuildSettingsPerLayer[0], true);
 
 									// \TODO: The QualityFix filter is used while the internal mutable runtime compression doesn't provide enough quality for some large block formats.
-									mu::EImageFormat MutableFormat = QualityFix(UnrealToMutablePixelFormat(UnrealTargetPlatformFormat,false));
-									mu::EImageFormat MutableFormatIfAlpha = QualityFix(UnrealToMutablePixelFormat(UnrealTargetPlatformFormatAlpha,true));
+									mu::EImageFormat MutableFormat = QualityAndPerformanceFix(UnrealToMutablePixelFormat(UnrealTargetPlatformFormat,false));
+									mu::EImageFormat MutableFormatIfAlpha = QualityAndPerformanceFix(UnrealToMutablePixelFormat(UnrealTargetPlatformFormatAlpha,true));
 
 									// Unsupported format: look for something generic
 									if (MutableFormat == mu::EImageFormat::IF_NONE)
diff --git a/Engine/Plugins/Experimental/Mutable/Source/CustomizableObjectEditor/Private/MuCOE/GenerateMutableSource/GenerateMutableSourceTable.cpp b/Engine/Plugins/Experimental/Mutable/Source/CustomizableObjectEditor/Private/MuCOE/GenerateMutableSource/GenerateMutableSourceTable.cpp
--- a/Engine/Plugins/Experimental/Mutable/Source/CustomizableObjectEditor/Private/MuCOE/GenerateMutableSource/GenerateMutableSourceTable.cpp
+++ b/Engine/Plugins/Experimental/Mutable/Source/CustomizableObjectEditor/Private/MuCOE/GenerateMutableSource/GenerateMutableSourceTable.cpp
@@ -367,8 +367,6 @@ bool FillTableColumn(const UCustomizableObjectNodeTable* TableNode,	mu::TablePtr
 
 		else if (SoftObjectProperty->PropertyClass->IsChildOf(UTexture::StaticClass()))
 		{
-			UTexture* Texture = nullptr;
-
 			// Two supported texture types
 			UTexture2D* Texture2D = Cast<UTexture2D>(Object);
 			UTexture2DArray* TextureArray = Cast<UTexture2DArray>(Object);
@@ -383,7 +381,9 @@ bool FillTableColumn(const UCustomizableObjectNodeTable* TableNode,	mu::TablePtr
 				GenerationContext.Compiler->CompilerLog(FText::FromString(WarningMessage), TableNode);
 			}
 
+			// There will be always one of the two options
 			check(Texture2D || TextureArray);
+			UTexture* Texture = Texture2D ? Cast<UTexture>(Texture2D) : Cast<UTexture>(TextureArray);
 
 			// Getting column index from column name
 			CurrentColumn = MutableTable->FindColumn(StringCast<ANSICHAR>(*ColumnName).Get());
@@ -395,23 +395,22 @@ bool FillTableColumn(const UCustomizableObjectNodeTable* TableNode,	mu::TablePtr
 
 			if (TableNode->GetColumnImageMode(ColumnName) == ETableTextureType::PASSTHROUGH_TEXTURE)
 			{
-				// There will be always one of the two options
-				Texture = Texture2D ? Cast<UTexture>(Texture2D) : Cast<UTexture>(TextureArray);
 
-				uint32* FoundIndex = GenerationContext.PassThroughTextureToIndexMap.Find(Texture);
-				uint32 ImageReferenceID;
+				FMutableGraphGenerationContext::FGeneratedPassThroughTexture* FoundIndex = GenerationContext.PassThroughTextureMap.Find(Texture);
+				FMutableGraphGenerationContext::FGeneratedPassThroughTexture NewEntry;
 
 				if (!FoundIndex)
 				{
-					ImageReferenceID = GenerationContext.PassThroughTextureToIndexMap.Num();
-					GenerationContext.PassThroughTextureToIndexMap.Add(Texture, ImageReferenceID);
-				}
-				else
-				{
-					ImageReferenceID = *FoundIndex;
+					FoundIndex = &NewEntry;
+					NewEntry.ID = GenerationContext.PassThroughTextureMap.Num();
+					NewEntry.ImageDesc.m_size[0] = Texture->Source.GetSizeX();
+					NewEntry.ImageDesc.m_size[1] = Texture->Source.GetSizeY();
+					NewEntry.ImageDesc.m_lods = Texture->Source.GetNumMips();
+					NewEntry.ImageDesc.m_format = mu::EImageFormat::IF_RGBA_UBYTE; //TODO: it cannot be known without actually loading the image, which we don't want to do here.
+					GenerationContext.PassThroughTextureMap.Add(Texture, NewEntry);
 				}
 
-				mu::Ptr<mu::ResourceProxyMemory<mu::Image>> Proxy = new mu::ResourceProxyMemory<mu::Image>(mu::Image::CreateAsReference(ImageReferenceID));
+				mu::Ptr<mu::ResourceProxyMemory<mu::Image>> Proxy = new mu::ResourceProxyMemory<mu::Image>(mu::Image::CreateAsReference(FoundIndex->ID,FoundIndex->ImageDesc, false));
 				MutableTable->SetCell(CurrentColumn, RowIdx, Proxy.get());
 			}
 			else 
@@ -419,9 +418,34 @@ bool FillTableColumn(const UCustomizableObjectNodeTable* TableNode,	mu::TablePtr
 				if (Texture2D)
 				{
 					GenerationContext.AddParticipatingObject(*Texture2D);
+				}
+
+				if (GenerationContext.Options.OptimizationLevel == 0)
+				{
+					// In the "None" optimization level, we don't process any image, and we set them all as image references
+					// that will be loaded at instance generation time: fast compilation x slow generation.
+					FMutableGraphGenerationContext::FGeneratedPassThroughTexture* FoundIndex = GenerationContext.PassThroughTextureMap.Find(Texture);
+					FMutableGraphGenerationContext::FGeneratedPassThroughTexture NewEntry;
 
+					if (!FoundIndex)
+					{
+						FoundIndex = &NewEntry;
+						NewEntry.ID = GenerationContext.PassThroughTextureMap.Num();
+						NewEntry.ImageDesc.m_size[0] = Texture->Source.GetSizeX();
+						NewEntry.ImageDesc.m_size[1] = Texture->Source.GetSizeY();
+						NewEntry.ImageDesc.m_lods = Texture->Source.GetNumMips();
+						NewEntry.ImageDesc.m_format = mu::EImageFormat::IF_RGBA_UBYTE; //TODO: it cannot be known without actually loading the image, which we don't want to do here.
+						GenerationContext.PassThroughTextureMap.Add(Texture, NewEntry);
+					}
+
+					mu::Ptr<mu::ResourceProxyMemory<mu::Image>> Proxy = new mu::ResourceProxyMemory<mu::Image>(mu::Image::CreateAsReference(FoundIndex->ID, FoundIndex->ImageDesc, true));
+					MutableTable->SetCell(CurrentColumn, RowIdx, Proxy.get());
+				}
+				else
+				{
 					GenerationContext.ArrayTextureUnrealToMutableTask.Add(FTextureUnrealToMutableTask(MutableTable, Texture2D, TableNode, CurrentColumn, RowIdx));
 				}
+
 			}
 		}
 
diff --git a/Engine/Plugins/Experimental/Mutable/Source/CustomizableObjectEditor/Private/MuCOE/SMutableCodeViewer.cpp b/Engine/Plugins/Experimental/Mutable/Source/CustomizableObjectEditor/Private/MuCOE/SMutableCodeViewer.cpp
--- a/Engine/Plugins/Experimental/Mutable/Source/CustomizableObjectEditor/Private/MuCOE/SMutableCodeViewer.cpp
+++ b/Engine/Plugins/Experimental/Mutable/Source/CustomizableObjectEditor/Private/MuCOE/SMutableCodeViewer.cpp
@@ -9,6 +9,7 @@
 #include "Framework/Views/TableViewMetadata.h"
 #include "Misc/Paths.h"
 #include "MuCO/CustomizableObject.h"
+#include "MuCO/UnrealToMutableTextureConversionUtils.h"
 #include "MuCOE/SMutableBoolViewer.h"
 #include "MuCOE/SMutableColorViewer.h"
 #include "MuCOE/SMutableConstantsWidget.h"
@@ -309,9 +310,10 @@ void SMutableCodeViewer::ClearSelectedTreeRow() const
 	TreeView->ClearSelection();
 }
 
-void SMutableCodeViewer::SetCurrentModel(const TSharedPtr<mu::Model, ESPMode::ThreadSafe>& InMutableModel)
+void SMutableCodeViewer::SetCurrentModel(const TSharedPtr<mu::Model, ESPMode::ThreadSafe>& InMutableModel, const TArray<TSoftObjectPtr<UTexture>>& InReferencedTextures)
 {
 	MutableModel = InMutableModel;
+	ReferencedTextures = InReferencedTextures;
 	PreviewParameters = mu::Model::NewParameters(MutableModel);
 
 	RootNodes.Empty();
@@ -356,13 +358,14 @@ void SMutableCodeViewer::SetCurrentModel(const TSharedPtr<mu::Model, ESPMode::Th
 }
 
 
-void SMutableCodeViewer::Construct(const FArguments& InArgs, const TSharedPtr<mu::Model, ESPMode::ThreadSafe>& InMutableModel/*, const TSharedPtr<SDockTab>& ConstructUnderMajorTab*/)
+void SMutableCodeViewer::Construct(const FArguments& InArgs, const TSharedPtr<mu::Model, ESPMode::ThreadSafe>& InMutableModel,
+	const TArray<TSoftObjectPtr<UTexture>>& InReferencedTextures )
 {
 	// Min width allowed for the column. Needed to avoid having issues with the constants space being to small
 	// and then getting too tall on the y axis crashing the UI drawer.
 	constexpr float MinParametersCollWidth = 400;
 
-	SetCurrentModel(InMutableModel);
+	SetCurrentModel(InMutableModel, InReferencedTextures);
 	
 	FToolBarBuilder ToolbarBuilder(TSharedPtr<const FUICommandList>(), FMultiBoxCustomization::None, TSharedPtr<FExtender>(), true);
 	ToolbarBuilder.SetLabelVisibility(EVisibility::Visible);
@@ -2457,8 +2460,15 @@ namespace
 	{
 		static inline const mu::FImageDesc IMAGE_DESC = 
 			mu::FImageDesc(mu::FImageSize(1024, 1024), mu::EImageFormat::IF_RGBA_UBYTE, 1);
-		
+
 	public:
+
+		/** */
+		TArray<TSoftObjectPtr<UTexture>> ReferencedTextures;
+
+	public:
+
+
 #ifdef MUTABLE_USE_NEW_TASKGRAPH
 		TTuple<UE::Tasks::FTask, TFunction<void()>> GetImageAsync(FName Id, uint8 MipmapsToSkip, TFunction<void(mu::Ptr<mu::Image>)>& ResultCallback) override
 #else
@@ -2517,6 +2527,53 @@ namespace
 		{
 			return IMAGE_DESC;
 		}
+
+
+		//-------------------------------------------------------------------------------------------------
+#ifdef MUTABLE_USE_NEW_TASKGRAPH
+		TTuple<UE::Tasks::FTask, TFunction<void()>> GetReferencedImageAsync(const void* ModelPtr, int32 Id, uint8 MipmapsToSkip, TFunction<void(mu::Ptr<mu::Image>)>& ResultCallback)
+#else
+		TTuple<FGraphEventRef, TFunction<void()>> GetReferencedImageAsync(const void* ModelPtr, int32 Id, uint8 MipmapsToSkip, TFunction<void(mu::Ptr<mu::Image>)>& ResultCallback)
+#endif
+		{
+			check(ReferencedTextures.IsValidIndex(Id));
+
+			TSoftObjectPtr<UTexture> TexturePtr = ReferencedTextures[Id].Get();
+			UTexture2D* Texture = Cast<UTexture2D>( TexturePtr.Get() );
+			check(Texture);
+			
+			// In the editor the src data can be directly accessed
+			int32 MipIndex = (MipmapsToSkip < Texture->GetPlatformData()->Mips.Num()) ? MipmapsToSkip : Texture->GetPlatformData()->Mips.Num() - 1;
+			check(MipIndex >= 0);
+
+			mu::Ptr<mu::Image> ResultImage = new mu::Image();
+			bool bIsNormalComposite = false; // TODO?
+			EUnrealToMutableConversionError Error = ConvertTextureUnrealSourceToMutable(ResultImage.get(), Texture, bIsNormalComposite, MipmapsToSkip);
+			check(Error == EUnrealToMutableConversionError::Success);
+
+			ResultCallback(ResultImage);
+
+#ifdef MUTABLE_USE_NEW_TASKGRAPH
+			auto TrivialReturn = []() -> TTuple<UE::Tasks::FTask, TFunction<void()>>
+			{
+				UE::Tasks::FTaskEvent CompletionEvent(TEXT("GetImageAsyncCompleted"));
+				CompletionEvent.Trigger();
+
+				return MakeTuple(CompletionEvent, []() -> void {});
+			};
+#else
+			auto TrivialReturn = []() -> TTuple<FGraphEventRef, TFunction<void()>>
+			{
+				FGraphEventRef CompletionEvent = FGraphEvent::CreateGraphEvent();
+				CompletionEvent->DispatchSubsequents();
+
+				return MakeTuple(CompletionEvent, []() -> void {});
+			};
+#endif
+
+			return Invoke(TrivialReturn);
+		}
+
 	};
 }
 
@@ -2552,6 +2609,7 @@ void SMutableCodeViewer::Tick(const FGeometry& AllottedGeometry, const double In
 	const mu::Ptr<mu::System> System = new mu::System(Settings);
 
 	TSharedPtr<TestImageProvider> ImageProvider = MakeShared<TestImageProvider>();
+	ImageProvider->ReferencedTextures = ReferencedTextures;
 	System->SetImageParameterGenerator(ImageProvider);
 
 	System->GetPrivate()->BeginBuild(MutableModel);
@@ -2852,7 +2910,8 @@ FReply SMutableCodeViewer::OnDrop(const FGeometry& MyGeometry, const FDragDropEv
 					{
 						mu::InputArchive arch(&stream);
 						TSharedPtr<mu::Model, ESPMode::ThreadSafe> Model = mu::Model::StaticUnserialise(arch);
-						SetCurrentModel(Model);
+						TArray<TSoftObjectPtr<UTexture>> DummyReferencedTextures;
+						SetCurrentModel(Model, DummyReferencedTextures);
 
 						TreeView->RequestTreeRefresh();
 
diff --git a/Engine/Plugins/Experimental/Mutable/Source/CustomizableObjectEditor/Private/MuCOE/SMutableCodeViewer.h b/Engine/Plugins/Experimental/Mutable/Source/CustomizableObjectEditor/Private/MuCOE/SMutableCodeViewer.h
--- a/Engine/Plugins/Experimental/Mutable/Source/CustomizableObjectEditor/Private/MuCOE/SMutableCodeViewer.h
+++ b/Engine/Plugins/Experimental/Mutable/Source/CustomizableObjectEditor/Private/MuCOE/SMutableCodeViewer.h
@@ -56,7 +56,8 @@ public:
 		
 	SLATE_END_ARGS()
 
-	void Construct(const FArguments& InArgs, const TSharedPtr<mu::Model, ESPMode::ThreadSafe>& InMutableModel /*, const TSharedPtr<SDockTab>& ConstructUnderMajorTab*/);
+	void Construct(const FArguments& InArgs, const TSharedPtr<mu::Model, ESPMode::ThreadSafe>& InMutableModel,
+		const TArray<TSoftObjectPtr<UTexture>>& ReferencedTextures );
 
 	// SWidget interface
 	virtual void Tick(const FGeometry& AllottedGeometry, const double InCurrentTime, const float InDeltaTime) override;
@@ -75,6 +76,9 @@ private:
 	/** The Mutable Model that we are showing. */
 	TSharedPtr<mu::Model, ESPMode::ThreadSafe> MutableModel;
 	
+	/** Array of external referenced textures in MutableModel, indexed by id. */
+	TArray<TSoftObjectPtr<UTexture>> ReferencedTextures;
+
 	/** Selected model operation for preview. */
 	mu::OP::ADDRESS SelectedOperationAddress = 0;
 
@@ -178,7 +182,7 @@ private:
 	TMap< mu::OP::ADDRESS, TSharedPtr<FMutableCodeTreeElement>> ExpandedElements;
 
 	/** Prepare the widget for the given model. */
-	void SetCurrentModel(const TSharedPtr<mu::Model, ESPMode::ThreadSafe>&);
+	void SetCurrentModel(const TSharedPtr<mu::Model, ESPMode::ThreadSafe>&, const TArray<TSoftObjectPtr<UTexture>>& ReferencedTextures);
 
 	/** Before any UI operation generate all the elements that may be navigable over the tree. No children of duplicated 
 	 * addresses will be generated.
diff --git a/Engine/Plugins/Experimental/Mutable/Source/CustomizableObjectEditor/Private/MuCOE/SMutableGraphViewer.cpp b/Engine/Plugins/Experimental/Mutable/Source/CustomizableObjectEditor/Private/MuCOE/SMutableGraphViewer.cpp
--- a/Engine/Plugins/Experimental/Mutable/Source/CustomizableObjectEditor/Private/MuCOE/SMutableGraphViewer.cpp
+++ b/Engine/Plugins/Experimental/Mutable/Source/CustomizableObjectEditor/Private/MuCOE/SMutableGraphViewer.cpp
@@ -116,12 +116,14 @@ FString SMutableGraphViewer::GetReferencerName() const
 }
 
 
-void SMutableGraphViewer::Construct(const FArguments& InArgs, const mu::NodePtr& InRootNode, 
+void SMutableGraphViewer::Construct(const FArguments& InArgs, const mu::NodePtr& InRootNode,
+	const TArray<TSoftObjectPtr<UTexture>>& InReferencedTextures,
 	const FCompilationOptions& InCompileOptions,
 	TWeakPtr<FTabManager> InParentTabManager, const FName& InParentNewTabId)
 {
 	DataTag = InArgs._DataTag;
 	RootNode = InRootNode;
+	ReferencedTextures = InReferencedTextures;
 	CompileOptions = InCompileOptions;	
 	ParentTabManager = InParentTabManager;
 	ParentNewTabId = InParentNewTabId;
@@ -265,7 +267,7 @@ void SMutableGraphViewer::CompileMutableCodePressed()
 	TSharedPtr<SDockTab> NewMutableCodeTab = SNew(SDockTab)
 		.Label(LOCTEXT("MutableCode", "Mutable Code"))
 		[
-			SNew(SMutableCodeViewer, CompileTask->Model)
+			SNew(SMutableCodeViewer, CompileTask->Model, ReferencedTextures)
 			.DataTag(NewDataTag)
 		];
 
@@ -290,7 +292,6 @@ TSharedRef<SWidget> SMutableGraphViewer::GenerateCompileOptionsMenuContent()
 		CompileOptimizationStrings.Empty();
 		CompileOptimizationStrings.Add(MakeShareable(new FString(NSLOCTEXT("CustomizableObjectEditor", "OptimizationNone", "None").ToString())));
 		CompileOptimizationStrings.Add(MakeShareable(new FString(NSLOCTEXT("CustomizableObjectEditor", "OptimizationMin", "Minimal").ToString())));
-		CompileOptimizationStrings.Add(MakeShareable(new FString(NSLOCTEXT("CustomizableObjectEditor", "OptimizationMed", "Medium").ToString())));
 		CompileOptimizationStrings.Add(MakeShareable(new FString(NSLOCTEXT("CustomizableObjectEditor", "OptimizationMax", "Maximum").ToString())));
 
 		CompileOptions.OptimizationLevel = FMath::Min(CompileOptions.OptimizationLevel, CompileOptimizationStrings.Num() - 1);
diff --git a/Engine/Plugins/Experimental/Mutable/Source/CustomizableObjectEditor/Private/MuCOE/SMutableGraphViewer.h b/Engine/Plugins/Experimental/Mutable/Source/CustomizableObjectEditor/Private/MuCOE/SMutableGraphViewer.h
--- a/Engine/Plugins/Experimental/Mutable/Source/CustomizableObjectEditor/Private/MuCOE/SMutableGraphViewer.h
+++ b/Engine/Plugins/Experimental/Mutable/Source/CustomizableObjectEditor/Private/MuCOE/SMutableGraphViewer.h
@@ -37,7 +37,7 @@ public:
 
 	SLATE_END_ARGS()
 
-	void Construct(const FArguments& InArgs, const mu::NodePtr& InRootNode, const FCompilationOptions& InCompileOptions,
+	void Construct(const FArguments& InArgs, const mu::NodePtr& InRootNode, const TArray<TSoftObjectPtr<UTexture>>& ReferencedTextures, const FCompilationOptions& InCompileOptions,
 		TWeakPtr<FTabManager> InParentTabManager, const FName& InParentNewTabId);
 
 	// SWidget interface
@@ -56,6 +56,9 @@ private:
 	/** The Mutable Graph to show, represented by its root. */
 	mu::NodePtr RootNode;
 
+	/** Array of external referenced textures in MutableModel, indexed by id. */
+	TArray<TSoftObjectPtr<UTexture>> ReferencedTextures;
+
 	/** Compilation options to use in the debugger operations. */
 	FCompilationOptions CompileOptions;
 
diff --git a/Engine/Plugins/Experimental/Mutable/Source/CustomizableObjectEditor/Private/MuCOE/SMutableObjectViewer.cpp b/Engine/Plugins/Experimental/Mutable/Source/CustomizableObjectEditor/Private/MuCOE/SMutableObjectViewer.cpp
--- a/Engine/Plugins/Experimental/Mutable/Source/CustomizableObjectEditor/Private/MuCOE/SMutableObjectViewer.cpp
+++ b/Engine/Plugins/Experimental/Mutable/Source/CustomizableObjectEditor/Private/MuCOE/SMutableObjectViewer.cpp
@@ -215,7 +215,8 @@ void SMutableObjectViewer::Construct(const FArguments& InArgs, UCustomizableObje
 void SMutableObjectViewer::GenerateMutableGraphPressed()
 {
 	// Convert from Unreal graph to Mutable graph.
-	mu::Ptr<mu::Node> RootNode = Compiler.Export(CustomizableObject, CompileOptions);
+	TArray<TSoftObjectPtr<UTexture>> ReferencedTextures;
+	mu::Ptr<mu::Node> RootNode = Compiler.Export(CustomizableObject, CompileOptions, ReferencedTextures);
 	if (!RootNode)
 	{
 		// TODO: Show errors
@@ -227,7 +228,7 @@ void SMutableObjectViewer::GenerateMutableGraphPressed()
 	TSharedPtr<SDockTab> NewMutableGraphTab = SNew(SDockTab)
 		.Label(LOCTEXT("MutableGraph", "Mutable Graph"))
 		[
-			SNew(SMutableGraphViewer, RootNode, CompileOptions, ParentTabManager, ParentNewTabId)
+			SNew(SMutableGraphViewer, RootNode, ReferencedTextures, CompileOptions, ParentTabManager, ParentNewTabId)
 			.DataTag(DataTag)
 		];
 
@@ -239,6 +240,7 @@ void SMutableObjectViewer::GenerateMutableGraphPressed()
 
 void SMutableObjectViewer::CompileMutableCodePressed()
 {
+	TArray<TSoftObjectPtr<UTexture>> ReferencedTextures;
 	if (CompileOptions.bForceLargeLODBias)
 	{
 		// Debug compile with many different biasses
@@ -247,7 +249,7 @@ void SMutableObjectViewer::CompileMutableCodePressed()
 		{
 			CompileOptions.DebugBias = Bias;
 
-			mu::NodePtr RootNode = Compiler.Export(CustomizableObject, CompileOptions);
+			mu::NodePtr RootNode = Compiler.Export(CustomizableObject, CompileOptions, ReferencedTextures);
 			if (!RootNode)
 			{
 				// TODO: Show errors
@@ -265,7 +267,7 @@ void SMutableObjectViewer::CompileMutableCodePressed()
 	}
 
 	// Convert from Unreal graph to Mutable graph.
-	mu::NodePtr RootNode = Compiler.Export(CustomizableObject, CompileOptions);
+	mu::NodePtr RootNode = Compiler.Export(CustomizableObject, CompileOptions, ReferencedTextures);
 	if (!RootNode)
 	{
 		// TODO: Show errors
@@ -284,7 +286,7 @@ void SMutableObjectViewer::CompileMutableCodePressed()
 	TSharedPtr<SDockTab> NewMutableCodeTab = SNew(SDockTab)
 		.Label(LOCTEXT("MutableCode", "Mutable Code"))
 		[
-			SNew(SMutableCodeViewer, CompileTask->Model)
+			SNew(SMutableCodeViewer, CompileTask->Model, ReferencedTextures)
 			.DataTag(DataTag)
 		];
 
@@ -347,7 +349,6 @@ TSharedRef<SWidget> SMutableObjectViewer::GenerateCompileOptionsMenuContent()
 		CompileOptimizationStrings.Empty();
 		CompileOptimizationStrings.Add(MakeShareable(new FString(LOCTEXT("OptimizationNone", "None").ToString())));
 		CompileOptimizationStrings.Add(MakeShareable(new FString(LOCTEXT("OptimizationMin", "Minimal").ToString())));
-		CompileOptimizationStrings.Add(MakeShareable(new FString(LOCTEXT("OptimizationMed", "Medium").ToString())));
 		CompileOptimizationStrings.Add(MakeShareable(new FString(LOCTEXT("OptimizationMax", "Maximum").ToString())));
 
 		CompileOptions.OptimizationLevel = FMath::Min(CompileOptions.OptimizationLevel, CompileOptimizationStrings.Num() - 1);
diff --git a/Engine/Plugins/Experimental/Mutable/Source/CustomizableObjectEditor/Private/MuCOE/UnrealToMutableTextureConversionUtils.cpp b/Engine/Plugins/Experimental/Mutable/Source/CustomizableObjectEditor/Private/MuCOE/UnrealToMutableTextureConversionUtils.cpp
deleted file mode 100644
diff --git a/Engine/Plugins/Experimental/Mutable/Source/CustomizableObjectEditor/Private/MuCOE/UnrealToMutableTextureConversionUtils.h b/Engine/Plugins/Experimental/Mutable/Source/CustomizableObjectEditor/Private/MuCOE/UnrealToMutableTextureConversionUtils.h
deleted file mode 100644
diff --git a/Engine/Plugins/Experimental/Mutable/Source/MutableRuntime/Private/MuR/CodeRunner.cpp b/Engine/Plugins/Experimental/Mutable/Source/MutableRuntime/Private/MuR/CodeRunner.cpp
--- a/Engine/Plugins/Experimental/Mutable/Source/MutableRuntime/Private/MuR/CodeRunner.cpp
+++ b/Engine/Plugins/Experimental/Mutable/Source/MutableRuntime/Private/MuR/CodeRunner.cpp
@@ -151,9 +151,9 @@ namespace mu
 
     //---------------------------------------------------------------------------------------------
 #ifdef MUTABLE_USE_NEW_TASKGRAPH
-	TTuple<UE::Tasks::FTask, TFunction<void()>> CodeRunner::LoadExternalImageAsync(FName Id, uint8 MipmapsToSkip, TFunction<void(Ptr<Image>)>& ResultCallback)
+	TTuple<UE::Tasks::FTask, TFunction<void()>> CodeRunner::LoadExternalImageAsync(FExternalImageId Id, uint8 MipmapsToSkip, TFunction<void(Ptr<Image>)>& ResultCallback)
 #else
-	TTuple<FGraphEventRef, TFunction<void()>> CodeRunner::LoadExternalImageAsync(FName Id, uint8 MipmapsToSkip, TFunction<void(Ptr<Image>)>& ResultCallback)
+	TTuple<FGraphEventRef, TFunction<void()>> CodeRunner::LoadExternalImageAsync(FExternalImageId Id, uint8 MipmapsToSkip, TFunction<void(Ptr<Image>)>& ResultCallback)
 #endif
     {
 		MUTABLE_CPUPROFILER_SCOPE(LoadExternalImageAsync);
@@ -162,7 +162,16 @@ namespace mu
 
 		if (m_pSystem->ImageParameterGenerator)
 		{
-			return m_pSystem->ImageParameterGenerator->GetImageAsync(Id, MipmapsToSkip, ResultCallback);
+			if (Id.ReferenceImageId < 0)
+			{
+				// It's a parameter image
+				return m_pSystem->ImageParameterGenerator->GetImageAsync(Id.ParameterId, MipmapsToSkip, ResultCallback);
+			}
+			else
+			{
+				// It's an image reference
+				return m_pSystem->ImageParameterGenerator->GetReferencedImageAsync(m_pModel.Get(), Id.ReferenceImageId, MipmapsToSkip, ResultCallback);
+			}
 		}
 		else
 		{
@@ -4026,7 +4035,16 @@ namespace mu
 			{
 			case 0:
 			{
-				Ptr<Image> Result = Image::CreateAsReference(Args.ID);
+				Ptr<Image> Result;
+				if (Args.ForceLoad)
+				{
+					// This should never be reached because it should have been caught as a Task in IssueOp
+					check(false);
+				}
+				else
+				{
+					Result = Image::CreateAsReference(Args.ID, Args.ImageDesc, false);
+				}
 				StoreImage(item, Result);
 				break;
 			}
@@ -5901,11 +5919,9 @@ namespace mu
 		case OP_TYPE::IM_REFERENCE:
 		{
 			check(item.Stage == 0);
+			OP::ResourceReferenceArgs Args = program.GetOpArgs<OP::ResourceReferenceArgs>(item.At);
 			FImageDesc& Result = m_heapImageDesc[item.CustomState];
-			Result.m_format = EImageFormat::IF_NONE;
-			Result.m_size[0] = 0;
-			Result.m_size[1] = 0;
-			Result.m_lods = 0;
+			Result = Args.ImageDesc;
 			StoreValidDesc(item);
 			break;
 		}
diff --git a/Engine/Plugins/Experimental/Mutable/Source/MutableRuntime/Private/MuR/CodeRunner.h b/Engine/Plugins/Experimental/Mutable/Source/MutableRuntime/Private/MuR/CodeRunner.h
--- a/Engine/Plugins/Experimental/Mutable/Source/MutableRuntime/Private/MuR/CodeRunner.h
+++ b/Engine/Plugins/Experimental/Mutable/Source/MutableRuntime/Private/MuR/CodeRunner.h
@@ -136,12 +136,21 @@ namespace  mu
 
     public:
 
+		struct FExternalImageId
+		{
+			/** If it is an image reference. */
+			int32 ReferenceImageId = -1;
+
+			/** If it is an image parameter.*/
+			FName ParameterId; 
+		};
+
 		//! Load an external image asynchronously, retuns an event to wait for complition and a cleanup function 
 		//! that must be called once the event has completed.
 #ifdef MUTABLE_USE_NEW_TASKGRAPH
-		TTuple<UE::Tasks::FTask, TFunction<void()>> LoadExternalImageAsync(FName Id, uint8 MipmapsToSkip, TFunction<void(Ptr<Image>)>& ResultCallback);
+		TTuple<UE::Tasks::FTask, TFunction<void()>> LoadExternalImageAsync(FExternalImageId Id, uint8 MipmapsToSkip, TFunction<void(Ptr<Image>)>& ResultCallback);
 #else
-		TTuple<FGraphEventRef, TFunction<void()>> LoadExternalImageAsync(FName Id, uint8 MipmapsToSkip, TFunction<void(Ptr<Image>)>& ResultCallback);
+		TTuple<FGraphEventRef, TFunction<void()>> LoadExternalImageAsync(FExternalImageId Id, uint8 MipmapsToSkip, TFunction<void(Ptr<Image>)>& ResultCallback);
 #endif
  	    mu::FImageDesc GetExternalImageDesc(FName Id, uint8 MipmapsToSkip);
 
diff --git a/Engine/Plugins/Experimental/Mutable/Source/MutableRuntime/Private/MuR/CodeRunnerTasked.cpp b/Engine/Plugins/Experimental/Mutable/Source/MutableRuntime/Private/MuR/CodeRunnerTasked.cpp
--- a/Engine/Plugins/Experimental/Mutable/Source/MutableRuntime/Private/MuR/CodeRunnerTasked.cpp
+++ b/Engine/Plugins/Experimental/Mutable/Source/MutableRuntime/Private/MuR/CodeRunnerTasked.cpp
@@ -2407,13 +2407,13 @@ namespace mu
 	}
 
 
-	//---------------------------------------------------------------------------------------------
-	//---------------------------------------------------------------------------------------------
-	//---------------------------------------------------------------------------------------------
+	/** This task is used to load an image parameter (by its FName) or an image reference (from its ID).
+	*/
 	class FImageExternalLoadTask : public CodeRunner::FIssuedTask
 	{
 	public:
-		FImageExternalLoadTask(const FScheduledOp& InItem, uint8 InMipmapsToSkip, FName InId);
+
+		FImageExternalLoadTask(const FScheduledOp& InItem, uint8 InMipmapsToSkip, CodeRunner::FExternalImageId InId);
 
 		// FIssuedTask interface
 		virtual bool Prepare(CodeRunner*, bool& bOutFailed) override;
@@ -2421,7 +2421,7 @@ namespace mu
 		
 	private:
 		uint8 MipmapsToSkip;
-		FName Id;
+		CodeRunner::FExternalImageId Id;
 
 		Ptr<Image> Result;
 		
@@ -2429,15 +2429,14 @@ namespace mu
 	};
 
 
-	//---------------------------------------------------------------------------------------------
-	FImageExternalLoadTask::FImageExternalLoadTask(const FScheduledOp& InOp,  uint8 InMipmapsToSkip, FName InId)
+	FImageExternalLoadTask::FImageExternalLoadTask(const FScheduledOp& InOp, uint8 InMipmapsToSkip, CodeRunner::FExternalImageId InId)
 		: FIssuedTask(InOp)
 	{
 		MipmapsToSkip = InMipmapsToSkip;
 		Id = InId;
 	}
 
-	//---------------------------------------------------------------------------------------------
+
 	bool FImageExternalLoadTask::Prepare(CodeRunner* Runner, bool& bOutFailed)
 	{
 		// This runs in the mutable Runner thread
@@ -2461,7 +2460,7 @@ namespace mu
 		return false;
 	}
 
-	//---------------------------------------------------------------------------------------------
+
 	void FImageExternalLoadTask::Complete(CodeRunner* Runner)
 	{
 		if (ExternalCleanUpFunc)
@@ -2573,11 +2572,32 @@ namespace mu
 
 			const uint8 MipmapsToSkip = item.ExecutionOptions + static_cast<uint8>(ImageLOD);
 
-			Issued = MakeShared<FImageExternalLoadTask>(item, MipmapsToSkip, Id);
+			CodeRunner::FExternalImageId FullId;
+			FullId.ParameterId = Id;
+			Issued = MakeShared<FImageExternalLoadTask>(item, MipmapsToSkip, FullId);
 
 			break;
 		}
-			
+
+		case OP_TYPE::IM_REFERENCE:
+		{
+			OP::ResourceReferenceArgs Args = m_pModel->GetPrivate()->m_program.GetOpArgs<OP::ResourceReferenceArgs>(item.At);
+
+			// We only convert references to images if indicated in the operation.
+			if (Args.ForceLoad)
+			{
+				check(item.Stage==0);
+
+				const uint8 MipmapsToSkip = item.ExecutionOptions + static_cast<uint8>(ImageLOD);
+
+				FExternalImageId FullId;
+				FullId.ReferenceImageId = Args.ID;
+				Issued = MakeShared<FImageExternalLoadTask>(item, MipmapsToSkip, FullId);
+			}
+
+			break;
+		}
+
 		case OP_TYPE::IM_PIXELFORMAT:
 		{
 			if (item.Stage == 1)
diff --git a/Engine/Plugins/Experimental/Mutable/Source/MutableRuntime/Private/MuR/Image.cpp b/Engine/Plugins/Experimental/Mutable/Source/MutableRuntime/Private/MuR/Image.cpp
--- a/Engine/Plugins/Experimental/Mutable/Source/MutableRuntime/Private/MuR/Image.cpp
+++ b/Engine/Plugins/Experimental/Mutable/Source/MutableRuntime/Private/MuR/Image.cpp
@@ -117,29 +117,35 @@ namespace mu
 
 
     //---------------------------------------------------------------------------------------------
-    Image::Image( uint32 sizeX, uint32 sizeY, uint32 lods, EImageFormat format, EInitializationType Init )
-    {
+	Image::Image(uint32 SizeX, uint32 SizeY, uint32 Lods, EImageFormat Format, EInitializationType InitType)
+	{
+		Init(SizeX, SizeY, Lods, Format, InitType);
+	}
+
+	
+	void Image::Init(uint32 SizeX, uint32 SizeY, uint32 Lods, EImageFormat Format, EInitializationType InitType)
+	{
 		MUTABLE_CPUPROFILER_SCOPE(NewImage)
 		LLM_SCOPE_BYNAME(TEXT("MutableRuntime"));
 
-        check(format != EImageFormat::IF_NONE);
-		check(format < EImageFormat::IF_COUNT);
+        check(Format != EImageFormat::IF_NONE);
+		check(Format < EImageFormat::IF_COUNT);
 
-		check(sizeX <= TNumericLimits<uint16>::Max());
-		check(sizeY <= TNumericLimits<uint16>::Max());
-		check(lods <= TNumericLimits<uint8>::Max());
+		check(SizeX <= TNumericLimits<uint16>::Max());
+		check(SizeY <= TNumericLimits<uint16>::Max());
+		check(Lods <= TNumericLimits<uint8>::Max());
 
 		// TODO: check that lods is sensible for the size.
 
-        m_format = format;
-        m_size = FImageSize( (uint16)sizeX, (uint16)sizeY );
-        m_lods = (uint8)lods;
+        m_format = Format;
+        m_size = FImageSize( (uint16)SizeX, (uint16)SizeY );
+        m_lods = (uint8)Lods;
 
-        const FImageFormatData& fdata = GetImageFormatData( format );
+        const FImageFormatData& fdata = GetImageFormatData( Format );
         int32 PixelsPerBlock = fdata.PixelsPerBlockX*fdata.PixelsPerBlockY;
         if (PixelsPerBlock)
         {
-			if (Init == EInitializationType::Black)
+			if (InitType == EInitializationType::Black)
 			{
 				InitToBlack();
 			}
@@ -196,11 +202,18 @@ namespace mu
 	}
 
 	//---------------------------------------------------------------------------------------------
-	Ptr<Image> Image::CreateAsReference(uint32 ID)
+	Ptr<Image> Image::CreateAsReference(uint32 ID, const FImageDesc& Desc, bool bForceLoad)
 	{
 		Ptr<Image> Result = new Image;
 		Result->ReferenceID = ID;
+		Result->m_size = Desc.m_size;
+		Result->m_lods = Desc.m_lods;
+		Result->m_format = Desc.m_format;
 		Result->m_flags = EImageFlags::IF_IS_REFERENCE;
+		if (bForceLoad)
+		{
+			Result->m_flags = Result->m_flags | EImageFlags::IF_IS_FORCELOAD;
+		}
 		return Result;
 	}
 
diff --git a/Engine/Plugins/Experimental/Mutable/Source/MutableRuntime/Private/MuR/Operations.cpp b/Engine/Plugins/Experimental/Mutable/Source/MutableRuntime/Private/MuR/Operations.cpp
--- a/Engine/Plugins/Experimental/Mutable/Source/MutableRuntime/Private/MuR/Operations.cpp
+++ b/Engine/Plugins/Experimental/Mutable/Source/MutableRuntime/Private/MuR/Operations.cpp
@@ -294,10 +294,6 @@ namespace mu
             f(&op.args.ImagePlainColour.colour );
             break;
 
-        case OP_TYPE::IM_CROP:
-            f(&op.args.ImageCrop.source );
-            break;
-
         case OP_TYPE::IM_DISPLACE:
             f(&op.args.ImageDisplace.source );
             f(&op.args.ImageDisplace.displacementMap );
diff --git a/Engine/Plugins/Experimental/Mutable/Source/MutableRuntime/Private/MuR/Operations.h b/Engine/Plugins/Experimental/Mutable/Source/MutableRuntime/Private/MuR/Operations.h
--- a/Engine/Plugins/Experimental/Mutable/Source/MutableRuntime/Private/MuR/Operations.h
+++ b/Engine/Plugins/Experimental/Mutable/Source/MutableRuntime/Private/MuR/Operations.h
@@ -435,7 +435,9 @@ namespace mu
 
 		struct ResourceReferenceArgs
 		{
+			FImageDesc ImageDesc;
 			int32 ID;
+			int8 ForceLoad;
 		};
 
         //-------------------------------------------------------------------------------------
@@ -1129,7 +1131,6 @@ namespace mu
             ImageBinariseArgs ImageBinarise;
             ImagePlainColourArgs ImagePlainColour;
             ImageGPUArgs ImageGPU;
-            ImageCropArgs ImageCrop;
             ImageDisplaceArgs ImageDisplace;
 			ImageInvertArgs ImageInvert;
 
diff --git a/Engine/Plugins/Experimental/Mutable/Source/MutableRuntime/Public/MuR/Image.h b/Engine/Plugins/Experimental/Mutable/Source/MutableRuntime/Public/MuR/Image.h
--- a/Engine/Plugins/Experimental/Mutable/Source/MutableRuntime/Public/MuR/Image.h
+++ b/Engine/Plugins/Experimental/Mutable/Source/MutableRuntime/Public/MuR/Image.h
@@ -210,9 +210,9 @@ namespace mu
         //!         depends on the image size.
         //! \param format Pixel format.
         Image( uint32 sizeX, uint32 sizeY, uint32 lods, EImageFormat format, EInitializationType Init );
-		
+
 		/** */
-		static Ptr<Image> CreateAsReference( uint32 ID );
+		static Ptr<Image> CreateAsReference( uint32 ID, const FImageDesc& Desc, bool bForceLoad);
 
 		//! Serialisation
 		static void Serialise( const Image* p, OutputArchive& arch );
@@ -225,6 +225,9 @@ namespace mu
 		// Own interface
 		//-----------------------------------------------------------------------------------------
 
+		/** */
+		void Init(uint32 sizeX, uint32 sizeY, uint32 lods, EImageFormat format, EInitializationType Init);
+
 		//! Return the width of the image.
         uint16 GetSizeX() const;
 
@@ -293,6 +296,9 @@ namespace mu
 
 			/** If this is set, this is a reference to an external image, and the ReferenceID is valid. */
 			IF_IS_REFERENCE = 1 << 4,
+
+			/** For reference images, this indicates that they should be loaded into full images as soon as they are generated. */
+			IF_IS_FORCELOAD = 1 << 5
 		} EImageFlags;
 
 		/** Persistent flags with some image properties. The meaning will depend of every context. */
diff --git a/Engine/Plugins/Experimental/Mutable/Source/MutableRuntime/Public/MuR/Model.h b/Engine/Plugins/Experimental/Mutable/Source/MutableRuntime/Public/MuR/Model.h
--- a/Engine/Plugins/Experimental/Mutable/Source/MutableRuntime/Public/MuR/Model.h
+++ b/Engine/Plugins/Experimental/Mutable/Source/MutableRuntime/Public/MuR/Model.h
@@ -13,7 +13,7 @@
 //! data structures. Compiled models are not necessarily compatible when the runtime is updated,
 //! so this version number can be used externally to verify this. It is not used internally, and
 //! serializing models from different versions than this runtime will probably result in a crash.
-#define MUTABLE_COMPILED_MODEL_CODE_VERSION		uint32( 79 )
+#define MUTABLE_COMPILED_MODEL_CODE_VERSION		uint32( 80 )
 #define MUTABLE_PARAMETERS_VERSION              uint32( 2 )
 
 
diff --git a/Engine/Plugins/Experimental/Mutable/Source/MutableRuntime/Public/MuR/System.h b/Engine/Plugins/Experimental/Mutable/Source/MutableRuntime/Public/MuR/System.h
--- a/Engine/Plugins/Experimental/Mutable/Source/MutableRuntime/Public/MuR/System.h
+++ b/Engine/Plugins/Experimental/Mutable/Source/MutableRuntime/Public/MuR/System.h
@@ -91,8 +91,10 @@ namespace mu
         //! Returns the completion event and a cleanup function that must be called once event is completed.
 #ifdef MUTABLE_USE_NEW_TASKGRAPH
 		virtual TTuple<UE::Tasks::FTask, TFunction<void()>> GetImageAsync(FName Id, uint8 MipmapsToSkip, TFunction<void(Ptr<Image>)>& ResultCallback) = 0;
+		virtual TTuple<UE::Tasks::FTask, TFunction<void()>> GetReferencedImageAsync(const void* ModelPtr, int32 Id, uint8 MipmapsToSkip, TFunction<void(Ptr<Image>)>& ResultCallback) { check(false); return {}; }
 #else
 		virtual TTuple<FGraphEventRef, TFunction<void()>> GetImageAsync(FName Id, uint8 MipmapsToSkip, TFunction<void(Ptr<Image>)>& ResultCallback) = 0;
+		virtual TTuple<FGraphEventRef, TFunction<void()>> GetReferencedImageAsync(const void* ModelPtr, int32 Id, uint8 MipmapsToSkip, TFunction<void(Ptr<Image>)>& ResultCallback) { check(false); return {}; };
 #endif
 
         virtual mu::FImageDesc GetImageDesc(FName Id, uint8 MipmapsToSkip) = 0;
diff --git a/Engine/Plugins/Experimental/Mutable/Source/MutableTools/Private/MuT/AST.cpp b/Engine/Plugins/Experimental/Mutable/Source/MutableTools/Private/MuT/AST.cpp
--- a/Engine/Plugins/Experimental/Mutable/Source/MutableTools/Private/MuT/AST.cpp
+++ b/Engine/Plugins/Experimental/Mutable/Source/MutableTools/Private/MuT/AST.cpp
@@ -1147,19 +1147,6 @@ FImageDesc ASTOpFixed::GetImageDesc( bool returnBestOption, FGetImageDescContext
         check( res.m_format != EImageFormat::IF_NONE );
         break;
 
-    case OP_TYPE::IM_CROP:
-        res = GetImageDesc( op.args.ImageCrop.source, returnBestOption, context );
-
-		check(op.args.ImageCrop.sizeX > 0);
-		check(op.args.ImageCrop.sizeY > 0);
-
-        res.m_size = FImageSize
-            (
-                op.args.ImageCrop.sizeX,
-                op.args.ImageCrop.sizeY
-            );
-        break;
-
     case OP_TYPE::IM_RESIZE:
         res = GetImageDesc( op.args.ImageResize.source, returnBestOption, context );
 
@@ -1280,12 +1267,6 @@ void ASTOpFixed::GetLayoutBlockSize( int* pBlockX, int* pBlockY )
 		break;
 	}
 
-	case OP_TYPE::IM_CROP:
-	{
-		GetLayoutBlockSize(op.args.ImageCrop.source, pBlockX, pBlockY);
-		break;
-	}
-
     default:
         checkf( false, TEXT("Instruction not supported") );
     }
@@ -1665,13 +1646,7 @@ mu::Ptr<ImageSizeExpression> ASTOpFixed::GetImageSizeExpression() const
 			pRes = children[op.args.ImageInvert.base].child()->GetImageSizeExpression();
 		}
 		break;
-    	
-    case OP_TYPE::IM_CROP:
-		pRes->type = ImageSizeExpression::ISET_CONSTANT;
-		pRes->size[0] = op.args.ImageCrop.sizeX;
-		pRes->size[1] = op.args.ImageCrop.sizeY;
-    	break;
-
+ 
     default:
     	check( false );
     	break;
diff --git a/Engine/Plugins/Experimental/Mutable/Source/MutableTools/Private/MuT/AST.h b/Engine/Plugins/Experimental/Mutable/Source/MutableTools/Private/MuT/AST.h
--- a/Engine/Plugins/Experimental/Mutable/Source/MutableTools/Private/MuT/AST.h
+++ b/Engine/Plugins/Experimental/Mutable/Source/MutableTools/Private/MuT/AST.h
@@ -218,8 +218,8 @@ namespace mu
 
 		int32 MinTextureResidentMipCount = 0;
 
-		/** */
-		struct FDeduplicationMeshFuncs : TDefaultMapHashableKeyFuncs<mu::Ptr<const mu::Mesh>, int32, false>
+		/** Structure used to speedup mesh constant comparison. */
+		struct FDeduplicationMeshFuncs : TDefaultMapHashableKeyFuncs<Ptr<const Mesh>, int32, false>
 		{
 			static FORCEINLINE bool Matches(KeyInitType A, KeyInitType B)
 			{
@@ -229,14 +229,43 @@ namespace mu
 			static FORCEINLINE uint32 GetKeyHash(KeyInitType Key)
 			{
 				const mu::Mesh* Data = Key.get();
-				return HashCombine(
+				return HashCombineFast(
 					::GetTypeHash(Data->m_VertexBuffers.GetElementCount()),
 					::GetTypeHash(Data->m_IndexBuffers.GetElementCount())
 				);
 			}
 		};
 
-		TMap<mu::Ptr<const mu::Mesh>, int32, FDefaultSetAllocator, FDeduplicationMeshFuncs> MeshConstantMap;
+		TMap<Ptr<const Mesh>, int32, FDefaultSetAllocator, FDeduplicationMeshFuncs> MeshConstantMap;
+
+		/** Structure used to speedup image mip comparison. */
+		struct FDeduplicationImageFuncs : TDefaultMapHashableKeyFuncs<Ptr<const Image>, int32, false>
+		{
+			static FORCEINLINE bool Matches(KeyInitType A, KeyInitType B)
+			{
+				return *A == *B;
+			}
+
+			static FORCEINLINE uint32 GetKeyHash(KeyInitType Key)
+			{
+				const mu::Image* Data = Key.get();
+				uint32 Hash = HashCombineFast(
+					::GetTypeHash(Data->m_format),
+					HashCombineFast(::GetTypeHash(Data->m_size[0]),
+								::GetTypeHash(Data->m_size[1]))
+				);
+
+				if (Data->m_data.Num()>=4)
+				{
+					uint32 FirstData = *(reinterpret_cast<const uint32*>(Data->m_data.GetData()));
+					Hash = HashCombineFast( Hash, ::GetTypeHash(FirstData) );
+				}
+
+				return Hash;
+			}
+		};
+
+		TMap<Ptr<const Image>, int32, FDefaultSetAllocator, FDeduplicationImageFuncs> ImageConstantMipMap;
 
 		/** Image operation functions, so that they can be overriden. */
 		FImageOperator& ImageOperator;
@@ -251,7 +280,7 @@ namespace mu
 
 		friend inline uint32 GetTypeHash(const FSinkerOldToNewKey& Key)
 		{
-			return HashCombine(::GetTypeHash(Key.Op.get()), ::GetTypeHash(Key.SinkingOp.get()));
+			return HashCombineFast(::GetTypeHash(Key.Op.get()), ::GetTypeHash(Key.SinkingOp.get()));
 		}
 
 		friend inline bool operator==(const FSinkerOldToNewKey& A, const FSinkerOldToNewKey& B)
@@ -265,16 +294,16 @@ namespace mu
 	{
 	public:
 
-		Ptr<ASTOp> Apply(const ASTOp* root);
+		Ptr<ASTOp> Apply(const class ASTOpImageCrop* root);
 
 	protected:
 
-		const ASTOp* m_root = nullptr;
+		const class ASTOpImageCrop* m_root = nullptr;
 		Ptr<ASTOp> m_initialSource;
 		//! For each operation we sink, the map from old instructions to new instructions.
 		TMap<FSinkerOldToNewKey, Ptr<ASTOp>> OldToNew;
 
-		Ptr<ASTOp> Visit(Ptr<ASTOp> at, const ASTOpFixed* currentCropOp);
+		Ptr<ASTOp> Visit(Ptr<ASTOp> at, const class ASTOpImageCrop* currentCropOp);
 	};
 
 	//---------------------------------------------------------------------------------------------
diff --git a/Engine/Plugins/Experimental/Mutable/Source/MutableTools/Private/MuT/ASTOpConstantResource.cpp b/Engine/Plugins/Experimental/Mutable/Source/MutableTools/Private/MuT/ASTOpConstantResource.cpp
--- a/Engine/Plugins/Experimental/Mutable/Source/MutableTools/Private/MuT/ASTOpConstantResource.cpp
+++ b/Engine/Plugins/Experimental/Mutable/Source/MutableTools/Private/MuT/ASTOpConstantResource.cpp
@@ -67,21 +67,19 @@ namespace mu
 	namespace
 	{
 		/** Adds a constant image data to a program and returns its constant index. */
-		int32 AddConstantImage(FProgram& Program, Ptr<const Image> pImage, int32 MinTextureResidentMipCount, FLinkerOptions* Options)
+		int32 AddConstantImage(FProgram& Program, const Ptr<const Image>& pImage, FLinkerOptions& Options)
 		{
+			MUTABLE_CPUPROFILER_SCOPE(AddConstantImage);
+
 			check(pImage->GetSizeX() * pImage->GetSizeY() > 0);
 
 			// Mips to store
 			int32 MipsToStore = 1;
-			if (MinTextureResidentMipCount > 0 && MinTextureResidentMipCount < 255)
-			{
-				int32 MaxMipmaps = Image::GetMipmapCount(pImage->GetSizeX(), pImage->GetSizeY());
 
-				// This is not true. We may want the full mipmaps for fragments of images, regardless of the resident mip size.
-				//MipsToStore = FMath::Max(1, MaxMipmaps-MinTextureResidentMipCount);
-				// \TODO: Calculate the mip ranges that makes sense to store.
-				MipsToStore = MaxMipmaps;
-			}
+			// We may want the full mipmaps for fragments of images, regardless of the resident mip size, for intermediate operations.
+			// \TODO: Calculate the mip ranges that makes sense to store.
+			int32 MaxMipmaps = Image::GetMipmapCount(pImage->GetSizeX(), pImage->GetSizeY());
+			MipsToStore = MaxMipmaps;
 
 			int32 FirstLODIndexIndex = Program.m_constantImageLODIndices.Num();
 
@@ -94,41 +92,41 @@ namespace mu
 				MipsToStore = pImage->GetLODCount();
 			}
 
-			// TODO:
-			int32 CompressionQuality = 4;
+			FImageOperator& ImOp = Options.ImageOperator;
 
-			FImageOperator& ImOp = Options->ImageOperator;
-
-			// TODO: Not efficient if we don't make mips (no need to clone base)
 			// TODO: If the image already has mips, we will be duplicating them...
-			Ptr<Image> pMip = ImOp.ExtractMip(pImage.get(), 0);
+			Ptr<const Image> pMip;
+			if (pImage->GetLODCount() == 1)
+			{
+				pMip = pImage;
+			}
+			else
+			{
+				pMip = ImOp.ExtractMip(pImage.get(), 0);
+			}
 			for (int Mip = 0; Mip < MipsToStore; ++Mip)
 			{
 				check(pMip->GetFormat() == pImage->GetFormat());
 
-				// Shrink the buffer to the minimum necessary size.
-				uint32 CalculatedSize = pMip->CalculateDataSize();
-				if (CalculatedSize)
-				{
-					pMip->m_data.SetNum(CalculatedSize);
-				}
-
 				// Ensure unique at mip level
 				int32 MipIndex = -1;
-				for (int32 CandidateMipIndex = 0; CandidateMipIndex < Program.m_constantImageLODs.Num(); ++CandidateMipIndex)
+
+				// Use a map-based deduplication
 				{
-					const Image* pCandidate = Program.m_constantImageLODs[CandidateMipIndex].Value.get();
-					if (*pCandidate == *pMip)
+					MUTABLE_CPUPROFILER_SCOPE(Deduplicate);
+
+					Ptr<const mu::Image> Key = pMip;
+					const int32* IndexPtr = Options.ImageConstantMipMap.Find(Key);
+					if (IndexPtr)
 					{
-						// Reuse mip
-						MipIndex = CandidateMipIndex;
+						MipIndex = *IndexPtr;
+					}
+					else
+					{
+						check(pMip->GetLODCount() == 1);
+						MipIndex = Program.m_constantImageLODs.Add(TPair<int32, Ptr<const Image>>(-1, pMip));
+						Options.ImageConstantMipMap.Add(Key, MipIndex);
 					}
-				}
-
-				if (MipIndex < 0)
-				{
-					check(pMip->GetLODCount() == 1);
-					MipIndex = Program.m_constantImageLODs.Add(TPair<int32, Ptr<const Image>>(-1, pMip));
 				}
 
 				Program.m_constantImageLODIndices.Add(uint32(MipIndex));
@@ -136,16 +134,26 @@ namespace mu
 				// Generate next mip if necessary
 				if (Mip + 1 < MipsToStore)
 				{
+					Ptr<Image> NewMip;
 					if (Mip > pImage->GetLODCount())
 					{
 						// Generate from the last mip.
-						pMip = ImOp.ExtractMip(pMip.get(), 1);
+						NewMip = ImOp.ExtractMip(pMip.get(), 1);
 					}
 					else
 					{
-						pMip = ImOp.ExtractMip(pImage.get(), Mip + 1);
+						NewMip = ImOp.ExtractMip(pImage.get(), Mip + 1);
 					}
-					check(pMip);
+					check(NewMip);
+
+					// Shrink the buffer to the minimum necessary size.
+					uint32 CalculatedSize = NewMip->CalculateDataSize();
+					if (CalculatedSize)
+					{
+						NewMip->m_data.SetNum(CalculatedSize);
+					}
+
+					pMip = NewMip;
 				}
 			}
 
@@ -229,8 +237,7 @@ namespace mu
 					}
 					else
 					{
-						int32 MinTextureResidentMipCount = Options->MinTextureResidentMipCount;
-						args.value = AddConstantImage( program, pTyped, MinTextureResidentMipCount, Options);
+						args.value = AddConstantImage( program, pTyped, *Options);
 					}
 					break;
 				}
@@ -508,24 +515,16 @@ namespace mu
 	//-------------------------------------------------------------------------------------------------
 	mu::Ptr<ImageSizeExpression> ASTOpConstantResource::GetImageSizeExpression() const
 	{
-		Ptr<ImageSizeExpression> pRes = new ImageSizeExpression;
-		pRes->type = ImageSizeExpression::ISET_CONSTANT;
-
-		switch (type)
-		{
-		case OP_TYPE::IM_CONSTANT:
+		if (type==OP_TYPE::IM_CONSTANT)
 		{
+			Ptr<ImageSizeExpression> pRes = new ImageSizeExpression;
+			pRes->type = ImageSizeExpression::ISET_CONSTANT;
 			Ptr<const Image> pConst = static_cast<const Image*>(GetValue().get());
 			pRes->size = pConst->m_size;
-			break;
-		}
-
-		default:
-			check(false);
-			return nullptr;
+			return pRes;
 		}
 
-		return pRes;
+		return nullptr;
 	}
 
 }
\ No newline at end of file
diff --git a/Engine/Plugins/Experimental/Mutable/Source/MutableTools/Private/MuT/ASTOpImageCrop.cpp b/Engine/Plugins/Experimental/Mutable/Source/MutableTools/Private/MuT/ASTOpImageCrop.cpp
new file mode 100644
--- /dev/null
+++ b/Engine/Plugins/Experimental/Mutable/Source/MutableTools/Private/MuT/ASTOpImageCrop.cpp
@@ -0,0 +1,498 @@
+// Copyright Epic Games, Inc. All Rights Reserved.
+
+#include "MuT/ASTOpImageCrop.h"
+
+#include "Containers/Map.h"
+#include "MuR/ModelPrivate.h"
+#include "MuR/RefCounted.h"
+#include "MuR/Types.h"
+#include "MuT/StreamsPrivate.h"
+#include "MuT/ASTOpConditional.h"
+#include "MuT/ASTOpSwitch.h"
+#include "MuT/ASTOpImagePatch.h"
+#include "MuT/ASTOpImagePixelFormat.h"
+#include "MuT/ASTOpImageLayer.h"
+#include "MuT/ASTOpImageLayerColor.h"
+#include "MuT/ASTOpImageRasterMesh.h"
+
+
+namespace mu
+{
+
+	ASTOpImageCrop::ASTOpImageCrop()
+		: Source(this)
+	{
+	}
+
+
+	ASTOpImageCrop::~ASTOpImageCrop()
+	{
+		// Explicit call needed to avoid recursive destruction
+		ASTOp::RemoveChildren();
+	}
+
+
+	bool ASTOpImageCrop::IsEqual(const ASTOp& InOther) const
+	{
+		if (const ASTOpImageCrop* other = dynamic_cast<const ASTOpImageCrop*>(&InOther))
+		{
+			return Source == other->Source &&
+				Min == other->Min &&
+				Size == other->Size;
+		}
+		return false;
+	}
+
+
+	uint64 ASTOpImageCrop::Hash() const
+	{
+		uint64 res = std::hash<OP_TYPE>()(OP_TYPE::IM_CROP);
+		hash_combine(res, Source.child().get());
+		hash_combine(res, Min[0]);
+		hash_combine(res, Min[1]);
+		hash_combine(res, Size[0]);
+		hash_combine(res, Size[1]);
+		return res;
+	}
+
+
+	mu::Ptr<ASTOp> ASTOpImageCrop::Clone(MapChildFuncRef mapChild) const
+	{
+		Ptr<ASTOpImageCrop> n = new ASTOpImageCrop();
+		n->Source = mapChild(Source.child());
+		n->Min = Min;
+		n->Size = Size;
+		return n;
+	}
+
+
+	void ASTOpImageCrop::ForEachChild(const TFunctionRef<void(ASTChild&)> f)
+	{
+		f(Source);
+	}
+
+
+	void ASTOpImageCrop::Link(FProgram& program, FLinkerOptions*)
+	{
+		// Already linked?
+		if (!linkedAddress)
+		{
+			OP::ImageCropArgs args;
+			memset(&args, 0, sizeof(args));
+
+			if (Source) args.source = Source->linkedAddress;
+			args.minX = Min[0];
+			args.minY = Min[1];
+			args.sizeX = Size[0];
+			args.sizeY = Size[1];
+
+			linkedAddress = (OP::ADDRESS)program.m_opAddress.Num();
+			program.m_opAddress.Add((uint32_t)program.m_byteCode.Num());
+			AppendCode(program.m_byteCode, OP_TYPE::IM_CROP);
+			AppendCode(program.m_byteCode, args);
+		}
+
+	}
+
+
+	FImageDesc ASTOpImageCrop::GetImageDesc(bool returnBestOption, FGetImageDescContext* context) const
+	{
+		FImageDesc Result;
+
+		// Local context in case it is necessary
+		FGetImageDescContext localContext;
+		if (!context)
+		{
+			context = &localContext;
+		}
+		else
+		{
+			// Cached result?
+			FImageDesc* PtrValue = context->m_results.Find(this);
+			if (PtrValue)
+			{
+				return *PtrValue;
+			}
+		}
+
+		// Actual work
+		if (Source)
+		{
+			Result = Source->GetImageDesc(returnBestOption, context);
+			//Result.m_lods = 1;
+			Result.m_size[0] = Size[0];
+			Result.m_size[1] = Size[1];
+		}
+
+
+		// Cache the result
+		if (context)
+		{
+			context->m_results.Add(this, Result);
+		}
+
+		return Result;
+	}
+
+
+	mu::Ptr<ImageSizeExpression> ASTOpImageCrop::GetImageSizeExpression() const
+	{
+		Ptr<ImageSizeExpression> pRes = new ImageSizeExpression;
+		pRes->type = ImageSizeExpression::ISET_CONSTANT;
+		pRes->size[0] = Size[0];
+		pRes->size[1] = Size[1];
+		return pRes;
+	}
+
+
+	void ASTOpImageCrop::GetLayoutBlockSize(int* pBlockX, int* pBlockY)
+	{
+		// We didn't find any layout yet.
+		*pBlockX = 0;
+		*pBlockY = 0;
+
+		// Try the source
+		if (Source)
+		{
+			Source->GetLayoutBlockSize( pBlockX, pBlockY );
+		}
+	}
+
+
+	Ptr<ASTOp> ASTOpImageCrop::OptimiseSemantic(const FModelOptimizationOptions& options, int32 Pass) const
+	{
+		Ptr<ASTOp> at;
+
+		Ptr<ASTOp> sourceAt = Source.child();
+
+		// The instruction can be sunk
+		OP_TYPE sourceType = sourceAt->GetOpType();
+		switch (sourceType)
+		{
+		case OP_TYPE::IM_PLAINCOLOUR:
+		{
+			Ptr<ASTOpFixed> NewOp = mu::Clone<ASTOpFixed>(sourceAt.get());
+			NewOp->op.args.ImagePlainColour.size[0] = Size[0];
+			NewOp->op.args.ImagePlainColour.size[1] = Size[1];
+			NewOp->op.args.ImagePlainColour.LODs = 1; // TODO
+			at = NewOp;
+			break;
+		}
+
+		default:
+			break;
+
+		}
+
+		return at;
+	}
+
+
+	Ptr<ASTOp> ASTOpImageCrop::OptimiseSink(const FModelOptimizationOptions& options, FOptimizeSinkContext& context) const
+	{
+		Ptr<ASTOp> at;
+
+		Ptr<ASTOp> sourceAt = Source.child();
+
+		switch (sourceAt->GetOpType())
+		{
+			// In case we have other operations with special optimisation rules.
+		case OP_TYPE::NONE:
+			break;
+
+		default:
+		{
+			at = context.ImageCropSinker.Apply(this);
+
+			break;
+		} // default
+
+		}
+
+
+		return at;
+	}
+
+
+	Ptr<ASTOp> Sink_ImageCropAST::Apply(const ASTOpImageCrop* InRoot)
+	{
+		check(InRoot->GetOpType() == OP_TYPE::IM_CROP);
+
+		m_root = InRoot;
+		OldToNew.Empty();
+
+		m_initialSource = InRoot->Source.child();
+		Ptr<ASTOp> newSource = Visit(m_initialSource, InRoot);
+
+		// If there is any change, it is the new root.
+		if (newSource != m_initialSource)
+		{
+			return newSource;
+		}
+
+		return nullptr;
+	}
+
+
+	Ptr<ASTOp> Sink_ImageCropAST::Visit(Ptr<ASTOp> at, const ASTOpImageCrop* currentCropOp)
+	{
+		if (!at) return nullptr;
+
+		// Already visited?
+		const Ptr<ASTOp>* Cached = OldToNew.Find({ at, currentCropOp });
+		if (Cached)
+		{
+			return *Cached;
+		}
+
+		bool skipSinking = false;
+		Ptr<ASTOp> newAt = at;
+		switch (at->GetOpType())
+		{
+
+		case OP_TYPE::IM_CONDITIONAL:
+		{
+			// We move down the two paths
+			Ptr<ASTOpConditional> newOp = mu::Clone<ASTOpConditional>(at);
+			newOp->yes = Visit(newOp->yes.child(), currentCropOp);
+			newOp->no = Visit(newOp->no.child(), currentCropOp);
+			newAt = newOp;
+			break;
+		}
+
+		case OP_TYPE::IM_SWITCH:
+		{
+			// We move down all the paths
+			Ptr<ASTOpSwitch> newOp = mu::Clone<ASTOpSwitch>(at);
+			newOp->def = Visit(newOp->def.child(), currentCropOp);
+			for (auto& c : newOp->cases)
+			{
+				c.branch = Visit(c.branch.child(), currentCropOp);
+			}
+			newAt = newOp;
+			break;
+		}
+
+		case OP_TYPE::IM_PIXELFORMAT:
+		{
+			Ptr<ASTOpImagePixelFormat> nop = mu::Clone<ASTOpImagePixelFormat>(at);
+			nop->Source = Visit(nop->Source.child(), currentCropOp);
+			newAt = nop;
+			break;
+		}
+
+		case OP_TYPE::IM_PATCH:
+		{
+			const ASTOpImagePatch* typedPatch = dynamic_cast<const ASTOpImagePatch*>(at.get());
+
+			Ptr<ASTOp> rectOp = typedPatch->patch.child();
+			ASTOp::FGetImageDescContext context;
+			FImageDesc patchDesc = rectOp->GetImageDesc(false, &context);
+			box<vec2<int16>> patchBox;
+			patchBox.min[0] = typedPatch->location[0];
+			patchBox.min[1] = typedPatch->location[1];
+			patchBox.size[0] = patchDesc.m_size[0];
+			patchBox.size[1] = patchDesc.m_size[1];
+
+			box<vec2<int16>> cropBox;
+			cropBox.min[0] = currentCropOp->Min[0];
+			cropBox.min[1] = currentCropOp->Min[1];
+			cropBox.size[0] = currentCropOp->Size[0];
+			cropBox.size[1] = currentCropOp->Size[1];
+
+			if (!patchBox.IntersectsExclusive(cropBox))
+			{
+				// We can ignore the patch
+				newAt = Visit(typedPatch->base.child(), currentCropOp);
+			}
+			else
+			{
+				// Crop the base with the full crop, and the patch with the intersected part,
+				// adapting the patch origin
+				Ptr<ASTOpImagePatch> newOp = mu::Clone<ASTOpImagePatch>(at);
+				newOp->base = Visit(newOp->base.child(), currentCropOp);
+
+				box<vec2<int16>> ibox = patchBox.Intersect(cropBox);
+				check(ibox.size[0] > 0 && ibox.size[1] > 0);
+
+				Ptr<ASTOpImageCrop> patchCropOp = mu::Clone<ASTOpImageCrop>(currentCropOp);
+				patchCropOp->Min[0] = ibox.min[0] - patchBox.min[0];
+				patchCropOp->Min[1] = ibox.min[1] - patchBox.min[1];
+				patchCropOp->Size[0] = ibox.size[0];
+				patchCropOp->Size[1] = ibox.size[1];
+				newOp->patch = Visit(newOp->patch.child(), patchCropOp.get());
+
+				newOp->location[0] = ibox.min[0] - cropBox.min[0];
+				newOp->location[1] = ibox.min[1] - cropBox.min[1];
+				newAt = newOp;
+			}
+
+			break;
+		}
+
+		case OP_TYPE::IM_CROP:
+		{
+			// We can combine the two crops into a possibly smaller crop
+			const ASTOpImageCrop* childCrop = dynamic_cast<const ASTOpImageCrop*>(at.get());
+
+			box<vec2<int16>> childCropBox;
+			childCropBox.min[0] = childCrop->Min[0];
+			childCropBox.min[1] = childCrop->Min[1];
+			childCropBox.size[0] = childCrop->Size[0];
+			childCropBox.size[1] = childCrop->Size[1];
+
+			box<vec2<int16>> cropBox;
+			cropBox.min[0] = currentCropOp->Min[0];
+			cropBox.min[1] = currentCropOp->Min[1];
+			cropBox.size[0] = currentCropOp->Size[0];
+			cropBox.size[1] = currentCropOp->Size[1];
+
+			// Compose the crops: in the final image the child crop is applied first and the
+			// current ctop is applied to the result. So the final crop box would be:
+			box<vec2<int16>> ibox;
+			ibox.min = childCropBox.min + cropBox.min;
+			ibox.size = vec2<int16>::min(cropBox.size, childCropBox.size);
+			check(cropBox.min.AllSmallerOrEqualThan(childCropBox.size));
+			check((cropBox.min + cropBox.size).AllSmallerOrEqualThan(childCropBox.size));
+			check((ibox.min + ibox.size).AllSmallerOrEqualThan(childCropBox.min + childCropBox.size));
+
+			// This happens more often that one would think
+			if (ibox == childCropBox)
+			{
+				// the parent crop is not necessary
+				skipSinking = true;
+			}
+			else if (ibox == cropBox)
+			{
+				// The child crop is not necessary
+				Ptr<ASTOp> childSource = childCrop->Source.child();
+				newAt = Visit(childSource, currentCropOp);
+			}
+			else
+			{
+				// combine into one crop
+				Ptr<ASTOpImageCrop> newCropOp = mu::Clone<ASTOpImageCrop>(currentCropOp);
+				newCropOp->Min[0] = ibox.min[0];
+				newCropOp->Min[1] = ibox.min[1];
+				newCropOp->Size[0] = ibox.size[0];
+				newCropOp->Size[1] = ibox.size[1];
+
+				Ptr<ASTOp> childSource = childCrop->Source.child();
+				newAt = Visit(childSource, newCropOp.get());
+			}
+			break;
+		}
+
+		case OP_TYPE::IM_LAYER:
+		{
+			// We move the op down the arguments
+			Ptr<ASTOpImageLayer> nop = mu::Clone<ASTOpImageLayer>(at);
+
+			Ptr<ASTOp> aOp = nop->base.child();
+			nop->base = Visit(aOp, currentCropOp);
+
+			Ptr<ASTOp> bOp = nop->blend.child();
+			nop->blend = Visit(bOp, currentCropOp);
+
+			Ptr<ASTOp> mOp = nop->mask.child();
+			nop->mask = Visit(mOp, currentCropOp);
+
+			newAt = nop;
+			break;
+		}
+
+		case OP_TYPE::IM_LAYERCOLOUR:
+		{
+			// We move the op down the arguments
+			Ptr<ASTOpImageLayerColor> nop = mu::Clone<ASTOpImageLayerColor>(at);
+
+			Ptr<ASTOp> aOp = nop->base.child();
+			nop->base = Visit(aOp, currentCropOp);
+
+			Ptr<ASTOp> mOp = nop->mask.child();
+			nop->mask = Visit(mOp, currentCropOp);
+
+			newAt = nop;
+			break;
+		}
+
+		case OP_TYPE::IM_DISPLACE:
+		{
+			// We move the op down the arguments
+			Ptr<ASTOpFixed> nop = mu::Clone<ASTOpFixed>(at);
+
+			Ptr<ASTOp> aOp = nop->children[nop->op.args.ImageDisplace.source].child();
+			nop->SetChild(nop->op.args.ImageDisplace.source, Visit(aOp, currentCropOp));
+
+			Ptr<ASTOp> bOp = nop->children[nop->op.args.ImageDisplace.displacementMap].child();
+			nop->SetChild(nop->op.args.ImageDisplace.displacementMap, Visit(bOp, currentCropOp));
+
+			newAt = nop;
+			break;
+		}
+
+		case OP_TYPE::IM_RASTERMESH:
+		{
+			// We add cropping data to the raster mesh if it doesn't have any
+			// \TODO: Is is possible to hit 2 crops on a raster mesh? Combine the crop.
+			Ptr<ASTOpImageRasterMesh> nop = mu::Clone<ASTOpImageRasterMesh>(at);
+
+			bool bRasterHasCrop = nop->UncroppedSizeX != 0;
+			if (!bRasterHasCrop)
+			{
+				box<vec2<int16>> cropBox;
+				cropBox.min[0] = currentCropOp->Min[0];
+				cropBox.min[1] = currentCropOp->Min[1];
+				cropBox.size[0] = currentCropOp->Size[0];
+				cropBox.size[1] = currentCropOp->Size[1];
+
+				nop->UncroppedSizeX = nop->SizeX;
+				nop->UncroppedSizeY = nop->SizeY;
+				nop->CropMinX = cropBox.min[0];
+				nop->CropMinY = cropBox.min[1];
+				nop->SizeX = cropBox.size[0];
+				nop->SizeY = cropBox.size[1];
+
+				newAt = nop;
+			}
+			break;
+		}
+
+		case OP_TYPE::IM_INTERPOLATE:
+		{
+			// Move the op  down all the paths
+			auto newOp = mu::Clone<ASTOpFixed>(at);
+
+			for (int v = 0; v < MUTABLE_OP_MAX_INTERPOLATE_COUNT; ++v)
+			{
+				Ptr<ASTOp> child = newOp->children[newOp->op.args.ImageInterpolate.targets[v]].child();
+				Ptr<ASTOp> bOp = Visit(child, currentCropOp);
+				newOp->SetChild(newOp->op.args.ImageInterpolate.targets[v], bOp);
+			}
+
+			newAt = newOp;
+			break;
+		}
+
+		default:
+			break;
+		}
+
+		// end on line, replace with crop
+		if (at == newAt && at != m_initialSource && !skipSinking)
+		{
+			Ptr<ASTOpImageCrop> newOp = mu::Clone<ASTOpImageCrop>(currentCropOp);
+			check(newOp->GetOpType() == OP_TYPE::IM_CROP);
+
+			newOp->Source = at;
+
+			newAt = newOp;
+		}
+
+		OldToNew.Add({ at, currentCropOp }, newAt);
+
+		return newAt;
+	}
+
+
+}
diff --git a/Engine/Plugins/Experimental/Mutable/Source/MutableTools/Private/MuT/ASTOpImageCrop.h b/Engine/Plugins/Experimental/Mutable/Source/MutableTools/Private/MuT/ASTOpImageCrop.h
new file mode 100644
--- /dev/null
+++ b/Engine/Plugins/Experimental/Mutable/Source/MutableTools/Private/MuT/ASTOpImageCrop.h
@@ -0,0 +1,48 @@
+// Copyright Epic Games, Inc. All Rights Reserved.
+
+#pragma once
+
+#include "HAL/PlatformMath.h"
+#include "MuR/Image.h"
+#include "MuR/MutableMath.h"
+#include "MuR/Operations.h"
+#include "MuR/Ptr.h"
+#include "MuT/AST.h"
+
+
+namespace mu
+{
+struct FProgram;
+
+	class ASTOpImageCrop : public ASTOp
+	{
+	public:
+
+		ASTChild Source;
+		UE::Math::TIntVector2<uint16> Min = UE::Math::TIntVector2<uint16>(0, 0);
+		UE::Math::TIntVector2<uint16> Size = UE::Math::TIntVector2<uint16>(0, 0);
+
+	public:
+
+		ASTOpImageCrop();
+		ASTOpImageCrop(const ASTOpImageCrop&) = delete;
+		~ASTOpImageCrop();
+
+		virtual OP_TYPE GetOpType() const override { return OP_TYPE::IM_CROP; }
+		virtual uint64 Hash() const override;
+		virtual bool IsEqual(const ASTOp& otherUntyped) const override;
+		virtual Ptr<ASTOp> Clone(MapChildFuncRef mapChild) const override;
+		virtual void ForEachChild(const TFunctionRef<void(ASTChild&)>) override;
+		virtual void Link(FProgram& program, FLinkerOptions* Options) override;
+		virtual FImageDesc GetImageDesc(bool returnBestOption, FGetImageDescContext* context) const override;
+		virtual Ptr<ImageSizeExpression> GetImageSizeExpression() const override;
+		virtual void GetLayoutBlockSize(int* pBlockX, int* pBlockY) override;
+		//TODO: virtual bool IsImagePlainConstant(vec4<float>& colour) const override;
+		virtual Ptr<ASTOp> OptimiseSemantic(const FModelOptimizationOptions& options, int32 Pass) const override;
+		virtual Ptr<ASTOp> OptimiseSink(const FModelOptimizationOptions& options, FOptimizeSinkContext& context) const override;
+
+	};
+
+
+}
+
diff --git a/Engine/Plugins/Experimental/Mutable/Source/MutableTools/Private/MuT/ASTOpImageLayer.cpp b/Engine/Plugins/Experimental/Mutable/Source/MutableTools/Private/MuT/ASTOpImageLayer.cpp
--- a/Engine/Plugins/Experimental/Mutable/Source/MutableTools/Private/MuT/ASTOpImageLayer.cpp
+++ b/Engine/Plugins/Experimental/Mutable/Source/MutableTools/Private/MuT/ASTOpImageLayer.cpp
@@ -7,6 +7,7 @@
 #include "MuT/ASTOpImageLayerColor.h"
 #include "MuT/ASTOpImageSwizzle.h"
 #include "MuT/ASTOpImageRasterMesh.h"
+#include "MuT/ASTOpImageCrop.h"
 #include "MuT/ASTOpSwitch.h"
 #include "MuR/ModelPrivate.h"
 #include "MuR/RefCounted.h"
@@ -484,29 +485,26 @@ namespace mu
 					check(maskUsage.size[0] > 0);
 					check(maskUsage.size[1] > 0);
 
-					Ptr<ASTOpFixed> cropMask = new ASTOpFixed();
-					cropMask->op.type = OP_TYPE::IM_CROP;
-					cropMask->SetChild(cropMask->op.args.ImageCrop.source, mask.child());
-					cropMask->op.args.ImageCrop.minX = maskUsage.min[0];
-					cropMask->op.args.ImageCrop.minY = maskUsage.min[1];
-					cropMask->op.args.ImageCrop.sizeX = maskUsage.size[0];
-					cropMask->op.args.ImageCrop.sizeY = maskUsage.size[1];
-
-					Ptr<ASTOpFixed> cropBlended = new ASTOpFixed();
-					cropBlended->op.type = OP_TYPE::IM_CROP;
-					cropBlended->SetChild(cropBlended->op.args.ImageCrop.source, blend.child());
-					cropBlended->op.args.ImageCrop.minX = maskUsage.min[0];
-					cropBlended->op.args.ImageCrop.minY = maskUsage.min[1];
-					cropBlended->op.args.ImageCrop.sizeX = maskUsage.size[0];
-					cropBlended->op.args.ImageCrop.sizeY = maskUsage.size[1];
-
-					Ptr<ASTOpFixed> cropBase = new ASTOpFixed();
-					cropBase->op.type = OP_TYPE::IM_CROP;
-					cropBase->SetChild(cropBase->op.args.ImageCrop.source, base.child());
-					cropBase->op.args.ImageCrop.minX = maskUsage.min[0];
-					cropBase->op.args.ImageCrop.minY = maskUsage.min[1];
-					cropBase->op.args.ImageCrop.sizeX = maskUsage.size[0];
-					cropBase->op.args.ImageCrop.sizeY = maskUsage.size[1];
+					Ptr<ASTOpImageCrop> cropMask = new ASTOpImageCrop();
+					cropMask->Source = mask.child();
+					cropMask->Min[0] = maskUsage.min[0];
+					cropMask->Min[1] = maskUsage.min[1];
+					cropMask->Size[0] = maskUsage.size[0];
+					cropMask->Size[1] = maskUsage.size[1];
+
+					Ptr<ASTOpImageCrop> cropBlended = new ASTOpImageCrop();
+					cropBlended->Source = blend.child();
+					cropBlended->Min[0] = maskUsage.min[0];
+					cropBlended->Min[1] = maskUsage.min[1];
+					cropBlended->Size[0] = maskUsage.size[0];
+					cropBlended->Size[1] = maskUsage.size[1];
+
+					Ptr<ASTOpImageCrop> cropBase = new ASTOpImageCrop();
+					cropBase->Source = base.child();
+					cropBase->Min[0] = maskUsage.min[0];
+					cropBase->Min[1] = maskUsage.min[1];
+					cropBase->Size[0] = maskUsage.size[0];
+					cropBase->Size[1] = maskUsage.size[1];
 
 					Ptr<ASTOpImageLayer> newLayer = mu::Clone<ASTOpImageLayer>(this);
 					newLayer->base = cropBase;
diff --git a/Engine/Plugins/Experimental/Mutable/Source/MutableTools/Private/MuT/ASTOpReferenceResource.cpp b/Engine/Plugins/Experimental/Mutable/Source/MutableTools/Private/MuT/ASTOpReferenceResource.cpp
--- a/Engine/Plugins/Experimental/Mutable/Source/MutableTools/Private/MuT/ASTOpReferenceResource.cpp
+++ b/Engine/Plugins/Experimental/Mutable/Source/MutableTools/Private/MuT/ASTOpReferenceResource.cpp
@@ -23,7 +23,7 @@ namespace mu
 	{
 		if (const ASTOpReferenceResource* other = dynamic_cast<const ASTOpReferenceResource*>(&otherUntyped))
 		{
-			return type == other->type && ID == other->ID;
+			return type == other->type && ID == other->ID && bForceLoad == other->bForceLoad && ImageDesc == other->ImageDesc;
 		}
 		return false;
 	}
@@ -35,6 +35,8 @@ namespace mu
 		Ptr<ASTOpReferenceResource> n = new ASTOpReferenceResource();
 		n->type = type;
 		n->ID = ID;
+		n->bForceLoad = bForceLoad;
+		n->ImageDesc = ImageDesc;
 		return n;
 	}
 
@@ -56,6 +58,8 @@ namespace mu
 			OP::ResourceReferenceArgs Args;
 			FMemory::Memset(&Args, 0, sizeof(Args));
 			Args.ID = ID;
+			Args.ForceLoad = bForceLoad?1:0;
+			Args.ImageDesc = ImageDesc;
 
 			linkedAddress = (OP::ADDRESS)program.m_opAddress.Num();
 			program.m_opAddress.Add((uint32)program.m_byteCode.Num());
diff --git a/Engine/Plugins/Experimental/Mutable/Source/MutableTools/Private/MuT/ASTOpReferenceResource.h b/Engine/Plugins/Experimental/Mutable/Source/MutableTools/Private/MuT/ASTOpReferenceResource.h
--- a/Engine/Plugins/Experimental/Mutable/Source/MutableTools/Private/MuT/ASTOpReferenceResource.h
+++ b/Engine/Plugins/Experimental/Mutable/Source/MutableTools/Private/MuT/ASTOpReferenceResource.h
@@ -28,6 +28,11 @@ namespace mu
 
 		//!
 		uint32 ID = 0;
+		
+		/** */
+		bool bForceLoad = false;
+
+		FImageDesc ImageDesc;
 
 	public:
 
diff --git a/Engine/Plugins/Experimental/Mutable/Source/MutableTools/Private/MuT/CodeGenerator.cpp b/Engine/Plugins/Experimental/Mutable/Source/MutableTools/Private/MuT/CodeGenerator.cpp
--- a/Engine/Plugins/Experimental/Mutable/Source/MutableTools/Private/MuT/CodeGenerator.cpp
+++ b/Engine/Plugins/Experimental/Mutable/Source/MutableTools/Private/MuT/CodeGenerator.cpp
@@ -24,6 +24,7 @@
 #include "MuT/ASTOpImageLayer.h"
 #include "MuT/ASTOpImageLayerColor.h"
 #include "MuT/ASTOpImagePatch.h"
+#include "MuT/ASTOpImageCrop.h"
 #include "MuT/ASTOpInstanceAdd.h"
 #include "MuT/ASTOpMeshBindShape.h"
 #include "MuT/ASTOpMeshClipDeform.h"
@@ -706,13 +707,12 @@ namespace mu
 				int32 TileSizeX = FMath::Min(TileSize, Size[0] - MinX);
 				int32 TileSizeY = FMath::Min(TileSize, Size[1] - MinY);
 
-				Ptr<ASTOpFixed> TileImage = new ASTOpFixed();
-				TileImage->op.type = OP_TYPE::IM_CROP;
-				TileImage->SetChild(TileImage->op.args.ImageCrop.source, Source);
-				TileImage->op.args.ImageCrop.minX = MinX;
-				TileImage->op.args.ImageCrop.minY = MinY;
-				TileImage->op.args.ImageCrop.sizeX = TileSizeX;
-				TileImage->op.args.ImageCrop.sizeY = TileSizeY;
+				Ptr<ASTOpImageCrop> TileImage = new ASTOpImageCrop();
+				TileImage->Source = Source;
+				TileImage->Min[0] = MinX;
+				TileImage->Min[1] = MinY;
+				TileImage->Size[0] = TileSizeX;
+				TileImage->Size[1] = TileSizeY;
 
 				Ptr<ASTOpImagePatch> PatchedImage = new ASTOpImagePatch();
 				PatchedImage->base = CurrentImage;
@@ -733,7 +733,7 @@ namespace mu
                                          NodeSurfaceNewPtrConst surfaceNode,
                                          const TArray<FirstPassGenerator::FSurface::FEdit>& edits )
     {
-        //MUTABLE_CPUPROFILER_SCOPE(GenerateSurface);
+        MUTABLE_CPUPROFILER_SCOPE(GenerateSurface);
 
         const NodeSurfaceNew::Private& node = *surfaceNode->GetPrivate();
 
@@ -1188,7 +1188,7 @@ namespace mu
 		{
 			for (int32 t = 0; t < node.m_images.Num(); ++t)
 			{
-				//MUTABLE_CPUPROFILER_SCOPE(SurfaceTexture);
+				MUTABLE_CPUPROFILER_SCOPE(SurfaceTexture);
 
 				if (NodeImagePtr pImageNode = node.m_images[t].m_pImage)
 				{
diff --git a/Engine/Plugins/Experimental/Mutable/Source/MutableTools/Private/MuT/CodeGenerator.h b/Engine/Plugins/Experimental/Mutable/Source/MutableTools/Private/MuT/CodeGenerator.h
--- a/Engine/Plugins/Experimental/Mutable/Source/MutableTools/Private/MuT/CodeGenerator.h
+++ b/Engine/Plugins/Experimental/Mutable/Source/MutableTools/Private/MuT/CodeGenerator.h
@@ -289,6 +289,9 @@ namespace mu
 
 		TArray<FConditionalExtensionDataOp> m_conditionalExtensionDataOps;
 
+		/** This is used to cache the conatant image op generation. */
+		TMap<Ptr<const Image>,Ptr<class ASTOpConstantResource>> ImageConstantOpPerImage;
+
 		//-----------------------------------------------------------------------------------------
 
 		// Get the modifiers that have to be applied to elements with a specific tag.
diff --git a/Engine/Plugins/Experimental/Mutable/Source/MutableTools/Private/MuT/CodeGenerator_Image.cpp b/Engine/Plugins/Experimental/Mutable/Source/MutableTools/Private/MuT/CodeGenerator_Image.cpp
--- a/Engine/Plugins/Experimental/Mutable/Source/MutableTools/Private/MuT/CodeGenerator_Image.cpp
+++ b/Engine/Plugins/Experimental/Mutable/Source/MutableTools/Private/MuT/CodeGenerator_Image.cpp
@@ -30,6 +30,7 @@
 #include "MuT/ASTOpImageMakeGrowMap.h"
 #include "MuT/ASTOpImageSwizzle.h"
 #include "MuT/ASTOpImageRasterMesh.h"
+#include "MuT/ASTOpImageCrop.h"
 #include "MuT/ASTOpMeshExtractLayoutBlocks.h"
 #include "MuT/ASTOpMeshFormat.h"
 #include "MuT/ASTOpParameter.h"
@@ -176,9 +177,6 @@ namespace mu
     {
 		const NodeImageConstant::Private& node = *InNode->GetPrivate();
 		
-        Ptr<ASTOpConstantResource> op = new ASTOpConstantResource();
-        op->type = OP_TYPE::IM_CONSTANT;
-
         // TODO: check duplicates
         Ptr<const Image> pImage;
 		if (node.m_pProxy)
@@ -227,24 +225,63 @@ namespace mu
 			// Do we need to crop?
 			if (rect.min[0]!=0 || rect.min[1]!=0 || pImage->GetSizeX() != rect.size[0] || pImage->GetSizeY() != rect.size[1])
 			{
-				FImageOperator ImOp = FImageOperator::GetDefault(m_compilerOptions->ImageFormatFunc);
+				// Option 1: Crop now
+				if (true)
+				{
+					FImageOperator ImOp = FImageOperator::GetDefault(m_compilerOptions->ImageFormatFunc);
 
-				Ptr<Image> pCropped = new Image(rect.size[0], rect.size[1], 1, pImage->GetFormat(), EInitializationType::NotInitialized);
-				ImOp.ImageCrop(pCropped.get(), m_compilerOptions->ImageCompressionQuality, pImage.get(), rect);
-				op->SetValue(pCropped, m_compilerOptions->OptimisationOptions.bUseDiskCache);
+					Ptr<Image> pCropped = new Image(rect.size[0], rect.size[1], 1, pImage->GetFormat(), EInitializationType::NotInitialized);
+					ImOp.ImageCrop(pCropped.get(), m_compilerOptions->ImageCompressionQuality, pImage.get(), rect);
+
+					Ptr<ASTOpConstantResource> op = new ASTOpConstantResource();
+					op->type = OP_TYPE::IM_CONSTANT;
+					op->SetValue(pCropped, m_compilerOptions->OptimisationOptions.bUseDiskCache);
+					Result.op = op;
+				}
+
+				// Option 2: Generate a crop instruction: seems to be a lot slower and give worse results.
+				else
+				{
+					//Ptr<ASTOpConstantResource>* Found = ImageConstantOpPerImage.Find( pImage );
+					//Ptr<ASTOpConstantResource> ConstantOp;
+					//if (Found)
+					//{
+					//	ConstantOp = *Found;
+					//}
+					//else
+					//{ 
+					//	ConstantOp = new ASTOpConstantResource();
+					//	ConstantOp->type = OP_TYPE::IM_CONSTANT;
+					//	ConstantOp->SetValue(pImage, m_compilerOptions->OptimisationOptions.bUseDiskCache);
+					//	ImageConstantOpPerImage.Add(pImage, ConstantOp);
+					//}
+
+					//Ptr<ASTOpImageCrop> CropOp = new ASTOpImageCrop();
+					//CropOp->Source = ConstantOp;
+					//CropOp->Min[0] = rect.min[0];
+					//CropOp->Min[1] = rect.min[1];
+					//CropOp->Size[0] = rect.size[0];
+					//CropOp->Size[1] = rect.size[1];
+
+					//Result.op = CropOp;
+				}
 			}
 			else
 			{
 				// No need to crop.
+				Ptr<ASTOpConstantResource> op = new ASTOpConstantResource();
+				op->type = OP_TYPE::IM_CONSTANT;
 				op->SetValue(pImage, m_compilerOptions->OptimisationOptions.bUseDiskCache);
+				Result.op = op;
 			}
 		}
         else
         {
-            op->SetValue( pImage, m_compilerOptions->OptimisationOptions.bUseDiskCache );
-        }
-
-		Result.op = op;
+			Ptr<ASTOpConstantResource> op = new ASTOpConstantResource();
+			op->type = OP_TYPE::IM_CONSTANT;
+			op->SetValue( pImage, m_compilerOptions->OptimisationOptions.bUseDiskCache );
+			Result.op = op;
+		}
     }
 
 
@@ -253,12 +290,64 @@ namespace mu
 	{
 		const NodeImageReference::Private& node = *InNode->GetPrivate();
 
-		Ptr<ASTOpReferenceResource> op = new ASTOpReferenceResource();
-		op->type = OP_TYPE::IM_REFERENCE;
-		op->ID = node.ImageReferenceID;
+		Ptr<ASTOpReferenceResource> ReferenceOp = new ASTOpReferenceResource();
+		ReferenceOp->type = OP_TYPE::IM_REFERENCE;
+		ReferenceOp->ID = node.ImageReferenceID;
+		ReferenceOp->bForceLoad = node.bForceLoad;
 
-		// TODO: check no crop
-		Result.op = op;
+		if (Options.ImageLayoutStrategy != CompilerOptions::TextureLayoutStrategy::None && Options.LayoutToApply)
+		{
+			// We want to generate only a block from the image.
+
+			FIntVector2 SourceImageSize(node.ImageDesc.m_size[0], node.ImageDesc.m_size[1]);
+
+			int32 BlockIndex = Options.LayoutToApply->FindBlock(Options.LayoutBlockId);
+			check(BlockIndex >= 0);
+
+			// Block in layout grid units
+			box< UE::Math::TIntVector2<uint16> > RectInCells;
+			Options.LayoutToApply->GetBlock
+			(
+				BlockIndex,
+				&RectInCells.min[0], &RectInCells.min[1],
+				&RectInCells.size[0], &RectInCells.size[1]
+			);
+
+			FIntPoint grid = Options.LayoutToApply->GetGridSize();
+			grid[0] = FMath::Max(1, grid[0]);
+			grid[1] = FMath::Max(1, grid[1]);
+
+			// Transform to pixels
+			box< UE::Math::TIntVector2<int32> > rect;
+			rect.min[0] = (RectInCells.min[0] * SourceImageSize[0]) / grid[0];
+			rect.min[1] = (RectInCells.min[1] * SourceImageSize[1]) / grid[1];
+			rect.size[0] = (RectInCells.size[0] * SourceImageSize[0]) / grid[0];
+			rect.size[1] = (RectInCells.size[1] * SourceImageSize[1]) / grid[1];
+
+			// Do we need to crop?
+			if (rect.min[0] != 0 || rect.min[1] != 0 || node.ImageDesc.m_size[0] != rect.size[0] || node.ImageDesc.m_size[1] != rect.size[1])
+			{
+				// If need to crop it has to be the case of an image we will force to load.
+				check(node.bForceLoad);
+
+				Ptr<ASTOpImageCrop> CropOp = new ASTOpImageCrop();
+				CropOp->Source = ReferenceOp;
+				CropOp->Min[0] = rect.min[0];
+				CropOp->Min[1] = rect.min[1];
+				CropOp->Size[0] = rect.size[0];
+				CropOp->Size[1] = rect.size[1];
+
+				Result.op = CropOp;
+			}
+			else
+			{
+				Result.op = ReferenceOp;
+			}
+		}
+		else
+		{
+			Result.op = ReferenceOp;
+		}
 	}
 
 
@@ -1721,8 +1810,15 @@ namespace mu
 				{
 					if (pImage->IsReference())
 					{
+						FImageDesc ImageDesc;
+						ImageDesc.m_format = pImage->GetFormat();
+						ImageDesc.m_lods = pImage->GetLODCount();
+						ImageDesc.m_size[0] = pImage->GetSizeX();
+						ImageDesc.m_size[1] = pImage->GetSizeY();
+
 						Ptr<NodeImageReference> ImageRef = new NodeImageReference();
-						ImageRef->SetImageReference(pImage->GetReferencedTexture());
+						ImageRef->SetImageReference(pImage->GetReferencedTexture(), ImageDesc);
+						ImageRef->SetForceLoad( pImage->m_flags & Image::IF_IS_FORCELOAD );
 
 						CellImage = ImageRef;
 					}
diff --git a/Engine/Plugins/Experimental/Mutable/Source/MutableTools/Private/MuT/CodeOptimiser_Semantic.cpp b/Engine/Plugins/Experimental/Mutable/Source/MutableTools/Private/MuT/CodeOptimiser_Semantic.cpp
--- a/Engine/Plugins/Experimental/Mutable/Source/MutableTools/Private/MuT/CodeOptimiser_Semantic.cpp
+++ b/Engine/Plugins/Experimental/Mutable/Source/MutableTools/Private/MuT/CodeOptimiser_Semantic.cpp
@@ -384,33 +384,6 @@ namespace mu
 			break;
 		}
 
-		//-------------------------------------------------------------------------------------
-		case OP_TYPE::IM_CROP:
-		{
-			Ptr<ASTOp> sourceAt = children[op.args.ImageCrop.source].child();
-
-			// The instruction can be sunk
-			OP_TYPE sourceType = sourceAt->GetOpType();
-			switch (sourceType)
-			{
-			case OP_TYPE::IM_PLAINCOLOUR:
-			{
-				Ptr<ASTOpFixed> NewOp = mu::Clone<ASTOpFixed>(sourceAt.get());
-				NewOp->op.args.ImagePlainColour.size[0] = op.args.ImageCrop.sizeX;
-				NewOp->op.args.ImagePlainColour.size[1] = op.args.ImageCrop.sizeY;
-				NewOp->op.args.ImagePlainColour.LODs = 1; // TODO
-				at = NewOp;
-				break;
-			}
-
-			default:
-				break;
-
-			}
-
-			break;
-		}
-
         //-------------------------------------------------------------------------------------
         case OP_TYPE::IM_PIXELFORMAT:
 		case OP_TYPE::IM_SWIZZLE:
@@ -459,294 +432,6 @@ namespace mu
     }
 
 
-	//---------------------------------------------------------------------------------------------
-	//---------------------------------------------------------------------------------------------
-	//---------------------------------------------------------------------------------------------
-    Ptr<ASTOp> Sink_ImageCropAST::Apply( const ASTOp* InRoot )
-    {
-		check(InRoot->GetOpType() == OP_TYPE::IM_CROP);
-		
-		m_root = InRoot;
-        OldToNew.Empty();
-
-		const ASTOpFixed* typedRoot = dynamic_cast<const ASTOpFixed*>(InRoot);
-        m_initialSource = typedRoot->children[typedRoot->op.args.ImageCrop.source].child();
-        Ptr<ASTOp> newSource = Visit( m_initialSource, typedRoot );
-
-        // If there is any change, it is the new root.
-        if (newSource!=m_initialSource)
-        {
-            return newSource;
-        }
-
-        return nullptr;
-    }
-
-
-	//---------------------------------------------------------------------------------------------
-	Ptr<ASTOp> Sink_ImageCropAST::Visit( Ptr<ASTOp> at, const ASTOpFixed* currentCropOp )
-    {
-        if (!at) return nullptr;
-
-        // Already visited?
-		const Ptr<ASTOp>* Cached = OldToNew.Find({ at, currentCropOp });
-		if (Cached)
-		{
-			return *Cached;
-		}
-
-        bool skipSinking=false;
-        Ptr<ASTOp> newAt = at;
-        switch ( at->GetOpType() )
-        {
-
-        case OP_TYPE::IM_CONDITIONAL:
-        {
-            // We move down the two paths
-            Ptr<ASTOpConditional> newOp = mu::Clone<ASTOpConditional>(at);
-            newOp->yes = Visit(newOp->yes.child(), currentCropOp);
-            newOp->no = Visit(newOp->no.child(), currentCropOp);
-            newAt = newOp;
-            break;
-        }
-
-        case OP_TYPE::IM_SWITCH:
-        {
-            // We move down all the paths
-            Ptr<ASTOpSwitch> newOp = mu::Clone<ASTOpSwitch>(at);
-            newOp->def = Visit(newOp->def.child(), currentCropOp);
-            for( auto& c:newOp->cases )
-            {
-                c.branch = Visit(c.branch.child(), currentCropOp);
-            }
-            newAt = newOp;
-            break;
-        }
-
-        case OP_TYPE::IM_PIXELFORMAT:
-        {
-            Ptr<ASTOpImagePixelFormat> nop = mu::Clone<ASTOpImagePixelFormat>(at);
-            nop->Source = Visit(nop->Source.child(), currentCropOp);
-            newAt = nop;
-            break;
-        }
-
-        case OP_TYPE::IM_PATCH:
-        {
-			const ASTOpImagePatch* typedPatch = dynamic_cast<const ASTOpImagePatch*>(at.get());
-
-			Ptr<ASTOp> rectOp = typedPatch->patch.child();
-            ASTOp::FGetImageDescContext context;
-            FImageDesc patchDesc = rectOp->GetImageDesc( false, &context );
-            box<vec2<int16>> patchBox;
-			patchBox.min[0] = typedPatch->location[0];
-			patchBox.min[1] = typedPatch->location[1];
-			patchBox.size[0] = patchDesc.m_size[0];
-			patchBox.size[1] = patchDesc.m_size[1];
-
-            box<vec2<int16_t>> cropBox;
-            cropBox.min[0] = currentCropOp->op.args.ImageCrop.minX;
-            cropBox.min[1] = currentCropOp->op.args.ImageCrop.minY;
-            cropBox.size[0] = currentCropOp->op.args.ImageCrop.sizeX;
-            cropBox.size[1] = currentCropOp->op.args.ImageCrop.sizeY;
-
-            if (!patchBox.IntersectsExclusive(cropBox))
-            {
-                // We can ignore the patch
-                newAt = Visit( typedPatch->base.child(), currentCropOp );
-            }
-            else
-            {
-                // Crop the base with the full crop, and the patch with the intersected part,
-                // adapting the patch origin
-                Ptr<ASTOpImagePatch> newOp = mu::Clone<ASTOpImagePatch>(at);
-                newOp->base = Visit(newOp->base.child(), currentCropOp);
-
-                box<vec2<int16>> ibox = patchBox.Intersect(cropBox);
-                check(ibox.size[0]>0 && ibox.size[1]>0);
-
-                Ptr<ASTOpFixed> patchCropOp = mu::Clone<ASTOpFixed>(currentCropOp);
-                patchCropOp->op.args.ImageCrop.minX = ibox.min[0] - patchBox.min[0];
-                patchCropOp->op.args.ImageCrop.minY = ibox.min[1] - patchBox.min[1];
-                patchCropOp->op.args.ImageCrop.sizeX = ibox.size[0];
-                patchCropOp->op.args.ImageCrop.sizeY = ibox.size[1];
-                newOp->patch = Visit(newOp->patch.child(), patchCropOp.get());
-
-                newOp->location[0] = ibox.min[0] - cropBox.min[0];
-                newOp->location[1] = ibox.min[1] - cropBox.min[1];
-                newAt = newOp;
-            }
-
-            break;
-        }
-
-        case OP_TYPE::IM_CROP:
-        {
-            // We can combine the two crops into a possibly smaller crop
-			const ASTOpFixed* childCrop = dynamic_cast<const ASTOpFixed*>(at.get());
-
-            box<vec2<int16>> childCropBox;
-            childCropBox.min[0] = childCrop->op.args.ImageCrop.minX;
-            childCropBox.min[1] = childCrop->op.args.ImageCrop.minY;
-            childCropBox.size[0] = childCrop->op.args.ImageCrop.sizeX;
-            childCropBox.size[1] = childCrop->op.args.ImageCrop.sizeY;
-
-            box<vec2<int16>> cropBox;
-            cropBox.min[0] = currentCropOp->op.args.ImageCrop.minX;
-            cropBox.min[1] = currentCropOp->op.args.ImageCrop.minY;
-            cropBox.size[0] = currentCropOp->op.args.ImageCrop.sizeX;
-            cropBox.size[1] = currentCropOp->op.args.ImageCrop.sizeY;
-
-            // Compose the crops: in the final image the child crop is applied first and the
-            // current ctop is applied to the result. So the final crop box would be:
-            box<vec2<int16_t>> ibox;
-            ibox.min = childCropBox.min + cropBox.min;
-            ibox.size = vec2<int16>::min(cropBox.size, childCropBox.size);
-            check(cropBox.min.AllSmallerOrEqualThan(childCropBox.size));
-            check((cropBox.min + cropBox.size)
-                                .AllSmallerOrEqualThan(childCropBox.size));
-            check((ibox.min + ibox.size)
-                                .AllSmallerOrEqualThan(childCropBox.min + childCropBox.size));
-
-            // This happens more often that one would think
-            if (ibox==childCropBox)
-            {
-                // the parent crop is not necessary
-                skipSinking = true;
-            }
-            else if (ibox == cropBox)
-            {
-                // The child crop is not necessary
-                Ptr<ASTOp> childSource = childCrop->children[childCrop->op.args.ImageCrop.source].child();
-                newAt = Visit(childSource, currentCropOp);
-            }
-            else
-            {
-                // combine into one crop
-                Ptr<ASTOpFixed> newCropOp = mu::Clone<ASTOpFixed>(currentCropOp);
-                newCropOp->op.args.ImageCrop.minX = ibox.min[0];
-                newCropOp->op.args.ImageCrop.minY = ibox.min[1];
-                newCropOp->op.args.ImageCrop.sizeX = ibox.size[0];
-                newCropOp->op.args.ImageCrop.sizeY = ibox.size[1];
-
-				Ptr<ASTOp> childSource = childCrop->children[childCrop->op.args.ImageCrop.source].child();
-                newAt = Visit(childSource, newCropOp.get());
-            }
-            break;
-        }
-
-        case OP_TYPE::IM_LAYER:
-        {
-            // We move the op down the arguments
-			Ptr<ASTOpImageLayer> nop = mu::Clone<ASTOpImageLayer>(at);
-
-			Ptr<ASTOp> aOp = nop->base.child();
-            nop->base = Visit(aOp, currentCropOp);
-
-			Ptr<ASTOp> bOp = nop->blend.child();
-            nop->blend = Visit(bOp, currentCropOp);
-
-			Ptr<ASTOp> mOp = nop->mask.child();
-            nop->mask = Visit(mOp, currentCropOp);
-
-            newAt = nop;
-            break;
-        }
-
-        case OP_TYPE::IM_LAYERCOLOUR:
-        {
-            // We move the op down the arguments
-			Ptr<ASTOpImageLayerColor> nop = mu::Clone<ASTOpImageLayerColor>(at);
-
-			Ptr<ASTOp> aOp = nop->base.child();
-            nop->base = Visit(aOp, currentCropOp);
-
-			Ptr<ASTOp> mOp = nop->mask.child();
-            nop->mask = Visit(mOp, currentCropOp);
-
-            newAt = nop;
-            break;
-        }
-
-		case OP_TYPE::IM_DISPLACE:
-		{
-			// We move the op down the arguments
-			Ptr<ASTOpFixed> nop = mu::Clone<ASTOpFixed>(at);
-
-			Ptr<ASTOp> aOp = nop->children[nop->op.args.ImageDisplace.source].child();
-			nop->SetChild(nop->op.args.ImageDisplace.source, Visit(aOp, currentCropOp));
-
-			Ptr<ASTOp> bOp = nop->children[nop->op.args.ImageDisplace.displacementMap].child();
-			nop->SetChild(nop->op.args.ImageDisplace.displacementMap, Visit(bOp, currentCropOp));
-
-			newAt = nop;
-			break;
-		}
-
-		case OP_TYPE::IM_RASTERMESH:
-		{
-			// We add cropping data to the raster mesh if it doesn't have any
-			// \TODO: Is is possible to hit 2 crops on a raster mesh? Combine the crop.
-			Ptr<ASTOpImageRasterMesh> nop = mu::Clone<ASTOpImageRasterMesh>(at);
-
-			bool bRasterHasCrop = nop->UncroppedSizeX != 0;
-			if (!bRasterHasCrop)
-			{
-				box<vec2<int16>> cropBox;
-				cropBox.min[0] = currentCropOp->op.args.ImageCrop.minX;
-				cropBox.min[1] = currentCropOp->op.args.ImageCrop.minY;
-				cropBox.size[0] = currentCropOp->op.args.ImageCrop.sizeX;
-				cropBox.size[1] = currentCropOp->op.args.ImageCrop.sizeY;
-
-				nop->UncroppedSizeX = nop->SizeX;
-				nop->UncroppedSizeY = nop->SizeY;
-				nop->CropMinX = cropBox.min[0];
-				nop->CropMinY = cropBox.min[1];
-				nop->SizeX = cropBox.size[0];
-				nop->SizeY = cropBox.size[1];
-
-				newAt = nop;
-			}
-			break;
-		}
-
-        case OP_TYPE::IM_INTERPOLATE:
-        {
-            // Move the op  down all the paths
-            auto newOp = mu::Clone<ASTOpFixed>(at);
-
-            for (int v = 0; v < MUTABLE_OP_MAX_INTERPOLATE_COUNT; ++v)
-            {
-				Ptr<ASTOp> child = newOp->children[newOp->op.args.ImageInterpolate.targets[v]].child();
-				Ptr<ASTOp> bOp = Visit(child, currentCropOp);
-                newOp->SetChild(newOp->op.args.ImageInterpolate.targets[v], bOp);
-            }
-
-            newAt = newOp;
-            break;
-        }
-
-		default:
-			break;
-		}
-
-        // end on line, replace with crop
-        if (at==newAt && at!=m_initialSource && !skipSinking)
-        {
-            Ptr<ASTOpFixed> newOp = mu::Clone<ASTOpFixed>(currentCropOp);
-            check(newOp->GetOpType()==OP_TYPE::IM_CROP);
-
-            newOp->SetChild( newOp->op.args.ImageCrop.source,at);
-
-            newAt = newOp;
-        }
-
-		OldToNew.Add({ at, currentCropOp }, newAt);
-
-        return newAt;
-    }
-
-
 	//---------------------------------------------------------------------------------------------
 	//---------------------------------------------------------------------------------------------
 	//---------------------------------------------------------------------------------------------
@@ -841,32 +526,6 @@ namespace mu
         }
 
 
-        //-----------------------------------------------------------------------------------------
-        //-----------------------------------------------------------------------------------------
-        //-----------------------------------------------------------------------------------------
-        case OP_TYPE::IM_CROP:
-        {
-            auto sourceAt = children[op.args.ImageCrop.source].child();
-
-            switch ( sourceAt->GetOpType() )
-            {
-            // In case we have other operations with special optimisation rules.
-            case OP_TYPE::NONE:
-                break;
-
-            default:
-            {
-                at = context.ImageCropSinker.Apply(this);
-
-                break;
-            } // default
-
-            } // switch
-
-            break;
-        }
-
-
         //-----------------------------------------------------------------------------------------
         //-----------------------------------------------------------------------------------------
         //-----------------------------------------------------------------------------------------
diff --git a/Engine/Plugins/Experimental/Mutable/Source/MutableTools/Private/MuT/ImageDescGenerator.cpp b/Engine/Plugins/Experimental/Mutable/Source/MutableTools/Private/MuT/ImageDescGenerator.cpp
--- a/Engine/Plugins/Experimental/Mutable/Source/MutableTools/Private/MuT/ImageDescGenerator.cpp
+++ b/Engine/Plugins/Experimental/Mutable/Source/MutableTools/Private/MuT/ImageDescGenerator.cpp
@@ -16,6 +16,7 @@
 #include "MuT/NodeImageColourMapPrivate.h"
 #include "MuT/NodeImageConditionalPrivate.h"
 #include "MuT/NodeImageConstantPrivate.h"
+#include "MuT/NodeImageReferencePrivate.h"
 #include "MuT/NodeImageFormatPrivate.h"
 #include "MuT/NodeImageGradientPrivate.h"
 #include "MuT/NodeImageInterpolatePrivate.h"
@@ -99,7 +100,7 @@ namespace mu
 	//---------------------------------------------------------------------------------------------
 	Ptr<ASTOp> ImageDescGenerator::Visit( const NodeImageReference::Private& node )
 	{
-		m_desc = MUTABLE_MISSING_IMAGE_DESC;
+		m_desc = node.ImageDesc;
 		return 0;
 	}
 
@@ -132,12 +133,6 @@ namespace mu
         std::size_t i = 0;
         while ( !pImage && i<node.m_pTable->GetPrivate()->Rows.Num() )
         {
-			if (node.m_pTable->GetPrivate()->Rows[i].Values[colIndex].ProxyImage->Get()->IsReference())
-			{
-				// Image References do not need an image desc
-				break;
-			}
-
             pImage = node.m_pTable->GetPrivate()->Rows[i].Values[ colIndex ].ProxyImage->Get();
             ++i;
         }
diff --git a/Engine/Plugins/Experimental/Mutable/Source/MutableTools/Private/MuT/NodeImageParameter.cpp b/Engine/Plugins/Experimental/Mutable/Source/MutableTools/Private/MuT/NodeImageParameter.cpp
--- a/Engine/Plugins/Experimental/Mutable/Source/MutableTools/Private/MuT/NodeImageParameter.cpp
+++ b/Engine/Plugins/Experimental/Mutable/Source/MutableTools/Private/MuT/NodeImageParameter.cpp
@@ -52,30 +52,16 @@ namespace mu
 	//---------------------------------------------------------------------------------------------
 	// Own Interface
 	//---------------------------------------------------------------------------------------------
-    void NodeImageParameter::SetName( const char* strName )
+    void NodeImageParameter::SetName( const FString& Name )
 	{
-		if ( strName )
-		{
-			m_pD->m_name = strName;
-		}
-		else
-		{
-			m_pD->m_name = "";
-		}
+		m_pD->m_name = Name;
 	}
 
 
 	//---------------------------------------------------------------------------------------------
-	void NodeImageParameter::SetUid( const char* strUid )
+	void NodeImageParameter::SetUid(const FString& Uid)
 	{
-		if ( strUid )
-		{
-			m_pD->m_uid = strUid;
-		}
-		else
-		{
-			m_pD->m_uid = "";
-		}
+		m_pD->m_uid = Uid;
 	}
 
 	
diff --git a/Engine/Plugins/Experimental/Mutable/Source/MutableTools/Private/MuT/NodeImageReference.cpp b/Engine/Plugins/Experimental/Mutable/Source/MutableTools/Private/MuT/NodeImageReference.cpp
--- a/Engine/Plugins/Experimental/Mutable/Source/MutableTools/Private/MuT/NodeImageReference.cpp
+++ b/Engine/Plugins/Experimental/Mutable/Source/MutableTools/Private/MuT/NodeImageReference.cpp
@@ -12,25 +12,18 @@
 
 namespace mu
 {
-
-	//---------------------------------------------------------------------------------------------
-	// Static initialisation
-	//---------------------------------------------------------------------------------------------
 	NODE_TYPE NodeImageReference::Private::s_type = NODE_TYPE( "ImageReference", NodeImage::GetStaticType() );
 
 
-	//---------------------------------------------------------------------------------------------
 	MUTABLE_IMPLEMENT_NODE( NodeImageReference, EType::Reference, Node, Node::EType::Image);
 
 
-	//---------------------------------------------------------------------------------------------
 	int NodeImageReference::GetInputCount() const
 	{
 		return 0;
 	}
 
 
-	//---------------------------------------------------------------------------------------------
     Node* NodeImageReference::GetInputNode( int ) const
 	{
 		check( false );
@@ -38,17 +31,22 @@ namespace mu
 	}
 
 
-	//---------------------------------------------------------------------------------------------
     void NodeImageReference::SetInputNode( int, NodePtr )
 	{
 		check( false );
 	}
 
 
-	//---------------------------------------------------------------------------------------------
-	void NodeImageReference::SetImageReference( uint32 ID)
+	void NodeImageReference::SetImageReference(uint32 ID, const FImageDesc& Desc )
 	{
 		m_pD->ImageReferenceID = ID;
+		m_pD->ImageDesc = Desc;
+	}
+
+
+	void NodeImageReference::SetForceLoad(bool bForce)
+	{
+		m_pD->bForceLoad = bForce;
 	}
 
 
diff --git a/Engine/Plugins/Experimental/Mutable/Source/MutableTools/Private/MuT/NodeImageReferencePrivate.h b/Engine/Plugins/Experimental/Mutable/Source/MutableTools/Private/MuT/NodeImageReferencePrivate.h
--- a/Engine/Plugins/Experimental/Mutable/Source/MutableTools/Private/MuT/NodeImageReferencePrivate.h
+++ b/Engine/Plugins/Experimental/Mutable/Source/MutableTools/Private/MuT/NodeImageReferencePrivate.h
@@ -22,23 +22,37 @@ namespace mu
 
 		uint32 ImageReferenceID = 0;
 
-		//!
+		FImageDesc ImageDesc;
+		
+		bool bForceLoad = false;
+
 		void Serialise(OutputArchive& arch) const
 		{
-			uint32 Ver = 0;
+			uint32 Ver = 1;
 			arch << Ver;
 
 			arch << ImageReferenceID;
+			arch << ImageDesc.m_size;
+			arch << ImageDesc.m_lods;
+			arch << ImageDesc.m_format;
+			arch << bForceLoad;
 		}
 
-		//!
 		void Unserialise(InputArchive& arch)
 		{
 			uint32 Ver;
 			arch >> Ver;
-			check(Ver == 0);
+			check(Ver >= 0 && Ver <= 1);
 
 			arch >> ImageReferenceID;
+			
+			if (Ver >= 1)
+			{
+				arch >> ImageDesc.m_size;
+				arch >> ImageDesc.m_lods;
+				arch >> ImageDesc.m_format;
+				arch >> bForceLoad;
+			}
 		}
 	};
 
diff --git a/Engine/Plugins/Experimental/Mutable/Source/MutableTools/Private/MuT/UnrealPixelFormatOverride.cpp b/Engine/Plugins/Experimental/Mutable/Source/MutableTools/Private/MuT/UnrealPixelFormatOverride.cpp
--- a/Engine/Plugins/Experimental/Mutable/Source/MutableTools/Private/MuT/UnrealPixelFormatOverride.cpp
+++ b/Engine/Plugins/Experimental/Mutable/Source/MutableTools/Private/MuT/UnrealPixelFormatOverride.cpp
@@ -206,7 +206,7 @@ mu::EImageFormat UnrealToMutablePixelFormat(EPixelFormat PlatformFormat, bool bH
 }
 
 
-mu::EImageFormat QualityFix(mu::EImageFormat Format)
+mu::EImageFormat QualityAndPerformanceFix(mu::EImageFormat Format)
 {
 	switch (Format)
 	{
@@ -222,6 +222,10 @@ mu::EImageFormat QualityFix(mu::EImageFormat Format)
 	case mu::EImageFormat::IF_ASTC_10x10_RGB_LDR:	return mu::EImageFormat::IF_ASTC_4x4_RGB_LDR; break;
 	case mu::EImageFormat::IF_ASTC_10x10_RGBA_LDR:	return mu::EImageFormat::IF_ASTC_4x4_RGBA_LDR; break;
 	case mu::EImageFormat::IF_ASTC_10x10_RG_LDR:	return mu::EImageFormat::IF_ASTC_4x4_RG_LDR; break;
+
+	// This is more of a performance fix.
+	case mu::EImageFormat::IF_BGRA_UBYTE:			return mu::EImageFormat::IF_RGBA_UBYTE;
+
 	default:
 		break;
 	}
diff --git a/Engine/Plugins/Experimental/Mutable/Source/MutableTools/Public/MuT/NodeImageParameter.h b/Engine/Plugins/Experimental/Mutable/Source/MutableTools/Public/MuT/NodeImageParameter.h
--- a/Engine/Plugins/Experimental/Mutable/Source/MutableTools/Public/MuT/NodeImageParameter.h
+++ b/Engine/Plugins/Experimental/Mutable/Source/MutableTools/Public/MuT/NodeImageParameter.h
@@ -47,10 +47,10 @@ namespace mu
 		//-----------------------------------------------------------------------------------------
 
 		//! Set the name of the parameter.
-		void SetName( const char* );
+		void SetName(const FString&);
 
 		//! Set the uid of the parameter.
-		void SetUid( const char* );
+		void SetUid(const FString&);
 
     	//! Get the default value of the parameter.
     	void SetDefaultValue(FName Value);
diff --git a/Engine/Plugins/Experimental/Mutable/Source/MutableTools/Public/MuT/NodeImageReference.h b/Engine/Plugins/Experimental/Mutable/Source/MutableTools/Public/MuT/NodeImageReference.h
--- a/Engine/Plugins/Experimental/Mutable/Source/MutableTools/Public/MuT/NodeImageReference.h
+++ b/Engine/Plugins/Experimental/Mutable/Source/MutableTools/Public/MuT/NodeImageReference.h
@@ -49,7 +49,12 @@ namespace mu
 		//-----------------------------------------------------------------------------------------
 
 		/** Set this node value to be an "image reference" (to point to an unreal engine image). */
-		void SetImageReference(uint32 ID);
+		void SetImageReference(uint32 ID, const FImageDesc& Desc);
+
+		/** Set the image to be loaded as soon as it is used in the core. This will convert it into a normal image for other operations to process. 
+		* Disabled by default.
+		*/
+		void SetForceLoad(bool bForce);
 
 		//-----------------------------------------------------------------------------------------
 		// Interface pattern
diff --git a/Engine/Plugins/Experimental/Mutable/Source/MutableTools/Public/MuT/UnrealPixelFormatOverride.h b/Engine/Plugins/Experimental/Mutable/Source/MutableTools/Public/MuT/UnrealPixelFormatOverride.h
--- a/Engine/Plugins/Experimental/Mutable/Source/MutableTools/Public/MuT/UnrealPixelFormatOverride.h
+++ b/Engine/Plugins/Experimental/Mutable/Source/MutableTools/Public/MuT/UnrealPixelFormatOverride.h
@@ -26,7 +26,7 @@ void ImageCoreToMutable(const FCompressedImage2D& Compressed, mu::Image* Mutable
 MUTABLETOOLS_API extern mu::EImageFormat UnrealToMutablePixelFormat(EPixelFormat PlatformFormat, bool bHasAlpha);
 
 /** Function that remaps some formats to the ones that provide more acceptable quality in Mutable. */
-MUTABLETOOLS_API extern mu::EImageFormat QualityFix(mu::EImageFormat Format);
+MUTABLETOOLS_API extern mu::EImageFormat QualityAndPerformanceFix(mu::EImageFormat Format);
 
 /** Try to convert an image using unreal's compression. If the format is not supported, it will return false in bOutSuccess.
 * Can be called from any thread.
