[Mutable] Moved CustomizableObjectNodeSwitchBase source from public to private.

#rnx

diff --git a/Engine/Plugins/Experimental/Mutable/Source/CustomizableObjectEditor/Classes/MuCOE/Nodes/CustomizableObjectNodeSwitchBase.cpp b/Engine/Plugins/Experimental/Mutable/Source/CustomizableObjectEditor/Classes/MuCOE/Nodes/CustomizableObjectNodeSwitchBase.cpp
deleted file mode 100644
diff --git a/Engine/Plugins/Experimental/Mutable/Source/CustomizableObjectEditor/Private/MuCOE/Nodes/CustomizableObjectNodeSwitchBase.cpp b/Engine/Plugins/Experimental/Mutable/Source/CustomizableObjectEditor/Private/MuCOE/Nodes/CustomizableObjectNodeSwitchBase.cpp
new file mode 100644
--- /dev/null
+++ b/Engine/Plugins/Experimental/Mutable/Source/CustomizableObjectEditor/Private/MuCOE/Nodes/CustomizableObjectNodeSwitchBase.cpp
@@ -0,0 +1,217 @@
+// Copyright Epic Games, Inc. All Rights Reserved.
+
+#include "MuCOE/Nodes/CustomizableObjectNodeSwitchBase.h"
+
+#include "MuCOE/EdGraphSchema_CustomizableObject.h"
+#include "MuCOE/GraphTraversal.h"
+#include "MuCOE/Nodes/CustomizableObjectNodeEnumParameter.h"
+
+class UCustomizableObjectNodeRemapPins;
+
+#define LOCTEXT_NAMESPACE "CustomizableObjectEditor"
+
+void UCustomizableObjectNodeSwitchBase::ReloadEnumParam()
+{
+	// Get the names of the enum parameter for the element pins
+	if (const UEdGraphPin* EnumPin = SwitchParameter())
+	{
+		if (const UEdGraphPin* LinkedPin = FollowInputPin(*EnumPin))
+		{
+			if (UCustomizableObjectNodeEnumParameter* EnumNode = Cast<UCustomizableObjectNodeEnumParameter>(LinkedPin->GetOwningNode()))
+			{
+				ReloadingElementsNames.Empty(EnumNode->Values.Num());
+				for (int i = 0; i < EnumNode->Values.Num(); ++i)
+				{
+					ReloadingElementsNames.Add(EnumNode->Values[i].Name);
+				}
+			}
+		}
+	}
+}
+
+
+void UCustomizableObjectNodeSwitchBase::ReconstructNode(UCustomizableObjectNodeRemapPins* RemapPins)
+{
+	ReloadEnumParam();
+
+	Super::ReconstructNode(RemapPins);
+}
+
+
+UEdGraphPin* UCustomizableObjectNodeSwitchBase::OutputPin() const
+{
+	return OutputPinReference.Get();
+}
+
+
+UEdGraphPin* UCustomizableObjectNodeSwitchBase::SwitchParameter() const
+{
+	return SwitchParameterPinReference.Get();
+}
+
+
+void UCustomizableObjectNodeSwitchBase::PinConnectionListChanged(UEdGraphPin* Pin)
+{
+	if (Pin == SwitchParameter())
+	{
+		LinkPostEditChangePropertyDelegate(*Pin);
+	}
+
+	Super::PinConnectionListChanged(Pin);
+}
+
+
+void UCustomizableObjectNodeSwitchBase::AllocateDefaultPins(UCustomizableObjectNodeRemapPins* RemapPins)
+{
+	const UEdGraphSchema_CustomizableObject* Schema = GetDefault<UEdGraphSchema_CustomizableObject>();
+
+	UEdGraphPin* OutputPin = CustomCreatePin(EGPD_Output, GetCategory(), FName(UEdGraphSchema_CustomizableObject::GetPinCategoryName(GetCategory()).ToString()));
+	OutputPin->bDefaultValueIsIgnored = true;
+
+	OutputPinReference = FEdGraphPinReference(OutputPin);
+	
+	for (int LayerIndex = ReloadingElementsNames.Num() - 1; LayerIndex >= 0; --LayerIndex)
+	{
+		FString PinName = GetPinPrefix(LayerIndex);
+		UEdGraphPin* InputPin = CustomCreatePin(EGPD_Input, GetCategory(), FName(*PinName));
+		FString PinDisplayName = ReloadingElementsNames[LayerIndex].IsEmpty() ? PinName : ReloadingElementsNames[LayerIndex];
+		InputPin->PinFriendlyName = FText::FromString(PinDisplayName);
+		InputPin->bDefaultValueIsIgnored = true;
+		InputPin->SetOwningNode(this);
+	}
+
+	UEdGraphPin* SwitchParameterPin = CustomCreatePin(EGPD_Input, Schema->PC_Enum, FName(TEXT("Switch Parameter")));
+	SwitchParameterPin->bDefaultValueIsIgnored = true;
+	SwitchParameterPin->SetOwningNode(this);
+
+	SwitchParameterPinReference = FEdGraphPinReference(SwitchParameterPin);
+}
+
+
+void UCustomizableObjectNodeSwitchBase::Serialize(FArchive& Ar)
+{
+	Super::Serialize(Ar);
+
+	if (Ar.CustomVer(FCustomizableObjectCustomVersion::GUID) < FCustomizableObjectCustomVersion::BugPinsSwitch)
+	{
+		SwitchParameterPinReference = FindPin(TEXT("Switch Parameter"));	
+	}
+}
+
+
+FText UCustomizableObjectNodeSwitchBase::GetNodeTitle(ENodeTitleType::Type TitleType) const
+{
+	return FText::Format(LOCTEXT("Switch_Title", "{0} Switch"), UEdGraphSchema_CustomizableObject::GetPinCategoryName(GetCategory()));
+}
+
+
+FLinearColor UCustomizableObjectNodeSwitchBase::GetNodeTitleColor() const
+{
+	const UEdGraphSchema_CustomizableObject* Schema = GetDefault<UEdGraphSchema_CustomizableObject>();
+	return Schema->GetPinTypeColor(GetCategory());
+}
+
+
+FText UCustomizableObjectNodeSwitchBase::GetTooltipText() const
+{
+	return LOCTEXT("Switch_Tooltip", "Change the resulting value depending on what is currently chosen among a predefined amount of sources.");
+}
+
+
+void UCustomizableObjectNodeSwitchBase::EnumParameterPostEditChangeProperty(FPostEditChangePropertyDelegateParameters& Parameters)
+{
+	if (const UEdGraphPin* SwitchPin = SwitchParameter())
+	{
+		if (const UEdGraphPin* ConnectedPin = FollowInputPin(*SwitchPin); ConnectedPin && ConnectedPin->GetOwningNode() == Parameters.Node)
+		{
+			Super::ReconstructNode();
+		}
+		else if (UCustomizableObjectNode* EnumNode = Cast<UCustomizableObjectNode>(Parameters.Node))
+		{
+			EnumNode->PostEditChangePropertyDelegate.RemoveDynamic(this, &UCustomizableObjectNodeSwitchBase::EnumParameterPostEditChangeProperty);
+		}
+	}
+ }
+
+
+int32 UCustomizableObjectNodeSwitchBase::GetNumElements() const
+{
+	int32 Count = 0;
+	for (UEdGraphPin* Pin : GetAllNonOrphanPins())
+	{
+		if (Pin->GetName().StartsWith(GetPinPrefix()))
+		{
+			Count++;
+		}
+	}
+
+	return Count;
+}
+
+
+FString UCustomizableObjectNodeSwitchBase::GetPinPrefix(int32 Index) const
+{
+	return GetPinPrefix() + FString::FromInt(Index) + TEXT(" ");
+}
+
+
+void UCustomizableObjectNodeSwitchBase::PostPasteNode()
+{
+	Super::PostPasteNode();
+
+	if (UEdGraphPin* SwitchPin = SwitchParameter())
+	{
+		LinkPostEditChangePropertyDelegate(*SwitchPin);
+	}
+}
+
+
+void UCustomizableObjectNodeSwitchBase::PostBackwardsCompatibleFixup()
+{
+	Super::PostBackwardsCompatibleFixup();
+
+	if (UEdGraphPin* SwitchPin = SwitchParameter())
+	{
+		LinkPostEditChangePropertyDelegate(*SwitchPin);
+	}
+}
+
+
+void UCustomizableObjectNodeSwitchBase::LinkPostEditChangePropertyDelegate(const UEdGraphPin& Pin)
+{
+	if (LastNodeEnumParameterConnected.IsValid())
+	{
+		LastNodeEnumParameterConnected->PostEditChangePropertyDelegate.RemoveDynamic(this, &UCustomizableObjectNodeSwitchBase::EnumParameterPostEditChangeProperty);
+	}
+
+	if (const UEdGraphPin* ConnectedPin = FollowInputPin(Pin))
+	{
+		LastNodeEnumParameterConnected = Cast<UCustomizableObjectNode>(ConnectedPin->GetOwningNode());
+
+		if (LastNodeEnumParameterConnected.IsValid())
+		{
+			LastNodeEnumParameterConnected->PostEditChangePropertyDelegate.AddUniqueDynamic(this, &UCustomizableObjectNodeSwitchBase::EnumParameterPostEditChangeProperty);
+		}
+
+		if (Cast<UCustomizableObjectNodeEnumParameter>(ConnectedPin->GetOwningNode()))
+		{
+			Super::ReconstructNode();
+		}
+	}
+}
+
+
+FString UCustomizableObjectNodeSwitchBase::GetOutputPinName() const
+{
+	return FString();
+}
+
+
+FString UCustomizableObjectNodeSwitchBase::GetPinPrefix() const
+{
+	return UEdGraphSchema_CustomizableObject::GetPinCategoryName(GetCategory()).ToString() + " ";
+}
+
+
+#undef LOCTEXT_NAMESPACE
+
