[mutable] More compilation fixes.
#rnx

diff --git a/Engine/Plugins/Experimental/Mutable/Source/CustomizableObject/Private/MuCO/UnrealMutableImageProvider.h b/Engine/Plugins/Experimental/Mutable/Source/CustomizableObject/Private/MuCO/UnrealMutableImageProvider.h
--- a/Engine/Plugins/Experimental/Mutable/Source/CustomizableObject/Private/MuCO/UnrealMutableImageProvider.h
+++ b/Engine/Plugins/Experimental/Mutable/Source/CustomizableObject/Private/MuCO/UnrealMutableImageProvider.h
@@ -6,6 +6,8 @@
 #include "MuR/System.h"
 #include "UObject/GCObject.h"
 #include "Containers/Ticker.h"
+#include "Containers/Map.h"
+#include "Containers/Queue.h"
 
 #ifdef MUTABLE_USE_NEW_TASKGRAPH
 #include "Tasks/Task.h"
diff --git a/Engine/Plugins/Experimental/Mutable/Source/CustomizableObjectEditor/Private/MuCOE/SMutableCodeViewer.h b/Engine/Plugins/Experimental/Mutable/Source/CustomizableObjectEditor/Private/MuCOE/SMutableCodeViewer.h
--- a/Engine/Plugins/Experimental/Mutable/Source/CustomizableObjectEditor/Private/MuCOE/SMutableCodeViewer.h
+++ b/Engine/Plugins/Experimental/Mutable/Source/CustomizableObjectEditor/Private/MuCOE/SMutableCodeViewer.h
@@ -8,6 +8,8 @@
 #include "UObject/GCObject.h"
 #include "Widgets/SCompoundWidget.h"
 #include "Widgets/Input/SSearchBox.h"
+#include "Engine/Texture.h"
+#include "UObject/SoftObjectPtr.h"
 
 class STableViewBase;
 namespace ESelectInfo { enum Type : int; }
