[mutable] Removed some unused code.
#rnx

diff --git a/Engine/Plugins/Experimental/Mutable/Source/MutableTools/Private/MuT/CodeGenerator.h b/Engine/Plugins/Experimental/Mutable/Source/MutableTools/Private/MuT/CodeGenerator.h
--- a/Engine/Plugins/Experimental/Mutable/Source/MutableTools/Private/MuT/CodeGenerator.h
+++ b/Engine/Plugins/Experimental/Mutable/Source/MutableTools/Private/MuT/CodeGenerator.h
@@ -47,6 +47,7 @@
 namespace mu
 {
 	class ASTOpParameter;
+	class ASTOpConstantResource;
 	class Layout;
 	class NodeColourArithmeticOperation;
 	class NodeColourConstant;
@@ -289,8 +290,8 @@ namespace mu
 
 		TArray<FConditionalExtensionDataOp> m_conditionalExtensionDataOps;
 
-		/** This is used to cache the conatant image op generation. */
-		TMap<Ptr<const Image>,Ptr<class ASTOpConstantResource>> ImageConstantOpPerImage;
+		/** This is used to cache the constant image op generation. */
+		//TMap<Ptr<const Image>,Ptr<ASTOpConstantResource>> ImageConstantOpPerImage;
 
 		//-----------------------------------------------------------------------------------------
 
diff --git a/Engine/Plugins/Experimental/Mutable/Source/MutableTools/Private/MuT/CodeGenerator_Image.cpp b/Engine/Plugins/Experimental/Mutable/Source/MutableTools/Private/MuT/CodeGenerator_Image.cpp
--- a/Engine/Plugins/Experimental/Mutable/Source/MutableTools/Private/MuT/CodeGenerator_Image.cpp
+++ b/Engine/Plugins/Experimental/Mutable/Source/MutableTools/Private/MuT/CodeGenerator_Image.cpp
@@ -225,46 +225,16 @@ namespace mu
 			// Do we need to crop?
 			if (rect.min[0]!=0 || rect.min[1]!=0 || pImage->GetSizeX() != rect.size[0] || pImage->GetSizeY() != rect.size[1])
 			{
-				// Option 1: Crop now
-				if (true)
-				{
-					FImageOperator ImOp = FImageOperator::GetDefault(m_compilerOptions->ImageFormatFunc);
-
-					Ptr<Image> pCropped = new Image(rect.size[0], rect.size[1], 1, pImage->GetFormat(), EInitializationType::NotInitialized);
-					ImOp.ImageCrop(pCropped.get(), m_compilerOptions->ImageCompressionQuality, pImage.get(), rect);
+				// Crop now
+				FImageOperator ImOp = FImageOperator::GetDefault(m_compilerOptions->ImageFormatFunc);
 
-					Ptr<ASTOpConstantResource> op = new ASTOpConstantResource();
-					op->type = OP_TYPE::IM_CONSTANT;
-					op->SetValue(pCropped, m_compilerOptions->OptimisationOptions.bUseDiskCache);
-					Result.op = op;
-				}
+				Ptr<Image> pCropped = new Image(rect.size[0], rect.size[1], 1, pImage->GetFormat(), EInitializationType::NotInitialized);
+				ImOp.ImageCrop(pCropped.get(), m_compilerOptions->ImageCompressionQuality, pImage.get(), rect);
 
-				// Option 2: Generate a crop instruction: seems to be a lot slower and give worse results.
-				else
-				{
-					//Ptr<ASTOpConstantResource>* Found = ImageConstantOpPerImage.Find( pImage );
-					//Ptr<ASTOpConstantResource> ConstantOp;
-					//if (Found)
-					//{
-					//	ConstantOp = *Found;
-					//}
-					//else
-					//{ 
-					//	ConstantOp = new ASTOpConstantResource();
-					//	ConstantOp->type = OP_TYPE::IM_CONSTANT;
-					//	ConstantOp->SetValue(pImage, m_compilerOptions->OptimisationOptions.bUseDiskCache);
-					//	ImageConstantOpPerImage.Add(pImage, ConstantOp);
-					//}
-
-					//Ptr<ASTOpImageCrop> CropOp = new ASTOpImageCrop();
-					//CropOp->Source = ConstantOp;
-					//CropOp->Min[0] = rect.min[0];
-					//CropOp->Min[1] = rect.min[1];
-					//CropOp->Size[0] = rect.size[0];
-					//CropOp->Size[1] = rect.size[1];
-
-					//Result.op = CropOp;
-				}
+				Ptr<ASTOpConstantResource> op = new ASTOpConstantResource();
+				op->type = OP_TYPE::IM_CONSTANT;
+				op->SetValue(pCropped, m_compilerOptions->OptimisationOptions.bUseDiskCache);
+				Result.op = op;
 			}
 			else
 			{
