#rnx

diff --git a/Engine/Plugins/Experimental/Mutable/Source/CustomizableObject/Private/MuCO/CustomizableObjectInstance.cpp b/Engine/Plugins/Experimental/Mutable/Source/CustomizableObject/Private/MuCO/CustomizableObjectInstance.cpp
--- a/Engine/Plugins/Experimental/Mutable/Source/CustomizableObject/Private/MuCO/CustomizableObjectInstance.cpp
+++ b/Engine/Plugins/Experimental/Mutable/Source/CustomizableObject/Private/MuCO/CustomizableObjectInstance.cpp
@@ -6802,10 +6802,8 @@ TSet<UAssetUserData*> UCustomizableObjectInstance::GetMergedAssetUserData(int32
 {
 	UCustomizableInstancePrivateData* PrivateInstanceData = GetPrivate();
 
-	if (PrivateInstanceData)
+	if (PrivateInstanceData && PrivateInstanceData->ComponentsData.IsValidIndex(ComponentIndex))
 	{
-		check(PrivateInstanceData->ComponentsData.IsValidIndex(ComponentIndex));
-
 		TSet<UAssetUserData*> Set;
 		
 		// Have to convert to UAssetUserData* because BP functions don't support TObjectPtr
