[Mutable] Remove bReuseMesh and related code to simplify the generation of meshes to ease the implementation of mesh streaming.
* Remove LastUpdateData
* Remove CopySockets and CopySkeletonData

#rnx

diff --git a/Engine/Plugins/Experimental/Mutable/Source/CustomizableObject/Private/MuCO/CustomizableInstancePrivateData.h b/Engine/Plugins/Experimental/Mutable/Source/CustomizableObject/Private/MuCO/CustomizableInstancePrivateData.h
--- a/Engine/Plugins/Experimental/Mutable/Source/CustomizableObject/Private/MuCO/CustomizableInstancePrivateData.h
+++ b/Engine/Plugins/Experimental/Mutable/Source/CustomizableObject/Private/MuCO/CustomizableInstancePrivateData.h
@@ -285,8 +285,6 @@ public:
 
 private:
 
-	void InitLastUpdateData(const TSharedPtr<FMutableOperationData>& OperationData);
-
 	void InitSkeletalMeshData(const TSharedPtr<FMutableOperationData>& OperationData, USkeletalMesh* SkeletalMesh, const FMutableRefSkeletalMeshData* RefSkeletalMeshData, const UCustomizableObject& CustomizableObject, int32 ComponentIndex);
 
 	bool BuildSkeletonData(const TSharedPtr<FMutableOperationData>& OperationData, USkeletalMesh& SkeletalMesh, const FMutableRefSkeletalMeshData& RefSkeletalMeshData, UCustomizableObject& CustomizableObject, int32 ComponentIndex);
@@ -295,10 +293,7 @@ private:
 	void BuildOrCopyMorphTargetsData(const TSharedPtr<FMutableOperationData>& OperationData, USkeletalMesh* SkeletalMesh, const USkeletalMesh* SrcSkeletalMesh, UCustomizableObjectInstance* CustomizableObjectInstance, int32 ComponentIndex);
 	bool BuildOrCopyRenderData(const TSharedPtr<FMutableOperationData>& OperationData, USkeletalMesh* SkeletalMesh, const USkeletalMesh* SrcSkeletalMesh, UCustomizableObjectInstance* CustomizableObjectInstance, int32 ComponentIndex);
 	void BuildOrCopyClothingData(const TSharedPtr<FMutableOperationData>& OperationData, USkeletalMesh* SkeletalMesh, const USkeletalMesh* SrcSkeletalMesh, UCustomizableObjectInstance* CustomizableObjectInstance, int32 ComponentIndex);
-
-	bool CopySkeletonData(const TSharedPtr<FMutableOperationData>& OperationData, USkeletalMesh* SrcSkeletalMesh, USkeletalMesh* DestSkeletalMesh, const UCustomizableObject& CustomizableObject, int32 ComponentIndex);
-	void CopyMeshSockets(USkeletalMesh* SrcSkeletalMesh, USkeletalMesh* DestSkeletalMesh);
-
+	
 	//
 	USkeleton* MergeSkeletons(UCustomizableObject& CustomizableObject, const FMutableRefSkeletalMeshData& RefSkeletalMeshData, int32 ComponentIndex);
 
@@ -367,9 +362,6 @@ public:
 	UPROPERTY(Transient)
 	TArray<TObjectPtr<UTexture>> LoadedPassThroughTexturesPendingSetMaterial;
 
-	// Struct used during an update to avoid generating resources that can be reused.
-	FInstanceGeneratedData LastUpdateData;
-
 private:
 	
 	ECOInstanceFlags InstanceFlagsPrivate = ECOInstanceFlags::ECONone;
diff --git a/Engine/Plugins/Experimental/Mutable/Source/CustomizableObject/Private/MuCO/CustomizableObjectInstance.cpp b/Engine/Plugins/Experimental/Mutable/Source/CustomizableObject/Private/MuCO/CustomizableObjectInstance.cpp
--- a/Engine/Plugins/Experimental/Mutable/Source/CustomizableObject/Private/MuCO/CustomizableObjectInstance.cpp
+++ b/Engine/Plugins/Experimental/Mutable/Source/CustomizableObject/Private/MuCO/CustomizableObjectInstance.cpp
@@ -103,40 +103,6 @@ void UCustomizableInstancePrivateData::SetSkeletalMeshStatus(ESkeletalMeshStatus
 }
 
 
-void UCustomizableInstancePrivateData::InitLastUpdateData(const TSharedPtr<FMutableOperationData>& OperationData)
-{
-	LastUpdateData.LODs.Init(FInstanceGeneratedData::FLOD(), OperationData->NumLODsAvailable);
-	for (int32 LODIndex = OperationData->CurrentMinLOD; LODIndex <= OperationData->CurrentMaxLOD; ++LODIndex)
-	{
-		LastUpdateData.LODs[LODIndex].FirstComponent = OperationData->InstanceUpdateData.LODs[LODIndex].FirstComponent;
-		LastUpdateData.LODs[LODIndex].ComponentCount = OperationData->InstanceUpdateData.LODs[LODIndex].ComponentCount;
-	}
-
-	const int32 NumComponents = OperationData->InstanceUpdateData.Components.Num();
-	LastUpdateData.Components.Init(FInstanceGeneratedData::FComponent(), NumComponents);
-	for (int32 ComponentIndex = 0; ComponentIndex < NumComponents; ++ComponentIndex)
-	{
-		FInstanceGeneratedData::FComponent& LastUpdateDataComponent = LastUpdateData.Components[ComponentIndex];
-		const FInstanceUpdateData::FComponent& UpdateComponent = OperationData->InstanceUpdateData.Components[ComponentIndex];
-		
-		LastUpdateDataComponent.ComponentId = UpdateComponent.Id;
-		LastUpdateDataComponent.bGenerated = UpdateComponent.bGenerated;
-
-		LastUpdateDataComponent.MeshID = UpdateComponent.MeshID;
-
-		LastUpdateDataComponent.FirstSurface = UpdateComponent.FirstSurface;
-		LastUpdateDataComponent.SurfaceCount = UpdateComponent.SurfaceCount;
-	}
-
-	const int32 NumSurfaces = OperationData->InstanceUpdateData.Surfaces.Num();
-	LastUpdateData.SurfaceIds.Init(0, NumSurfaces);
-	for (int32 SurfaceIndex = 0; SurfaceIndex < NumSurfaces; ++SurfaceIndex)
-	{
-		LastUpdateData.SurfaceIds[SurfaceIndex] = OperationData->InstanceUpdateData.Surfaces[SurfaceIndex].SurfaceId;
-	}
-}
-
-
 void UCustomizableInstancePrivateData::InvalidateGeneratedData()
 {
 	// Init Component Data
@@ -145,7 +111,6 @@ void UCustomizableInstancePrivateData::InvalidateGeneratedData()
 	ComponentsData.Init(TemplateComponentData, ComponentsData.Num());
 
 	GeneratedMaterials.Empty();
-	LastUpdateData.Clear();
 	ClearCOInstanceFlags(Generated);
 }
 
@@ -224,9 +189,6 @@ void UCustomizableObjectInstance::SetDescriptor(const FCustomizableObjectInstanc
 
 void UCustomizableInstancePrivateData::PrepareForUpdate(const TSharedPtr<FMutableOperationData>& OperationData)
 {
-	// Prepare LastUpdateData to allow reuse of meshes and surfaces
-	InitLastUpdateData(OperationData);
-
 	// Clear the ComponentData from previous updates
 	for (FCustomizableInstanceComponentData& ComponentData : ComponentsData)
 	{
@@ -1735,11 +1697,6 @@ bool UCustomizableInstancePrivateData::DoComponentsNeedUpdate(UCustomizableObjec
 							bFoundEmptyMesh = true;
 						}
 					}
-					else if (Component.bReuseMesh)
-					{
-						bFoundNonEmptyMesh = true;
-						check(LastMeshID != MAX_uint64);
-					}
 					else
 					{
 						OutComponentNeedsUpdate[Component.Id] = true;
@@ -1834,7 +1791,7 @@ bool UCustomizableInstancePrivateData::UpdateSkeletalMesh_PostBeginUpdate0(UCust
 			continue;
 		}
 
-		if (!Component.bGenerated || (!Component.Mesh && !Component.bReuseMesh) || Component.SurfaceCount == 0)
+		if (!Component.bGenerated || !Component.Mesh || Component.SurfaceCount == 0)
 		{
 			continue;
 		}
@@ -1935,26 +1892,6 @@ bool UCustomizableInstancePrivateData::UpdateSkeletalMesh_PostBeginUpdate0(UCust
 			// Add sockets from the SkeletalMesh of reference and from the MutableMesh
 			BuildMeshSockets(OperationData, SkeletalMesh, RefSkeletalMeshData, Public, Component.Mesh);
 		}
-		else if (Component.bReuseMesh)
-		{
-			// Source mesh to reuse data from
-			USkeletalMesh* SrcSkeletalMesh = OldSkeletalMeshes[Component.Id];
-			check(SrcSkeletalMesh);
-
-			// Set the Skeleton from the previously generated mesh and fix up ActiveBones and Bonemap arrays from new components
-			bSuccess = CopySkeletonData(OperationData, SrcSkeletalMesh, SkeletalMesh, *CustomizableObject, Component.Id);
-
-			if (!bSuccess)
-			{
-				break;
-			}
-
-			// Set the PhysicsAsset from the previously generated mesh
-			SkeletalMesh->SetPhysicsAsset(SrcSkeletalMesh->GetPhysicsAsset());
-
-			// Copy Sockets from the previously generated mesh
-			CopyMeshSockets(SrcSkeletalMesh, SkeletalMesh);
-		}
 		else
 		{
 			bSuccess = false;
@@ -2104,16 +2041,6 @@ bool UCustomizableInstancePrivateData::UpdateSkeletalMesh_PostBeginUpdate0(UCust
 }
 
 
-static TAutoConsoleVariable<int32> CVarPendingReleaseSkeletalMeshMode(
-	TEXT("mutable.PendingReleaseSkeletalMesh"),
-	1,
-	TEXT("Whether use pending for release or not \n")
-	TEXT(" 0: Use GC\n")
-	TEXT(" 1: Use PendingRelease\n")
-	TEXT(" Default use"),
-	ECVF_SetByConsole);
-
-
 UCustomizableInstancePrivateData::UCustomizableInstancePrivateData()
 {
 	MinSquareDistFromComponentToPlayer = FLT_MAX;
@@ -3463,93 +3390,6 @@ bool UCustomizableInstancePrivateData::BuildSkeletonData(const TSharedPtr<FMutab
 }
 
 
-bool UCustomizableInstancePrivateData::CopySkeletonData(const TSharedPtr<FMutableOperationData>& OperationData, USkeletalMesh* SrcSkeletalMesh, USkeletalMesh* DestSkeletalMesh, const UCustomizableObject& CustomizableObject, int32 ComponentIndex)
-{
-	MUTABLE_CPUPROFILER_SCOPE(UCustomizableInstancePrivateData::CopySkeletonData);
-
-	check(SrcSkeletalMesh);
-	check(DestSkeletalMesh);
-
-	const TObjectPtr<USkeleton> Skeleton = SrcSkeletalMesh->GetSkeleton();
-	if (!Skeleton)
-	{
-		return false;
-	}
-
-	DestSkeletalMesh->SetSkeleton(Skeleton);
-	DestSkeletalMesh->SetRefSkeleton(Skeleton->GetReferenceSkeleton());
-	DestSkeletalMesh->SetRefBasesInvMatrix(SrcSkeletalMesh->GetRefBasesInvMatrix());
-
-	const FReferenceSkeleton& RefSkeleton = DestSkeletalMesh->GetRefSkeleton();
-
-	// Check that the bones we need are present in the current Skeleton
-	FInstanceUpdateData::FSkeletonData& MutSkeletonData = OperationData->InstanceUpdateData.Skeletons[ComponentIndex];
-	const int32 MutBoneCount = MutSkeletonData.BoneIds.Num();
-	
-	TMap<uint16, uint16> BoneToFinalBoneIndexMap;
-	BoneToFinalBoneIndexMap.Reserve(MutBoneCount);
-
-	{
-		MUTABLE_CPUPROFILER_SCOPE(CopySkeletonData_CheckForMissingBones);
-
-		const TArray<FName>& BoneNames = CustomizableObject.GetBoneNamesArray();
-
-		// Ensure all the required bones are present in the skeleton
-		for (int32 BoneIndex = 0; BoneIndex < MutBoneCount; ++BoneIndex)
-		{
-			const uint16 BoneId = MutSkeletonData.BoneIds[BoneIndex];
-			check(BoneNames.IsValidIndex(BoneId));
-
-			const FName BoneName = BoneNames[BoneId];
-			check(BoneName != NAME_None);
-
-			const int32 SourceBoneIndex = RefSkeleton.FindRawBoneIndex(BoneName);
-			
-			if (SourceBoneIndex == INDEX_NONE)
-			{
-				// Merged skeleton is missing some bones! This happens if one of the skeletons involved in the merge is discarded due to being incompatible with the rest
-				// or if the source mesh is not in sync with the skeleton. 
-				UE_LOG(LogMutable, Warning, TEXT("Building instance: generated mesh has a bone [%s] not present in the reference mesh [%s]. Failing to generate mesh. "),
-					*BoneName.ToString(), *DestSkeletalMesh->GetName());
-				return false;
-			}
-
-			BoneToFinalBoneIndexMap.Add(BoneId, SourceBoneIndex);
-		}
-	}
-
-	{
-		MUTABLE_CPUPROFILER_SCOPE(CopySkeletonData_FixBoneIndices);
-
-		// Fix up BoneMaps and ActiveBones indices
-		for (FInstanceUpdateData::FComponent& Component : OperationData->InstanceUpdateData.Components)
-		{
-			if (Component.Id != ComponentIndex)
-			{
-				continue;
-			}
-
-			for (uint32 BoneMapIndex = Component.FirstBoneMap; BoneMapIndex < Component.FirstBoneMap + Component.BoneMapCount; ++BoneMapIndex)
-			{
-				const uint16 BoneId = OperationData->InstanceUpdateData.BoneMaps[BoneMapIndex];
-				OperationData->InstanceUpdateData.BoneMaps[BoneMapIndex] = BoneToFinalBoneIndexMap[BoneId];
-			}
-
-			TArray<uint16> UniqueActiveBones;
-			for (const uint16 BoneId : Component.ActiveBones)
-			{
-				UniqueActiveBones.AddUnique(BoneToFinalBoneIndexMap[BoneId]);
-			}
-
-			Exchange(Component.ActiveBones, UniqueActiveBones);
-			Component.ActiveBones.Sort();
-		}
-	}
-
-	return true;
-}
-
-
 void UCustomizableInstancePrivateData::BuildMeshSockets(const TSharedPtr<FMutableOperationData>& OperationData, USkeletalMesh* SkeletalMesh, const FMutableRefSkeletalMeshData* RefSkeletalMeshData, UCustomizableObjectInstance* Public, mu::MeshPtrConst MutableMesh)
 {
 	// Build mesh sockets.
@@ -3650,40 +3490,6 @@ void UCustomizableInstancePrivateData::BuildMeshSockets(const TSharedPtr<FMutabl
 }
 
 
-void UCustomizableInstancePrivateData::CopyMeshSockets(USkeletalMesh* SrcSkeletalMesh, USkeletalMesh* DestSkeletalMesh)
-{
-	// Copy mesh sockets.
-	MUTABLE_CPUPROFILER_SCOPE(UCustomizableInstancePrivateData::CopyMeshSockets);
-
-	check(SrcSkeletalMesh);
-	check(DestSkeletalMesh);
-
-	const TArray<USkeletalMeshSocket*>& SrcSockets = SrcSkeletalMesh->GetMeshOnlySocketList();
-	TArray<TObjectPtr<USkeletalMeshSocket>>& DestSockets = DestSkeletalMesh->GetMeshOnlySocketList();
-	DestSockets.Empty(SrcSockets.Num());
-
-	for (const USkeletalMeshSocket* SrcSocket : SrcSockets)
-	{
-		if (SrcSocket)
-		{
-			USkeletalMeshSocket* NewSocket = NewObject<USkeletalMeshSocket>(DestSkeletalMesh);
-
-
-			NewSocket->SocketName = SrcSocket->SocketName;
-			NewSocket->BoneName = SrcSocket->BoneName;
-
-			NewSocket->RelativeLocation = SrcSocket->RelativeLocation;
-			NewSocket->RelativeRotation = SrcSocket->RelativeRotation;
-			NewSocket->RelativeScale = SrcSocket->RelativeScale;
-
-			NewSocket->bForceAlwaysAnimated = SrcSocket->bForceAlwaysAnimated;
-			DestSockets.Add(NewSocket);
-
-		}
-	}
-}
-
-
 void UCustomizableInstancePrivateData::BuildOrCopyElementData(const TSharedPtr<FMutableOperationData>& OperationData, USkeletalMesh* SkeletalMesh, UCustomizableObjectInstance* CustomizableObjectInstance, int32 ComponentIndex)
 {
 	MUTABLE_CPUPROFILER_SCOPE(UCustomizableInstancePrivateData::BuildOrCopyElementData);
@@ -3699,7 +3505,7 @@ void UCustomizableInstancePrivateData::BuildOrCopyElementData(const TSharedPtr<F
 
 		const FInstanceUpdateData::FComponent& Component = OperationData->InstanceUpdateData.Components[LOD.FirstComponent + ComponentIndex];
 
-		if (!Component.bGenerated || Component.bReuseMesh)
+		if (!Component.bGenerated)
 		{
 			continue;
 		}
@@ -3750,7 +3556,7 @@ void UCustomizableInstancePrivateData::BuildOrCopyMorphTargetsData(const TShared
 	{
 		const FInstanceUpdateData::FLOD& LOD = OperationData->InstanceUpdateData.LODs[LODIndex];
 
-		if (LODIndex >= OperationData->CurrentMinLOD && OperationData->InstanceUpdateData.Components[OperationData->InstanceUpdateData.LODs[LODIndex].FirstComponent + ComponentIndex].bGenerated)
+		if (LODIndex >= OperationData->CurrentMinLOD && OperationData->InstanceUpdateData.Components[LOD.FirstComponent + ComponentIndex].bGenerated)
 		{
 			LastValidLODIndex = LODIndex;
 		}
@@ -3837,20 +3643,10 @@ void UCustomizableInstancePrivateData::BuildOrCopyMorphTargetsData(const TShared
 			MUTABLE_CPUPROFILER_SCOPE(BuildOrCopyMorphTargetsData_Copy);
 
 			// Copy data from the last valid LOD generated from either the current mesh or the previous mesh
-			const FInstanceUpdateData::FLOD& LastLOD = OperationData->InstanceUpdateData.LODs[LastValidLODIndex];
-			const FInstanceUpdateData::FComponent& LastComponent = OperationData->InstanceUpdateData.Components[LastLOD.FirstComponent + ComponentIndex];
-
-			const USkeletalMesh* SrcSkeletalMesh = LastComponent.bReuseMesh ? LastUpdateSkeletalMesh : SkeletalMesh;
-			if (!SrcSkeletalMesh)
-			{
-				check(false);
-				return;
-			}
-
-			const int32 NumMorphTargets = SrcSkeletalMesh->GetMorphTargets().Num();
+			const int32 NumMorphTargets = SkeletalMesh->GetMorphTargets().Num();
 			for (int32 MorphTargetIndex = 0; MorphTargetIndex < NumMorphTargets; ++MorphTargetIndex)
 			{
-				SkeletalMesh->GetMorphTargets()[MorphTargetIndex]->GetMorphLODModels()[LODIndex] = SrcSkeletalMesh->GetMorphTargets()[MorphTargetIndex]->GetMorphLODModels()[LastValidLODIndex];
+				SkeletalMesh->GetMorphTargets()[MorphTargetIndex]->GetMorphLODModels()[LODIndex] = SkeletalMesh->GetMorphTargets()[MorphTargetIndex]->GetMorphLODModels()[LastValidLODIndex];
 			}
 		}
 	}
@@ -3931,14 +3727,6 @@ void UCustomizableInstancePrivateData::BuildOrCopyClothingData(const TSharedPtr<
 				continue;
 			}
 
-
-			if (!Component.Mesh && Component.bReuseMesh)
-			{
-				// TODO PERE: Implement copying LODRenderDatas with clothing
-				unimplemented()
-			}
-
-
 			if (mu::MeshPtrConst MutableMesh = Component.Mesh)
 			{
 				const mu::FMeshBufferSet& MeshSet = MutableMesh->GetVertexBuffers();
@@ -4027,9 +3815,6 @@ void UCustomizableInstancePrivateData::BuildOrCopyClothingData(const TSharedPtr<
 
 					SectionsWithCloth.Add
 							(FSectionWithClothData{ ClothAssetIndex, ClothAssetLodIndex, Section, LODIndex, FirstVertex, IndicesView16Bits, IndicesView32Bits, ClothingDataView, TArray<FMeshToMeshVertData>() });
-
-					// TODO Pere: Remove once we support the copy of LODRenderDatas with clothing
-					LastUpdateData.Components[LOD.FirstComponent + ComponentIndex].bGenerated = false;
 				}
 			}
 		}
@@ -4753,9 +4538,9 @@ void UCustomizableInstancePrivateData::BuildOrCopyClothingData(const TSharedPtr<
 
 bool UCustomizableInstancePrivateData::BuildOrCopyRenderData(const TSharedPtr<FMutableOperationData>& OperationData, USkeletalMesh* SkeletalMesh, const USkeletalMesh* LastUpdateSkeletalMesh, UCustomizableObjectInstance* Public, int32 ComponentIndex)
 {
-	MUTABLE_CPUPROFILER_SCOPE(UCustomizableInstancePrivateData::BuildOrCopyRenderData)
+	MUTABLE_CPUPROFILER_SCOPE(UCustomizableInstancePrivateData::BuildOrCopyRenderData);
 
-		UCustomizableObject* CustomizableObject = Public->GetCustomizableObject();
+	UCustomizableObject* CustomizableObject = Public->GetCustomizableObject();
 	check(CustomizableObject);
 
 	int32 LastValidLODIndex = OperationData->CurrentMaxLOD;
@@ -4764,9 +4549,9 @@ bool UCustomizableInstancePrivateData::BuildOrCopyRenderData(const TSharedPtr<FM
 		MUTABLE_CPUPROFILER_SCOPE(BuildOrCopyRenderData_LODLoop);
 
 		const FInstanceUpdateData::FLOD& LOD = OperationData->InstanceUpdateData.LODs[LODIndex];
-
-		bool bCopyComponent = false;
-		if (!LOD.ComponentCount)
+		FInstanceUpdateData::FComponent* Component = nullptr;
+		
+		if (ComponentIndex >=  LOD.ComponentCount)
 		{
 			// Interrupt the generation if the LOD is empty and it should have been generated.
 			if (LODIndex >= OperationData->CurrentMinLOD && LODIndex <= OperationData->CurrentMaxLOD)
@@ -4778,56 +4563,28 @@ bool UCustomizableInstancePrivateData::BuildOrCopyRenderData(const TSharedPtr<FM
 				// End with failure
 				return false;
 			}
-
-			// LODIndex is below the CurrentMinLOD copy the data from the LastValidLODIndex
-			bCopyComponent = true;
 		}
-
 		else
 		{
-			check(ComponentIndex < LOD.ComponentCount);
-			FInstanceUpdateData::FComponent& Component = OperationData->InstanceUpdateData.Components[LOD.FirstComponent + ComponentIndex];
-
-			if (!Component.bGenerated || Component.bReuseMesh)
-			{
-				bCopyComponent = true;
-			}
-			else
-			{
-				LastValidLODIndex = LODIndex;
-			}
+			Component = &OperationData->InstanceUpdateData.Components[LOD.FirstComponent + ComponentIndex];
 		}
 
-
-		if (bCopyComponent)
+		if (!Component || !Component->bGenerated)
 		{
-			// Copy data from the last valid LOD generated from either the current mesh or the previous mesh
-			const FInstanceUpdateData::FLOD& LastValidLOD = OperationData->InstanceUpdateData.LODs[LastValidLODIndex];
-			const FInstanceUpdateData::FComponent& LastValidComponent = OperationData->InstanceUpdateData.Components[LastValidLOD.FirstComponent + ComponentIndex];
-
-			const USkeletalMesh* SrcSkeletalMesh = LastValidComponent.bReuseMesh ? LastUpdateSkeletalMesh : SkeletalMesh;
-			if (!SrcSkeletalMesh)
-			{
-				check(false);
-				return false;
-			}
-
 			TRACE_CPUPROFILER_EVENT_SCOPE_TEXT(*FString::Printf(TEXT("BuildOrCopyRenderData_CopyData: From LOD %d to LOD %d"), LastValidLODIndex, LODIndex));
 
 			// Render Data will be reused from the previously generated component
-			UnrealConversionUtils::CopySkeletalMeshLODRenderData(SrcSkeletalMesh, SkeletalMesh, LastValidLODIndex, LODIndex);
+			UnrealConversionUtils::CopySkeletalMeshLODRenderData(SkeletalMesh, SkeletalMesh, LastValidLODIndex, LODIndex);
 			continue;
 		}
 
-		FInstanceUpdateData::FComponent& Component = OperationData->InstanceUpdateData.Components[OperationData->InstanceUpdateData.LODs[LastValidLODIndex].FirstComponent + ComponentIndex];
-
 		// There could be components without a mesh in LODs
-		if (!Component.Mesh)
+		if (!Component->Mesh)
 		{
 			UE_LOG(LogMutable, Warning, TEXT("Building instance: generated mesh [%s] has LOD [%d] of component [%d] with no mesh.")
 				, *SkeletalMesh->GetName()
 				, LODIndex
-				, Component.Id);
+				, Component->Id);
 
 			// End with failure
 			return false;
@@ -4835,18 +4592,18 @@ bool UCustomizableInstancePrivateData::BuildOrCopyRenderData(const TSharedPtr<FM
 
 		TRACE_CPUPROFILER_EVENT_SCOPE_TEXT(*FString::Printf(TEXT("BuildOrCopyRenderData_BuildData: LOD %d"), LODIndex));
 
-		SetLastMeshId(Component.Id, LODIndex, Component.MeshID);
+		SetLastMeshId(Component->Id, LODIndex, Component->MeshID);
 
 		UnrealConversionUtils::SetupRenderSections(
 			SkeletalMesh,
-			Component.Mesh,
+			Component->Mesh,
 			LODIndex,
 			OperationData->InstanceUpdateData.BoneMaps,
-			Component.FirstBoneMap);
+			Component->FirstBoneMap);
 
 		UnrealConversionUtils::CopyMutableVertexBuffers(
 			SkeletalMesh,
-			Component.Mesh,
+			Component->Mesh,
 			LODIndex);
 
 		FSkeletalMeshLODRenderData& LODModel = SkeletalMesh->GetResourceForRendering()->LODRenderData[LODIndex];
@@ -4869,15 +4626,15 @@ bool UCustomizableInstancePrivateData::BuildOrCopyRenderData(const TSharedPtr<FM
 		}
 
 		// Update active and required bones
-		LODModel.ActiveBoneIndices.Append(Component.ActiveBones);
-		LODModel.RequiredBones.Append(Component.ActiveBones);
+		LODModel.ActiveBoneIndices.Append(Component->ActiveBones);
+		LODModel.RequiredBones.Append(Component->ActiveBones);
 
 		const TArray<FMutableSkinWeightProfileInfo>& SkinWeightProfilesInfo = CustomizableObject->SkinWeightProfilesInfo;
 		if (SkinWeightProfilesInfo.Num())
 		{
 			bool bHasSkinWeightProfiles = false;
 
-			const mu::FMeshBufferSet& MutableMeshVertexBuffers = Component.Mesh->GetVertexBuffers();
+			const mu::FMeshBufferSet& MutableMeshVertexBuffers = Component->Mesh->GetVertexBuffers();
 
 			const int32 SkinWeightProfilesCount = CustomizableObject->SkinWeightProfilesInfo.Num();
 			for (int32 ProfileIndex = 0; ProfileIndex < SkinWeightProfilesCount; ++ProfileIndex)
@@ -4919,14 +4676,14 @@ bool UCustomizableInstancePrivateData::BuildOrCopyRenderData(const TSharedPtr<FM
 		}
 
 		// Copy indices.
-		if (!UnrealConversionUtils::CopyMutableIndexBuffers(Component.Mesh, LODModel))
+		if (!UnrealConversionUtils::CopyMutableIndexBuffers(Component->Mesh, LODModel))
 		{
 			// End with failure
 			return false;
 		}
 
 		// Update LOD and streaming data
-		const FMutableRefLODRenderData& RefLODRenderData = CustomizableObject->GetRefSkeletalMeshData(Component.Id)->LODData[LODIndex].RenderData;
+		const FMutableRefLODRenderData& RefLODRenderData = CustomizableObject->GetRefSkeletalMeshData(Component->Id)->LODData[LODIndex].RenderData;
 		LODModel.bIsLODOptional = RefLODRenderData.bIsLODOptional;
 		LODModel.bStreamedDataInlined = RefLODRenderData.bStreamedDataInlined;
 
@@ -4935,6 +4692,7 @@ bool UCustomizableInstancePrivateData::BuildOrCopyRenderData(const TSharedPtr<FM
 		// in the map of textures to stream.
 		LODModel.BuffersSize = 1;
 
+		LastValidLODIndex = LODIndex;
 	}
 
 
diff --git a/Engine/Plugins/Experimental/Mutable/Source/CustomizableObject/Private/MuCO/CustomizableObjectSystem.cpp b/Engine/Plugins/Experimental/Mutable/Source/CustomizableObject/Private/MuCO/CustomizableObjectSystem.cpp
--- a/Engine/Plugins/Experimental/Mutable/Source/CustomizableObject/Private/MuCO/CustomizableObjectSystem.cpp
+++ b/Engine/Plugins/Experimental/Mutable/Source/CustomizableObject/Private/MuCO/CustomizableObjectSystem.cpp
@@ -646,13 +646,6 @@ static FAutoConsoleVariableRef CVarEnableMutableReuseInstanceTextures(
 	ECVF_Default);
 
 
-bool FCustomizableObjectSystemPrivate::bEnableMutableReusePreviousUpdateData = false;
-static FAutoConsoleVariableRef CVarEnableMutableReusePreviousUpdateData(
-	TEXT("mutable.EnableMutableReusePreviousUpdateData"), FCustomizableObjectSystemPrivate::bEnableMutableReusePreviousUpdateData,
-	TEXT("If true, Mutable will try to reuse render sections from previous SkeletalMeshes. If false, all SkeletalMeshes will be build from scratch."),
-	ECVF_Default);
-
-
 int32 FCustomizableObjectSystemPrivate::EnableOnlyGenerateRequestedLODs = 1;
 
 static FAutoConsoleVariableRef CVarEnableOnlyGenerateRequestedLODs(
@@ -1571,10 +1564,6 @@ namespace impl
 			LOD.FirstComponent = OperationData->InstanceUpdateData.Components.Num();
 			LOD.ComponentCount = Instance->GetComponentCount(MutableLODIndex);
 
-			const FInstanceGeneratedData::FLOD& PrevUpdateLODData = OperationData->LastUpdateData.LODs.IsValidIndex(MutableLODIndex) ?
-				OperationData->LastUpdateData.LODs[MutableLODIndex] : FInstanceGeneratedData::FLOD();
-
-
 			for (int32 ComponentIndex = 0; ComponentIndex < LOD.ComponentCount; ++ComponentIndex)
 			{
 				OperationData->InstanceUpdateData.Components.Push(FInstanceUpdateData::FComponent());
@@ -1583,9 +1572,6 @@ namespace impl
 				Component.FirstSurface = OperationData->InstanceUpdateData.Surfaces.Num();
 				Component.SurfaceCount = 0;
 
-				const FInstanceGeneratedData::FComponent& PrevUpdateComponent = OperationData->LastUpdateData.Components.IsValidIndex(PrevUpdateLODData.FirstComponent + ComponentIndex) ?
-					OperationData->LastUpdateData.Components[PrevUpdateLODData.FirstComponent +  ComponentIndex] : FInstanceGeneratedData::FComponent();
-
 				const bool bGenerateLOD = OperationData->RequestedLODs.IsValidIndex(Component.Id) ? (OperationData->RequestedLODs[Component.Id] & (1 << MutableLODIndex)) != 0 : true;
 
 				// Mesh
@@ -1595,38 +1581,27 @@ namespace impl
 
 					Component.MeshID = Instance->GetMeshId(MutableLODIndex, ComponentIndex, 0);
 
-					// Mesh cache is not enabled yet.
-					//auto CachedMesh = Cache.Meshes.Find(meshId);
-					//if (CachedMesh && CachedMesh->IsValid(false, true))
-					//{
-					//	UE_LOG(LogMutable, Verbose, TEXT("Mesh resource with id [%d] can be cached."), meshId);
-					//	INC_DWORD_STAT(STAT_MutableNumCachedSkeletalMeshes);
-					//}
-
-					// Check if we can use data from the currently generated SkeletalMesh
-					Component.bReuseMesh = OperationData->bCanReuseGeneratedData && PrevUpdateComponent.bGenerated ? PrevUpdateComponent.MeshID == Component.MeshID : false;
-
-					if(bGenerateLOD && !Component.bReuseMesh)
+					if(bGenerateLOD)
 					{
 						Component.Mesh = System->GetMesh(OperationData->InstanceID, Component.MeshID);
 					}
 				}
 
-				if (!Component.Mesh && !Component.bReuseMesh)
+				if (!Component.Mesh)
 				{
 					continue;
 				}
 
 				Component.bGenerated = true;
 
-				const int32 SurfaceCount = Component.Mesh ? Component.Mesh->GetSurfaceCount() : PrevUpdateComponent.SurfaceCount;
 
 				// Materials and images
+				const int32 SurfaceCount = Component.Mesh->GetSurfaceCount();
 				for (int32 MeshSurfaceIndex = 0; MeshSurfaceIndex < SurfaceCount; ++MeshSurfaceIndex)
 				{
-					const uint32 SurfaceId = Component.Mesh ? Component.Mesh->GetSurfaceId(MeshSurfaceIndex) : OperationData->LastUpdateData.SurfaceIds[PrevUpdateComponent.FirstSurface + MeshSurfaceIndex];
+					const uint32 SurfaceId = Component.Mesh->GetSurfaceId(MeshSurfaceIndex);
 					const int32 InstanceSurfaceIndex = Instance->FindSurfaceById(MutableLODIndex, ComponentIndex, SurfaceId);
-					check(Component.bReuseMesh || Component.Mesh->GetVertexCount() == 0 || InstanceSurfaceIndex >= 0);
+					check(Component.Mesh->GetVertexCount() > 0 || InstanceSurfaceIndex >= 0);
 
 					int32 BaseSurfaceIndex = InstanceSurfaceIndex;
 					int32 BaseLODIndex = MutableLODIndex;
@@ -2700,8 +2675,6 @@ namespace impl
 		CurrentOperationData->Instance = Operation->CustomizableObjectInstance;
 		CurrentOperationData->TextureCoverageQueries_MutableThreadParams = CandidateInstancePrivateData->TextureCoverageQueries;
 		CurrentOperationData->TextureCoverageQueries_MutableThreadResults.Empty();
-		CurrentOperationData->bCanReuseGeneratedData = SystemPrivateData->bEnableMutableReusePreviousUpdateData;
-		CurrentOperationData->LastUpdateData = SystemPrivateData->bEnableMutableReusePreviousUpdateData ? CandidateInstancePrivateData->LastUpdateData : FInstanceGeneratedData();
 		CurrentOperationData->CurrentMinLOD = Operation->InstanceDescriptorRuntimeHash.GetMinLOD();
 		CurrentOperationData->CurrentMaxLOD = Operation->InstanceDescriptorRuntimeHash.GetMaxLOD();
 		CurrentOperationData->bNeverStream = bNeverStream;
diff --git a/Engine/Plugins/Experimental/Mutable/Source/CustomizableObject/Private/MuCO/CustomizableObjectSystemPrivate.h b/Engine/Plugins/Experimental/Mutable/Source/CustomizableObject/Private/MuCO/CustomizableObjectSystemPrivate.h
--- a/Engine/Plugins/Experimental/Mutable/Source/CustomizableObject/Private/MuCO/CustomizableObjectSystemPrivate.h
+++ b/Engine/Plugins/Experimental/Mutable/Source/CustomizableObject/Private/MuCO/CustomizableObjectSystemPrivate.h
@@ -352,44 +352,6 @@ struct FMutableTask
 };
 
 
-// Mutable data generated during the last update of an instance.
-struct FInstanceGeneratedData
-{
-	struct FComponent
-	{
-		uint16 ComponentId = 0;
-
-		/** True if it can be reused */
-		bool bGenerated = false;
-
-		mu::FResourceID MeshID;
-
-		/** Range in the Surfaces array */
-		uint16 FirstSurface = 0;
-		uint16 SurfaceCount = 0;
-	};
-
-	struct FLOD
-	{
-		/** Range in the Components array */
-		uint16 FirstComponent = 0;
-		uint16 ComponentCount = 0;
-	};
-
-	TArray<FLOD> LODs;
-	TArray<FComponent> Components;
-	TArray<uint32> SurfaceIds;
-
-	/** Clear data, called upon failing to generate a mesh and after recompiling the CO */
-	void Clear()
-	{
-		LODs.Empty();
-		Components.Empty();
-		SurfaceIds.Empty();
-	}
-};
-
-
 // Mutable data generated during the update steps.
 // We keep it from begin to end update, and it is used in several steps.
 struct FInstanceUpdateData
@@ -446,12 +408,9 @@ struct FInstanceUpdateData
 	{
 		uint16 Id = 0;
 		
-		// True if the Mesh is valid or if we're reusing a component
+		// True if the Mesh is valid
 		bool bGenerated = false;
 
-		// Reuse component from a previously generated SkeletalMesh
-		bool bReuseMesh = false;
-
 		mu::FResourceID MeshID;
 		mu::MeshPtrConst Mesh;
 
@@ -546,9 +505,6 @@ struct FPendingTextureCoverageQuery
 struct FMutableOperationData
 {
 	TWeakObjectPtr<UCustomizableObjectInstance> Instance;
-	
-	bool bCanReuseGeneratedData = false;
-	FInstanceGeneratedData LastUpdateData;
 
 	FInstanceUpdateData InstanceUpdateData;
 	TArray<int> RelevantParametersInProgress;
@@ -678,7 +634,6 @@ public:
 	static int32 EnableOnlyGenerateRequestedLODs;
 	static int32 MaxTextureSizeToGenerate;
 	static int32 SkeletalMeshMinLodQualityLevel;
-	static bool bEnableMutableReusePreviousUpdateData;
 
 #if WITH_EDITOR
 	mu::FImageOperator::FImagePixelFormatFunc ImageFormatOverrideFunc;
