[mutable] Debugger: fix a mismatch between the generic generated code tree and the op-specific. Add a validation step to detect future problems.
#rnx

diff --git a/Engine/Plugins/Experimental/Mutable/Source/CustomizableObjectEditor/Private/MuCOE/SMutableCodeViewer.cpp b/Engine/Plugins/Experimental/Mutable/Source/CustomizableObjectEditor/Private/MuCOE/SMutableCodeViewer.cpp
--- a/Engine/Plugins/Experimental/Mutable/Source/CustomizableObjectEditor/Private/MuCOE/SMutableCodeViewer.cpp
+++ b/Engine/Plugins/Experimental/Mutable/Source/CustomizableObjectEditor/Private/MuCOE/SMutableCodeViewer.cpp
@@ -1436,7 +1436,6 @@ void SMutableCodeViewer::GenerateElementRecursive(const int32& InStateIndex, mu:
 
 					const FSlateColor LabelColor = ColorPerComputationalCost[StaticCast<uint8>(GetOperationTypeComputationalCost(InProgram.GetOpType(ChildAddress)))];
 
-					// No caption for the generic tree
 					const TSharedPtr<FMutableCodeTreeElement> Item = MakeShareable(new FMutableCodeTreeElement(ItemCache.Num(),InStateIndex, MutableModel, ChildAddress, Caption, LabelColor, MainItemPtr));
 
 					// Cache this element for later access
@@ -1456,6 +1455,7 @@ void SMutableCodeViewer::GenerateElementRecursive(const int32& InStateIndex, mu:
 
 	
 	// For some specific parent operation types we create more detailed subtrees.
+	bool bUseGeneric = false;
 	const mu::OP_TYPE ParentOperationType = InProgram.GetOpType(InParentAddress);
 	switch (ParentOperationType)
 	{
@@ -1494,8 +1494,12 @@ void SMutableCodeViewer::GenerateElementRecursive(const int32& InStateIndex, mu:
 			FMemory::Memcpy(&At, OpData, sizeof(mu::OP::ADDRESS));
 			OpData += sizeof(mu::OP::ADDRESS);
 
-			FString Caption = FString::Printf(TEXT("case %d "),Condition);
-			AddOpFunc(At, Caption);
+			// This conditional is necessary to exactly match the generic op behaviour (see ForEachReference in Operations.cpp)
+			if (At)
+			{
+				FString Caption = FString::Printf(TEXT("case %d "), Condition);
+				AddOpFunc(At, Caption);
+			}
 		}
 
 		break;
@@ -1563,15 +1567,40 @@ void SMutableCodeViewer::GenerateElementRecursive(const int32& InStateIndex, mu:
 	default:
 	{
 		// Generic list of child operations
+		bUseGeneric = true;
+		break;
+	}
+
+	}
 
+	if (bUseGeneric)
+	{
 		// Find children of the provided element
 		mu::ForEachReference(InProgram, InParentAddress, [this, &InProgram, AddOpFunc](mu::OP::ADDRESS ChildAddress)
 			{
 				AddOpFunc(ChildAddress,TEXT(""));
 			});
-		break;
 	}
+	else
+	{
+		// Validate in case there is a mismatch in the custom processing of children and the generic one, which would cause problems.
+		ChildIndex = 0;
 
+		auto ValidateOpFunc = [this, InParentAddress, &ChildIndex](mu::OP::ADDRESS ChildAddress)
+		{
+			if (ChildAddress)
+			{
+				const FItemCacheKey Key = { InParentAddress, ChildAddress, ChildIndex };
+				const TSharedPtr<FMutableCodeTreeElement>* CachedItem = ItemCache.Find(Key);
+				check(CachedItem);
+			}
+			++ChildIndex;
+		};
+
+		mu::ForEachReference(InProgram, InParentAddress, [this, &InProgram, ValidateOpFunc](mu::OP::ADDRESS ChildAddress)
+			{
+				ValidateOpFunc(ChildAddress);
+			});
 	}
 }
 
