[Mutable] Disable Compile Only Selected option due to crashing.

#rnx

diff --git a/Engine/Plugins/Experimental/Mutable/Source/CustomizableObjectEditor/Private/MuCOE/CustomizableObjectEditor.cpp b/Engine/Plugins/Experimental/Mutable/Source/CustomizableObjectEditor/Private/MuCOE/CustomizableObjectEditor.cpp
--- a/Engine/Plugins/Experimental/Mutable/Source/CustomizableObjectEditor/Private/MuCOE/CustomizableObjectEditor.cpp
+++ b/Engine/Plugins/Experimental/Mutable/Source/CustomizableObjectEditor/Private/MuCOE/CustomizableObjectEditor.cpp
@@ -1257,7 +1257,7 @@ void FCustomizableObjectEditor::ExtendToolbar()
 		{
 			ToolbarBuilder.BeginSection("Compilation");
 			ToolbarBuilder.AddToolBarButton(FCustomizableObjectEditorCommands::Get().Compile);
-			ToolbarBuilder.AddToolBarButton(FCustomizableObjectEditorCommands::Get().CompileOnlySelected);
+			//ToolbarBuilder.AddToolBarButton(FCustomizableObjectEditorCommands::Get().CompileOnlySelected); // Disabled due to a crash with this mode (MTBL-1716). 
 			ToolbarBuilder.AddComboButton(
 				FUIAction(),
 				FOnGetContent::CreateSP(Editor, &FCustomizableObjectEditor::GenerateCompileOptionsMenuContent, CommandList.ToSharedRef()),
