[Mutable] Added more specific context to Non-normalized UVs warnings.

- Added ErrorContext to all resources.
- Changed Log functoins to support UObjects in general instead of only UCustomizableObjectNode.

#rnx

diff --git a/Engine/Plugins/Experimental/Mutable/Source/CustomizableObjectEditor/Private/MuCOE/CompilationMessageCache.cpp b/Engine/Plugins/Experimental/Mutable/Source/CustomizableObjectEditor/Private/MuCOE/CompilationMessageCache.cpp
--- a/Engine/Plugins/Experimental/Mutable/Source/CustomizableObjectEditor/Private/MuCOE/CompilationMessageCache.cpp
+++ b/Engine/Plugins/Experimental/Mutable/Source/CustomizableObjectEditor/Private/MuCOE/CompilationMessageCache.cpp
@@ -7,10 +7,10 @@ FCompilationMessageCache::FCompilationMessageCache()
 	ClearMessageCounters();
 }
 
-bool FCompilationMessageCache::AddMessage(const FText& InMessage, const TArray<const UCustomizableObjectNode*>& InArrayNode, EMessageSeverity::Type MessageSeverity /* = EMessageSeverity::Warning */, const ELoggerSpamBin SpamBin /*= ELoggerSpamBin::ShowAll*/)
+bool FCompilationMessageCache::AddMessage(const FText& InMessage, const TArray<const UObject*>& InContext, EMessageSeverity::Type MessageSeverity /* = EMessageSeverity::Warning */, const ELoggerSpamBin SpamBin /*= ELoggerSpamBin::ShowAll*/)
 {
 	// Skip message if an identical one has already been reported
-	const FLoggedMessage Cached = { InMessage, InArrayNode, MessageSeverity };
+	const FLoggedMessage Cached = { InMessage, InContext, MessageSeverity };
 	if (LoggedMessages.Contains(Cached))
 	{
 		return false;
diff --git a/Engine/Plugins/Experimental/Mutable/Source/CustomizableObjectEditor/Private/MuCOE/CustomizableObjectCompiler.cpp b/Engine/Plugins/Experimental/Mutable/Source/CustomizableObjectEditor/Private/MuCOE/CustomizableObjectCompiler.cpp
--- a/Engine/Plugins/Experimental/Mutable/Source/CustomizableObjectEditor/Private/MuCOE/CustomizableObjectCompiler.cpp
+++ b/Engine/Plugins/Experimental/Mutable/Source/CustomizableObjectEditor/Private/MuCOE/CustomizableObjectCompiler.cpp
@@ -407,7 +407,7 @@ void FCustomizableObjectCompiler::ProcessChildObjectsRecursively(UCustomizableOb
 
 void FCustomizableObjectCompiler::DisplayParameterWarning(FMutableGraphGenerationContext& GenerationContext)
 {
-	for (const TPair<FString, TArray<const UCustomizableObjectNode*>>& It : GenerationContext.ParameterNamesMap)
+	for (const TPair<FString, TArray<const UObject*>>& It : GenerationContext.ParameterNamesMap)
 	{
 		if (It.Key == "")
 		{
@@ -425,7 +425,7 @@ void FCustomizableObjectCompiler::DisplayParameterWarning(FMutableGraphGeneratio
 
 void FCustomizableObjectCompiler::DisplayDuplicatedNodeIdsWarning(FMutableGraphGenerationContext & GenerationContext)
 {
-	for (const TPair<FGuid, TArray<const UCustomizableObjectNode*>>& It : GenerationContext.NodeIdsMap)
+	for (const TPair<FGuid, TArray<const UObject*>>& It : GenerationContext.NodeIdsMap)
 	{
 		if (It.Value.Num() > 1)
 		{
@@ -1424,11 +1424,8 @@ void FCustomizableObjectCompiler::CompileInternal(UCustomizableObject* Object, c
 		}
 	}
 
-	// Get a list of generated nodes, to be able to understand error messages later.
 	for (UCustomizableObjectNode* Node : GenerationContext.GeneratedNodes)
 	{
-		GeneratedNodes.Add(Node, Node);
-
 		Node->ResetAttachedErrorData();
 	}
 
@@ -1682,13 +1679,13 @@ void FCustomizableObjectCompiler::NotifyCompilationErrors() const
 }
 
 
-void FCustomizableObjectCompiler::CompilerLog(const FText& Message, const TArray<const UCustomizableObjectNode*>& ArrayNode, const EMessageSeverity::Type MessageSeverity, const bool bAddBaseObjectInfo, const ELoggerSpamBin SpamBin)
+void FCustomizableObjectCompiler::CompilerLog(const FText& Message, const TArray<const UObject*>& Context, const EMessageSeverity::Type MessageSeverity, const bool bAddBaseObjectInfo, const ELoggerSpamBin SpamBin)
 {
-	if (CompilationLogsContainer.AddMessage(Message, ArrayNode, MessageSeverity, SpamBin)) // Cache the message for later reference
+	if (CompilationLogsContainer.AddMessage(Message, Context, MessageSeverity, SpamBin)) // Cache the message for later reference
 	{
 		FCustomizableObjectEditorLogger::CreateLog(Message)
 			.Severity(MessageSeverity)
-			.Nodes(ArrayNode)
+			.Context(Context)
 			.BaseObject(bAddBaseObjectInfo)
 			.SpamBin(SpamBin)
 			.Log();
@@ -1696,14 +1693,14 @@ void FCustomizableObjectCompiler::CompilerLog(const FText& Message, const TArray
 }
 
 
-void FCustomizableObjectCompiler::CompilerLog(const FText& Message, const UCustomizableObjectNode* Node, const EMessageSeverity::Type MessageSeverity, const bool bAddBaseObjectInfo, const ELoggerSpamBin SpamBin)
+void FCustomizableObjectCompiler::CompilerLog(const FText& Message, const UObject* Context, const EMessageSeverity::Type MessageSeverity, const bool bAddBaseObjectInfo, const ELoggerSpamBin SpamBin)
 {
-	TArray<const UCustomizableObjectNode*> ArrayNode;
-	if (Node)
+	TArray<const UObject*> ContextArray;
+	if (Context)
 	{
-		ArrayNode.Add(Node);
+		ContextArray.Add(Context);
 	}
-	CompilerLog(Message, ArrayNode, MessageSeverity, bAddBaseObjectInfo, SpamBin);
+	CompilerLog(Message, ContextArray, MessageSeverity, bAddBaseObjectInfo, SpamBin);
 }
 
 
@@ -1713,25 +1710,26 @@ void FCustomizableObjectCompiler::UpdateCompilerLogData()
 	MessageLogModule.RegisterLogListing(FName("Mutable"), LOCTEXT("MutableLog", "Mutable"));
 	const TArray<FCustomizableObjectCompileRunnable::FError>& ArrayCompileErrors = CompileTask->GetArrayErrors();
 
-	FText ObjectName = CurrentObject ? FText::FromString(CurrentObject->GetName()) : LOCTEXT("Unknown Object", "Unknown Object");
+	const FText ObjectName = CurrentObject ? FText::FromString(CurrentObject->GetName()) : LOCTEXT("Unknown Object", "Unknown Object");
 
-	int32 i;
-	for (i = 0; i < ArrayCompileErrors.Num(); ++i)
+	for (const FCustomizableObjectCompileRunnable::FError& CompileError : ArrayCompileErrors)
 	{
-		const UCustomizableObjectNode** pNode = GeneratedNodes.Find(ArrayCompileErrors[i].Context);
+		const UObject* Object = static_cast<const UObject*>(CompileError.Context); // Context are always UObjects
 
-		if (ArrayCompileErrors[i].AttachedData && pNode)
+		if (const UCustomizableObjectNode* Node = Cast<UCustomizableObjectNode>(Object))
 		{
-			UCustomizableObjectNode::FAttachedErrorDataView ErrorDataView;
-			ErrorDataView.UnassignedUVs = { ArrayCompileErrors[i].AttachedData->UnassignedUVs.GetData(),
-											ArrayCompileErrors[i].AttachedData->UnassignedUVs.Num() };
+			if (CompileError.AttachedData)
+			{
+				UCustomizableObjectNode::FAttachedErrorDataView ErrorDataView;
+				ErrorDataView.UnassignedUVs = { CompileError.AttachedData->UnassignedUVs.GetData(),
+												CompileError.AttachedData->UnassignedUVs.Num() };
 
-			const UCustomizableObjectNode* Node = Cast<UCustomizableObjectNode>(*pNode);
-			const_cast<UCustomizableObjectNode*>(Node)->AddAttachedErrorData(ErrorDataView);
+				const_cast<UCustomizableObjectNode*>(Node)->AddAttachedErrorData(ErrorDataView);
+			}			
 		}
 
-		FText FullMsg = FText::Format(LOCTEXT("MutableMessage", "{0} : {1}"), ObjectName, ArrayCompileErrors[i].Message);
-		CompilerLog(FullMsg, pNode ? *pNode : nullptr, ArrayCompileErrors[i].Severity, true, ArrayCompileErrors[i].SpamBin);
+		FText FullMsg = FText::Format(LOCTEXT("MutableMessage", "{0} : {1}"), ObjectName, CompileError.Message);
+		CompilerLog(FullMsg, Object, CompileError.Severity, true, CompileError.SpamBin);
 	}
 }
 
diff --git a/Engine/Plugins/Experimental/Mutable/Source/CustomizableObjectEditor/Private/MuCOE/CustomizableObjectCompiler.h b/Engine/Plugins/Experimental/Mutable/Source/CustomizableObjectEditor/Private/MuCOE/CustomizableObjectCompiler.h
--- a/Engine/Plugins/Experimental/Mutable/Source/CustomizableObjectEditor/Private/MuCOE/CustomizableObjectCompiler.h
+++ b/Engine/Plugins/Experimental/Mutable/Source/CustomizableObjectEditor/Private/MuCOE/CustomizableObjectCompiler.h
@@ -45,8 +45,8 @@ public:
 	/** Generate the Mutable Graph from the Unreal Graph. */
 	mu::NodePtr Export(UCustomizableObject* Object, const FCompilationOptions& Options, TArray<TSoftObjectPtr<UTexture>>& OutReferencedTextures);
 
-	void CompilerLog(const FText& Message, const TArray<const UCustomizableObjectNode*>& ArrayNode, const EMessageSeverity::Type MessageSeverity = EMessageSeverity::Warning, const bool bAddBaseObjectInfo = true, const ELoggerSpamBin SpamBin = ELoggerSpamBin::ShowAll);
-	void CompilerLog(const FText& Message, const UCustomizableObjectNode* Node = nullptr, const EMessageSeverity::Type MessageSeverity = EMessageSeverity::Warning, const bool bAddBaseObjectInfo = true, const ELoggerSpamBin SpamBin = ELoggerSpamBin::ShowAll);
+	void CompilerLog(const FText& Message, const TArray<const UObject*>& UObject, const EMessageSeverity::Type MessageSeverity = EMessageSeverity::Warning, const bool bAddBaseObjectInfo = true, const ELoggerSpamBin SpamBin = ELoggerSpamBin::ShowAll);
+	void CompilerLog(const FText& Message, const UObject* Context = nullptr, const EMessageSeverity::Type MessageSeverity = EMessageSeverity::Warning, const bool bAddBaseObjectInfo = true, const ELoggerSpamBin SpamBin = ELoggerSpamBin::ShowAll);
 	void NotifyCompilationErrors() const;
 
 	void FinishCompilation();
@@ -239,9 +239,6 @@ private:
 	// Stores the only option of an Int Param that should be compiled
 	TMap<FString, FString> ParamNamesToSelectedOptions;
 
-	// List of nodes used during generation, so that they can be found when processing mutable runtime error messages.
-	TMap<const void*, const UCustomizableObjectNode*> GeneratedNodes;
-
 	/** Pointer to the Asynchronous Preloading process call back */
 	TSharedPtr<struct FStreamableHandle> AsynchronousStreamableHandlePtr;
 	
diff --git a/Engine/Plugins/Experimental/Mutable/Source/CustomizableObjectEditor/Private/MuCOE/CustomizableObjectEditor.cpp b/Engine/Plugins/Experimental/Mutable/Source/CustomizableObjectEditor/Private/MuCOE/CustomizableObjectEditor.cpp
--- a/Engine/Plugins/Experimental/Mutable/Source/CustomizableObjectEditor/Private/MuCOE/CustomizableObjectEditor.cpp
+++ b/Engine/Plugins/Experimental/Mutable/Source/CustomizableObjectEditor/Private/MuCOE/CustomizableObjectEditor.cpp
@@ -2476,7 +2476,7 @@ void FCustomizableObjectEditor::LogSearchResult(UCustomizableObjectNode* Node, F
 	}
 	
 	FCustomizableObjectEditorLogger::CreateLog(FText::FromString(Type + ": " + Result))
-	.Node(*Node)
+	.Context(*Node)
 	.BaseObject()
 	.Notification(false)
 	.Log();
diff --git a/Engine/Plugins/Experimental/Mutable/Source/CustomizableObjectEditor/Private/MuCOE/CustomizableObjectEditorLogger.cpp b/Engine/Plugins/Experimental/Mutable/Source/CustomizableObjectEditor/Private/MuCOE/CustomizableObjectEditorLogger.cpp
--- a/Engine/Plugins/Experimental/Mutable/Source/CustomizableObjectEditor/Private/MuCOE/CustomizableObjectEditorLogger.cpp
+++ b/Engine/Plugins/Experimental/Mutable/Source/CustomizableObjectEditor/Private/MuCOE/CustomizableObjectEditorLogger.cpp
@@ -161,16 +161,16 @@ FLogParameters& FLogParameters::Severity(const EMessageSeverity::Type Severity)
 }
 
 
-FLogParameters& FLogParameters::Nodes(const TArray<const UCustomizableObjectNode*>& Nodes)
+FLogParameters& FLogParameters::Context(const TArray<const UObject*>& Context)
 {
-	ParamNodes = Nodes;
+	ParamContext = Context;
 	return *this;
 }
 
 
-FLogParameters& FLogParameters::Node(const UCustomizableObjectNode& Node)
+FLogParameters& FLogParameters::Context(const UObject& Context)
 {
-	ParamNodes.Add(&Node);
+	ParamContext.Add(&Context);
 	return *this;
 }
 
@@ -240,20 +240,26 @@ void FCustomizableObjectEditorLogger::Log(FLogParameters& LogParameters)
 	
 	Message->AddToken(FTextToken::Create(MessageText));
 
-	for (const UCustomizableObjectNode* Node : LogParameters.ParamNodes)
+	for (const UObject* Context : LogParameters.ParamContext)
 	{
-		check(Node); // Can not be nullptr.
+		if (const UCustomizableObjectNode* Node = Cast<const UCustomizableObjectNode>(Context))
+		{
+			if (LogParameters.bParamBaseObject)
+			{
+				Message->AddToken(FTextToken::Create(FText::FromString(TEXT(" "))));
+				const UObject* Asset = Node->GetCustomizableObjectGraph()->GetOuter();
+				Message->AddToken(FUObjectToken::Create(Asset)->OnMessageTokenActivated(FOnMessageTokenActivated::CreateStatic(&OnMessageLogLinkActivated)));
+			}
 
-		if (LogParameters.bParamBaseObject)
+			Message->AddToken(FTextToken::Create(FText::FromString(TEXT(" (Node "))));
+			Message->AddToken(FCustomizableObjectToken::Create(Node)->OnMessageTokenActivated(FOnMessageTokenActivated::CreateStatic(&OnMessageLogLinkActivated)));
+			Message->AddToken(FTextToken::Create(FText::FromString(TEXT(")"))));			
+		}
+		else
 		{
 			Message->AddToken(FTextToken::Create(FText::FromString(TEXT(" "))));
-			const UObject* Asset = Node->GetCustomizableObjectGraph()->GetOuter();
-			Message->AddToken(FUObjectToken::Create(Asset)->OnMessageTokenActivated(FOnMessageTokenActivated::CreateStatic(&OnMessageLogLinkActivated)));
+			Message->AddToken(FUObjectToken::Create(Context)->OnMessageTokenActivated(FOnMessageTokenActivated::CreateStatic(&OnMessageLogLinkActivated)));
 		}
-
-		Message->AddToken(FTextToken::Create(FText::FromString(TEXT(" (Node "))));
-		Message->AddToken(FCustomizableObjectToken::Create(Node)->OnMessageTokenActivated(FOnMessageTokenActivated::CreateStatic(&OnMessageLogLinkActivated)));
-		Message->AddToken(FTextToken::Create(FText::FromString(TEXT(")"))));
 	}
 
 	if (LogParameters.bParamNotification)
diff --git a/Engine/Plugins/Experimental/Mutable/Source/CustomizableObjectEditor/Private/MuCOE/GenerateMutableSource/GenerateMutableSource.cpp b/Engine/Plugins/Experimental/Mutable/Source/CustomizableObjectEditor/Private/MuCOE/GenerateMutableSource/GenerateMutableSource.cpp
--- a/Engine/Plugins/Experimental/Mutable/Source/CustomizableObjectEditor/Private/MuCOE/GenerateMutableSource/GenerateMutableSource.cpp
+++ b/Engine/Plugins/Experimental/Mutable/Source/CustomizableObjectEditor/Private/MuCOE/GenerateMutableSource/GenerateMutableSource.cpp
@@ -157,13 +157,13 @@ mu::MeshPtr FMutableGraphGenerationContext::FindGeneratedMesh( const FGeneratedM
 /** Adds to ParameterNamesMap the node Node to the array of elements with name Name */
 void FMutableGraphGenerationContext::AddParameterNameUnique(const UCustomizableObjectNode* Node, FString Name)
 {
-	if (TArray<const UCustomizableObjectNode*>* ArrayResult = ParameterNamesMap.Find(Name))
+	if (TArray<const UObject*>* ArrayResult = ParameterNamesMap.Find(Name))
 	{
 		ArrayResult->AddUnique(Node);
 	}
 	else
 	{
-		TArray<const UCustomizableObjectNode*> ArrayTemp;
+		TArray<const UObject*> ArrayTemp;
 		ArrayTemp.Add(Node);
 		ParameterNamesMap.Add(Name, ArrayTemp);
 	}
@@ -172,11 +172,11 @@ void FMutableGraphGenerationContext::AddParameterNameUnique(const UCustomizableO
 
 const FGuid FMutableGraphGenerationContext::GetNodeIdUnique(const UCustomizableObjectNode* Node)
 {
-	TArray<const UCustomizableObjectNode*>* ArrayResult = NodeIdsMap.Find(Node->NodeGuid);
+	TArray<const UObject*>* ArrayResult = NodeIdsMap.Find(Node->NodeGuid);
 
 	if (ArrayResult == nullptr)
 	{
-		TArray<const UCustomizableObjectNode*> ArrayTemp;
+		TArray<const UObject*> ArrayTemp;
 		ArrayTemp.Add(Node);
 		NodeIdsMap.Add(Node->NodeGuid, ArrayTemp);
 		return Node->NodeGuid;
diff --git a/Engine/Plugins/Experimental/Mutable/Source/CustomizableObjectEditor/Private/MuCOE/GenerateMutableSource/GenerateMutableSource.h b/Engine/Plugins/Experimental/Mutable/Source/CustomizableObjectEditor/Private/MuCOE/GenerateMutableSource/GenerateMutableSource.h
--- a/Engine/Plugins/Experimental/Mutable/Source/CustomizableObjectEditor/Private/MuCOE/GenerateMutableSource/GenerateMutableSource.h
+++ b/Engine/Plugins/Experimental/Mutable/Source/CustomizableObjectEditor/Private/MuCOE/GenerateMutableSource/GenerateMutableSource.h
@@ -626,7 +626,7 @@ struct FMutableGraphGenerationContext
 	TArray<UMaterialInterface*> ReferencedMaterials;
 	TArray<FName> ReferencedMaterialSlotNames;
 	TArray<FGeneratedImageProperties> ImageProperties;
-	TMap<FString, TArray<const UCustomizableObjectNode*>> ParameterNamesMap;
+	TMap<FString, TArray<const UObject*>> ParameterNamesMap;
 	TArray<const UCustomizableObjectNode*> NoNameNodeObjectArray;
 	TMap<FString, FCustomizableObjectIdPair> GroupNodeMap;
 	TMap<FString, FString> CustomizableObjectPathMap;
@@ -641,7 +641,7 @@ struct FMutableGraphGenerationContext
 	TArray<FName> BoneNames;
 
 	// Used to aviod Nodes with duplicated ids
-	TMap<FGuid, TArray<const UCustomizableObjectNode*>> NodeIdsMap;
+	TMap<FGuid, TArray<const UObject*>> NodeIdsMap;
 	TMultiMap<const UCustomizableObject*, FGroupNodeIdsTempData> DuplicatedGroupNodeIds;
 
 	// For a given material node (the key is node package path + node uid + image index in node) stores images generated for the same node at a higher quality LOD to reuse that image node
diff --git a/Engine/Plugins/Experimental/Mutable/Source/CustomizableObjectEditor/Private/MuCOE/GenerateMutableSource/GenerateMutableSourceTable.cpp b/Engine/Plugins/Experimental/Mutable/Source/CustomizableObjectEditor/Private/MuCOE/GenerateMutableSource/GenerateMutableSourceTable.cpp
--- a/Engine/Plugins/Experimental/Mutable/Source/CustomizableObjectEditor/Private/MuCOE/GenerateMutableSource/GenerateMutableSourceTable.cpp
+++ b/Engine/Plugins/Experimental/Mutable/Source/CustomizableObjectEditor/Private/MuCOE/GenerateMutableSource/GenerateMutableSourceTable.cpp
@@ -285,7 +285,7 @@ bool FillTableColumn(const UCustomizableObjectNodeTable* TableNode,	mu::TablePtr
 					AddTagToMutableMeshUnique(*MutableMesh, MeshTag);
 				}
 
-				MutableTable->SetCell(CurrentColumn, RowIdx, MutableMesh.get());
+				MutableTable->SetCell(CurrentColumn, RowIdx, MutableMesh.get(), SkeletalMesh);
 			}
 			else
 			{
@@ -354,7 +354,7 @@ bool FillTableColumn(const UCustomizableObjectNodeTable* TableNode,	mu::TablePtr
 
 			if (MutableMesh)
 			{
-				MutableTable->SetCell(CurrentColumn, RowIdx, MutableMesh.get());
+				MutableTable->SetCell(CurrentColumn, RowIdx, MutableMesh.get(), StaticMesh);
 			}
 			else
 			{
diff --git a/Engine/Plugins/Experimental/Mutable/Source/CustomizableObjectEditor/Private/MuCOE/Nodes/CustomizableObjectNode.cpp b/Engine/Plugins/Experimental/Mutable/Source/CustomizableObjectEditor/Private/MuCOE/Nodes/CustomizableObjectNode.cpp
--- a/Engine/Plugins/Experimental/Mutable/Source/CustomizableObjectEditor/Private/MuCOE/Nodes/CustomizableObjectNode.cpp
+++ b/Engine/Plugins/Experimental/Mutable/Source/CustomizableObjectEditor/Private/MuCOE/Nodes/CustomizableObjectNode.cpp
@@ -287,7 +287,7 @@ void UCustomizableObjectNode::ReconstructNode(UCustomizableObjectNodeRemapPins*
 		{
 			FCustomizableObjectEditorLogger::CreateLog(LOCTEXT("OrphanPinsWarningReconstruct", "Failed to remap old pins"))
 			.Severity(EMessageSeverity::Warning)
-			.Node(*this)
+			.Context(*this)
 			.Log();
 		}
 
diff --git a/Engine/Plugins/Experimental/Mutable/Source/CustomizableObjectEditor/Private/MuCOE/Nodes/CustomizableObjectNodeObject.cpp b/Engine/Plugins/Experimental/Mutable/Source/CustomizableObjectEditor/Private/MuCOE/Nodes/CustomizableObjectNodeObject.cpp
--- a/Engine/Plugins/Experimental/Mutable/Source/CustomizableObjectEditor/Private/MuCOE/Nodes/CustomizableObjectNodeObject.cpp
+++ b/Engine/Plugins/Experimental/Mutable/Source/CustomizableObjectEditor/Private/MuCOE/Nodes/CustomizableObjectNodeObject.cpp
@@ -270,7 +270,7 @@ void UCustomizableObjectNodeObject::PostBackwardsCompatibleFixup()
 	{
 		FCustomizableObjectEditorLogger::CreateLog(LOCTEXT("ResaveNode","Please re-save this Customizable Object to avoid binary differences when packaging"))
 		.Severity(EMessageSeverity::Warning)
-		.Node(*this)
+		.Context(*this)
 		.Log();
 	}
 
diff --git a/Engine/Plugins/Experimental/Mutable/Source/CustomizableObjectEditor/Private/MuCOE/Nodes/CustomizableObjectNodeTable.cpp b/Engine/Plugins/Experimental/Mutable/Source/CustomizableObjectEditor/Private/MuCOE/Nodes/CustomizableObjectNodeTable.cpp
--- a/Engine/Plugins/Experimental/Mutable/Source/CustomizableObjectEditor/Private/MuCOE/Nodes/CustomizableObjectNodeTable.cpp
+++ b/Engine/Plugins/Experimental/Mutable/Source/CustomizableObjectEditor/Private/MuCOE/Nodes/CustomizableObjectNodeTable.cpp
@@ -334,7 +334,7 @@ void UCustomizableObjectNodeTable::AllocateDefaultPins(UCustomizableObjectNodeRe
 				FCustomizableObjectEditorLogger::CreateLog(Text)
 					.Category(ELoggerCategory::General)
 					.Severity(EMessageSeverity::Warning)
-					.Node(*this)
+					.Context(*this)
 					.Log();
 			}
 		}
@@ -364,7 +364,7 @@ void UCustomizableObjectNodeTable::AllocateDefaultPins(UCustomizableObjectNodeRe
 			FCustomizableObjectEditorLogger::CreateLog(Text)
 				.Category(ELoggerCategory::General)
 				.Severity(EMessageSeverity::Warning)
-				.Node(*this)
+				.Context(*this)
 				.Log();
 		}
 
diff --git a/Engine/Plugins/Experimental/Mutable/Source/CustomizableObjectEditor/Public/MuCOE/CompilationMessageCache.h b/Engine/Plugins/Experimental/Mutable/Source/CustomizableObjectEditor/Public/MuCOE/CompilationMessageCache.h
--- a/Engine/Plugins/Experimental/Mutable/Source/CustomizableObjectEditor/Public/MuCOE/CompilationMessageCache.h
+++ b/Engine/Plugins/Experimental/Mutable/Source/CustomizableObjectEditor/Public/MuCOE/CompilationMessageCache.h
@@ -16,11 +16,11 @@ public:
 	/** Method designed to serve as a way to add a new message to the collection of messages the compiler has produced as
 	 * part of the compilation process.
 	 * @param InMessage The message's text.
-	 * @param InArrayNode The node that is related with the InMessage provided.
+	 * @param InContext The UObject that is related with the InMessage provided.
 	 * @param MessageSeverity Severity of the message.
 	 * @return True if the message is new and had not been added before. If it returns false it means it has not been logged again.
 	 */
-	 bool AddMessage(const FText& InMessage, const TArray<const UCustomizableObjectNode*>& InArrayNode, const EMessageSeverity::Type MessageSeverity = EMessageSeverity::Warning, const ELoggerSpamBin SpamBin = ELoggerSpamBin::ShowAll);
+	 bool AddMessage(const FText& InMessage, const TArray<const UObject*>& InContext, const EMessageSeverity::Type MessageSeverity = EMessageSeverity::Warning, const ELoggerSpamBin SpamBin = ELoggerSpamBin::ShowAll);
 	
 	
 	/** Provides the caller with separated arrays with all the warning and error messages.
@@ -62,7 +62,7 @@ private:
 	struct FLoggedMessage
 	{
 		FText Message;
-		TArray<const UCustomizableObjectNode*> Context;
+		TArray<const UObject*> Context;
 		EMessageSeverity::Type Severity;
 	
 		bool operator==( const FLoggedMessage& o) const
diff --git a/Engine/Plugins/Experimental/Mutable/Source/CustomizableObjectEditor/Public/MuCOE/CustomizableObjectEditorLogger.h b/Engine/Plugins/Experimental/Mutable/Source/CustomizableObjectEditor/Public/MuCOE/CustomizableObjectEditorLogger.h
--- a/Engine/Plugins/Experimental/Mutable/Source/CustomizableObjectEditor/Public/MuCOE/CustomizableObjectEditorLogger.h
+++ b/Engine/Plugins/Experimental/Mutable/Source/CustomizableObjectEditor/Public/MuCOE/CustomizableObjectEditorLogger.h
@@ -41,9 +41,9 @@ public:
 
 	FLogParameters& Severity(EMessageSeverity::Type MessageSeverity);
 
-	FLogParameters& Nodes(const TArray<const UCustomizableObjectNode*>& Nodes);
+	FLogParameters& Context(const TArray<const UObject*>& Context);
 
-	FLogParameters& Node(const UCustomizableObjectNode& Node);
+	FLogParameters& Context(const UObject& Context);
 
 	FLogParameters& BaseObject(bool BaseObject = true);
 
@@ -77,7 +77,7 @@ private:
 	
 	EMessageSeverity::Type ParamSeverity = EMessageSeverity::Info;
 	
-	TArray<const UCustomizableObjectNode*> ParamNodes;
+	TArray<const UObject*> ParamContext;
 
 	bool bParamBaseObject = false;
 
diff --git a/Engine/Plugins/Experimental/Mutable/Source/MutableTools/Private/MuT/CodeGenerator.h b/Engine/Plugins/Experimental/Mutable/Source/MutableTools/Private/MuT/CodeGenerator.h
--- a/Engine/Plugins/Experimental/Mutable/Source/MutableTools/Private/MuT/CodeGenerator.h
+++ b/Engine/Plugins/Experimental/Mutable/Source/MutableTools/Private/MuT/CodeGenerator.h
@@ -503,6 +503,10 @@ namespace mu
 			*/
 			TArray<Ptr<const Layout>> OverrideLayouts;			
 
+			/** Optional context to use instead of the node error context.
+			 * Be careful since it is not used everywhere. Check usages before assigning a value to it. */
+			TOptional<const void*> OverrideContext;
+			
 			friend FORCEINLINE uint32 GetTypeHash(const FMeshGenerationOptions& InKey)
 			{
 				uint32 KeyHash = 0;
diff --git a/Engine/Plugins/Experimental/Mutable/Source/MutableTools/Private/MuT/CodeGenerator_Mesh.cpp b/Engine/Plugins/Experimental/Mutable/Source/MutableTools/Private/MuT/CodeGenerator_Mesh.cpp
--- a/Engine/Plugins/Experimental/Mutable/Source/MutableTools/Private/MuT/CodeGenerator_Mesh.cpp
+++ b/Engine/Plugins/Experimental/Mutable/Source/MutableTools/Private/MuT/CodeGenerator_Mesh.cpp
@@ -1143,6 +1143,8 @@ class Node;
 						TargetOptions.OverrideLayouts = NewResult.GeneratedLayouts;
 					}
 
+					TargetOptions.OverrideContext = node.m_pTable->GetPrivate()->Rows[row].Values[colIndex].ErrorContext;
+
 					GenerateMesh(TargetOptions, BranchResults, pCell);
 
 					if (t == 0)
@@ -1370,7 +1372,8 @@ class Node;
 						{
 							Ptr<const Layout> SourceLayout = TypedNode->GetPrivate()->m_pLayout;
 							Ptr<const Layout> GeneratedLayout = AddLayout( SourceLayout );
-							PrepareForLayout(GeneratedLayout, pCloned, LayoutIndex, TypedNode->GetPrivate()->m_errorContext, InOptions);
+							const void* Context = InOptions.OverrideContext.Get(node.m_errorContext);
+							PrepareForLayout(GeneratedLayout, pCloned, LayoutIndex, Context, InOptions);
 
 							OutResult.GeneratedLayouts.Add(GeneratedLayout);
 						}
@@ -1382,7 +1385,8 @@ class Node;
 					for (int32 LayoutIndex = 0; LayoutIndex < InOptions.OverrideLayouts.Num(); ++LayoutIndex)
 					{
 						Ptr<const Layout> GeneratedLayout = InOptions.OverrideLayouts[LayoutIndex];
-						PrepareForLayout(GeneratedLayout, pCloned, LayoutIndex, node.m_errorContext, InOptions);
+						const void* Context = InOptions.OverrideContext.Get(node.m_errorContext);
+						PrepareForLayout(GeneratedLayout, pCloned, LayoutIndex, Context, InOptions);
 
 						OutResult.GeneratedLayouts.Add(GeneratedLayout);
 					}
diff --git a/Engine/Plugins/Experimental/Mutable/Source/MutableTools/Private/MuT/Table.cpp b/Engine/Plugins/Experimental/Mutable/Source/MutableTools/Private/MuT/Table.cpp
--- a/Engine/Plugins/Experimental/Mutable/Source/MutableTools/Private/MuT/Table.cpp
+++ b/Engine/Plugins/Experimental/Mutable/Source/MutableTools/Private/MuT/Table.cpp
@@ -126,7 +126,7 @@ namespace mu
 	}
 
 
-    void Table::SetCell( int32 Column, uint32 RowId, float Value)
+    void Table::SetCell( int32 Column, uint32 RowId, float Value, const void* ErrorContext)
 	{
 		int32 Row = m_pD->FindRow(RowId);
 		check( Row>=0 );
@@ -135,10 +135,11 @@ namespace mu
 		check( m_pD->Columns[Column].Type == ETableColumnType::Scalar );
 
 		m_pD->Rows[ Row ].Values[Column].Scalar = Value;
+		m_pD->Rows[ Row ].Values[Column].ErrorContext = ErrorContext;
 	}
 
 
-    void Table::SetCell( int32 Column, uint32 RowId, const FVector4f& Value)
+    void Table::SetCell( int32 Column, uint32 RowId, const FVector4f& Value, const void* ErrorContext)
 	{
 		int32 Row = m_pD->FindRow(RowId);
 		check( Row>=0 );
@@ -147,10 +148,11 @@ namespace mu
 		check( m_pD->Columns[Column].Type == ETableColumnType::Color );
 
 		m_pD->Rows[ Row ].Values[Column].Color = Value;
+		m_pD->Rows[ Row ].Values[Column].ErrorContext = ErrorContext;
 	}
 
 
-    void Table::SetCell(int32 Column, uint32 RowId, ResourceProxy<Image>* Value)
+    void Table::SetCell(int32 Column, uint32 RowId, ResourceProxy<Image>* Value, const void* ErrorContext)
 	{
 		int32 Row = m_pD->FindRow(RowId);
 		check( Row>=0 );
@@ -159,10 +161,11 @@ namespace mu
 		check( m_pD->Columns[Column].Type == ETableColumnType::Image );
 
 		m_pD->Rows[ Row ].Values[Column].ProxyImage = Value;
+		m_pD->Rows[ Row ].Values[Column].ErrorContext = ErrorContext;
 	}
 
 
-    void Table::SetCell( int32 Column, uint32 RowId, Mesh* Value )
+    void Table::SetCell( int32 Column, uint32 RowId, Mesh* Value, const void* ErrorContext )
 	{
 		int32 Row = m_pD->FindRow(RowId);
 		check( Row>=0 );
@@ -171,10 +174,11 @@ namespace mu
 		check( m_pD->Columns[Column].Type == ETableColumnType::Mesh );
 
 		m_pD->Rows[ Row ].Values[Column].Mesh = Value;
+		m_pD->Rows[ Row ].Values[Column].ErrorContext = ErrorContext;
 	}
 
 
-    void Table::SetCell( int32 Column, uint32 RowId, const FString& Value )
+    void Table::SetCell( int32 Column, uint32 RowId, const FString& Value, const void* ErrorContext )
 	{
 		int32 Row = m_pD->FindRow(RowId);
 		check( Row>=0 );
@@ -183,6 +187,7 @@ namespace mu
 		check( m_pD->Columns[Column].Type == ETableColumnType::String );
 
 		m_pD->Rows[ Row ].Values[Column].String = Value;
+		m_pD->Rows[ Row ].Values[Column].ErrorContext = ErrorContext;
 	}
 
 
diff --git a/Engine/Plugins/Experimental/Mutable/Source/MutableTools/Private/MuT/TablePrivate.h b/Engine/Plugins/Experimental/Mutable/Source/MutableTools/Private/MuT/TablePrivate.h
--- a/Engine/Plugins/Experimental/Mutable/Source/MutableTools/Private/MuT/TablePrivate.h
+++ b/Engine/Plugins/Experimental/Mutable/Source/MutableTools/Private/MuT/TablePrivate.h
@@ -32,6 +32,8 @@ namespace mu
 		Ptr<ResourceProxy<Image>> ProxyImage;
 		Ptr<Mesh> Mesh;
 		FString String;
+
+		const void* ErrorContext;
 	};
 
 
diff --git a/Engine/Plugins/Experimental/Mutable/Source/MutableTools/Public/MuT/Table.h b/Engine/Plugins/Experimental/Mutable/Source/MutableTools/Public/MuT/Table.h
--- a/Engine/Plugins/Experimental/Mutable/Source/MutableTools/Public/MuT/Table.h
+++ b/Engine/Plugins/Experimental/Mutable/Source/MutableTools/Public/MuT/Table.h
@@ -68,11 +68,11 @@ namespace mu
 		void SetNoneOption(bool bAddOption);
 
 		//!
-        void SetCell( int32 Column, uint32 RowId, float );
-        void SetCell( int32 Column, uint32 RowId, const FVector4f& );
-		void SetCell( int32 Column, uint32 RowId, ResourceProxy<Image>* );
-		void SetCell( int32 Column, uint32 RowId, Mesh* );
-        void SetCell( int32 Column, uint32 RowId, const FString& );
+        void SetCell( int32 Column, uint32 RowId, float Value, const void* ErrorContext = nullptr);
+        void SetCell( int32 Column, uint32 RowId, const FVector4f& Value, const void* ErrorContext = nullptr);
+		void SetCell( int32 Column, uint32 RowId, ResourceProxy<Image>* Value, const void* ErrorContext = nullptr);
+		void SetCell( int32 Column, uint32 RowId, Mesh* Value, const void* ErrorContext = nullptr);
+        void SetCell( int32 Column, uint32 RowId, const FString& Value, const void* ErrorContext = nullptr);
 
 
 		//-----------------------------------------------------------------------------------------
