[Mutable] PinViewer use Unreal's DetailsView.

- Indirectly adds reset value.

#rnx

diff --git a/Engine/Plugins/Experimental/Mutable/Source/CustomizableObject/Public/MuCO/CustomizableObjectCustomVersion.h b/Engine/Plugins/Experimental/Mutable/Source/CustomizableObject/Public/MuCO/CustomizableObjectCustomVersion.h
--- a/Engine/Plugins/Experimental/Mutable/Source/CustomizableObject/Public/MuCO/CustomizableObjectCustomVersion.h
+++ b/Engine/Plugins/Experimental/Mutable/Source/CustomizableObject/Public/MuCO/CustomizableObjectCustomVersion.h
@@ -101,6 +101,8 @@ struct CUSTOMIZABLEOBJECT_API FCustomizableObjectCustomVersion
 		NodeVariationSerializationIssue, // The parent class and child classes had a member with the same name. Unreal was serialized the member from the child class but deserialized it to the parent class.
 
 		RegenerateNodeObjectsIds,
+		
+		NodeMaterialPinDataImageDetails,
 
 		// -----<new versions can be added above this line>-------------------------------------------------
 		VersionPlusOne,
diff --git a/Engine/Plugins/Experimental/Mutable/Source/CustomizableObjectEditor/Classes/MuCOE/Nodes/CustomizableObjectNodeMaterial.h b/Engine/Plugins/Experimental/Mutable/Source/CustomizableObjectEditor/Classes/MuCOE/Nodes/CustomizableObjectNodeMaterial.h
--- a/Engine/Plugins/Experimental/Mutable/Source/CustomizableObjectEditor/Classes/MuCOE/Nodes/CustomizableObjectNodeMaterial.h
+++ b/Engine/Plugins/Experimental/Mutable/Source/CustomizableObjectEditor/Classes/MuCOE/Nodes/CustomizableObjectNodeMaterial.h
@@ -83,6 +83,20 @@ enum class EPinMode
 	Passthrough
 };
 
+
+/** Image Pin, UV Layout Mode. */
+UENUM()
+enum class EUVLayoutMode
+{
+	/** Does not override the UV Index specified in the Material. */
+	FromMaterial,
+	/* Texture should not be transformed by any layout. Theses textures will not be reduced automatically for LODs. */
+	Ignore,
+	/** User specified UV Index. */
+	Index
+};
+
+
 /** Enum to FText. */
 FText EPinModeToText(EPinMode PinMode);
 
@@ -299,7 +313,10 @@ class CUSTOMIZABLEOBJECTEDITOR_API UCustomizableObjectNodeMaterialPinDataImage :
 public:
 	// UObject interface
 	virtual void PostEditChangeProperty(FPropertyChangedEvent& PropertyChangedEvent) override;
-
+private:
+	virtual void PostLoad() override;
+	
+public:
 	// NodePinDataParameter interface
 	/** Virtual function used to copy pin data when remapping pins. */
 	virtual void Copy(const UCustomizableObjectNodePinData& Other) override;
@@ -317,24 +334,23 @@ public:
 	
 private:
 	/** Image pin mode. If is not default, overrides the defined node behaviour. */
-	UPROPERTY()
+	UPROPERTY(EditAnywhere, Category = NoCategory)
 	EPinMode PinMode = EPinMode::Default;
 
 public:
-	/* UVLayout Mode. Indicates that the texture should not be transformed by any layout. Theses textures will not be reduced automatically for LODs. */
 	constexpr static int32 UV_LAYOUT_IGNORE = -1;
-	
-	/** UVLayout Mode. Does not override the Material Texture Parameter. */
-	constexpr static int32 UV_LAYOUT_DEFAULT = -2;
+
+	UPROPERTY(EditAnywhere, Category = NoCategory)
+	EUVLayoutMode UVLayoutMode;
 	
 	/** Index of the UV channel that will be used with this image.It is necessary to apply the proper layout transformations to it. */
-	UPROPERTY()
-	int32 UVLayout = UV_LAYOUT_DEFAULT;
+	UPROPERTY(EditAnywhere, Category = NoCategory, meta = (ClampMin = 0, EditCondition = "UVLayoutMode == EUVLayoutMode::Index"))
+	int32 UVLayout = 0;
 
 	/** Reference Texture used to decide the texture properties of the mutable-generated textures
 	* connected to this material. If null, it will try to be guessed at compile time from
 	* the graph. */
-	UPROPERTY(EditAnywhere, Category=CustomizableObject) // Required to be EditAnywhere for the selector to work.
+	UPROPERTY(EditAnywhere, Category = NoCategory) // Required to be EditAnywhere for the selector to work.
 	TObjectPtr<UTexture2D> ReferenceTexture = nullptr;
 
 private:
diff --git a/Engine/Plugins/Experimental/Mutable/Source/CustomizableObjectEditor/Private/MuCOE/Nodes/CustomizableObjectNodeMaterial.cpp b/Engine/Plugins/Experimental/Mutable/Source/CustomizableObjectEditor/Private/MuCOE/Nodes/CustomizableObjectNodeMaterial.cpp
--- a/Engine/Plugins/Experimental/Mutable/Source/CustomizableObjectEditor/Private/MuCOE/Nodes/CustomizableObjectNodeMaterial.cpp
+++ b/Engine/Plugins/Experimental/Mutable/Source/CustomizableObjectEditor/Private/MuCOE/Nodes/CustomizableObjectNodeMaterial.cpp
@@ -16,8 +16,10 @@
 #include "MuCOE/Nodes/CustomizableObjectNodeStaticMesh.h"
 #include "MuCOE/Nodes/CustomizableObjectNodeTable.h"
 #include "MuCOE/Nodes/SCustomizableObjectNodeMaterial.h"
-#include "MuCOE/Nodes/SPinViewerNodeMaterialPinImageDetails.h"
 #include "ObjectEditorUtils.h"
+#include "PropertyCustomizationHelpers.h"
+#include "Modules/ModuleManager.h"
+
 
 class SGraphNode;
 class SWidget;
@@ -89,9 +91,11 @@ bool UCustomizableObjectNodeMaterialRemapPinsByName::HasSavedPinData(const UCust
 
 bool UCustomizableObjectNodeMaterialPinDataImage::IsDefault() const
 {
-	return PinMode == EPinMode::Default &&
-		UVLayout == UV_LAYOUT_DEFAULT &&
-		!ReferenceTexture;
+	const UCustomizableObjectNodeMaterialPinDataImage* Default = Cast<UCustomizableObjectNodeMaterialPinDataImage>(GetClass()->ClassDefaultObject);
+
+	return PinMode == Default->PinMode &&
+		UVLayoutMode == Default->UVLayoutMode &&
+		ReferenceTexture == Default->ReferenceTexture;
 }
 
 
@@ -240,6 +244,29 @@ void UCustomizableObjectNodeMaterialPinDataImage::PostEditChangeProperty(FProper
 	}
 }
 
+
+void UCustomizableObjectNodeMaterialPinDataImage::PostLoad()
+{
+	Super::PostLoad();
+
+	const int32 CustomizableObjectCustomVersion = GetLinkerCustomVersion(FCustomizableObjectCustomVersion::GUID);
+	
+	if (CustomizableObjectCustomVersion < FCustomizableObjectCustomVersion::NodeMaterialPinDataImageDetails)
+	{
+		if (UVLayout == UV_LAYOUT_IGNORE)
+		{
+			UVLayoutMode = EUVLayoutMode::Ignore;
+			UVLayout = 0;
+		}
+		else if (UVLayout == -2) // UV_LAYOUT_DEFAULT
+		{
+			UVLayoutMode = EUVLayoutMode::FromMaterial;
+			UVLayout = 0;
+		}
+	}
+}
+
+
 void UCustomizableObjectNodeMaterial::AllocateDefaultPins(UCustomizableObjectNodeRemapPins* RemapPins)
 {
 	const UEdGraphSchema_CustomizableObject* Schema = GetDefault<UEdGraphSchema_CustomizableObject>();
@@ -775,9 +802,16 @@ int32 UCustomizableObjectNodeMaterial::GetImageUVLayout(const int32 ImageIndex)
 	if (const UEdGraphPin* Pin = GetParameterPin(EMaterialParameterType::Texture, ImageIndex))
 	{
 		const UCustomizableObjectNodeMaterialPinDataImage& PinData = GetPinData<UCustomizableObjectNodeMaterialPinDataImage>(*Pin);
-		if (PinData.UVLayout != UCustomizableObjectNodeMaterialPinDataImage::UV_LAYOUT_DEFAULT)
+		switch(PinData.UVLayoutMode)
 		{
+		case EUVLayoutMode::FromMaterial:
+			break;
+		case EUVLayoutMode::Ignore:
+			return UCustomizableObjectNodeMaterialPinDataImage::UV_LAYOUT_IGNORE;
+		case EUVLayoutMode::Index:
 			return PinData.UVLayout;
+		default:
+			unimplemented();
 		}
 	}
 
@@ -1056,7 +1090,16 @@ TSharedPtr<SWidget> UCustomizableObjectNodeMaterial::CustomizePinDetails(UEdGrap
 {
 	if (UCustomizableObjectNodeMaterialPinDataImage* PinData = Cast<UCustomizableObjectNodeMaterialPinDataImage>(GetPinData(Pin)))
 	{
-		return SNew(SPinViewerNodeMaterialPinImageDetails).Pin(&Pin).PinData(PinData);
+		FPropertyEditorModule& EditModule = FModuleManager::Get().GetModuleChecked<FPropertyEditorModule>("PropertyEditor");
+
+		FDetailsViewArgs DetailsViewArgs;
+		DetailsViewArgs.bAllowSearch = false;
+		DetailsViewArgs.bHideSelectionTip = true;
+		
+		const TSharedRef<IDetailsView> SettingsView = EditModule.CreateDetailView(DetailsViewArgs);
+		SettingsView->SetObject(PinData);
+		
+		return SettingsView;
 	}
 	else
 	{
diff --git a/Engine/Plugins/Experimental/Mutable/Source/CustomizableObjectEditor/Private/MuCOE/Nodes/SCustomizableObjectNodeMaterialPinImage.cpp b/Engine/Plugins/Experimental/Mutable/Source/CustomizableObjectEditor/Private/MuCOE/Nodes/SCustomizableObjectNodeMaterialPinImage.cpp
--- a/Engine/Plugins/Experimental/Mutable/Source/CustomizableObjectEditor/Private/MuCOE/Nodes/SCustomizableObjectNodeMaterialPinImage.cpp
+++ b/Engine/Plugins/Experimental/Mutable/Source/CustomizableObjectEditor/Private/MuCOE/Nodes/SCustomizableObjectNodeMaterialPinImage.cpp
@@ -68,9 +68,14 @@ const FSlateBrush* SCustomizableObjectNodeMaterialPinImage::GetPinIcon() const
 
 FText SCustomizableObjectNodeMaterialPinImage::GetTooltipText() const
 {
+	if (GraphPinObj->bOrphanedPin)
+	{
+		return LOCTEXT("PinModeMutableOrpahan", "Pin not disapearing due to being connected or having a property modified.");
+	}
+	
 	if (CastChecked<UCustomizableObjectNodeMaterial>(GraphPinObj->GetOwningNode())->IsImageMutableMode(*GraphPinObj))
 	{
-		return  LOCTEXT("PinModeMutableTooltip", "Texture Parameter goes through Mutable.");
+		return LOCTEXT("PinModeMutableTooltip", "Texture Parameter goes through Mutable.");
 	}
 	else
 	{
diff --git a/Engine/Plugins/Experimental/Mutable/Source/CustomizableObjectEditor/Private/MuCOE/Nodes/SPinViewerNodeMaterialPinImageDetails.cpp b/Engine/Plugins/Experimental/Mutable/Source/CustomizableObjectEditor/Private/MuCOE/Nodes/SPinViewerNodeMaterialPinImageDetails.cpp
deleted file mode 100644
diff --git a/Engine/Plugins/Experimental/Mutable/Source/CustomizableObjectEditor/Private/MuCOE/Nodes/SPinViewerNodeMaterialPinImageDetails.h b/Engine/Plugins/Experimental/Mutable/Source/CustomizableObjectEditor/Private/MuCOE/Nodes/SPinViewerNodeMaterialPinImageDetails.h
deleted file mode 100644
