[Mutable] Initialize the CO identifier when setting a new Model

#rnx

diff --git a/Engine/Plugins/Experimental/Mutable/Source/CustomizableObject/Private/MuCO/CustomizableObject.cpp b/Engine/Plugins/Experimental/Mutable/Source/CustomizableObject/Private/MuCO/CustomizableObject.cpp
--- a/Engine/Plugins/Experimental/Mutable/Source/CustomizableObject/Private/MuCO/CustomizableObject.cpp
+++ b/Engine/Plugins/Experimental/Mutable/Source/CustomizableObject/Private/MuCO/CustomizableObject.cpp
@@ -862,16 +862,20 @@ FString UCustomizableObject::GetCompiledDataFolderPath(bool bIsEditorData) const
 }
 
 
+FGuid GenerateIdentifier(const UCustomizableObject& CustomizableObject)
+{
+	// Generate the Identifier using the path and name of the asset
+	uint32 FullPathHash = GetTypeHash(CustomizableObject.GetFullName());
+	uint32 OutermostHash = GetTypeHash(GetNameSafe(CustomizableObject.GetOutermost()));
+	uint32 OuterHash = GetTypeHash(CustomizableObject.GetName());
+	return FGuid(0, FullPathHash, OutermostHash, OuterHash);
+}
+
+
 FString UCustomizableObject::GetCompiledDataFileName(bool bIsModel, const ITargetPlatform* InTargetPlatform, bool bIsDiskStreamer)
 {
 	// Generate the Identifier using the path and name of the asset
-	if (!bIsDiskStreamer)
-	{
-		uint32 FullPathHash = GetTypeHash(GetFullName());
-		uint32 OutermostHash = GetTypeHash(GetNameSafe(GetOutermost()));
-		uint32 OuterHash = GetTypeHash(GetName());
-		Identifier = FGuid(0, FullPathHash, OutermostHash, OuterHash);
-	}
+	Identifier = GenerateIdentifier(*this);
 
 	const FString PlatformName = InTargetPlatform ? InTargetPlatform->PlatformName() : FPlatformProperties::PlatformName();
 	const FString FileIdentifier = bIsDiskStreamer ? GetPrivate()->Identifier.ToString() : Identifier.ToString();
@@ -1257,7 +1261,7 @@ TSharedPtr<mu::Model, ESPMode::ThreadSafe> UCustomizableObject::GetModel() const
 #if WITH_EDITOR
 void UCustomizableObject::SetModel(TSharedPtr<mu::Model, ESPMode::ThreadSafe> Model)
 {
-	PrivateData->SetModel(Model, Identifier);
+	PrivateData->SetModel(Model, GenerateIdentifier(*this));
 	
 	UpdateCompiledDataFromModel();
 }
