[Mutable] Implement MinimizeMemory execution strategy for IM_COMPOSE operations.

#rnx

diff --git a/Engine/Plugins/Experimental/Mutable/Source/MutableRuntime/Private/MuR/CodeRunner.cpp b/Engine/Plugins/Experimental/Mutable/Source/MutableRuntime/Private/MuR/CodeRunner.cpp
--- a/Engine/Plugins/Experimental/Mutable/Source/MutableRuntime/Private/MuR/CodeRunner.cpp
+++ b/Engine/Plugins/Experimental/Mutable/Source/MutableRuntime/Private/MuR/CodeRunner.cpp
@@ -3554,47 +3554,96 @@ namespace mu
 
         case OP_TYPE::IM_COMPOSE:
         {
-			OP::ImageComposeArgs args = pModel->GetPrivate()->m_program.GetOpArgs<OP::ImageComposeArgs>( item.At );
-            switch ( item.Stage )
-            {
-            case 0:
-                AddOp( FScheduledOp( item.At, item, 1 ), FScheduledOp::FromOpAndOptions( args.layout, item, 0 ) );
-                break;
+			OP::ImageComposeArgs Args = pModel->GetPrivate()->m_program.GetOpArgs<OP::ImageComposeArgs>(item.At);
 
-            case 1:
-            {
-            	MUTABLE_CPUPROFILER_SCOPE(IM_COMPOSE_1)
-            		
-                Ptr<const Layout> ComposeLayout = LoadLayout( FCacheAddress( args.layout, FScheduledOp::FromOpAndOptions(args.layout, item, 0)) );
+			if (ExecutionStrategy == EExecutionStrategy::MinimizeMemory)
+			{
+            	switch (item.Stage)
+				{
+				case 0:
+					AddOp(FScheduledOp(item.At, item, 1), FScheduledOp::FromOpAndOptions(Args.layout, item, 0));
+					break;
+				case 1:
+				{
+					Ptr<const Layout> ComposeLayout = 
+							LoadLayout(FCacheAddress(Args.layout, FScheduledOp::FromOpAndOptions(Args.layout, item, 0)));
 
-                FScheduledOpData data;
-                data.Resource = const_cast<Layout*>(ComposeLayout.get());
-				int32 dataPos = m_heapData.Add( data );
+					FScheduledOpData Data;
+					Data.Resource = const_cast<Layout*>(ComposeLayout.get());
+					int32 DataPos = m_heapData.Add(Data);
 
-                int relBlockIndex = ComposeLayout->FindBlock( args.blockIndex );
-                if ( relBlockIndex >= 0 )
-                {
-                    AddOp( FScheduledOp( item.At, item, 2, dataPos ),
-                           FScheduledOp( args.base, item ),
-                           FScheduledOp( args.blockImage, item ),
-                           FScheduledOp( args.mask, item ) );
-                }
-                else
-                {
-                    AddOp( FScheduledOp( item.At, item, 2, dataPos ),
-                           FScheduledOp( args.base, item ) );
-                }
-                break;
-            }
+					int32 RelBlockIndex = ComposeLayout->FindBlock(Args.blockIndex);
 
-            case 2:
-				// This has been moved to a task. It should have been intercepted in IssueOp.
-				check(false);
-                break;
+					if (RelBlockIndex >= 0)
+					{
+						AddOp(FScheduledOp(item.At, item, 2, DataPos), FScheduledOp(Args.base, item));
+					}
+					else
+					{
+						// Jump directly to stage 3, no need to load mask or blockImage.
+						AddOp(FScheduledOp(item.At, item, 3, DataPos), FScheduledOp(Args.base, item));
+					}
 
-            default:
-                check( false );
-            }
+					break;
+				}
+				case 2:
+				{
+					AddOp(FScheduledOp(item.At, item, 3, item.CustomState),
+						  FScheduledOp(Args.blockImage, item),
+						  FScheduledOp(Args.mask, item));
+					break;
+				}
+
+				case 3:
+					// This has been moved to a task. It should have been intercepted in IssueOp.
+					check(false);
+					break;
+
+				default:
+					check(false);
+				}
+			}
+			else
+			{
+            	switch (item.Stage)
+				{
+				case 0:
+					AddOp(FScheduledOp(item.At, item, 1), FScheduledOp::FromOpAndOptions(Args.layout, item, 0));
+					break;
+
+				case 1:
+				{	
+					Ptr<const Layout> ComposeLayout = 
+							LoadLayout(FCacheAddress(Args.layout, FScheduledOp::FromOpAndOptions(Args.layout, item, 0)));
+
+					FScheduledOpData Data;
+					Data.Resource = const_cast<Layout*>(ComposeLayout.get());
+					int32 DataPos = m_heapData.Add(Data);
+
+					int32 RelBlockIndex = ComposeLayout->FindBlock(Args.blockIndex);
+					if (RelBlockIndex >= 0)
+					{
+						AddOp(FScheduledOp(item.At, item, 2, DataPos),
+							  FScheduledOp(Args.base, item),
+							  FScheduledOp(Args.blockImage, item),
+							  FScheduledOp(Args.mask, item));
+					}
+					else
+					{
+						AddOp(FScheduledOp(item.At, item, 2, DataPos), FScheduledOp(Args.base, item));
+					}
+					break;
+				}
+
+				case 2:
+					// This has been moved to a task. It should have been intercepted in IssueOp.
+					check(false);
+					break;
+
+				default:
+					check(false);
+				}
+			}
 
             break;
         }
diff --git a/Engine/Plugins/Experimental/Mutable/Source/MutableRuntime/Private/MuR/CodeRunnerTasked.cpp b/Engine/Plugins/Experimental/Mutable/Source/MutableRuntime/Private/MuR/CodeRunnerTasked.cpp
--- a/Engine/Plugins/Experimental/Mutable/Source/MutableRuntime/Private/MuR/CodeRunnerTasked.cpp
+++ b/Engine/Plugins/Experimental/Mutable/Source/MutableRuntime/Private/MuR/CodeRunnerTasked.cpp
@@ -158,7 +158,7 @@ namespace mu
 		{   0,   0,   0,   0 },	// IM_RESIZELIKE	(to be deprecated?)
 		{   0,   0,   0,   0 },	// IM_RESIZEREL		(special case)
 		{   0,  20,   0,   0 },	// IM_BLANKLAYOUT
-		{   0,   0, -20,   0 },	// IM_COMPOSE
+		{   0, -20, -20,   0 },	// IM_COMPOSE
 		{   0,   0, -20,   0 },	// IM_INTERPOLATE
 		{   0,   0,   0,   0 },	// IM_SATURATE
 		{   0,   0,   0,   0 },	// IM_LUMINANCE
@@ -2693,7 +2693,8 @@ namespace mu
 
 		case OP_TYPE::IM_COMPOSE:
 		{
-			if (item.Stage == 2)
+			if ((ExecutionStrategy == EExecutionStrategy::MinimizeMemory && item.Stage == 3) ||
+				(ExecutionStrategy != EExecutionStrategy::MinimizeMemory && item.Stage == 2))
 			{
 				OP::ImageComposeArgs Args = program.GetOpArgs<OP::ImageComposeArgs>(item.At);
 				Ptr<const Layout> ComposeLayout = static_cast<const Layout*>( m_heapData[item.CustomState].Resource.get());
