[mutable] Fixed memory increment when using the same table texture with different texture modes from different table nodes.

#rnx

diff --git a/Engine/Plugins/Experimental/Mutable/Source/CustomizableObjectEditor/Private/MuCOE/GenerateMutableSource/GenerateMutableSource.h b/Engine/Plugins/Experimental/Mutable/Source/CustomizableObjectEditor/Private/MuCOE/GenerateMutableSource/GenerateMutableSource.h
--- a/Engine/Plugins/Experimental/Mutable/Source/CustomizableObjectEditor/Private/MuCOE/GenerateMutableSource/GenerateMutableSource.h
+++ b/Engine/Plugins/Experimental/Mutable/Source/CustomizableObjectEditor/Private/MuCOE/GenerateMutableSource/GenerateMutableSource.h
@@ -594,6 +594,20 @@ struct FMutableGraphGenerationContext
 	};
 	TArray<FGeneratedMeshData> GeneratedMeshes;
 
+	struct FGeneratedTableImageData
+	{
+		FString PinName;
+		FName PinType;
+		const mu::Ptr<mu::Table> Table;
+		const UCustomizableObjectNodeTable* TableNode;
+
+		bool operator==(const FGeneratedTableImageData& Other) const
+		{
+			return PinName == Other.PinName && Table == Other.Table;
+		}
+	};
+	TArray<FGeneratedTableImageData> GeneratedTableImages;
+
 	// Stack of mesh generation flags. The last one is the currently valid.
 	// The value is a bit mask of EMutableMeshConversionFlags
 	TArray<EMutableMeshConversionFlags> MeshGenerationFlags;
diff --git a/Engine/Plugins/Experimental/Mutable/Source/CustomizableObjectEditor/Private/MuCOE/GenerateMutableSource/GenerateMutableSourceImage.cpp b/Engine/Plugins/Experimental/Mutable/Source/CustomizableObjectEditor/Private/MuCOE/GenerateMutableSource/GenerateMutableSourceImage.cpp
--- a/Engine/Plugins/Experimental/Mutable/Source/CustomizableObjectEditor/Private/MuCOE/GenerateMutableSource/GenerateMutableSourceImage.cpp
+++ b/Engine/Plugins/Experimental/Mutable/Source/CustomizableObjectEditor/Private/MuCOE/GenerateMutableSource/GenerateMutableSourceImage.cpp
@@ -1063,6 +1063,34 @@ mu::NodeImagePtr GenerateMutableSourceImage(const UEdGraphPin* Pin, FMutableGrap
 							// Material parameters use the Data Table Column Name + Parameter id as mutable column Name to aboid duplicated names (i.e. two MI columns with the same parents but different values).
 							ColumnName = Property->GetDisplayNameText().ToString() + GenerationContext.CurrentMaterialTableParameterId;
 						}
+						else
+						{
+							// Checking if this pin texture has been used in another table node but with a different texture mode
+							const int32 ImageDataIndex = GenerationContext.GeneratedTableImages.Find({ ColumnName, Pin->PinType.PinCategory, Table });
+
+							if (ImageDataIndex != INDEX_NONE )
+							{
+								if (GenerationContext.GeneratedTableImages[ImageDataIndex].PinType != Pin->PinType.PinCategory)
+								{
+									TArray<const UObject*> Nodes;
+									Nodes.Add(TypedNodeTable);
+									Nodes.Add(GenerationContext.GeneratedTableImages[ImageDataIndex].TableNode);
+
+									FString Msg = FString::Printf(TEXT("Texture pin [%s] with different texture modes found in more than one table node. This will add multiple times the texture reseource in the final cook."), *ColumnName);
+									GenerationContext.Compiler->CompilerLog(FText::FromString(Msg), Nodes);
+								}
+							}
+							else
+							{
+								GenerationContext.GeneratedTableImages.Add({ ColumnName, Pin->PinType.PinCategory, Table, TypedNodeTable });
+							}
+
+							// Encoding the texture mode of the type to allow different texture modes from the same pin
+							if (Pin->PinType.PinCategory == Schema->PC_PassThroughImage)
+							{
+								ColumnName += "--PassThrough";
+							}
+						}
 
 						// Generating a new Texture column if not exists
 						if (Table->FindColumn(StringCast<ANSICHAR>(*ColumnName).Get()) == INDEX_NONE)
diff --git a/Engine/Plugins/Experimental/Mutable/Source/CustomizableObjectEditor/Private/MuCOE/GenerateMutableSource/GenerateMutableSourceTable.cpp b/Engine/Plugins/Experimental/Mutable/Source/CustomizableObjectEditor/Private/MuCOE/GenerateMutableSource/GenerateMutableSourceTable.cpp
--- a/Engine/Plugins/Experimental/Mutable/Source/CustomizableObjectEditor/Private/MuCOE/GenerateMutableSource/GenerateMutableSourceTable.cpp
+++ b/Engine/Plugins/Experimental/Mutable/Source/CustomizableObjectEditor/Private/MuCOE/GenerateMutableSource/GenerateMutableSourceTable.cpp
@@ -393,9 +393,11 @@ bool FillTableColumn(const UCustomizableObjectNodeTable* TableNode,	mu::TablePtr
 				CurrentColumn = MutableTable->AddColumn(StringCast<ANSICHAR>(*ColumnName).Get(), mu::ETableColumnType::Image);
 			}
 
-			if (TableNode->GetColumnImageMode(ColumnName) == ETableTextureType::PASSTHROUGH_TEXTURE)
-			{
+			// Removing encoding part
+			const FString PinName = ColumnName.Replace(TEXT("--PassThrough"),TEXT(""), ESearchCase::CaseSensitive);
 
+			if (TableNode->GetColumnImageMode(PinName) == ETableTextureType::PASSTHROUGH_TEXTURE)
+			{
 				FMutableGraphGenerationContext::FGeneratedPassThroughTexture* FoundIndex = GenerationContext.PassThroughTextureMap.Find(Texture);
 				FMutableGraphGenerationContext::FGeneratedPassThroughTexture NewEntry;
 
