[Mutable] Added warning when Material Node can not deduce UV Index from Material.

#rnx

diff --git a/Engine/Plugins/Experimental/Mutable/Source/CustomizableObjectEditor/Classes/MuCOE/Nodes/CustomizableObjectNodeMaterial.h b/Engine/Plugins/Experimental/Mutable/Source/CustomizableObjectEditor/Classes/MuCOE/Nodes/CustomizableObjectNodeMaterial.h
--- a/Engine/Plugins/Experimental/Mutable/Source/CustomizableObjectEditor/Classes/MuCOE/Nodes/CustomizableObjectNodeMaterial.h
+++ b/Engine/Plugins/Experimental/Mutable/Source/CustomizableObjectEditor/Classes/MuCOE/Nodes/CustomizableObjectNodeMaterial.h
@@ -238,9 +238,10 @@ public:
 	 * @param ImageIndex Have to be valid. */
 	UTexture2D* GetImageValue(int32 ImageIndex) const;
 	
-	/** Get the Material Texture Parameter UV index.
+	/** Get the Material Texture Parameter UV Index.
 	 *
-	 * @param ImageIndex Have to be valid. */
+	 * @param ImageIndex Have to be valid.
+	 * @return Return -1 if the UV Index is set to Ignore. Return >= 0 for a valid UV Index.  */
 	int32 GetImageUVLayout(int32 ImageIndex) const;
 
 	/** Delegate called when a Texture Parameter Pin Mode changes. */
@@ -286,6 +287,9 @@ private:
 	/** Returns the Image Pin Mode the pin should be at. It does not update its mode, to update it call UpdateImagePinMode. */
 	EPinMode GetImagePinMode(const UEdGraphPin& Pin) const;
 
+	/** Get the UV Layout Index defined in the Material. */
+	int32 GetImageUVLayoutFromMaterial(int32 ImageIndex) const;
+	
 	// Deprecated properties
 	/** Set all pins to Mutable mode. Even so, each pin can override its behaviour. */
 	UPROPERTY()
diff --git a/Engine/Plugins/Experimental/Mutable/Source/CustomizableObjectEditor/Private/MuCOE/Nodes/CustomizableObjectNodeMaterial.cpp b/Engine/Plugins/Experimental/Mutable/Source/CustomizableObjectEditor/Private/MuCOE/Nodes/CustomizableObjectNodeMaterial.cpp
--- a/Engine/Plugins/Experimental/Mutable/Source/CustomizableObjectEditor/Private/MuCOE/Nodes/CustomizableObjectNodeMaterial.cpp
+++ b/Engine/Plugins/Experimental/Mutable/Source/CustomizableObjectEditor/Private/MuCOE/Nodes/CustomizableObjectNodeMaterial.cpp
@@ -19,7 +19,7 @@
 #include "ObjectEditorUtils.h"
 #include "PropertyCustomizationHelpers.h"
 #include "Modules/ModuleManager.h"
-
+#include "MuCOE/CustomizableObjectEditorLogger.h"
 
 class SGraphNode;
 class SWidget;
@@ -148,15 +148,9 @@ int32 UCustomizableObjectNodeMaterial::GetExpressionTextureCoordinate(UMaterial*
 		{
 			return TextureCoords->CoordinateIndex;
 		}
-		else
-		{
-			return 0;
-		}
-	}
-	else
-	{
-		return -1;
 	}
+
+	return -1;
 }
 
 
@@ -204,6 +198,38 @@ EPinMode UCustomizableObjectNodeMaterial::GetImagePinMode(const UEdGraphPin& Pin
 }
 
 
+int32 UCustomizableObjectNodeMaterial::GetImageUVLayoutFromMaterial(const int32 ImageIndex) const
+{
+	const FGuid ImageId = GetParameterId(EMaterialParameterType::Texture, ImageIndex);
+
+	if (const int32 TextureCoordinate = GetExpressionTextureCoordinate(Material->GetMaterial(), ImageId); TextureCoordinate >= 0)
+	{
+		return TextureCoordinate;
+	}
+
+	FMaterialLayersFunctions Layers;
+	Material->GetMaterialLayers(Layers);
+
+	TArray<TArray<TObjectPtr<UMaterialFunctionInterface>>*> MaterialFunctionInterfaces;
+	MaterialFunctionInterfaces.SetNumUninitialized(2);
+	MaterialFunctionInterfaces[0] = &Layers.Layers;
+	MaterialFunctionInterfaces[1] = &Layers.Blends;
+
+	for (const TArray<TObjectPtr<UMaterialFunctionInterface>>* MaterialFunctionInterface : MaterialFunctionInterfaces)
+	{
+		for (const TObjectPtr<UMaterialFunctionInterface>& Layer : *MaterialFunctionInterface)
+		{
+			if (const int32 TextureCoordinate = GetExpressionTextureCoordinate(Layer->GetPreviewMaterial()->GetMaterial(), ImageId); TextureCoordinate >= 0)
+			{
+				return TextureCoordinate;
+			}	
+		}
+	}
+
+	return -1;
+}
+
+
 bool UCustomizableObjectNodeMaterialPinDataParameter::IsDefault() const
 {
 	return true;
@@ -258,7 +284,7 @@ void UCustomizableObjectNodeMaterialPinDataImage::PostLoad()
 			UVLayoutMode = EUVLayoutMode::Ignore;
 			UVLayout = 0;
 		}
-		else if (UVLayout == -2) // UV_LAYOUT_DEFAULT
+		else if (UVLayout == UV_LAYOUT_DEFAULT)
 		{
 			UVLayoutMode = EUVLayoutMode::FromMaterial;
 			UVLayout = 0;
@@ -345,10 +371,27 @@ void UCustomizableObjectNodeMaterial::BackwardsCompatibleFixup()
 				{
 					PinData->ParameterId = GetParameterId(EMaterialParameterType::Texture, ParameterIndex);
 
-					if (const int32 UVLayout = GetImageUVLayout(ParameterIndex); Image.UVLayout != UVLayout)
+					if (Image.UVLayout == -1)
 					{
-						PinData->UVLayout = UVLayout;
+						PinData->UVLayout = Image.UVLayout;
 					}
+					else
+					{
+						const int32 UVLayout = GetImageUVLayoutFromMaterial(ParameterIndex);
+						if (UVLayout < 0) // Could not be deduced from the Material
+						{
+							PinData->UVLayout = Image.UVLayout;						
+						}
+						else if (UVLayout == Image.UVLayout)
+						{
+							PinData->UVLayout = UV_LAYOUT_DEFAULT;							
+						}
+						else
+						{
+							PinData->UVLayout = UVLayout;
+						}					
+					}
+					
 					break;
 				}
 			}
@@ -815,37 +858,18 @@ int32 UCustomizableObjectNodeMaterial::GetImageUVLayout(const int32 ImageIndex)
 		}
 	}
 
-	// Else
-	const FGuid ImageId = GetParameterId(EMaterialParameterType::Texture, ImageIndex);
-	
-	if (int32 TextureCoordinate = GetExpressionTextureCoordinate(Material->GetMaterial(), ImageId); TextureCoordinate >= 0)
+	const int32 UVIndex = GetImageUVLayoutFromMaterial(ImageIndex);
+	if (UVIndex == -1)
 	{
-		return TextureCoordinate;
+		FCustomizableObjectEditorLogger::CreateLog(LOCTEXT("UVLayoutMaterialError", "Could not deduce the UV Layout Index from the UMaterial. Only directly connected UMaterialExpressionTextureCoordinate are supported."))
+			.Severity(EMessageSeverity::Warning)
+			.Context(*this)
+			.Log();
+		
+		return 0;
 	}
-	else
-	{
-		FMaterialLayersFunctions Layers;
-		Material->GetMaterialLayers(Layers);
-
-		TArray<TArray<TObjectPtr<UMaterialFunctionInterface>>*> MaterialFunctionInterfaces;
-		MaterialFunctionInterfaces.SetNumUninitialized(2);
-		MaterialFunctionInterfaces[0] = &Layers.Layers;
-		MaterialFunctionInterfaces[1] = &Layers.Blends;
 
-		for (const TArray<TObjectPtr<UMaterialFunctionInterface>>* MaterialFunctionInterface : MaterialFunctionInterfaces)
-		{
-			for (const TObjectPtr<UMaterialFunctionInterface>& Layer : *MaterialFunctionInterface)
-			{
-				if (TextureCoordinate = GetExpressionTextureCoordinate(Layer->GetPreviewMaterial()->GetMaterial(), ImageId); TextureCoordinate >= 0)
-				{
-					return TextureCoordinate;
-				}	
-			}
-		}
-	}
-	
-	check(false); // Texture Parameter not found.
-	return 0;
+	return UVIndex;
 }
 
 
diff --git a/Engine/Plugins/Experimental/Mutable/Source/CustomizableObjectEditor/Public/MuCOE/CustomizableObjectEditor_Deprecated.h b/Engine/Plugins/Experimental/Mutable/Source/CustomizableObjectEditor/Public/MuCOE/CustomizableObjectEditor_Deprecated.h
--- a/Engine/Plugins/Experimental/Mutable/Source/CustomizableObjectEditor/Public/MuCOE/CustomizableObjectEditor_Deprecated.h
+++ b/Engine/Plugins/Experimental/Mutable/Source/CustomizableObjectEditor/Public/MuCOE/CustomizableObjectEditor_Deprecated.h
@@ -67,6 +67,9 @@ struct CUSTOMIZABLEOBJECTEDITOR_API FCustomizableObjectNodeMaterialScalar
 	FString PinName;
 };
 
+// Deprecated, do not use!
+constexpr int32 UV_LAYOUT_DEFAULT = -2;
+
 // UCustomizableObjectNodeEditMaterial
 // Deprecated, do not use!
 USTRUCT()
