[Mutable] Removed incorrect texture passthrough warning.

- GetImageUVLayout will return a UV Index even if it is not specified by the uer (deduced from the Material). It does not make sense warn that a non-Mutable texture is using a valid UV Index.

#rnx

diff --git a/Engine/Plugins/Experimental/Mutable/Source/CustomizableObjectEditor/Private/MuCOE/GenerateMutableSource/GenerateMutableSourceSurface.cpp b/Engine/Plugins/Experimental/Mutable/Source/CustomizableObjectEditor/Private/MuCOE/GenerateMutableSource/GenerateMutableSourceSurface.cpp
--- a/Engine/Plugins/Experimental/Mutable/Source/CustomizableObjectEditor/Private/MuCOE/GenerateMutableSource/GenerateMutableSourceSurface.cpp
+++ b/Engine/Plugins/Experimental/Mutable/Source/CustomizableObjectEditor/Private/MuCOE/GenerateMutableSource/GenerateMutableSourceSurface.cpp
@@ -440,14 +440,6 @@ mu::NodeSurfacePtr GenerateMutableSourceSurface(const UEdGraphPin * Pin, FMutabl
 					const FString ImageName = TypedNodeMat->GetParameterName(EMaterialParameterType::Texture, ImageIndex).ToString();
 					FString SurfNodeImageName = FString::Printf(TEXT("%d"), GenerationContext.ImageProperties.Num());
 					SurfNode->SetImageName(ImageIndex, StringCast<ANSICHAR>(*SurfNodeImageName).Get());
-					int32 UVLayout = TypedNodeMat->GetImageUVLayout(ImageIndex);
-
-					if (UVLayout >= 0)
-					{
-						FString msg = FString::Printf(TEXT("Passthrough texture [%s] will ignore layout despite set to use layout [%d]"), *ImageName, UVLayout);
-						GenerationContext.Compiler->CompilerLog(FText::FromString(msg), Node, EMessageSeverity::Warning);
-
-					}
 
 					SurfNode->SetImageLayoutIndex(ImageIndex, -1);
 					SurfNode->SetImageAdditionalNames(ImageIndex, StringCast<ANSICHAR>(*TypedNodeMat->Material->GetName()).Get(), StringCast<ANSICHAR>(*ImageName).Get());
