[Mutable] Added more objects to the Participating Objects system.

#rnx

diff --git a/Engine/Plugins/Experimental/Mutable/Source/CustomizableObjectEditor/Private/MuCOE/CustomizableObjectCompiler.cpp b/Engine/Plugins/Experimental/Mutable/Source/CustomizableObjectEditor/Private/MuCOE/CustomizableObjectCompiler.cpp
--- a/Engine/Plugins/Experimental/Mutable/Source/CustomizableObjectEditor/Private/MuCOE/CustomizableObjectCompiler.cpp
+++ b/Engine/Plugins/Experimental/Mutable/Source/CustomizableObjectEditor/Private/MuCOE/CustomizableObjectCompiler.cpp
@@ -1285,6 +1285,7 @@ void FCustomizableObjectCompiler::CompileInternal(UCustomizableObject* Object, c
 			}
 		}
 
+		if (!ParamNamesToSelectedOptions.Num())
 		{
 			GenerationContext.ParticipatingObjects.Remove(Object); // Remove self CO reference.
 
diff --git a/Engine/Plugins/Experimental/Mutable/Source/CustomizableObjectEditor/Private/MuCOE/GenerateMutableSource/GenerateMutableSource.cpp b/Engine/Plugins/Experimental/Mutable/Source/CustomizableObjectEditor/Private/MuCOE/GenerateMutableSource/GenerateMutableSource.cpp
--- a/Engine/Plugins/Experimental/Mutable/Source/CustomizableObjectEditor/Private/MuCOE/GenerateMutableSource/GenerateMutableSource.cpp
+++ b/Engine/Plugins/Experimental/Mutable/Source/CustomizableObjectEditor/Private/MuCOE/GenerateMutableSource/GenerateMutableSource.cpp
@@ -1418,6 +1418,8 @@ void PopulateReferenceSkeletalMeshesData(FMutableGraphGenerationContext& Generat
 		// Physics Asset
 		if (const UPhysicsAsset* PhysicsAsset = RefSkeletalMesh->GetPhysicsAsset())
 		{
+			GenerationContext.AddParticipatingObject(*PhysicsAsset);
+
 			Data.PhysicsAsset = PhysicsAsset;
 			GenerationContext.PhysicsAssetMap.FindOrAdd(Data.PhysicsAsset.ToString(), Data.PhysicsAsset);
 		}
@@ -1425,14 +1427,18 @@ void PopulateReferenceSkeletalMeshesData(FMutableGraphGenerationContext& Generat
 		// Post ProcessAnimInstance
 		if(const TSubclassOf<UAnimInstance> PostProcessAnimInstance = RefSkeletalMesh->GetPostProcessAnimBlueprint())
 		{
+			GenerationContext.AddParticipatingObject(*PostProcessAnimInstance.Get());
+
 			Data.PostProcessAnimInst = PostProcessAnimInstance;
 			GenerationContext.AnimBPAssetsMap.FindOrAdd(Data.PostProcessAnimInst.ToString(), Data.PostProcessAnimInst);
 		}
 		
 		// Shadow Physics Asset
-		if (const UPhysicsAsset* PhysicsAsset = RefSkeletalMesh->GetPhysicsAsset())
+		if (const UPhysicsAsset* PhysicsAsset = RefSkeletalMesh->GetShadowPhysicsAsset())
 		{
-			Data.PhysicsAsset = PhysicsAsset;
+			GenerationContext.AddParticipatingObject(*PhysicsAsset);
+
+			Data.ShadowPhysicsAsset = PhysicsAsset;
 			GenerationContext.PhysicsAssetMap.FindOrAdd(Data.PhysicsAsset.ToString(), Data.PhysicsAsset);
 		}
 
diff --git a/Engine/Plugins/Experimental/Mutable/Source/CustomizableObjectEditor/Private/MuCOE/GenerateMutableSource/GenerateMutableSourceMesh.cpp b/Engine/Plugins/Experimental/Mutable/Source/CustomizableObjectEditor/Private/MuCOE/GenerateMutableSource/GenerateMutableSourceMesh.cpp
--- a/Engine/Plugins/Experimental/Mutable/Source/CustomizableObjectEditor/Private/MuCOE/GenerateMutableSource/GenerateMutableSourceMesh.cpp
+++ b/Engine/Plugins/Experimental/Mutable/Source/CustomizableObjectEditor/Private/MuCOE/GenerateMutableSource/GenerateMutableSourceMesh.cpp
@@ -2857,6 +2857,11 @@ mu::NodeMeshPtr GenerateMutableSourceMesh(const UEdGraphPin* Pin,
 
 			if (!TypedNodeSkel->AnimInstance.IsNull())
 			{
+				if (UClass* AnimInstance = TypedNodeSkel->AnimInstance.LoadSynchronous())
+				{
+					GenerationContext.AddParticipatingObject(*AnimInstance);					
+				}
+
 				FName SlotIndex = TypedNodeSkel->AnimBlueprintSlotName;
 				GenerationContext.AnimBPAssetsMap.Add(TypedNodeSkel->AnimInstance.ToString(), TypedNodeSkel->AnimInstance);
 
@@ -2883,6 +2888,8 @@ mu::NodeMeshPtr GenerateMutableSourceMesh(const UEdGraphPin* Pin,
 				{
 					for (const UAssetUserData* AssetUserData : *AssetUserDataArray)
 					{
+						GenerationContext.AddParticipatingObject(*AssetUserData);
+
 						FString AuxString = AssetUserData->GetPathName();
 						GenerationContext.AssetUserDataAssetsMap.Add(AuxString, TSoftObjectPtr<UAssetUserData>(AssetUserData));
 
@@ -2905,6 +2912,9 @@ mu::NodeMeshPtr GenerateMutableSourceMesh(const UEdGraphPin* Pin,
 					MutableMesh->GetPhysicsBody()->GetBodyCount())
 				{
 					TSoftObjectPtr<UPhysicsAsset> PhysicsAsset = TypedNodeSkel->SkeletalMesh->GetPhysicsAsset();
+
+					GenerationContext.AddParticipatingObject(*PhysicsAsset);
+
 					GenerationContext.PhysicsAssetMap.Add(PhysicsAsset.ToString(), PhysicsAsset);
 					FString PhysicsAssetTag = FString("__PhysicsAsset:") + PhysicsAsset.ToString();
 
@@ -2929,6 +2939,8 @@ mu::NodeMeshPtr GenerateMutableSourceMesh(const UEdGraphPin* Pin,
 						TSoftObjectPtr<UPhysicsAsset> PhysicsAsset = ClothingAssetCommon->PhysicsAsset.Get();
 
 						FString ClothPhysicsAssetTag = FString("__ClothPhysicsAsset:") + FString::Printf(TEXT("%d_AssetIdx_"), AssetIndex) + PhysicsAsset.ToString();
+
+						GenerationContext.AddParticipatingObject(*ClothingAssetCommon->PhysicsAsset);
 						
 						GenerationContext.PhysicsAssetMap.Add(PhysicsAsset.ToString(), ClothingAssetCommon->PhysicsAsset.Get());
 
diff --git a/Engine/Plugins/Experimental/Mutable/Source/CustomizableObjectEditor/Private/MuCOE/GenerateMutableSource/GenerateMutableSourceTable.cpp b/Engine/Plugins/Experimental/Mutable/Source/CustomizableObjectEditor/Private/MuCOE/GenerateMutableSource/GenerateMutableSourceTable.cpp
--- a/Engine/Plugins/Experimental/Mutable/Source/CustomizableObjectEditor/Private/MuCOE/GenerateMutableSource/GenerateMutableSourceTable.cpp
+++ b/Engine/Plugins/Experimental/Mutable/Source/CustomizableObjectEditor/Private/MuCOE/GenerateMutableSource/GenerateMutableSourceTable.cpp
@@ -109,6 +109,11 @@ bool FillTableColumn(const UCustomizableObjectNodeTable* TableNode,	mu::TablePtr
 
 											if (!AnimInstance.IsNull())
 											{
+												if (UClass* Anim = AnimInstance.LoadSynchronous())
+												{
+													GenerationContext.AddParticipatingObject(*Anim);													
+												}
+
 												GenerationContext.AnimBPAssetsMap.Add(AnimInstance.ToString(), AnimInstance);
 
 												AnimBPAssetTag = GenerateAnimationInstanceTag(AnimInstance.ToString(), SlotIndex);
@@ -240,8 +245,11 @@ bool FillTableColumn(const UCustomizableObjectNodeTable* TableNode,	mu::TablePtr
 			if (MutableMesh)
 			{
  				if (SkeletalMesh->GetPhysicsAsset() && MutableMesh->GetPhysicsBody() && MutableMesh->GetPhysicsBody()->GetBodyCount())
-				{	
-					TSoftObjectPtr<UPhysicsAsset> PhysicsAsset = SkeletalMesh->GetPhysicsAsset();
+				{
+ 					TSoftObjectPtr<UPhysicsAsset> PhysicsAsset = SkeletalMesh->GetPhysicsAsset();
+
+					GenerationContext.AddParticipatingObject(*PhysicsAsset); 						
+
 					GenerationContext.PhysicsAssetMap.Add(PhysicsAsset.ToString(), PhysicsAsset);
 					FString PhysicsAssetTag = FString("__PhysicsAsset:") + PhysicsAsset.ToString();
 
@@ -266,6 +274,8 @@ bool FillTableColumn(const UCustomizableObjectNodeTable* TableNode,	mu::TablePtr
 					{
 						for (const UAssetUserData* AssetUserData : *AssetUserDataArray)
 						{
+							GenerationContext.AddParticipatingObject(*AssetUserData);
+
 							FString AuxString = AssetUserData->GetPathName();
 							GenerationContext.AssetUserDataAssetsMap.Add(AuxString, TSoftObjectPtr<UAssetUserData>(AssetUserData));
 
