[Mutable] Fix crash when deserializing CO with references to non-existing UObjects.

#rnx

diff --git a/Engine/Plugins/Experimental/Mutable/Source/CustomizableObject/Private/MuCO/CustomizableObject.cpp b/Engine/Plugins/Experimental/Mutable/Source/CustomizableObject/Private/MuCO/CustomizableObject.cpp
--- a/Engine/Plugins/Experimental/Mutable/Source/CustomizableObject/Private/MuCO/CustomizableObject.cpp
+++ b/Engine/Plugins/Experimental/Mutable/Source/CustomizableObject/Private/MuCO/CustomizableObject.cpp
@@ -183,6 +183,11 @@ void UCustomizableObject::PostLoad()
 #if WITH_EDITORONLY_DATA
 	for (TTuple<TObjectPtr<const UObject>, FGuid>& ParticipatingObject : ParticipatingObjects)
 	{
+		if (!ParticipatingObject.Key) // Object no longer exists.
+		{
+			continue;
+		}
+		
 		PRAGMA_DISABLE_DEPRECATION_WARNINGS
 		const FGuid PackageGuid = ParticipatingObject.Key.GetPackage()->GetGuid();
 		PRAGMA_ENABLE_DEPRECATION_WARNINGS
