[Mutable] Fix incorrect LOD bias being applied to textures.

#rnx

diff --git a/Engine/Plugins/Experimental/Mutable/Source/CustomizableObject/Classes/MuCO/CustomizableObject.h b/Engine/Plugins/Experimental/Mutable/Source/CustomizableObject/Classes/MuCO/CustomizableObject.h
--- a/Engine/Plugins/Experimental/Mutable/Source/CustomizableObject/Classes/MuCO/CustomizableObject.h
+++ b/Engine/Plugins/Experimental/Mutable/Source/CustomizableObject/Classes/MuCO/CustomizableObject.h
@@ -224,6 +224,9 @@ struct FCompilationOptions
 	// Used to reduce the number of notifications when compiling objects
 	bool bSilentCompilation = true;
 
+	// Used to reduce texture size on higher mesh LODs. Only active if LOD strategy is set to Automatic LODs from Mesh
+	bool bUseLODAsBias = true;
+
 	/** Force a very big number on the mips to skip during compilation. Useful to debug special cooks of the data. */
 	bool bForceLargeLODBias = false;
 	int32 DebugBias = 0;
diff --git a/Engine/Plugins/Experimental/Mutable/Source/CustomizableObjectEditor/Private/MuCOE/GenerateMutableSource/GenerateMutableSource.cpp b/Engine/Plugins/Experimental/Mutable/Source/CustomizableObjectEditor/Private/MuCOE/GenerateMutableSource/GenerateMutableSource.cpp
--- a/Engine/Plugins/Experimental/Mutable/Source/CustomizableObjectEditor/Private/MuCOE/GenerateMutableSource/GenerateMutableSource.cpp
+++ b/Engine/Plugins/Experimental/Mutable/Source/CustomizableObjectEditor/Private/MuCOE/GenerateMutableSource/GenerateMutableSource.cpp
@@ -1461,7 +1461,7 @@ void PopulateReferenceSkeletalMeshesData(FMutableGraphGenerationContext& Generat
 
 
 int32 ComputeLODBias(const FMutableGraphGenerationContext& GenerationContext, const UTexture2D* ReferenceTexture, int32 MaxTextureSize,
-	const UCustomizableObjectNodeMaterial* MaterialNode, const int32 ImageIndex, bool bUseLODAsBias)
+	const UCustomizableObjectNodeMaterial* MaterialNode, const int32 ImageIndex)
 {
 	constexpr int32 MaxAllowedLODBias = 6;
 
@@ -1484,12 +1484,12 @@ int32 ComputeLODBias(const FMutableGraphGenerationContext& GenerationContext, co
 	}
 
 	// Increment the LOD bias per each LOD if we are using automatic LODs
-	if (GenerationContext.CurrentLOD > 0
-		&&
-		GenerationContext.CurrentAutoLODStrategy == ECustomizableObjectAutomaticLODStrategy::AutomaticFromMesh)
+	if (GenerationContext.CurrentLOD > 0 &&
+		GenerationContext.CurrentAutoLODStrategy == ECustomizableObjectAutomaticLODStrategy::AutomaticFromMesh &&
+		GenerationContext.Options.bUseLODAsBias)
 	{
 		// Only if the texture actually uses a layout. Otherwise it could be a special texture we shouldn't scale.
-		if ((MaterialNode && MaterialNode->GetImageUVLayout(ImageIndex) >= 0) || bUseLODAsBias)
+		if (MaterialNode && MaterialNode->GetImageUVLayout(ImageIndex) >= 0)
 		{
 			// \todo: make it an object property to be tweaked
 			int MipsToSkipPerLOD = 1;
diff --git a/Engine/Plugins/Experimental/Mutable/Source/CustomizableObjectEditor/Private/MuCOE/GenerateMutableSource/GenerateMutableSource.h b/Engine/Plugins/Experimental/Mutable/Source/CustomizableObjectEditor/Private/MuCOE/GenerateMutableSource/GenerateMutableSource.h
--- a/Engine/Plugins/Experimental/Mutable/Source/CustomizableObjectEditor/Private/MuCOE/GenerateMutableSource/GenerateMutableSource.h
+++ b/Engine/Plugins/Experimental/Mutable/Source/CustomizableObjectEditor/Private/MuCOE/GenerateMutableSource/GenerateMutableSource.h
@@ -845,7 +845,7 @@ FString GenerateAssetUserDataTag(const FString& AssetUserData);
 // Computes the LOD bias for a texture given the current mesh LOD and automatic LOD settings, the reference texture settings
 // and whether it's being built for a server or not
 int32 ComputeLODBias(const FMutableGraphGenerationContext& GenerationContext, const UTexture2D* ReferenceTexture, int32 MaxTextureSize,
-	const UCustomizableObjectNodeMaterial* MaterialNode, const int32 ImageIndex, bool bUseLODAsBias = true);
+	const UCustomizableObjectNodeMaterial* MaterialNode, const int32 ImageIndex);
 
 
 int32 GetMaxTextureSize(const UTexture2D* ReferenceTexture, const FMutableGraphGenerationContext& GenerationContext);
diff --git a/Engine/Plugins/Experimental/Mutable/Source/CustomizableObjectEditor/Private/MuCOE/GenerateMutableSource/GenerateMutableSourceImage.cpp b/Engine/Plugins/Experimental/Mutable/Source/CustomizableObjectEditor/Private/MuCOE/GenerateMutableSource/GenerateMutableSourceImage.cpp
--- a/Engine/Plugins/Experimental/Mutable/Source/CustomizableObjectEditor/Private/MuCOE/GenerateMutableSource/GenerateMutableSourceImage.cpp
+++ b/Engine/Plugins/Experimental/Mutable/Source/CustomizableObjectEditor/Private/MuCOE/GenerateMutableSource/GenerateMutableSourceImage.cpp
@@ -619,7 +619,7 @@ mu::NodeImagePtr GenerateMutableSourceImage(const UEdGraphPin* Pin, FMutableGrap
 		{
 			GenerationContext.AddParticipatingObject(*TypedNodeProject->ReferenceTexture);
 
-			int32 LODBias = ComputeLODBias(GenerationContext, TypedNodeProject->ReferenceTexture, TypedNodeProject->ReferenceTexture->MaxTextureSize, nullptr, INDEX_NONE, false);
+			int32 LODBias = ComputeLODBias(GenerationContext, TypedNodeProject->ReferenceTexture, TypedNodeProject->ReferenceTexture->MaxTextureSize, nullptr, INDEX_NONE);
 
 			if (TextureSize.X > 0 && TextureSize.Y > 0)
 			{
@@ -856,7 +856,7 @@ mu::NodeImagePtr GenerateMutableSourceImage(const UEdGraphPin* Pin, FMutableGrap
 			
 			const int32 LODBias = ComputeLODBias(
 				GenerationContext, TypedNodeTransform->ReferenceTexture, TypedNodeTransform->ReferenceTexture->MaxTextureSize, 
-				nullptr, INDEX_NONE, false);
+				nullptr, INDEX_NONE);
 
 			if (TextureSize.X > 0 && TextureSize.Y > 0)
 			{
diff --git a/Engine/Plugins/Experimental/Mutable/Source/CustomizableObjectEditor/Private/MuCOE/GenerateMutableSource/GenerateMutableSourceSurface.cpp b/Engine/Plugins/Experimental/Mutable/Source/CustomizableObjectEditor/Private/MuCOE/GenerateMutableSource/GenerateMutableSourceSurface.cpp
--- a/Engine/Plugins/Experimental/Mutable/Source/CustomizableObjectEditor/Private/MuCOE/GenerateMutableSource/GenerateMutableSourceSurface.cpp
+++ b/Engine/Plugins/Experimental/Mutable/Source/CustomizableObjectEditor/Private/MuCOE/GenerateMutableSource/GenerateMutableSourceSurface.cpp
@@ -552,7 +552,7 @@ mu::NodeSurfacePtr GenerateMutableSourceSurface(const UEdGraphPin * Pin, FMutabl
 						SurfaceData.ImageProperties = Props;
 
 						// Calculate the LODBias for this texture
-						int32 LODBias = ComputeLODBias(GenerationContext, ReferenceTexture, Props.MaxTextureSize, TypedNodeMat, ImageIndex);
+						int32 LODBias = ComputeLODBias(GenerationContext, ReferenceTexture, ReferenceTexture ? ReferenceTexture->MaxTextureSize : 0, TypedNodeMat, ImageIndex);
 
 						GenerationContext.CurrentTextureLODBias = LODBias;
 
