[Mutable] Invalidate the old model after recompiling a CO in the editor.

#rnx

diff --git a/Engine/Plugins/Experimental/Mutable/Source/CustomizableObject/Private/MuCO/CustomizableObject.cpp b/Engine/Plugins/Experimental/Mutable/Source/CustomizableObject/Private/MuCO/CustomizableObject.cpp
--- a/Engine/Plugins/Experimental/Mutable/Source/CustomizableObject/Private/MuCO/CustomizableObject.cpp
+++ b/Engine/Plugins/Experimental/Mutable/Source/CustomizableObject/Private/MuCO/CustomizableObject.cpp
@@ -1858,6 +1858,13 @@ FGuid UCustomizableObject::GetCompilationGuid() const
 
 void FCustomizableObjectPrivateData::SetModel(const TSharedPtr<mu::Model, ESPMode::ThreadSafe>& Model, const FGuid Id)
 {
+#if WITH_EDITOR
+	if (MutableModel)
+	{
+		MutableModel->Invalidate();
+	}
+#endif
+	
 	MutableModel = Model;
 
 #if WITH_EDITOR
diff --git a/Engine/Plugins/Experimental/Mutable/Source/CustomizableObject/Private/MuCO/CustomizableObjectMipDataProvider.cpp b/Engine/Plugins/Experimental/Mutable/Source/CustomizableObject/Private/MuCO/CustomizableObjectMipDataProvider.cpp
--- a/Engine/Plugins/Experimental/Mutable/Source/CustomizableObject/Private/MuCO/CustomizableObjectMipDataProvider.cpp
+++ b/Engine/Plugins/Experimental/Mutable/Source/CustomizableObject/Private/MuCO/CustomizableObjectMipDataProvider.cpp
@@ -124,102 +124,108 @@ namespace impl
 			mu::SystemPtr System = OperationData->UpdateContext->GetSystem();
 			const TSharedPtr<mu::Model, ESPMode::ThreadSafe> Model = OperationData->UpdateContext->GetModel();
 
-			// For now, we are forcing the recreation of mutable-side instances with every update.
-			mu::Instance::ID InstanceID = System->NewInstance(Model);
-			UE_LOG(LogMutable, Verbose, TEXT("Creating Mutable instance with id [%d] for a single UpdateImage"), InstanceID)
-
-				const mu::Instance* Instance = nullptr;
-
-			// Main instance generation step
+#if WITH_EDITOR
+			// Recompiling a CO in the editor will invalidate the previously generated Model. Check that it is valid before accessing the streamed data.
+			if (Model && Model->IsValid())
+#endif
 			{
-				// LOD mask, set to all ones to build all LODs
-				uint32 LODMask = 0xFFFFFFFF;
 
-				Instance = System->BeginUpdate(InstanceID, OperationData->UpdateContext->GetParameters(), OperationData->UpdateContext->GetState(), LODMask);
-				check(Instance);
-			}
+				// For now, we are forcing the recreation of mutable-side instances with every update.
+				mu::Instance::ID InstanceID = System->NewInstance(Model);
+				UE_LOG(LogMutable, Verbose, TEXT("Creating Mutable instance with id [%d] for a single UpdateImage"), InstanceID)
 
+					const mu::Instance* Instance = nullptr;
 
-			// Generate the required image
-			{
-				MUTABLE_CPUPROFILER_SCOPE(RequestedImage);
+				// Main instance generation step
+				{
+					// LOD mask, set to all ones to build all LODs
+					uint32 LODMask = 0xFFFFFFFF;
 
-				const FMutableImageReference& ImageRef = OperationData->RequestedImage;
+					Instance = System->BeginUpdate(InstanceID, OperationData->UpdateContext->GetParameters(), OperationData->UpdateContext->GetState(), LODMask);
+					check(Instance);
+				}
 
-				int32 SurfaceIndex = Instance->FindSurfaceById(ImageRef.LOD, ImageRef.Component, ImageRef.SurfaceId);
-				check(SurfaceIndex >= 0);
 
-				// This ID may be different than the ID obtained the first time the image was generated, because the mutable
-				// runtime cannot remember all the resources it has built, and only remembers a fixed amount.
-				mu::FResourceID MipImageID = Instance->GetImageId(ImageRef.LOD, ImageRef.Component, SurfaceIndex, ImageRef.Image);
+				// Generate the required image
+				{
+					MUTABLE_CPUPROFILER_SCOPE(RequestedImage);
 
+					const FMutableImageReference& ImageRef = OperationData->RequestedImage;
 
-				mu::ImagePtrConst Image;
-				{
-					MUTABLE_CPUPROFILER_SCOPE(GetImage);
+					int32 SurfaceIndex = Instance->FindSurfaceById(ImageRef.LOD, ImageRef.Component, ImageRef.SurfaceId);
+					check(SurfaceIndex >= 0);
 
-					Image = System->GetImage(InstanceID, MipImageID, ImageRef.BaseMip + OperationData->MipsToSkip, ImageRef.LOD);
-				}
+					// This ID may be different than the ID obtained the first time the image was generated, because the mutable
+					// runtime cannot remember all the resources it has built, and only remembers a fixed amount.
+					mu::FResourceID MipImageID = Instance->GetImageId(ImageRef.LOD, ImageRef.Component, SurfaceIndex, ImageRef.Image);
 
-				check(Image);
 
-				int32 FullMipCount = Image->GetMipmapCount(Image->GetSizeX(), Image->GetSizeY());
-				int32 RealMipCount = Image->GetLODCount();
+					mu::ImagePtrConst Image;
+					{
+						MUTABLE_CPUPROFILER_SCOPE(GetImage);
 
-				bool bForceMipchain =
-					// Did we fail to generate the entire mipchain (if we have mips at all)?
-					(RealMipCount != 1) && (RealMipCount != FullMipCount);
+						Image = System->GetImage(InstanceID, MipImageID, ImageRef.BaseMip + OperationData->MipsToSkip, ImageRef.LOD);
+					}
 
-				if (bForceMipchain)
-				{
-					MUTABLE_CPUPROFILER_SCOPE(GetImage_MipFix);
+					check(Image);
+
+					int32 FullMipCount = Image->GetMipmapCount(Image->GetSizeX(), Image->GetSizeY());
+					int32 RealMipCount = Image->GetLODCount();
 
-					UE_LOG(LogMutable, Warning, TEXT("Mutable generated an incomplete mip chain for image."));
+					bool bForceMipchain =
+						// Did we fail to generate the entire mipchain (if we have mips at all)?
+						(RealMipCount != 1) && (RealMipCount != FullMipCount);
 
-					// Force the right number of mips. The missing data will be black.
-					mu::Ptr<mu::Image> NewImage = new mu::Image(Image->GetSizeX(), Image->GetSizeY(), FullMipCount, Image->GetFormat(), mu::EInitializationType::Black);
-					check(NewImage);
-					if (NewImage->GetDataSize() >= Image->GetDataSize())
+					if (bForceMipchain)
 					{
-						FMemory::Memcpy(NewImage->GetData(), Image->GetData(), Image->GetDataSize());
+						MUTABLE_CPUPROFILER_SCOPE(GetImage_MipFix);
+
+						UE_LOG(LogMutable, Warning, TEXT("Mutable generated an incomplete mip chain for image."));
+
+						// Force the right number of mips. The missing data will be black.
+						mu::Ptr<mu::Image> NewImage = new mu::Image(Image->GetSizeX(), Image->GetSizeY(), FullMipCount, Image->GetFormat(), mu::EInitializationType::Black);
+						check(NewImage);
+						if (NewImage->GetDataSize() >= Image->GetDataSize())
+						{
+							FMemory::Memcpy(NewImage->GetData(), Image->GetData(), Image->GetDataSize());
+						}
+						Image = NewImage;
 					}
-					Image = NewImage;
-				}
-
-				OperationData->Result = Image;
-			}
 
-			// End update
-			{
-				MUTABLE_CPUPROFILER_SCOPE(EndUpdate);
-				System->EndUpdate(InstanceID);
-				System->ReleaseInstance(InstanceID);
+					OperationData->Result = Image;
+				}
 
-				if (CVarClearWorkingMemoryOnUpdateEnd.GetValueOnAnyThread())
+				// End update
 				{
-					System->ClearWorkingMemory();
+					MUTABLE_CPUPROFILER_SCOPE(EndUpdate);
+					System->EndUpdate(InstanceID);
+					System->ReleaseInstance(InstanceID);
+
+					if (CVarClearWorkingMemoryOnUpdateEnd.GetValueOnAnyThread())
+					{
+						System->ClearWorkingMemory();
+					}
 				}
 			}
+		}
+
+		{
+			// The request could be cancelled in parallel from CancelCounterSafely and its value be changed
+			// between reading it and actually running Decrement() and RescheduleCallback(), so lock
+			FScopeLock Lock(&OperationData->CounterTaskLock);
 
+			if (OperationData->Counter) // If the request has been cancelled the counter will be null
 			{
-				// The request could be cancelled in parallel from CancelCounterSafely and its value be changed
-				// between reading it and actually running Decrement() and RescheduleCallback(), so lock
-				FScopeLock Lock(&OperationData->CounterTaskLock);
+				// Make the FMutableTextureMipDataProvider continue
+				OperationData->Counter->Decrement();
 
-				if (OperationData->Counter) // If the request has been cancelled the counter will be null
+				if (OperationData->Counter->GetValue() == 0)
 				{
-					// Make the FMutableTextureMipDataProvider continue
-					OperationData->Counter->Decrement();
-
-					if (OperationData->Counter->GetValue() == 0)
-					{
-						OperationData->RescheduleCallback();
-					}
+					OperationData->RescheduleCallback();
 				}
 			}
 		}
 	}
-
 } // namespace
 
 
@@ -343,33 +349,40 @@ bool FMutableTextureMipDataProvider::PollMips(const FTextureUpdateSyncOptions& S
 		return false;
 	}
 	
-	if (OperationData && OperationData->Result && OperationData->Levels.Num())
+	if (OperationData && OperationData->Levels.Num())
 	{
 		// The counter must be zero meaning the Mutable image operation has finished
 		check(SyncOptions.Counter->GetValue() == 0);
 
-		mu::Ptr<const mu::Image> Mip = OperationData->Result;
-		int32 MipIndex = 0;
-		check(Mip->GetSizeX() == OperationData->Levels[0].SizeX);
-		check(Mip->GetSizeY() == OperationData->Levels[0].SizeY);
+		mu::Ptr<const mu::Image> Image = OperationData->Result;
+		
+		int32 ImageLODCount = 0;
+		if (Image)
+		{
+			ImageLODCount = Image->GetLODCount();
+			// check(Image->GetLODCount() == OperationData->Levels.Num()); TODO PRP
+			check(Image->GetSizeX() == OperationData->Levels[0].SizeX);
+			check(Image->GetSizeY() == OperationData->Levels[0].SizeY);
+		}
 
+		int32 MipIndex = 0;
 		for (FMutableMipUpdateLevel& Level : OperationData->Levels)
 		{
 			void* Dest = Level.Dest;
 
-			if (MipIndex >= Mip->GetLODCount())
+			if (MipIndex < ImageLODCount)
 			{
-				// Mutable didn't generate all the expected mips
-				UE_LOG(LogMutable, Warning, TEXT("Mutable image is missing mips."));
-				FMemory::Memzero(Dest, Level.DataSize);
-			}
-			else
-			{
-				int32 MipDataSize = Mip->GetLODDataSize(MipIndex);
+				int32 MipDataSize = Image->GetLODDataSize(MipIndex);
 
 				// Check Mip DataSize for consistency, but skip if 0 because it's optional and might be zero in cooked mips
 				check(Level.DataSize == 0 || MipDataSize == Level.DataSize);
-				FMemory::Memcpy(Dest, Mip->GetMipData(MipIndex), MipDataSize);
+				FMemory::Memcpy(Dest, Image->GetMipData(MipIndex), MipDataSize);
+			}
+			else
+			{
+				// Mutable didn't generate all the expected mips
+				UE_LOG(LogMutable, Warning, TEXT("Mutable image is missing mips."));
+				FMemory::Memzero(Dest, Level.DataSize);
 			}
 			++MipIndex;
 		}
diff --git a/Engine/Plugins/Experimental/Mutable/Source/MutableRuntime/Private/MuR/Model.cpp b/Engine/Plugins/Experimental/Mutable/Source/MutableRuntime/Private/MuR/Model.cpp
--- a/Engine/Plugins/Experimental/Mutable/Source/MutableRuntime/Private/MuR/Model.cpp
+++ b/Engine/Plugins/Experimental/Mutable/Source/MutableRuntime/Private/MuR/Model.cpp
@@ -221,6 +221,20 @@ namespace mu
 		return m_pD->m_program.m_roms.Num() > 0;
     }
 
+#if WITH_EDITOR
+	//---------------------------------------------------------------------------------------------
+	bool Model::IsValid() const
+	{
+		return m_pD->m_program.bIsValid;
+	}
+
+
+	//---------------------------------------------------------------------------------------------
+	void Model::Invalidate()
+    {
+		m_pD->m_program.bIsValid = false;
+    }
+#endif
 
     //---------------------------------------------------------------------------------------------
     void Model::UnloadExternalData()
diff --git a/Engine/Plugins/Experimental/Mutable/Source/MutableRuntime/Private/MuR/ModelPrivate.h b/Engine/Plugins/Experimental/Mutable/Source/MutableRuntime/Private/MuR/ModelPrivate.h
--- a/Engine/Plugins/Experimental/Mutable/Source/MutableRuntime/Private/MuR/ModelPrivate.h
+++ b/Engine/Plugins/Experimental/Mutable/Source/MutableRuntime/Private/MuR/ModelPrivate.h
@@ -270,7 +270,11 @@ namespace mu
         //! The parameter lists are sorted.
 		TArray<TArray<uint16>> m_parameterLists;
 
-
+#if WITH_EDITOR
+		//! State of the program. True unless the streamed resources were destroyed,
+		//! which could happen in the editor after recompiling the CO.
+		bool bIsValid = true;
+#endif
         //!
         void Serialise( OutputArchive& arch ) const
         {
diff --git a/Engine/Plugins/Experimental/Mutable/Source/MutableRuntime/Public/MuR/Model.h b/Engine/Plugins/Experimental/Mutable/Source/MutableRuntime/Public/MuR/Model.h
--- a/Engine/Plugins/Experimental/Mutable/Source/MutableRuntime/Public/MuR/Model.h
+++ b/Engine/Plugins/Experimental/Mutable/Source/MutableRuntime/Public/MuR/Model.h
@@ -66,6 +66,16 @@ namespace mu
 		//! require data streaming when used.
 		bool HasExternalData() const;
 
+#if WITH_EDITOR
+		//! Return true unless the streamed resources were destroyed, which could happen in the
+		//! editor after recompiling the CO.
+		bool IsValid() const;
+
+		//! Invalidate the Model. Compiling a compiled CO will invalidate the model kept by previously
+		//! generated resources, like streamed textures.
+		void Invalidate();
+#endif
+
 		//-----------------------------------------------------------------------------------------
 		// Own interface
 		//-----------------------------------------------------------------------------------------
