[mutable] Added new option to reduce layout blocks.
+ Now layout blocks can be reduced by two unit blocks in a unitary block reduction strategy.
+ Minor Fix: Now layout index property will be hidden by default and only will be visible when the layout mode option is "Index".

#rnx

diff --git a/Engine/Plugins/Experimental/Mutable/Source/CustomizableObjectEditor/Classes/MuCOE/CustomizableObjectLayout.h b/Engine/Plugins/Experimental/Mutable/Source/CustomizableObjectEditor/Classes/MuCOE/CustomizableObjectLayout.h
--- a/Engine/Plugins/Experimental/Mutable/Source/CustomizableObjectEditor/Classes/MuCOE/CustomizableObjectLayout.h
+++ b/Engine/Plugins/Experimental/Mutable/Source/CustomizableObjectEditor/Classes/MuCOE/CustomizableObjectLayout.h
@@ -29,28 +29,39 @@ struct CUSTOMIZABLEOBJECTEDITOR_API FCustomizableObjectLayoutBlock
 {
 	GENERATED_USTRUCT_BODY()
 
-	FCustomizableObjectLayoutBlock()
+	FCustomizableObjectLayoutBlock(FIntPoint InMin = FIntPoint(0, 0), FIntPoint InMax = FIntPoint(1, 1))
 	{
-		Min = FIntPoint(0, 0);
-		Max = FIntPoint(1, 1);
+		Min = InMin;
+		Max = InMax;
 		Priority = 0;
+		Id = FGuid::NewGuid();
+		bReduceBothAxes = false;
+		bReduceByTwo = false;
 	}
 
+	/** Top left coordinate. */
 	UPROPERTY(EditAnywhere, Category = CustomizableObject)
 	FIntPoint Min;
 
+	/** Bottom right coordinate. */
 	UPROPERTY(EditAnywhere, Category = CustomizableObject)
 	FIntPoint Max;
 
+	/** Priority to be reduced. Only functional in fixed layouts. */
 	UPROPERTY(EditAnywhere, Category = CustomizableObject)
 	uint32 Priority;
 
-	//! Unique unchangeable id used to reference this block from other nodes.
+	/** Unique unchangeable id used to reference this block from other nodes. */
 	UPROPERTY()
 	FGuid Id;
 
+	/** Block will be reduced on both sizes at the same time on each reduction. */
 	UPROPERTY(EditAnywhere, Category = CustomizableObject)
-	bool bUseSymmetry = false;
+	bool bReduceBothAxes = false;
+
+	/** Block will be reduced by two in an Unitary Layout reduction. */
+	UPROPERTY(EditAnywhere, Category = CustomizableObject)
+	bool bReduceByTwo = false;
 };
 
 UCLASS()
diff --git a/Engine/Plugins/Experimental/Mutable/Source/CustomizableObjectEditor/Classes/MuCOE/Nodes/CustomizableObjectNodeMaterial.h b/Engine/Plugins/Experimental/Mutable/Source/CustomizableObjectEditor/Classes/MuCOE/Nodes/CustomizableObjectNodeMaterial.h
--- a/Engine/Plugins/Experimental/Mutable/Source/CustomizableObjectEditor/Classes/MuCOE/Nodes/CustomizableObjectNodeMaterial.h
+++ b/Engine/Plugins/Experimental/Mutable/Source/CustomizableObjectEditor/Classes/MuCOE/Nodes/CustomizableObjectNodeMaterial.h
@@ -348,7 +348,7 @@ public:
 	EUVLayoutMode UVLayoutMode = EUVLayoutMode::FromMaterial;
 	
 	/** Index of the UV channel that will be used with this image.It is necessary to apply the proper layout transformations to it. */
-	UPROPERTY(EditAnywhere, Category = NoCategory, meta = (ClampMin = 0, EditCondition = "UVLayoutMode == EUVLayoutMode::Index"))
+	UPROPERTY(EditAnywhere, Category = NoCategory, meta = (ClampMin = 0, EditCondition = "UVLayoutMode == EUVLayoutMode::Index", EditConditionHides))
 	int32 UVLayout = 0;
 
 	/** Reference Texture used to decide the texture properties of the mutable-generated textures
diff --git a/Engine/Plugins/Experimental/Mutable/Source/CustomizableObjectEditor/Private/MuCOE/CustomizableObjectLayout.cpp b/Engine/Plugins/Experimental/Mutable/Source/CustomizableObjectEditor/Private/MuCOE/CustomizableObjectLayout.cpp
--- a/Engine/Plugins/Experimental/Mutable/Source/CustomizableObjectEditor/Private/MuCOE/CustomizableObjectLayout.cpp
+++ b/Engine/Plugins/Experimental/Mutable/Source/CustomizableObjectEditor/Private/MuCOE/CustomizableObjectLayout.cpp
@@ -19,12 +19,7 @@ UCustomizableObjectLayout::UCustomizableObjectLayout()
 	GridSize = FIntPoint(4, 4);
 	MaxGridSize = FIntPoint(4, 4);
 
-	FCustomizableObjectLayoutBlock Block;
-	Block.Min = FIntPoint(0, 0);
-	Block.Max = FIntPoint(4, 4);
-	Block.Id = FGuid::NewGuid();
-	Block.Priority = 0;
-	Block.bUseSymmetry = false;
+	FCustomizableObjectLayoutBlock Block(FIntPoint(0, 0), FIntPoint(4, 4));
 	Blocks.Add(Block);
 
 	PackingStrategy = ECustomizableObjectTextureLayoutPackingStrategy::Resizable;
@@ -119,12 +114,7 @@ void UCustomizableObjectLayout::GenerateBlocksFromUVs()
 		
 				Layout->GetLayout()->GetBlock(i, &minX, &minY, &sizeX, &sizeY);
 		
-				FCustomizableObjectLayoutBlock block;
-				block.Min = FIntPoint(minX, minY);
-				block.Max = FIntPoint(minX + sizeX, minY + sizeY);
-				block.Id = FGuid::NewGuid();
-				block.Priority = 0;
-				block.bUseSymmetry = false;
+				FCustomizableObjectLayoutBlock block(FIntPoint(minX, minY), FIntPoint(minX + sizeX, minY + sizeY));
 				Blocks.Add(block);
 			}
 		
diff --git a/Engine/Plugins/Experimental/Mutable/Source/CustomizableObjectEditor/Private/MuCOE/GenerateMutableSource/GenerateMutableSourceLayout.cpp b/Engine/Plugins/Experimental/Mutable/Source/CustomizableObjectEditor/Private/MuCOE/GenerateMutableSource/GenerateMutableSourceLayout.cpp
--- a/Engine/Plugins/Experimental/Mutable/Source/CustomizableObjectEditor/Private/MuCOE/GenerateMutableSource/GenerateMutableSourceLayout.cpp
+++ b/Engine/Plugins/Experimental/Mutable/Source/CustomizableObjectEditor/Private/MuCOE/GenerateMutableSource/GenerateMutableSourceLayout.cpp
@@ -81,7 +81,7 @@ mu::NodeLayoutPtr GenerateMutableSourceLayout(const UEdGraphPin * Pin, FMutableG
 					TypedNodeBlocks->Layout->Blocks[BlockIndex].Max.X - TypedNodeBlocks->Layout->Blocks[BlockIndex].Min.X,
 					TypedNodeBlocks->Layout->Blocks[BlockIndex].Max.Y - TypedNodeBlocks->Layout->Blocks[BlockIndex].Min.Y);
 
-				LayoutNode->SetBlockOptions(BlockIndex, TypedNodeBlocks->Layout->Blocks[BlockIndex].Priority, TypedNodeBlocks->Layout->Blocks[BlockIndex].bUseSymmetry);
+				LayoutNode->SetBlockOptions(BlockIndex, TypedNodeBlocks->Layout->Blocks[BlockIndex].Priority, TypedNodeBlocks->Layout->Blocks[BlockIndex].bReduceBothAxes, TypedNodeBlocks->Layout->Blocks[BlockIndex].bReduceByTwo);
 			}
 		}
 		else
@@ -90,7 +90,7 @@ mu::NodeLayoutPtr GenerateMutableSourceLayout(const UEdGraphPin * Pin, FMutableG
 			GenerationContext.Compiler->CompilerLog(FText::FromString(msg), Node, EMessageSeverity::Warning);
 
 			LayoutNode->SetBlock(0, 0, 0, TypedNodeBlocks->Layout->GetGridSize().X, TypedNodeBlocks->Layout->GetGridSize().Y);
-			LayoutNode->SetBlockOptions(0, 0, false);
+			LayoutNode->SetBlockOptions(0, 0, false, false);
 		}
 	}
 	
diff --git a/Engine/Plugins/Experimental/Mutable/Source/CustomizableObjectEditor/Private/MuCOE/GenerateMutableSource/GenerateMutableSourceMesh.cpp b/Engine/Plugins/Experimental/Mutable/Source/CustomizableObjectEditor/Private/MuCOE/GenerateMutableSource/GenerateMutableSourceMesh.cpp
--- a/Engine/Plugins/Experimental/Mutable/Source/CustomizableObjectEditor/Private/MuCOE/GenerateMutableSource/GenerateMutableSourceMesh.cpp
+++ b/Engine/Plugins/Experimental/Mutable/Source/CustomizableObjectEditor/Private/MuCOE/GenerateMutableSource/GenerateMutableSourceMesh.cpp
@@ -2792,7 +2792,7 @@ mu::NodeLayoutBlocksPtr CreateDefaultLayout()
 	LayoutNode->SetBlockReductionMethod(mu::EReductionMethod::HALVE_REDUCTION);
 	LayoutNode->SetBlockCount(1);
 	LayoutNode->SetBlock(0, 0, 0, GridSize, GridSize);
-	LayoutNode->SetBlockOptions(0, 0, false);
+	LayoutNode->SetBlockOptions(0, 0, false, false);
 
 	return LayoutNode;
 }
@@ -3741,7 +3741,7 @@ mu::NodeMeshPtr GenerateMutableSourceMesh(const UEdGraphPin* Pin,
 											Layouts[i]->Blocks[BlockIndex].Max.X - Layouts[i]->Blocks[BlockIndex].Min.X,
 											Layouts[i]->Blocks[BlockIndex].Max.Y - Layouts[i]->Blocks[BlockIndex].Min.Y);
 
-										LayoutNode->SetBlockOptions(BlockIndex, Layouts[i]->Blocks[BlockIndex].Priority, Layouts[i]->Blocks[BlockIndex].bUseSymmetry);
+										LayoutNode->SetBlockOptions(BlockIndex, Layouts[i]->Blocks[BlockIndex].Priority, Layouts[i]->Blocks[BlockIndex].bReduceBothAxes, Layouts[i]->Blocks[BlockIndex].bReduceByTwo);
 									}
 								}
 								else
@@ -3750,7 +3750,7 @@ mu::NodeMeshPtr GenerateMutableSourceMesh(const UEdGraphPin* Pin,
 									GenerationContext.Compiler->CompilerLog(FText::FromString(msg), Node, EMessageSeverity::Warning);
 
 									LayoutNode->SetBlock(0, 0, 0, Layouts[i]->GetGridSize().X, Layouts[i]->GetGridSize().Y);
-									LayoutNode->SetBlockOptions(0, 0, false);
+									LayoutNode->SetBlockOptions(0, 0, false, false);
 								}
 
 								MeshTableNode->SetLayout(i, LayoutNode);
diff --git a/Engine/Plugins/Experimental/Mutable/Source/CustomizableObjectEditor/Private/MuCOE/SCustomizableObjectLayoutGrid.cpp b/Engine/Plugins/Experimental/Mutable/Source/CustomizableObjectEditor/Private/MuCOE/SCustomizableObjectLayoutGrid.cpp
--- a/Engine/Plugins/Experimental/Mutable/Source/CustomizableObjectEditor/Private/MuCOE/SCustomizableObjectLayoutGrid.cpp
+++ b/Engine/Plugins/Experimental/Mutable/Source/CustomizableObjectEditor/Private/MuCOE/SCustomizableObjectLayoutGrid.cpp
@@ -144,6 +144,7 @@ void SCustomizableObjectLayoutGrid::Construct( const FArguments& InArgs )
 	AddBlockAtDelegate = InArgs._OnAddBlockAt;
 	OnSetBlockPriority = InArgs._OnSetBlockPriority;
 	OnSetReduceBlockSymmetrically = InArgs._OnSetReduceBlockSymmetrically;
+	OnSetReduceBlockByTwo = InArgs._OnSetReduceBlockByTwo;
 
 	HasDragged = false;
 	Dragging = false;
@@ -450,13 +451,24 @@ FReply SCustomizableObjectLayoutGrid::OnMouseButtonDown( const FGeometry& MyGeom
 					MenuBuilder.AddWidget(
 						SNew(SBox)
 						.WidthOverride(125.0f)
-						.ToolTipText(LOCTEXT("SetBlockSymmetry_Tooltip", "if true, this block will be reduced in both axes at the same time in a Fixed Layout Strategy."))
+						.ToolTipText(LOCTEXT("SetBlockSymmetry_Tooltip", "If true, this block will be reduced in both axes at the same time in a Fixed Layout Strategy."))
 						[
 							SNew(SCheckBox)
-							.IsChecked(this, &SCustomizableObjectLayoutGrid::GetReductionMethodValue)
+							.IsChecked(this, &SCustomizableObjectLayoutGrid::GetReductionMethodBoolValue, EFRO_Symmetry)
 							.OnCheckStateChanged(this, &SCustomizableObjectLayoutGrid::OnReduceBlockSymmetricallyChanged)
 						]
 					, FText::FromString("Reduce Symmetrically"), true);
+
+					MenuBuilder.AddWidget(
+						SNew(SBox)
+						.WidthOverride(125.0f)
+						.ToolTipText(LOCTEXT("SetBlockReduceByTwo_Tooltip", "Only for Unitary reduction. If true, this option reduces each time the block by two block units."))
+						[
+							SNew(SCheckBox)
+							.IsChecked(this, &SCustomizableObjectLayoutGrid::GetReductionMethodBoolValue, EFRO_RedyceByTwo)
+							.OnCheckStateChanged(this, &SCustomizableObjectLayoutGrid::OnReduceBlockByTwoChanged)
+						]
+					, FText::FromString("Reduce by Two"), true);
 				}
 			}
 			MenuBuilder.EndSection();
@@ -1082,12 +1094,13 @@ TOptional<int32> SCustomizableObjectLayoutGrid::GetBlockPriortyValue() const
 }
 
 
-ECheckBoxState SCustomizableObjectLayoutGrid::GetReductionMethodValue() const
+ECheckBoxState SCustomizableObjectLayoutGrid::GetReductionMethodBoolValue(EFixedReductionOptions Option) const
 {
 	if (SelectedBlocks.Num())
 	{
 		TArray<FCustomizableObjectLayoutBlock> CurrentSelectedBlocks;
 
+		// Getting all selected blocks
 		for (const FCustomizableObjectLayoutBlock& Block : Blocks.Get())
 		{
 			if (SelectedBlocks.Contains(Block.Id))
@@ -1096,17 +1109,41 @@ ECheckBoxState SCustomizableObjectLayoutGrid::GetReductionMethodValue() const
 			}
 		}
 
-		bool bUsesSymmetry = CurrentSelectedBlocks[0].bUseSymmetry;
-
-		for (const FCustomizableObjectLayoutBlock& Block : CurrentSelectedBlocks)
+		switch (Option)
+		{
+		case EFRO_Symmetry:
 		{
-			if (Block.bUseSymmetry != bUsesSymmetry)
+			const bool bReduceBothAxes = CurrentSelectedBlocks[0].bReduceBothAxes;
+
+			// If one or more blocks have a different value than the rest of selected blocks return Undetermined
+			for (const FCustomizableObjectLayoutBlock& Block : CurrentSelectedBlocks)
 			{
-				return ECheckBoxState::Undetermined;
+				if (Block.bReduceBothAxes != bReduceBothAxes)
+				{
+					return ECheckBoxState::Undetermined;
+				}
 			}
+
+			return bReduceBothAxes ? ECheckBoxState::Checked : ECheckBoxState::Unchecked;
 		}
+		case EFRO_RedyceByTwo:
+		{
+			const bool bReduceByTwo = CurrentSelectedBlocks[0].bReduceByTwo;
 
-		return bUsesSymmetry ? ECheckBoxState::Checked : ECheckBoxState::Unchecked;
+			// If one or more blocks have a different value than the rest of selected blocks return Undetermined
+			for (const FCustomizableObjectLayoutBlock& Block : CurrentSelectedBlocks)
+			{
+				if (Block.bReduceByTwo != bReduceByTwo)
+				{
+					return ECheckBoxState::Undetermined;
+				}
+			}
+
+			return bReduceByTwo ? ECheckBoxState::Checked : ECheckBoxState::Unchecked;
+		}
+		default:
+			break;
+		}
 	}
 
 	return ECheckBoxState::Undetermined;
@@ -1131,6 +1168,15 @@ void SCustomizableObjectLayoutGrid::OnReduceBlockSymmetricallyChanged(ECheckBoxS
 }
 
 
+void SCustomizableObjectLayoutGrid::OnReduceBlockByTwoChanged(ECheckBoxState InCheckboxState)
+{
+	if (SelectedBlocks.Num())
+	{
+		OnSetReduceBlockByTwo.ExecuteIfBound(InCheckboxState == ECheckBoxState::Checked);
+	}
+}
+
+
 // Canvas Drawer --------------------------------------------------------------
 
 void FUVCanvasDrawer::Initialize(const FIntRect& InCanvasRect, const FIntRect& InClippingRect, const FVector2D& InOrigin, const FVector2D& InSize, const FIntPoint& InGridSize, const float InCellSize)
diff --git a/Engine/Plugins/Experimental/Mutable/Source/CustomizableObjectEditor/Private/MuCOE/SCustomizableObjectLayoutGrid.h b/Engine/Plugins/Experimental/Mutable/Source/CustomizableObjectEditor/Private/MuCOE/SCustomizableObjectLayoutGrid.h
--- a/Engine/Plugins/Experimental/Mutable/Source/CustomizableObjectEditor/Private/MuCOE/SCustomizableObjectLayoutGrid.h
+++ b/Engine/Plugins/Experimental/Mutable/Source/CustomizableObjectEditor/Private/MuCOE/SCustomizableObjectLayoutGrid.h
@@ -36,6 +36,13 @@ struct FBlockWidgetData
 	FRect2D HandleRect;
 };
 
+
+enum EFixedReductionOptions
+{
+	EFRO_Symmetry,
+	EFRO_RedyceByTwo
+};
+
 /** */
 class SCustomizableObjectLayoutGrid : public SCompoundWidget
 {
@@ -47,6 +54,7 @@ public:
 	DECLARE_DELEGATE_TwoParams(FAddBlockAtDelegate, FIntPoint, FIntPoint);
 	DECLARE_DELEGATE_OneParam(FSetBlockPriority, int32);
 	DECLARE_DELEGATE_OneParam(FSetReduceBlockSymmetrically, bool);
+	DECLARE_DELEGATE_OneParam(FSetReduceBlockByTwo, bool);
 
 	SLATE_BEGIN_ARGS( SCustomizableObjectLayoutGrid ){}
 
@@ -62,6 +70,7 @@ public:
 		SLATE_EVENT(FAddBlockAtDelegate, OnAddBlockAt)
 		SLATE_EVENT(FSetBlockPriority, OnSetBlockPriority)
 		SLATE_EVENT(FSetReduceBlockSymmetrically, OnSetReduceBlockSymmetrically)
+		SLATE_EVENT(FSetReduceBlockByTwo, OnSetReduceBlockByTwo)
 
 	SLATE_END_ARGS()
 
@@ -117,11 +126,14 @@ public:
 	/** Callback when the priority of a block changes */
 	void OnBlockPriorityChanged(int32 InValue);
 
-	/** Callback when a block reduction method changes */
+	/** Callback when symmetry block reduction option changes */
 	void OnReduceBlockSymmetricallyChanged(ECheckBoxState InCheckboxState);
 
+	/** Callback when ReduceByTwo block reduction option changes */
+	void OnReduceBlockByTwoChanged(ECheckBoxState InCheckboxState);
+
 	/** Gets block reduction value of the selected blocks */
-	ECheckBoxState GetReductionMethodValue() const;
+	ECheckBoxState GetReductionMethodBoolValue(EFixedReductionOptions Option) const;
 
 private:
 
@@ -136,6 +148,7 @@ private:
 	FAddBlockAtDelegate AddBlockAtDelegate;
 	FSetBlockPriority OnSetBlockPriority;
 	FSetReduceBlockSymmetrically OnSetReduceBlockSymmetrically;
+	FSetReduceBlockByTwo OnSetReduceBlockByTwo;
 
 	/** Size of the grid in blocks */
 	TAttribute<FIntPoint> GridSize;
diff --git a/Engine/Plugins/Experimental/Mutable/Source/CustomizableObjectEditor/Private/MuCOE/SCustomizableObjectNodeLayoutBlocksEditor.cpp b/Engine/Plugins/Experimental/Mutable/Source/CustomizableObjectEditor/Private/MuCOE/SCustomizableObjectNodeLayoutBlocksEditor.cpp
--- a/Engine/Plugins/Experimental/Mutable/Source/CustomizableObjectEditor/Private/MuCOE/SCustomizableObjectNodeLayoutBlocksEditor.cpp
+++ b/Engine/Plugins/Experimental/Mutable/Source/CustomizableObjectEditor/Private/MuCOE/SCustomizableObjectNodeLayoutBlocksEditor.cpp
@@ -121,6 +121,7 @@ void SCustomizableObjectNodeLayoutBlocksEditor::SetCurrentLayout(UCustomizableOb
 			.OnAddBlockAt(this, &SCustomizableObjectNodeLayoutBlocksEditor::OnAddBlockAt)
 			.OnSetBlockPriority(this, &SCustomizableObjectNodeLayoutBlocksEditor::OnSetBlockPriority)
 			.OnSetReduceBlockSymmetrically(this, &SCustomizableObjectNodeLayoutBlocksEditor::OnSetBlockReductionSymmetry)
+			.OnSetReduceBlockByTwo(this, &SCustomizableObjectNodeLayoutBlocksEditor::OnSetBlockReductionByTwo)
 		]
 	];	
 }
@@ -412,11 +413,6 @@ void SCustomizableObjectNodeLayoutBlocksEditor::OnAddBlock()
 	if (CurrentLayout)
 	{
 		FCustomizableObjectLayoutBlock block;
-		block.Min = FIntPoint( 0, 0 );
-		block.Max = FIntPoint(1, 1);
-		block.Id = FGuid::NewGuid();
-		block.Priority = 0;
-		block.bUseSymmetry = false;
 		CurrentLayout->Blocks.Add( block );
 		CurrentLayout->MarkPackageDirty();
 
@@ -432,13 +428,7 @@ void SCustomizableObjectNodeLayoutBlocksEditor::OnAddBlockAt(const FIntPoint Min
 {
 	if (CurrentLayout)
 	{
-		FCustomizableObjectLayoutBlock block;
-		block.Min = Min;
-		block.Max = Max;
-		block.Id = FGuid::NewGuid();
-		block.Priority = 0;
-		block.bUseSymmetry = false;
-
+		FCustomizableObjectLayoutBlock block(Min,Max);
 		CurrentLayout->Blocks.Add(block);
 		CurrentLayout->MarkPackageDirty();
 	}
@@ -546,7 +536,28 @@ void SCustomizableObjectNodeLayoutBlocksEditor::OnSetBlockReductionSymmetry(bool
 			{
 				if (SelectedBlocks.Contains(CurrentLayout->Blocks[i].Id))
 				{
-					CurrentLayout->Blocks[i].bUseSymmetry = bInValue;
+					CurrentLayout->Blocks[i].bReduceBothAxes = bInValue;
+					CurrentLayout->MarkPackageDirty();
+				}
+			}
+		}
+	}
+}
+
+
+void SCustomizableObjectNodeLayoutBlocksEditor::OnSetBlockReductionByTwo(bool bInValue)
+{
+	if (CurrentLayout)
+	{
+		if (LayoutGridWidget.IsValid())
+		{
+			TArray<FGuid> SelectedBlocks = LayoutGridWidget->GetSelectedBlocks();
+
+			for (int i = 0; i < CurrentLayout->Blocks.Num(); ++i)
+			{
+				if (SelectedBlocks.Contains(CurrentLayout->Blocks[i].Id))
+				{
+					CurrentLayout->Blocks[i].bReduceByTwo = bInValue;
 					CurrentLayout->MarkPackageDirty();
 				}
 			}
diff --git a/Engine/Plugins/Experimental/Mutable/Source/CustomizableObjectEditor/Private/MuCOE/SCustomizableObjectNodeLayoutBlocksEditor.h b/Engine/Plugins/Experimental/Mutable/Source/CustomizableObjectEditor/Private/MuCOE/SCustomizableObjectNodeLayoutBlocksEditor.h
--- a/Engine/Plugins/Experimental/Mutable/Source/CustomizableObjectEditor/Private/MuCOE/SCustomizableObjectNodeLayoutBlocksEditor.h
+++ b/Engine/Plugins/Experimental/Mutable/Source/CustomizableObjectEditor/Private/MuCOE/SCustomizableObjectNodeLayoutBlocksEditor.h
@@ -119,9 +119,12 @@ private:
 	/** Sets the block priority from the input text. */
 	void OnSetBlockPriority(int32 InValue);
 
-	/** Sets the block reduction symmetry method. */
+	/** Sets the block reduction symmetry option. */
 	void OnSetBlockReductionSymmetry(bool bInValue);
 
+	/** Sets the block reduction ReduceByTwo option. */
+	void OnSetBlockReductionByTwo(bool bInValue);
+
 	/** Called when the packing strategy has changed. */
 	void OnLayoutPackingStrategyChanged(TSharedPtr<FString> NewSelection, ESelectInfo::Type SelectInfo);
 
diff --git a/Engine/Plugins/Experimental/Mutable/Source/MutableRuntime/Private/MuR/CodeRunner.cpp b/Engine/Plugins/Experimental/Mutable/Source/MutableRuntime/Private/MuR/CodeRunner.cpp
--- a/Engine/Plugins/Experimental/Mutable/Source/MutableRuntime/Private/MuR/CodeRunner.cpp
+++ b/Engine/Plugins/Experimental/Mutable/Source/MutableRuntime/Private/MuR/CodeRunner.cpp
@@ -5720,7 +5720,8 @@ namespace mu
 					scratch.positions.SetNum(BlockCount);
 					scratch.priorities.SetNum(BlockCount);
 					scratch.reductions.SetNum(BlockCount);
-					scratch.useSymmetry.SetNum(BlockCount);
+					scratch.ReduceBothAxes.SetNum(BlockCount);
+					scratch.ReduceByTwo.SetNum(BlockCount);
 
 					LayoutPack3(pResult.get(), Source.get(), &scratch);
 				}
diff --git a/Engine/Plugins/Experimental/Mutable/Source/MutableRuntime/Private/MuR/Layout.cpp b/Engine/Plugins/Experimental/Mutable/Source/MutableRuntime/Private/MuR/Layout.cpp
--- a/Engine/Plugins/Experimental/Mutable/Source/MutableRuntime/Private/MuR/Layout.cpp
+++ b/Engine/Plugins/Experimental/Mutable/Source/MutableRuntime/Private/MuR/Layout.cpp
@@ -140,21 +140,13 @@ namespace mu {
 
 
 	//---------------------------------------------------------------------------------------------
-	void Layout::GetBlockOptions(int index, int* pPriority, bool* bUseSymmetry) const
+	void Layout::GetBlockOptions(int index, int& pPriority, bool& bReduceBothAxes, bool& bReduceByTwo) const
 	{
 		check(index >= 0 && index < m_blocks.Num());
-		check(pPriority);
-		check(bUseSymmetry);
 
-		if (pPriority)
-		{
-			*pPriority = m_blocks[index].m_priority;
-		}
-
-		if (bUseSymmetry)
-		{
-			*bUseSymmetry = m_blocks[index].bUseSymmetry;
-		}
+		pPriority = m_blocks[index].m_priority;
+		bReduceBothAxes = m_blocks[index].bReduceBothAxes;
+		bReduceByTwo = m_blocks[index].bReduceByTwo;
 	}
 
 
@@ -170,12 +162,13 @@ namespace mu {
 
 
 	//---------------------------------------------------------------------------------------------
-	void Layout::SetBlockOptions(int index, int priority, bool bUseSymmetry)
+	void Layout::SetBlockOptions(int index, int priority, bool bReduceBothAxes, bool bReduceByTwo)
 	{
 		check(index >= 0 && index < m_blocks.Num());
 
 		m_blocks[index].m_priority = priority;
-		m_blocks[index].bUseSymmetry = bUseSymmetry;
+		m_blocks[index].bReduceBothAxes = bReduceBothAxes;
+		m_blocks[index].bReduceByTwo = bReduceByTwo;
 	}
 
 
@@ -196,7 +189,7 @@ namespace mu {
 	//---------------------------------------------------------------------------------------------
 	void Layout::Serialise(OutputArchive& arch) const
 	{
-		uint32 ver = 5;
+		uint32 ver = 6;
 		arch << ver;
 
 		arch << m_size;
@@ -214,11 +207,11 @@ namespace mu {
 	{
 		uint32 ver;
 		arch >> ver;
-		check(ver <= 5);
+		check(ver <= 6);
 
 		arch >> m_size;
 
-		if (ver < 5)
+		if (ver < 6)
 		{
 			uint32_t Size = 0;
 			arch >> Size;
@@ -226,6 +219,7 @@ namespace mu {
 
 			for (uint32_t BlockIndex = 0; BlockIndex < Size; ++BlockIndex)
 			{
+				// Unserialise taking into account the Layout version
 				m_blocks[BlockIndex].UnserialiseOldVersion(arch, ver);
 			}
 		}
@@ -335,7 +329,8 @@ namespace mu {
 		arch << m_size;
 		arch << m_id;
 		arch << m_priority;
-		arch << bUseSymmetry;
+		arch << bReduceBothAxes;
+		arch << bReduceByTwo;
 	}
 
 
@@ -346,18 +341,23 @@ namespace mu {
 		arch >> m_size;
 		arch >> m_id;
 		arch >> m_priority;
-		arch >> bUseSymmetry;
+		arch >> bReduceBothAxes;
+		arch >> bReduceByTwo;
 	}
 
 	
 	//---------------------------------------------------------------------------------------------
 	void Layout::FBlock::UnserialiseOldVersion(InputArchive& Archive, const int32 Version)
 	{
-		// Use the version if the Layout version changes to 6
 		Archive >> m_min;
 		Archive >> m_size;
 		Archive >> m_id;
 		Archive >> m_priority;
+
+		if (Version >= 5)
+		{
+			Archive >> bReduceBothAxes;
+		}
 	}
 }
 
diff --git a/Engine/Plugins/Experimental/Mutable/Source/MutableRuntime/Private/MuR/OpLayoutPack.h b/Engine/Plugins/Experimental/Mutable/Source/MutableRuntime/Private/MuR/OpLayoutPack.h
--- a/Engine/Plugins/Experimental/Mutable/Source/MutableRuntime/Private/MuR/OpLayoutPack.h
+++ b/Engine/Plugins/Experimental/Mutable/Source/MutableRuntime/Private/MuR/OpLayoutPack.h
@@ -140,18 +140,20 @@ namespace mu
 			index = -1;
 		}
 
-		LAY_BLOCK( int32 i, UE::Math::TIntVector2<uint16> s, int32 p=0, bool sym = false )
+		LAY_BLOCK( int32 i, UE::Math::TIntVector2<uint16> s, int32 p=0, bool bInReduceBothAxes = false, bool bInReduceByTwo = false )
 		{
 			index = i;
 			size = s;
 			priority = p;
-			useSymmetry = sym;
+			bReduceBothAxes = bInReduceBothAxes;
+			bReduceByTwo = bInReduceByTwo;
 		}
 
 		int32 index;
 		UE::Math::TIntVector2<uint16> size = UE::Math::TIntVector2<uint16>(0,0);
 		int32 priority;
-		bool useSymmetry;
+		bool bReduceBothAxes;
+		bool bReduceByTwo;
 	};
 
     inline bool CompareBlocks( const LAY_BLOCK& a, const LAY_BLOCK& b )
@@ -188,6 +190,7 @@ namespace mu
 	{
 		if (a.priority == b.priority)
 		{
+			// TODO(Max): Check if this comparison modifies the block order while reducing them.
 			return CompareBlocks(a, b);
 		}
 
@@ -205,7 +208,8 @@ namespace mu
 		TArray< vec2<int> > positions;
 		TArray< int > priorities;
 		TArray< vec2<int> > reductions;
-		TArray< int > useSymmetry;
+		TArray< int > ReduceBothAxes;
+		TArray< int > ReduceByTwo;
     };
 
 
@@ -358,14 +362,16 @@ namespace mu
     // }
 
 
-	inline void ReductionOperation(uint16& BlockSize, EReductionMethod ReductionMethod)
+	inline void ReductionOperation(uint16& BlockSize, EReductionMethod ReductionMethod, int32 bReduceByTwo)
 	{
 		if (ReductionMethod == EReductionMethod::UNITARY_REDUCTION)
 		{
-			BlockSize -= 1;
+			int32 Reduction = (bReduceByTwo && BlockSize > 2) ? 2 : 1;
+			BlockSize -= Reduction;
 			return;
 		}
 		
+		// Reduce size by half
 		BlockSize /= 2;
 	}
 
@@ -380,15 +386,15 @@ namespace mu
 
 		if (scratch->sorted[r_it].size[0] != 1 || scratch->sorted[r_it].size[1] != 1)
 		{
-			if (scratch->sorted[r_it].useSymmetry)
+			if (scratch->sorted[r_it].bReduceBothAxes)
 			{
 				// We reduce both sides of the block at the same time
 				for (int32 Index = 0; Index <= 1; ++Index)
 				{
 					if (scratch->sorted[r_it].size[Index] > 1)
 					{
-						ReductionOperation(scratch->sorted[r_it].size[Index], ReductionMethod);
-						ReductionOperation(scratch->blocks[scratch->sorted[r_it].index][Index], ReductionMethod);
+						ReductionOperation(scratch->sorted[r_it].size[Index], ReductionMethod, scratch->sorted[r_it].bReduceByTwo);
+						ReductionOperation(scratch->blocks[scratch->sorted[r_it].index][Index], ReductionMethod, scratch->sorted[r_it].bReduceByTwo );
 
 						pass = true;
 					}
@@ -400,8 +406,8 @@ namespace mu
 				{
 					if (scratch->sorted[r_it].size[1] > 1)
 					{
-						ReductionOperation(scratch->sorted[r_it].size[1], ReductionMethod);
-						ReductionOperation(scratch->blocks[scratch->sorted[r_it].index][1], ReductionMethod);
+						ReductionOperation(scratch->sorted[r_it].size[1], ReductionMethod, scratch->sorted[r_it].bReduceByTwo);
+						ReductionOperation(scratch->blocks[scratch->sorted[r_it].index][1], ReductionMethod, scratch->sorted[r_it].bReduceByTwo);
 						pass = true;
 					}
 
@@ -411,8 +417,8 @@ namespace mu
 				{
 					if (scratch->sorted[r_it].size[0] > 1)
 					{
-						ReductionOperation(scratch->sorted[r_it].size[0], ReductionMethod);
-						ReductionOperation(scratch->blocks[scratch->sorted[r_it].index][0], ReductionMethod);
+						ReductionOperation(scratch->sorted[r_it].size[0], ReductionMethod, scratch->sorted[r_it].bReduceByTwo);
+						ReductionOperation(scratch->blocks[scratch->sorted[r_it].index][0], ReductionMethod, scratch->sorted[r_it].bReduceByTwo);
 						pass = true;
 					}
 
@@ -431,8 +437,8 @@ namespace mu
 
 					if (scratch->sorted[r_it].size[Index] > 1)
 					{
-						ReductionOperation(scratch->sorted[r_it].size[Index], ReductionMethod);
-						ReductionOperation(scratch->blocks[scratch->sorted[r_it].index][Index], ReductionMethod);
+						ReductionOperation(scratch->sorted[r_it].size[Index], ReductionMethod, scratch->sorted[r_it].bReduceByTwo);
+						ReductionOperation(scratch->blocks[scratch->sorted[r_it].index][Index], ReductionMethod, scratch->sorted[r_it].bReduceByTwo);
 						pass = true;
 					}
 
@@ -649,18 +655,22 @@ namespace mu
 
 		bool usePriority = false;
 
+		EPackStrategy LayoutStrategy = pSourceLayout->GetLayoutPackingStrategy();
+		EReductionMethod ReductionMethod = pSourceLayout->GetBlockReductionMethod();
+
         // Look for the maximum block sizes on the layout and the total area
 		uint16 maxX = 0;
 		uint16 maxY = 0;
         int area = 0;
+
         for ( int index=0; index<blockCount; ++index )
         {
             box< UE::Math::TIntVector2<uint16> > b;
             pSourceLayout->GetBlock( index, &b.min[0], &b.min[1], &b.size[0], &b.size[1] );
 
 			int p;
-			bool Symmetry;
-			pSourceLayout->GetBlockOptions(index, &p, &Symmetry);
+			bool bReduceBothAxes, bReduceByTwo;
+			pSourceLayout->GetBlockOptions(index, p, bReduceBothAxes, bReduceByTwo);
 
 			vec2<int> reductions;
 			reductions[0] = 0;
@@ -679,13 +689,14 @@ namespace mu
             scratch->blocks[index] = b.size;
 			scratch->priorities[index] = p;
 			scratch->reductions[index] = reductions;
-			scratch->useSymmetry[index] = (int)Symmetry;
+			scratch->ReduceBothAxes[index] = (int)bReduceBothAxes;
+			scratch->ReduceByTwo[index] = (int)bReduceByTwo;
         }
 
 
         // Grow until the area is big enough to fit all blocks. We always grow X first, because
         // in case we cannot pack everything, we will grow Y with the current horizon algorithm.
-		if (pSourceLayout->GetLayoutPackingStrategy() == EPackStrategy::RESIZABLE_LAYOUT)
+		if (LayoutStrategy == EPackStrategy::RESIZABLE_LAYOUT)
 		{
 			maxX = FGenericPlatformMath::RoundUpToPowerOfTwo(maxX);
 			maxY = FGenericPlatformMath::RoundUpToPowerOfTwo(maxY);
@@ -727,7 +738,7 @@ namespace mu
 				}
 			}
 
-			//The highest block is bigger than the layout max size
+			// Reducing blocks that do not fit in the layout grid
 			while (maxX > layoutSizeX)
 			{
 				area = 0;
@@ -736,16 +747,27 @@ namespace mu
 				{
 					if (scratch->blocks[index][0] == maxX)
 					{
-						scratch->blocks[index][0] /= 2;
+						ReductionOperation(scratch->blocks[index][0], ReductionMethod, scratch->ReduceByTwo[index]);
 						scratch->reductions[index][0]++;
+
+						if (scratch->ReduceBothAxes[index] == 1)
+						{
+							ReductionOperation(scratch->blocks[index][1], ReductionMethod, scratch->ReduceByTwo[index]);
+							scratch->reductions[index][1]++;
+						}
 					}
 				}
-				maxX = 0;
+
+				maxX = maxY = 0;
 				
 				//recalculating area and maximum block sizes
 				for (int index = 0; index < blockCount; ++index)
 				{
 					maxX = FMath::Max(maxX, scratch->blocks[index][0]);
+
+					// maxY could have been modified if a block had symmetry enabled in the previous reduction
+					maxY = FMath::Max(maxY, scratch->blocks[index][1]);
+
 					area += scratch->blocks[index][0] * scratch->blocks[index][1];
 				}
 			}
@@ -758,16 +780,27 @@ namespace mu
 				{
 					if (scratch->blocks[index][1] == maxY)
 					{
-						scratch->blocks[index][1] /= 2;
+						ReductionOperation(scratch->blocks[index][1], ReductionMethod, scratch->ReduceByTwo[index]);
 						scratch->reductions[index][1]++;
+
+						if (scratch->ReduceBothAxes[index] == 1)
+						{
+							ReductionOperation(scratch->blocks[index][0], ReductionMethod, scratch->ReduceByTwo[index]);
+							scratch->reductions[index][0]++;
+						}
 					}
 				}
-				maxY = 0;
+
+				maxX = maxY = 0;
 
 				//recalculating area and maximum block sizes
 				for (int index = 0; index < blockCount; ++index)
 				{
 					maxY = FMath::Max(maxY, scratch->blocks[index][1]);
+
+					// maxX could have been modified if a block had symmetry enabled in the previous reduction
+					maxX = FMath::Max(maxX, scratch->blocks[index][0]);
+
 					area += scratch->blocks[index][0] * scratch->blocks[index][1];
 				}
 			}
@@ -794,28 +827,29 @@ namespace mu
 		// This is used to iterate through blocks.
 		int32 BlockIterator = 0;
 
-        // Sort by height, area
+		// Making a copy of the blocks array to sort them
         check( (int)scratch->sorted.Num()==blockCount );
         for ( int index=0; index<blockCount; ++index )
         {
-            scratch->sorted[index] = LAY_BLOCK( index, scratch->blocks[index], scratch->priorities[index], (bool)scratch->useSymmetry[index]);
+            scratch->sorted[index] = LAY_BLOCK( index, scratch->blocks[index], scratch->priorities[index], (bool)scratch->ReduceBothAxes[index], (bool)scratch->ReduceByTwo[index]);
         }
 
-		scratch->sorted.Sort( CompareBlocks);
-		
-		if(pSourceLayout->GetLayoutPackingStrategy() == EPackStrategy::FIXED_LAYOUT)
+		// Sort blocks by height, area
+		scratch->sorted.Sort(CompareBlocks);
+
+		if(LayoutStrategy == EPackStrategy::FIXED_LAYOUT)
 		{
-			if (usePriority)	//Sort by priority
+			if (usePriority)
 			{
+				//Sort blocks by priority
 				scratch->sorted.Sort(CompareBlocksPriority);
 			}
 
 			// Shrink blocks in case we do not have enough space to pack everything
 			while (maxX*maxY < area)
 			{
-				ReduceBlock(blockCount, area, BlockIterator, scratch, pSourceLayout->GetBlockReductionMethod());
+				ReduceBlock(blockCount, area, BlockIterator, scratch, ReductionMethod);
 			}
-			
 		}
 
 		bool fits = false;
@@ -829,9 +863,9 @@ namespace mu
 			}
 
 			// Try to pack everything
-			fits = SetPositions(bestY, layoutSizeY, &maxX, &maxY, scratch, pSourceLayout->GetLayoutPackingStrategy());
+			fits = SetPositions(bestY, layoutSizeY, &maxX, &maxY, scratch, LayoutStrategy);
 
-			if (!fits && pSourceLayout->GetLayoutPackingStrategy() == EPackStrategy::FIXED_LAYOUT)
+			if (!fits && LayoutStrategy == EPackStrategy::FIXED_LAYOUT)
 			{
 				// Sort by priority before shrink
 				if (usePriority)
@@ -839,15 +873,15 @@ namespace mu
 					scratch->sorted.Sort(CompareBlocksPriority);
 				}
 
-				ReduceBlock(blockCount, area, BlockIterator, scratch, pSourceLayout->GetBlockReductionMethod());
+				ReduceBlock(blockCount, area, BlockIterator, scratch, ReductionMethod);
 			}
 		}
 
         // Set data in the result
         pResult->SetGridSize( maxX, maxY );
         pResult->SetMaxGridSize(layoutSizeX, layoutSizeY);
-		pResult->SetLayoutPackingStrategy(pSourceLayout->GetLayoutPackingStrategy());
-		pResult->SetBlockReductionMethod(pSourceLayout->GetBlockReductionMethod());
+		pResult->SetLayoutPackingStrategy(LayoutStrategy);
+		pResult->SetBlockReductionMethod(ReductionMethod);
 
         for ( int index=0; index<blockCount; ++index )
         {
@@ -858,7 +892,7 @@ namespace mu
                     scratch->blocks[index][0], scratch->blocks[index][1]
                 );
 
-			pResult->SetBlockOptions(index, scratch->priorities[index], (bool)scratch->useSymmetry[index]);
+			pResult->SetBlockOptions(index, scratch->priorities[index], (bool)scratch->ReduceBothAxes[index], (bool)scratch->ReduceByTwo[index]);
         }
     }
 }
diff --git a/Engine/Plugins/Experimental/Mutable/Source/MutableRuntime/Public/MuR/Layout.h b/Engine/Plugins/Experimental/Mutable/Source/MutableRuntime/Public/MuR/Layout.h
--- a/Engine/Plugins/Experimental/Mutable/Source/MutableRuntime/Public/MuR/Layout.h
+++ b/Engine/Plugins/Experimental/Mutable/Source/MutableRuntime/Public/MuR/Layout.h
@@ -107,8 +107,9 @@ namespace mu
 		//! Returns the reduction priority of a block.
 		//! \param index Block to get priority.
 		//! \param[out] priority reduction priority of the block
-		//! \param[out] bUSeSymmetry reduction method of the block
-		void GetBlockOptions(int index, int* priority, bool* bUseSymmetry) const;
+		//! \param[out] bReduceBothAxes reduction method of the block
+		//! \param[out] bReduceByTwo reduction method of the block
+		void GetBlockOptions(int index, int& priority, bool& bReduceBothAxes, bool& bReduceByTwo) const;
 
 		//! Set a block of the layout.
 		//! "Position" here means the lower-left corner of the block.
@@ -124,8 +125,9 @@ namespace mu
 		//! Set the reduction options of a block
 		//! \param index Block to set the options
 		//! \param priority will be set to the reduction priority of the block. The blocks with the highest values will be the last to be reduced
-		//! \param bUseSymmetry will be set to reduce the block in both axis at the same time.
-		void SetBlockOptions(int index, int priority, bool bUseSymmetry);
+		//! \param bReduceBothAxes will be set to reduce the block in both axis at the same time.
+		//! \param bReduceByTwo will reduce by two blocks on a unitary reduction.
+		void SetBlockOptions(int index, int priority, bool bReduceBothAxes, bool bReduceByTwo);
 
 		//! Set the texture layout packing strategy
 		//! By default the texture layout packing strategy is set to resizable layout
@@ -163,7 +165,8 @@ namespace mu
 				m_size = size;
 				m_id = -1;
 				m_priority = 0;
-				bUseSymmetry = false;
+				bReduceBothAxes = false;
+				bReduceByTwo = false;
 			}
 
 			UE::Math::TIntVector2<uint16> m_min = UE::Math::TIntVector2<uint16>(0,0);
@@ -176,7 +179,10 @@ namespace mu
 			int32 m_priority;
 
 			//! Value to control the method to reduce the block
-			bool bUseSymmetry;
+			bool bReduceBothAxes;
+
+			//! Value to control if a block has to be reduced by two in an unitary reduction strategy
+			bool bReduceByTwo;
 
 
 			//!
@@ -192,7 +198,8 @@ namespace mu
 					(m_size == o.m_size) &&
 					(m_id == o.m_id) &&
 					(m_priority == o.m_priority) &&
-					(bUseSymmetry == o.bUseSymmetry);
+					(bReduceBothAxes == o.bReduceBothAxes) &&
+					(bReduceByTwo == o.bReduceByTwo);
 			}
 
 			inline bool IsSimilar(const FBlock& o) const
@@ -201,7 +208,8 @@ namespace mu
 				return (m_min == o.m_min) &&
 					(m_size == o.m_size) &&
 					(m_priority == o.m_priority) &&
-					(bUseSymmetry == o.bUseSymmetry);
+					(bReduceBothAxes == o.bReduceBothAxes) &&
+					(bReduceByTwo == o.bReduceByTwo);
 			}
 
 			//!
diff --git a/Engine/Plugins/Experimental/Mutable/Source/MutableTools/Private/MuT/NodeLayout.cpp b/Engine/Plugins/Experimental/Mutable/Source/MutableTools/Private/MuT/NodeLayout.cpp
--- a/Engine/Plugins/Experimental/Mutable/Source/MutableTools/Private/MuT/NodeLayout.cpp
+++ b/Engine/Plugins/Experimental/Mutable/Source/MutableTools/Private/MuT/NodeLayout.cpp
@@ -183,9 +183,9 @@ namespace mu
 
 
 	//---------------------------------------------------------------------------------------------
-	void NodeLayoutBlocks::SetBlockOptions(int index, int priority, bool bUseSymmetry)
+	void NodeLayoutBlocks::SetBlockOptions(int index, int priority, bool bReduceBothAxes, bool bReduceByTwo)
 	{
-		m_pD->m_pLayout->SetBlockOptions(index, priority, bUseSymmetry);
+		m_pD->m_pLayout->SetBlockOptions(index, priority, bReduceBothAxes, bReduceByTwo);
 	}
 
 
diff --git a/Engine/Plugins/Experimental/Mutable/Source/MutableTools/Public/MuT/NodeLayout.h b/Engine/Plugins/Experimental/Mutable/Source/MutableTools/Public/MuT/NodeLayout.h
--- a/Engine/Plugins/Experimental/Mutable/Source/MutableTools/Public/MuT/NodeLayout.h
+++ b/Engine/Plugins/Experimental/Mutable/Source/MutableTools/Public/MuT/NodeLayout.h
@@ -140,7 +140,7 @@ namespace mu
         void SetBlock( int index, int minx, int miny, int sizex, int sizey );
 
 		//! Set reduction block options like priority or if the block has to be reduced symmetrically.
-		void SetBlockOptions(int index, int priority, bool bUseSymmetry);
+		void SetBlockOptions(int index, int priority, bool bReduceBothAxes, bool bReduceByTwo);
 
 		//! Set the texture layout packing strategy 
 		void SetLayoutPackingStrategy(EPackStrategy strategy);
