[Mutable] Fixed the UV visualizer of the Customizable Object instance.

#rnx

diff --git a/Engine/Plugins/Experimental/Mutable/Source/CustomizableObjectEditor/Private/MuCOE/CustomizableObjectInstanceEditor.cpp b/Engine/Plugins/Experimental/Mutable/Source/CustomizableObjectEditor/Private/MuCOE/CustomizableObjectInstanceEditor.cpp
--- a/Engine/Plugins/Experimental/Mutable/Source/CustomizableObjectEditor/Private/MuCOE/CustomizableObjectInstanceEditor.cpp
+++ b/Engine/Plugins/Experimental/Mutable/Source/CustomizableObjectEditor/Private/MuCOE/CustomizableObjectInstanceEditor.cpp
@@ -718,6 +718,8 @@ void FCustomizableObjectInstanceEditor::OnInstancePropertySelectionChanged(FProp
 
 void FCustomizableObjectInstanceEditor::OnUpdatePreviewInstance()
 {
+	check(CustomizableObjectInstance)
+
 	Viewport->GetViewportClient()->Invalidate();
 
 	UpdatePreviewVisibility();
@@ -740,11 +742,9 @@ void FCustomizableObjectInstanceEditor::OnUpdatePreviewInstance()
 		Viewport->SetPreviewComponents(ObjectPtrDecay(PreviewSkeletalMeshComponents));
 	}
 
-	if (!CustomizableObjectInstance)
-	{
-		return;
-	}
-	
+	// Sets which UV material to draw
+	Viewport->SetDrawDefaultUVMaterial(false);
+
 	if (TextureAnalyzer.IsValid())
 	{
 		TextureAnalyzer.Get()->RefreshTextureAnalyzerTable(CustomizableObjectInstance);
diff --git a/Engine/Plugins/Experimental/Mutable/Source/CustomizableObjectEditor/Private/MuCOE/SCustomizableObjectEditorViewport.cpp b/Engine/Plugins/Experimental/Mutable/Source/CustomizableObjectEditor/Private/MuCOE/SCustomizableObjectEditorViewport.cpp
--- a/Engine/Plugins/Experimental/Mutable/Source/CustomizableObjectEditor/Private/MuCOE/SCustomizableObjectEditorViewport.cpp
+++ b/Engine/Plugins/Experimental/Mutable/Source/CustomizableObjectEditor/Private/MuCOE/SCustomizableObjectEditorViewport.cpp
@@ -545,6 +545,7 @@ void SCustomizableObjectEditorViewportTabBody::SetDrawDefaultUVMaterial(bool bIs
 			}
 		}
 
+		// If the CO is recompiled or couldn't find the last used materia, draw the first material
 		if (bIsCompilation || !bMaterialFound)
 		{
 			LevelViewportClient->SetDrawUVOverlayMaterial(*(ArrayUVMaterialOptionString[0]), "0");
@@ -558,9 +559,30 @@ void SCustomizableObjectEditorViewportTabBody::SetDrawDefaultUVMaterial(bool bIs
 
 	GenerateUVChannelOptions();
 
-	if (!ArrayUVChannelOptionString.IsEmpty() && (bIsCompilation || !bMaterialFound))
+	if (!ArrayUVChannelOptionString.IsEmpty())
 	{
-		SelectedUVChannel = ArrayUVChannelOptionString[0];
+		bool bUVChannelFound = false;
+
+		if (!bIsCompilation && SelectedUVChannel.IsValid())
+		{
+			// We check if the selected Material still exists after the update
+			FString UVChannel = *SelectedUVChannel;
+
+			for (int32 MaterialIndex = 0; MaterialIndex < ArrayUVChannelOptionString.Num(); ++MaterialIndex)
+			{
+				if (UVChannel == *ArrayUVChannelOptionString[MaterialIndex])
+				{
+					bUVChannelFound = true;
+					break;
+				}
+			}
+		}
+
+		// If the CO is recompiled or couldn't find the last used materia, draw the UV channel 0
+		if (bIsCompilation || !bMaterialFound || !bUVChannelFound)
+		{
+			SelectedUVChannel = ArrayUVChannelOptionString[0];
+		}
 	}
 	else
 	{
