[mutable] Ensure images are reduced at compile time if we know the platform-specific maximum size they can be used in.
Add section-based insights tracing for mutable updates.
#rnx

diff --git a/Engine/Plugins/Experimental/Mutable/Source/CustomizableObject/Private/MuCO/CustomizableObjectSystem.cpp b/Engine/Plugins/Experimental/Mutable/Source/CustomizableObject/Private/MuCO/CustomizableObjectSystem.cpp
--- a/Engine/Plugins/Experimental/Mutable/Source/CustomizableObject/Private/MuCO/CustomizableObjectSystem.cpp
+++ b/Engine/Plugins/Experimental/Mutable/Source/CustomizableObject/Private/MuCO/CustomizableObjectSystem.cpp
@@ -50,6 +50,9 @@ class UAnimInstance;
 DECLARE_CYCLE_STAT(TEXT("MutablePendingRelease Time"), STAT_MutablePendingRelease, STATGROUP_Game);
 DECLARE_CYCLE_STAT(TEXT("MutableTask"), STAT_MutableTask, STATGROUP_Game);
 
+#define UE_MUTABLE_UPDATE_REGION		TEXT("Mutable Update")
+
+
 UCustomizableObjectSystem* FCustomizableObjectSystemPrivate::SSystem = nullptr;
 
 static TAutoConsoleVariable<int32> CVarWorkingMemory(
@@ -2207,11 +2210,18 @@ namespace impl
 			CustomizableObjectInstancePrivateData->TexturesToRelease.Empty();
 		}
 
+		{
+			double delta = FPlatformTime::Seconds() - CustomizableObjectSystemPrivateData->CurrentMutableOperation->StartUpdateTime;
+			UE_LOG(LogMutable, Log, TEXT("Finished update in %.3f ms."), delta*1000.0);
+		}
+
 		LogBenchmarkUtil::UpdateBuildTimeStats(CustomizableObjectSystemPrivateData->MutableStats, 
 												CustomizableObjectSystemPrivateData->CurrentMutableOperation->StartUpdateTime);
 
 		// End Update
 		System->ClearCurrentMutableOperation();
+
+		TRACE_END_REGION(UE_MUTABLE_UPDATE_REGION);
 	}
 
 
@@ -2985,7 +2995,8 @@ bool UCustomizableObjectSystem::Tick(float DeltaTime)
 
 				const TSharedPtr<FMutableOperation> Operation = MakeShared<FMutableOperation>(FMutableOperation::CreateInstanceUpdate(*PendingInstance, PendingInstanceUpdateFound->Callback));
 				Private->StartUpdateSkeletalMesh(Operation);
-				
+				TRACE_BEGIN_REGION(UE_MUTABLE_UPDATE_REGION);
+
 				Private->MutablePendingInstanceWork.RemoveUpdate(PendingInstanceUpdateFound->CustomizableObjectInstance);
 			}
 			else if (LODUpdateCandidateFound)
@@ -2995,6 +3006,7 @@ bool UCustomizableObjectSystem::Tick(float DeltaTime)
 
 				const TSharedRef<FMutableOperation> Operation = MakeShared<FMutableOperation>(FMutableOperation::CreateInstanceUpdate(*LODUpdateCandidateFound->CustomizableObjectInstance, nullptr));
 				Private->StartUpdateSkeletalMesh(Operation);
+				TRACE_BEGIN_REGION(UE_MUTABLE_UPDATE_REGION);
 			}
 		}
 
diff --git a/Engine/Plugins/Experimental/Mutable/Source/MutableRuntime/Public/MuR/Image.h b/Engine/Plugins/Experimental/Mutable/Source/MutableRuntime/Public/MuR/Image.h
--- a/Engine/Plugins/Experimental/Mutable/Source/MutableRuntime/Public/MuR/Image.h
+++ b/Engine/Plugins/Experimental/Mutable/Source/MutableRuntime/Public/MuR/Image.h
@@ -583,7 +583,7 @@ namespace mu
 		MUTABLERUNTIME_API Ptr<Image> ExtractMip(const Image* From, int32 Mip);
 
 		/** Bilinear filter image resize. */
-		void ImageResizeLinear(Image* Dest, int32 ImageCompressionQuality, const Image* Base);
+		MUTABLERUNTIME_API void ImageResizeLinear(Image* Dest, int32 ImageCompressionQuality, const Image* Base);
 
 		/** Fill the image with a plain colour. */
 		void FillColor(Image* Image, FVector4f Color);
diff --git a/Engine/Plugins/Experimental/Mutable/Source/MutableTools/Private/MuT/CodeGenerator.cpp b/Engine/Plugins/Experimental/Mutable/Source/MutableTools/Private/MuT/CodeGenerator.cpp
--- a/Engine/Plugins/Experimental/Mutable/Source/MutableTools/Private/MuT/CodeGenerator.cpp
+++ b/Engine/Plugins/Experimental/Mutable/Source/MutableTools/Private/MuT/CodeGenerator.cpp
@@ -1338,7 +1338,6 @@ namespace mu
 							// Image
 							//-------------------------------------
 
-							// Size of a layout block in pixels
 							FIntPoint grid = pLayout->GetGridSize();
 
 							check(desc.m_format != EImageFormat::IF_NONE);
@@ -1402,18 +1401,18 @@ namespace mu
 								);
 
 								// Transform to pixels
-								box< UE::Math::TIntVector2<uint16> > rect = rectInCells;
-								rect.min[0] *= blockSizeX;
-								rect.min[1] *= blockSizeY;
-								rect.size[0] *= blockSizeX;
-								rect.size[1] *= blockSizeY;
+								box< UE::Math::TIntVector2<uint16> > BlockRectPixels = rectInCells;
+								BlockRectPixels.min[0] *= blockSizeX;
+								BlockRectPixels.min[1] *= blockSizeY;
+								BlockRectPixels.size[0] *= blockSizeX;
+								BlockRectPixels.size[1] *= blockSizeY;
 
 								// Generate the image
 								FImageGenerationOptions ImageOptions;
 								ImageOptions.CurrentStateIndex = m_currentStateIndex;
 								ImageOptions.ImageLayoutStrategy = ImageLayoutStrategy;
 								ImageOptions.ActiveTags = node.m_tags;
-								ImageOptions.RectSize = UE::Math::TIntVector2<int32>(rect.size);
+								ImageOptions.RectSize = UE::Math::TIntVector2<int32>(BlockRectPixels.size);
 								ImageOptions.LayoutToApply = pLayout;
 								ImageOptions.LayoutBlockId = pLayout->m_blocks[b].m_id;
 								FImageGenerationResult Result;
@@ -1508,12 +1507,16 @@ namespace mu
 												rect.size[0] = (blockRect.size[0] * extendDesc.m_size[0]) / extlayout[0];
 												rect.size[1] = (blockRect.size[1] * extendDesc.m_size[1]) / extlayout[1];
 
+												UE::Math::TIntVector2<int32> ExpectedBlockPixels;
+												ExpectedBlockPixels[0] = blockSizeX * blockRect.size[0];
+												ExpectedBlockPixels[1] = blockSizeY * blockRect.size[1];
+
 												// Generate the image block
 												FImageGenerationOptions ImageOptions;
 												ImageOptions.CurrentStateIndex = m_currentStateIndex;
 												ImageOptions.ImageLayoutStrategy = ImageLayoutStrategy;
 												ImageOptions.ActiveTags = node.m_tags;
-												ImageOptions.RectSize = UE::Math::TIntVector2<int32>(extendDesc.m_size);
+												ImageOptions.RectSize = ExpectedBlockPixels;
 												ImageOptions.LayoutToApply = pExtendLayout;
 												ImageOptions.LayoutBlockId = pExtendLayout->m_blocks[b].m_id;
 												FImageGenerationResult ExtendResult;
@@ -1521,15 +1524,13 @@ namespace mu
 												Ptr<ASTOp> fragmentAd = ExtendResult.op;
 
 												// Adjust the format and size of the block to be added
-												Ptr<ASTOp> formatted = GenerateImageFormat(fragmentAd, GetUncompressedFormat(desc.m_format));
-												UE::Math::TIntVector2<int32> expectedSize;
-												expectedSize[0] = blockSizeX * blockRect.size[0];
-												expectedSize[1] = blockSizeY * blockRect.size[1];
-												formatted = GenerateImageSize(formatted, expectedSize);
+												// \TODO: review if this is still needed considering a higher format op, and the RectSize in the ImageOptions
+												Ptr<ASTOp> formatted = GenerateImageFormat(fragmentAd, desc.m_format);
+												formatted = GenerateImageSize(formatted, ExpectedBlockPixels);
 												fragmentAd = formatted;
 
 												// Apply tiling to avoid generating chunks of image that are too big.
-												fragmentAd = ApplyTiling(fragmentAd, expectedSize, desc.m_format);
+												fragmentAd = ApplyTiling(fragmentAd, ExpectedBlockPixels, desc.m_format);
 
 												// Compose operation
 												Ptr<ASTOpImageCompose> composeOp = new ASTOpImageCompose();
diff --git a/Engine/Plugins/Experimental/Mutable/Source/MutableTools/Private/MuT/CodeGenerator.h b/Engine/Plugins/Experimental/Mutable/Source/MutableTools/Private/MuT/CodeGenerator.h
--- a/Engine/Plugins/Experimental/Mutable/Source/MutableTools/Private/MuT/CodeGenerator.h
+++ b/Engine/Plugins/Experimental/Mutable/Source/MutableTools/Private/MuT/CodeGenerator.h
@@ -47,7 +47,6 @@
 namespace mu
 {
 	class ASTOpParameter;
-	class ASTOpConstantResource;
 	class Layout;
 	class NodeColourArithmeticOperation;
 	class NodeColourConstant;
@@ -290,9 +289,6 @@ namespace mu
 
 		TArray<FConditionalExtensionDataOp> m_conditionalExtensionDataOps;
 
-		/** This is used to cache the constant image op generation. */
-		//TMap<Ptr<const Image>,Ptr<ASTOpConstantResource>> ImageConstantOpPerImage;
-
 		//-----------------------------------------------------------------------------------------
 
 		// Get the modifiers that have to be applied to elements with a specific tag.
diff --git a/Engine/Plugins/Experimental/Mutable/Source/MutableTools/Private/MuT/CodeGenerator_Image.cpp b/Engine/Plugins/Experimental/Mutable/Source/MutableTools/Private/MuT/CodeGenerator_Image.cpp
--- a/Engine/Plugins/Experimental/Mutable/Source/MutableTools/Private/MuT/CodeGenerator_Image.cpp
+++ b/Engine/Plugins/Experimental/Mutable/Source/MutableTools/Private/MuT/CodeGenerator_Image.cpp
@@ -193,7 +193,7 @@ namespace mu
             m_pErrorLog->GetPrivate()->Add( "Constant image not set.", ELMT_WARNING, node.m_errorContext );
         }
 
-		if (Options.ImageLayoutStrategy!= CompilerOptions::TextureLayoutStrategy::None && Options.LayoutToApply)
+		if (Options.ImageLayoutStrategy!=CompilerOptions::TextureLayoutStrategy::None && Options.LayoutToApply)
 		{
 			// We want to generate only a block from the image.
 
@@ -216,42 +216,38 @@ namespace mu
 			grid[1] = FMath::Max(1, grid[1]);
 
 			// Transform to pixels
-			box< UE::Math::TIntVector2<int32> > rect;
-			rect.min[0]  = (RectInCells.min[0]  * SourceImageSize[0]) / grid[0];
-			rect.min[1]  = (RectInCells.min[1]  * SourceImageSize[1]) / grid[1];
-			rect.size[0] = (RectInCells.size[0] * SourceImageSize[0]) / grid[0];
-			rect.size[1] = (RectInCells.size[1] * SourceImageSize[1]) / grid[1];
+			box< UE::Math::TIntVector2<int32> > SourceRectPixels;
+			SourceRectPixels.min[0]  = (RectInCells.min[0]  * SourceImageSize[0]) / grid[0];
+			SourceRectPixels.min[1]  = (RectInCells.min[1]  * SourceImageSize[1]) / grid[1];
+			SourceRectPixels.size[0] = (RectInCells.size[0] * SourceImageSize[0]) / grid[0];
+			SourceRectPixels.size[1] = (RectInCells.size[1] * SourceImageSize[1]) / grid[1];
+
+			FImageOperator ImOp = FImageOperator::GetDefault(m_compilerOptions->ImageFormatFunc);
 
 			// Do we need to crop?
-			if (rect.min[0]!=0 || rect.min[1]!=0 || pImage->GetSizeX() != rect.size[0] || pImage->GetSizeY() != rect.size[1])
+			if (SourceRectPixels.min[0]!=0 || SourceRectPixels.min[1]!=0 || pImage->GetSizeX() != SourceRectPixels.size[0] || pImage->GetSizeY() != SourceRectPixels.size[1])
 			{
 				// Crop now
-				FImageOperator ImOp = FImageOperator::GetDefault(m_compilerOptions->ImageFormatFunc);
-
-				Ptr<Image> pCropped = new Image(rect.size[0], rect.size[1], 1, pImage->GetFormat(), EInitializationType::NotInitialized);
-				ImOp.ImageCrop(pCropped.get(), m_compilerOptions->ImageCompressionQuality, pImage.get(), rect);
-
-				Ptr<ASTOpConstantResource> op = new ASTOpConstantResource();
-				op->type = OP_TYPE::IM_CONSTANT;
-				op->SetValue(pCropped, m_compilerOptions->OptimisationOptions.bUseDiskCache);
-				Result.op = op;
+				Ptr<Image> pCropped = new Image(SourceRectPixels.size[0], SourceRectPixels.size[1], 1, pImage->GetFormat(), EInitializationType::NotInitialized);
+				ImOp.ImageCrop(pCropped.get(), m_compilerOptions->ImageCompressionQuality, pImage.get(), SourceRectPixels);
+				pImage = pCropped;
 			}
-			else
+
+			// If actual size we need this image in is going to be smaller, resize now too
+			bool bBothAxisEqualOrSmaller = (Options.RectSize[0] < SourceRectPixels.size[0]) && (Options.RectSize[1] < SourceRectPixels.size[1]);
+			bool bBothAxisEqual = (Options.RectSize[0] == SourceRectPixels.size[0]) && (Options.RectSize[1] == SourceRectPixels.size[1]);
+			if (bBothAxisEqualOrSmaller && !bBothAxisEqual)
 			{
-				// No need to crop.
-				Ptr<ASTOpConstantResource> op = new ASTOpConstantResource();
-				op->type = OP_TYPE::IM_CONSTANT;
-				op->SetValue(pImage, m_compilerOptions->OptimisationOptions.bUseDiskCache);
-				Result.op = op;
+				Ptr<Image> ResizedImage = new Image(Options.RectSize[0], Options.RectSize[1],1,pImage->GetFormat(),EInitializationType::NotInitialized);
+				ImOp.ImageResizeLinear(ResizedImage.get(), m_compilerOptions->ImageCompressionQuality, pImage.get() );
+				pImage = ResizedImage;
 			}
 		}
-        else
-        {
-			Ptr<ASTOpConstantResource> op = new ASTOpConstantResource();
-			op->type = OP_TYPE::IM_CONSTANT;
-			op->SetValue( pImage, m_compilerOptions->OptimisationOptions.bUseDiskCache );
-			Result.op = op;
-		}
+ 
+		Ptr<ASTOpConstantResource> op = new ASTOpConstantResource();
+		op->type = OP_TYPE::IM_CONSTANT;
+		op->SetValue( pImage, m_compilerOptions->OptimisationOptions.bUseDiskCache );
+		Result.op = op;
     }
 
 
