[Mutable] Created a single Update Context.

This update context is maintain through the full lifecycle of the update.

- Merged FMutableOperation and FMutableOperationData into the the new FUpdateContextPrivate.
- Renaming of the Operation and OperationData varaibles will come in a future submit.
- Moved header implementations to the source due to delcaration order issues.

#tests PIE Jumpsuit, Village Infinite
#rnx

diff --git a/Engine/Plugins/Experimental/Mutable/Source/CustomizableObject/Classes/MuCO/CustomizableObjectInstance.h b/Engine/Plugins/Experimental/Mutable/Source/CustomizableObject/Classes/MuCO/CustomizableObjectInstance.h
--- a/Engine/Plugins/Experimental/Mutable/Source/CustomizableObject/Classes/MuCO/CustomizableObjectInstance.h
+++ b/Engine/Plugins/Experimental/Mutable/Source/CustomizableObject/Classes/MuCO/CustomizableObjectInstance.h
@@ -17,6 +17,7 @@ class UCustomizableInstancePrivateData; // This is used to hide Mutable SDK memb
 class UCustomizableObject;
 class UCustomizableSkeletalComponent;
 class UTexture2D;
+class FUpdateContextPrivate;
 struct FFrame;
 struct FGameplayTagContainer;
 struct FPropertyChangedEvent;
@@ -698,7 +699,7 @@ private:
 	UPROPERTY()
 	TArray<FCustomizableObjectFloatParameterValue> FloatParameters_DEPRECATED;
 
-	UPROPERTY()
+	UPROPERTY()																																						
 	TArray<FCustomizableObjectTextureParameterValue> TextureParameters_DEPRECATED;
 
 	UPROPERTY()
@@ -740,7 +741,7 @@ public:
 
 	void Issue();
 
-	void ApplyLODUpdateParamsToInstance();
+	void ApplyLODUpdateParamsToInstance(FUpdateContextPrivate* Context = nullptr);
 
 private:
 	/** If true it means that EnqueueUpdateSkeletalMesh has decided this update should be performed, if false it should be ignored. Just used for consistency checks */
diff --git a/Engine/Plugins/Experimental/Mutable/Source/CustomizableObject/Classes/MuCO/CustomizableObjectSystem.h b/Engine/Plugins/Experimental/Mutable/Source/CustomizableObject/Classes/MuCO/CustomizableObjectSystem.h
--- a/Engine/Plugins/Experimental/Mutable/Source/CustomizableObject/Classes/MuCO/CustomizableObjectSystem.h
+++ b/Engine/Plugins/Experimental/Mutable/Source/CustomizableObject/Classes/MuCO/CustomizableObjectSystem.h
@@ -28,6 +28,7 @@ class USkeletalMesh;
 class UMaterialInterface;
 class UTexture2D;
 class FCustomizableObjectSystemPrivate; // This is used to hide Mutable SDK members in the public headers.
+class FUpdateContextPrivate;
 struct FFrame;
 struct FGuid;
 
@@ -170,9 +171,8 @@ public:
 };
 
 
-// Before the Mutable Queue rework this made sense, but this is no longer the case. Remove this when doing MTBL-1409.
 /** End a Customizable Object Instance Update. All code paths of an update have to end here. */
-void FinishUpdateGlobal(UCustomizableObjectInstance* Instance, EUpdateResult UpdateResult, const FInstanceUpdateDelegate* UpdateCallback, const FDescriptorRuntimeHash InUpdatedHash = FDescriptorRuntimeHash());
+void FinishUpdateGlobal(const TSharedRef<FUpdateContextPrivate>& Context);
 
 
 UCLASS(Blueprintable, BlueprintType)
diff --git a/Engine/Plugins/Experimental/Mutable/Source/CustomizableObject/Private/MuCO/CustomizableInstancePrivateData.h b/Engine/Plugins/Experimental/Mutable/Source/CustomizableObject/Private/MuCO/CustomizableInstancePrivateData.h
--- a/Engine/Plugins/Experimental/Mutable/Source/CustomizableObject/Private/MuCO/CustomizableInstancePrivateData.h
+++ b/Engine/Plugins/Experimental/Mutable/Source/CustomizableObject/Private/MuCO/CustomizableInstancePrivateData.h
@@ -236,29 +236,29 @@ public:
 	void ClearCOInstanceFlags(ECOInstanceFlags FlagsToClear) { InstanceFlagsPrivate = (ECOInstanceFlags)(InstanceFlagsPrivate & ~FlagsToClear); }
 	bool HasCOInstanceFlags(ECOInstanceFlags FlagsToCheck) const { return (InstanceFlagsPrivate & FlagsToCheck) != 0; }
 
-	void BuildMaterials(const TSharedPtr<FMutableOperationData>& OperationData, UCustomizableObjectInstance* Public);
+	void BuildMaterials(const TSharedRef<FUpdateContextPrivate>& OperationData, UCustomizableObjectInstance* Public);
 
 	void ReuseTexture(UTexture2D* Texture, TSharedRef<FTexturePlatformData, ESPMode::ThreadSafe>& PlatformData);
 
 	// Return an event that will be fired when the assets  have been loaded. It returns null if no asset needs loading.
-	FGraphEventRef LoadAdditionalAssetsAsync(const TSharedPtr<FMutableOperationData>& OperationData, UCustomizableObjectInstance* Public, struct FStreamableManager &StreamableManager);
+	FGraphEventRef LoadAdditionalAssetsAsync(const TSharedRef<FUpdateContextPrivate>& OperationData, UCustomizableObjectInstance* Public, struct FStreamableManager &StreamableManager);
 	void AdditionalAssetsAsyncLoaded(UCustomizableObjectInstance* Public);
 
 	void TickUpdateCloseCustomizableObjects(UCustomizableObjectInstance& Publics, FMutableInstanceUpdateMap& InOutRequestedUpdates);
 	void UpdateInstanceIfNotGenerated(UCustomizableObjectInstance& Public, FMutableInstanceUpdateMap& InOutRequestedUpdates);
 
 	// Returns true if success (?)
-	bool UpdateSkeletalMesh_PostBeginUpdate0(UCustomizableObjectInstance* Public, const TSharedPtr<FMutableOperationData>& OperationData);
+	bool UpdateSkeletalMesh_PostBeginUpdate0(UCustomizableObjectInstance* Public, const TSharedRef<FUpdateContextPrivate>& OperationData);
 
 	static void ReleaseMutableTexture(const FMutableImageCacheKey& MutableTextureKey, UTexture2D* Texture, struct FMutableResourceCache& Cache);
 
 	// Prepare the data for the unreal textures, but don't create them because
 	// it runs in the mutable thread
 	// \TODO: CustomizableObject shouldn't be here
-	static void ProcessTextureCoverageQueries(const TSharedPtr<FMutableOperationData>& OperationData, UCustomizableObject* CustomizableObject, const FString& ImageKeyName, FTexturePlatformData *PlatformData, UMaterialInterface* Material);
+	static void ProcessTextureCoverageQueries(const TSharedRef<FUpdateContextPrivate>& OperationData, UCustomizableObject* CustomizableObject, const FString& ImageKeyName, FTexturePlatformData *PlatformData, UMaterialInterface* Material);
 
 	// Copy data generated in the mutable thread over to the instance and initializes additional data required during the update
-	void PrepareForUpdate(const TSharedPtr<FMutableOperationData>& OperationData);
+	void PrepareForUpdate(const TSharedRef<FUpdateContextPrivate>& OperationData);
 
 	int32 GetNumLODsAvailable() const { return NumLODsAvailable; }
 	
@@ -285,14 +285,14 @@ public:
 
 private:
 
-	void InitSkeletalMeshData(const TSharedPtr<FMutableOperationData>& OperationData, USkeletalMesh* SkeletalMesh, const FMutableRefSkeletalMeshData* RefSkeletalMeshData, const UCustomizableObject& CustomizableObject, int32 ComponentIndex);
+	void InitSkeletalMeshData(const TSharedRef<FUpdateContextPrivate>& OperationData, USkeletalMesh* SkeletalMesh, const FMutableRefSkeletalMeshData* RefSkeletalMeshData, const UCustomizableObject& CustomizableObject, int32 ComponentIndex);
 
-	bool BuildSkeletonData(const TSharedPtr<FMutableOperationData>& OperationData, USkeletalMesh& SkeletalMesh, const FMutableRefSkeletalMeshData& RefSkeletalMeshData, UCustomizableObject& CustomizableObject, int32 ComponentIndex);
-	void BuildMeshSockets(const TSharedPtr<FMutableOperationData>& OperationData, USkeletalMesh* SkeletalMesh, const FMutableRefSkeletalMeshData* RefSkeletalMeshData, UCustomizableObjectInstance* CustomizableObjectInstance, mu::MeshPtrConst MutableMesh);
-	void BuildOrCopyElementData(const TSharedPtr<FMutableOperationData>& OperationData, USkeletalMesh* SkeletalMesh, UCustomizableObjectInstance* CustomizableObjectInstance, int32 ComponentIndex);
-	void BuildOrCopyMorphTargetsData(const TSharedPtr<FMutableOperationData>& OperationData, USkeletalMesh* SkeletalMesh, const USkeletalMesh* SrcSkeletalMesh, UCustomizableObjectInstance* CustomizableObjectInstance, int32 ComponentIndex);
-	bool BuildOrCopyRenderData(const TSharedPtr<FMutableOperationData>& OperationData, USkeletalMesh* SkeletalMesh, const USkeletalMesh* SrcSkeletalMesh, UCustomizableObjectInstance* CustomizableObjectInstance, int32 ComponentIndex);
-	void BuildOrCopyClothingData(const TSharedPtr<FMutableOperationData>& OperationData, USkeletalMesh* SkeletalMesh, const USkeletalMesh* SrcSkeletalMesh, UCustomizableObjectInstance* CustomizableObjectInstance, int32 ComponentIndex);
+	bool BuildSkeletonData(const TSharedRef<FUpdateContextPrivate>& OperationData, USkeletalMesh& SkeletalMesh, const FMutableRefSkeletalMeshData& RefSkeletalMeshData, UCustomizableObject& CustomizableObject, int32 ComponentIndex);
+	void BuildMeshSockets(const TSharedRef<FUpdateContextPrivate>& OperationData, USkeletalMesh* SkeletalMesh, const FMutableRefSkeletalMeshData* RefSkeletalMeshData, UCustomizableObjectInstance* CustomizableObjectInstance, mu::MeshPtrConst MutableMesh);
+	void BuildOrCopyElementData(const TSharedRef<FUpdateContextPrivate>& OperationData, USkeletalMesh* SkeletalMesh, UCustomizableObjectInstance* CustomizableObjectInstance, int32 ComponentIndex);
+	void BuildOrCopyMorphTargetsData(const TSharedRef<FUpdateContextPrivate>& OperationData, USkeletalMesh* SkeletalMesh, const USkeletalMesh* SrcSkeletalMesh, UCustomizableObjectInstance* CustomizableObjectInstance, int32 ComponentIndex);
+	bool BuildOrCopyRenderData(const TSharedRef<FUpdateContextPrivate>& OperationData, USkeletalMesh* SkeletalMesh, const USkeletalMesh* SrcSkeletalMesh, UCustomizableObjectInstance* CustomizableObjectInstance, int32 ComponentIndex);
+	void BuildOrCopyClothingData(const TSharedRef<FUpdateContextPrivate>& OperationData, USkeletalMesh* SkeletalMesh, const USkeletalMesh* SrcSkeletalMesh, UCustomizableObjectInstance* CustomizableObjectInstance, int32 ComponentIndex);
 	
 	//
 	USkeleton* MergeSkeletons(UCustomizableObject& CustomizableObject, const FMutableRefSkeletalMeshData& RefSkeletalMeshData, int32 ComponentIndex);
@@ -305,7 +305,7 @@ private:
 
 	void InvalidateGeneratedData();
 
-	bool DoComponentsNeedUpdate(UCustomizableObjectInstance* CustomizableObjectInstance, const TSharedPtr<FMutableOperationData>& OperationData, TArray<bool>& OutComponentNeedsUpdate, bool& bOutEmptyMesh);
+	bool DoComponentsNeedUpdate(UCustomizableObjectInstance* CustomizableObjectInstance, const TSharedRef<FUpdateContextPrivate>& OperationData, TArray<bool>& OutComponentNeedsUpdate, bool& bOutEmptyMesh);
 
 	mu::FResourceID GetLastMeshId(int32 ComponentIndex, int32 LODIndex) const;
 	void SetLastMeshId(int32 ComponentIndex, int32 LODIndex, mu::FResourceID MeshId);
diff --git a/Engine/Plugins/Experimental/Mutable/Source/CustomizableObject/Private/MuCO/CustomizableObjectInstance.cpp b/Engine/Plugins/Experimental/Mutable/Source/CustomizableObject/Private/MuCO/CustomizableObjectInstance.cpp
--- a/Engine/Plugins/Experimental/Mutable/Source/CustomizableObject/Private/MuCO/CustomizableObjectInstance.cpp
+++ b/Engine/Plugins/Experimental/Mutable/Source/CustomizableObject/Private/MuCO/CustomizableObjectInstance.cpp
@@ -187,7 +187,7 @@ void UCustomizableObjectInstance::SetDescriptor(const FCustomizableObjectInstanc
 }
 
 
-void UCustomizableInstancePrivateData::PrepareForUpdate(const TSharedPtr<FMutableOperationData>& OperationData)
+void UCustomizableInstancePrivateData::PrepareForUpdate(const TSharedRef<FUpdateContextPrivate>& OperationData)
 {
 	// Clear the ComponentData from previous updates
 	for (FCustomizableInstanceComponentData& ComponentData : ComponentsData)
@@ -673,15 +673,26 @@ bool UCustomizableObjectInstance::CanUpdateInstance() const
 
 void UCustomizableObjectInstance::UpdateSkeletalMeshAsync(bool bIgnoreCloseDist, bool bForceHighPriority)
 {
-	FCustomizableObjectSystemPrivate* SystemPrivate = UCustomizableObjectSystem::GetInstance()->GetPrivate();	
-	SystemPrivate->EnqueueUpdateSkeletalMesh(*this, false, bIgnoreCloseDist, bForceHighPriority, nullptr, nullptr);
+	FCustomizableObjectSystemPrivate* SystemPrivate = UCustomizableObjectSystem::GetInstance()->GetPrivate();
+
+	const TSharedRef<FUpdateContextPrivate> Context = MakeShared<FUpdateContextPrivate>(*this);
+	Context->bIgnoreCloseDist = bIgnoreCloseDist;
+	Context->bForceHighPriority = bForceHighPriority;
+	
+	SystemPrivate->EnqueueUpdateSkeletalMesh(Context);
 }
 
 
 void UCustomizableObjectInstance::UpdateSkeletalMeshAsyncResult(FInstanceUpdateDelegate Callback, bool bIgnoreCloseDist, bool bForceHighPriority)
 {
-	FCustomizableObjectSystemPrivate* SystemPrivate = UCustomizableObjectSystem::GetInstance()->GetPrivate();	
-	SystemPrivate->EnqueueUpdateSkeletalMesh(*this, false, bIgnoreCloseDist, bForceHighPriority, nullptr, &Callback);
+	FCustomizableObjectSystemPrivate* SystemPrivate = UCustomizableObjectSystem::GetInstance()->GetPrivate();
+
+	const TSharedRef<FUpdateContextPrivate> Context = MakeShared<FUpdateContextPrivate>(*this);
+	Context->bIgnoreCloseDist = bIgnoreCloseDist;
+	Context->bForceHighPriority = bForceHighPriority;
+	Context->UpdateCallback = Callback;
+	
+	SystemPrivate->EnqueueUpdateSkeletalMesh(Context);
 }
 
 
@@ -749,8 +760,12 @@ void UCustomizableInstancePrivateData::UpdateInstanceIfNotGenerated(UCustomizabl
 		return;
 	}
 
-	FCustomizableObjectSystemPrivate* SystemPrivate = UCustomizableObjectSystem::GetInstance()->GetPrivate();	
-	SystemPrivate->EnqueueUpdateSkeletalMesh(Public, true, false, false, nullptr, nullptr);
+	FCustomizableObjectSystemPrivate* SystemPrivate = UCustomizableObjectSystem::GetInstance()->GetPrivate();
+
+	const TSharedRef<FUpdateContextPrivate> Context = MakeShared<FUpdateContextPrivate>(Public);
+	Context->bOnlyUpdateIfNotGenerated = true;
+	
+	SystemPrivate->EnqueueUpdateSkeletalMesh(Context);
 
 	EQueuePriorityType Priority = SystemPrivate->GetUpdatePriority(Public, false);
 	FMutableUpdateCandidate* UpdateCandidate = InOutRequestedUpdates.Find(&Public);
@@ -1540,7 +1555,7 @@ FAutoConsoleVariableRef CVarApplyFixDoComponentsNeedUpdate(
 	TEXT("If true, the new version of DoComponentsNeedUpdate will be used. This fix tries to detect cases where the output mesh would be invalid and cancel the update."));
 
 
-bool UCustomizableInstancePrivateData::DoComponentsNeedUpdate(UCustomizableObjectInstance* Public, const TSharedPtr<FMutableOperationData>& OperationData, TArray<bool>& OutComponentNeedsUpdate, bool& bHasInvalidMesh)
+bool UCustomizableInstancePrivateData::DoComponentsNeedUpdate(UCustomizableObjectInstance* Public, const TSharedRef<FUpdateContextPrivate>& OperationData, TArray<bool>& OutComponentNeedsUpdate, bool& bHasInvalidMesh)
 {
 	UCustomizableObject* CustomizableObject = Public->GetCustomizableObject();
 	check(CustomizableObject);
@@ -1725,7 +1740,7 @@ bool UCustomizableInstancePrivateData::DoComponentsNeedUpdate(UCustomizableObjec
 }
 
 
-bool UCustomizableInstancePrivateData::UpdateSkeletalMesh_PostBeginUpdate0(UCustomizableObjectInstance* Public, const TSharedPtr<FMutableOperationData>& OperationData)
+bool UCustomizableInstancePrivateData::UpdateSkeletalMesh_PostBeginUpdate0(UCustomizableObjectInstance* Public, const TSharedRef<FUpdateContextPrivate>& OperationData)
 {
 	MUTABLE_CPUPROFILER_SCOPE(UCustomizableInstancePrivateData::UpdateSkeletalMesh_PostBeginUpdate0)
 
@@ -2801,7 +2816,7 @@ void FMutableUpdateCandidate::Issue()
 }
 
 
-void FMutableUpdateCandidate::ApplyLODUpdateParamsToInstance()
+void FMutableUpdateCandidate::ApplyLODUpdateParamsToInstance(FUpdateContextPrivate* Context)
 {
 	CustomizableObjectInstance->Descriptor.MinLOD = MinLOD;
 	CustomizableObjectInstance->Descriptor.MaxLOD = MaxLOD;
@@ -2810,6 +2825,11 @@ void FMutableUpdateCandidate::ApplyLODUpdateParamsToInstance()
 
 	CustomizableObjectInstance->GetPrivate()->UpdateDescriptorRuntimeHash.UpdateMinMaxLOD(MinLOD, MaxLOD);
 	CustomizableObjectInstance->GetPrivate()->UpdateDescriptorRuntimeHash.UpdateRequestedLODs(CustomizableObjectInstance->Descriptor.RequestedLODLevels);
+
+	if (Context)
+	{
+		Context->InstanceDescriptorRuntimeHash = CustomizableObjectInstance->GetPrivate()->UpdateDescriptorRuntimeHash;
+	}
 }
 
 
@@ -3176,7 +3196,7 @@ void UCustomizableInstancePrivateData::ConvertImage(UTexture2D* Texture, mu::Ima
 }
 
 
-void UCustomizableInstancePrivateData::InitSkeletalMeshData(const TSharedPtr<FMutableOperationData>& OperationData, USkeletalMesh* SkeletalMesh, const FMutableRefSkeletalMeshData* RefSkeletalMeshData, const UCustomizableObject& CustomizableObject, int32 ComponentIndex)
+void UCustomizableInstancePrivateData::InitSkeletalMeshData(const TSharedRef<FUpdateContextPrivate>& OperationData, USkeletalMesh* SkeletalMesh, const FMutableRefSkeletalMeshData* RefSkeletalMeshData, const UCustomizableObject& CustomizableObject, int32 ComponentIndex)
 {
 	MUTABLE_CPUPROFILER_SCOPE(UCustomizableInstancePrivateData::InitSkeletalMesh);
 
@@ -3261,7 +3281,7 @@ void UCustomizableInstancePrivateData::InitSkeletalMeshData(const TSharedPtr<FMu
 }
 
 
-bool UCustomizableInstancePrivateData::BuildSkeletonData(const TSharedPtr<FMutableOperationData>& OperationData, USkeletalMesh& SkeletalMesh, const FMutableRefSkeletalMeshData& RefSkeletalMeshData, UCustomizableObject& CustomizableObject, int32 ComponentIndex)
+bool UCustomizableInstancePrivateData::BuildSkeletonData(const TSharedRef<FUpdateContextPrivate>& OperationData, USkeletalMesh& SkeletalMesh, const FMutableRefSkeletalMeshData& RefSkeletalMeshData, UCustomizableObject& CustomizableObject, int32 ComponentIndex)
 {
 	MUTABLE_CPUPROFILER_SCOPE(UCustomizableInstancePrivateData::BuildSkeletonData);
 
@@ -3390,7 +3410,7 @@ bool UCustomizableInstancePrivateData::BuildSkeletonData(const TSharedPtr<FMutab
 }
 
 
-void UCustomizableInstancePrivateData::BuildMeshSockets(const TSharedPtr<FMutableOperationData>& OperationData, USkeletalMesh* SkeletalMesh, const FMutableRefSkeletalMeshData* RefSkeletalMeshData, UCustomizableObjectInstance* Public, mu::MeshPtrConst MutableMesh)
+void UCustomizableInstancePrivateData::BuildMeshSockets(const TSharedRef<FUpdateContextPrivate>& OperationData, USkeletalMesh* SkeletalMesh, const FMutableRefSkeletalMeshData* RefSkeletalMeshData, UCustomizableObjectInstance* Public, mu::MeshPtrConst MutableMesh)
 {
 	// Build mesh sockets.
 	MUTABLE_CPUPROFILER_SCOPE(UCustomizableInstancePrivateData::BuildMeshSockets);
@@ -3490,7 +3510,7 @@ void UCustomizableInstancePrivateData::BuildMeshSockets(const TSharedPtr<FMutabl
 }
 
 
-void UCustomizableInstancePrivateData::BuildOrCopyElementData(const TSharedPtr<FMutableOperationData>& OperationData, USkeletalMesh* SkeletalMesh, UCustomizableObjectInstance* CustomizableObjectInstance, int32 ComponentIndex)
+void UCustomizableInstancePrivateData::BuildOrCopyElementData(const TSharedRef<FUpdateContextPrivate>& OperationData, USkeletalMesh* SkeletalMesh, UCustomizableObjectInstance* CustomizableObjectInstance, int32 ComponentIndex)
 {
 	MUTABLE_CPUPROFILER_SCOPE(UCustomizableInstancePrivateData::BuildOrCopyElementData);
 
@@ -3518,7 +3538,7 @@ void UCustomizableInstancePrivateData::BuildOrCopyElementData(const TSharedPtr<F
 }
 
 
-void UCustomizableInstancePrivateData::BuildOrCopyMorphTargetsData(const TSharedPtr<FMutableOperationData>& OperationData, USkeletalMesh* SkeletalMesh, const USkeletalMesh* LastUpdateSkeletalMesh, UCustomizableObjectInstance* CustomizableObjectInstance, int32 ComponentIndex)
+void UCustomizableInstancePrivateData::BuildOrCopyMorphTargetsData(const TSharedRef<FUpdateContextPrivate>& OperationData, USkeletalMesh* SkeletalMesh, const USkeletalMesh* LastUpdateSkeletalMesh, UCustomizableObjectInstance* CustomizableObjectInstance, int32 ComponentIndex)
 {
 	MUTABLE_CPUPROFILER_SCOPE(UCustomizableInstancePrivateData::BuildOrCopyMorphTargetsData);
 
@@ -3670,7 +3690,7 @@ namespace
 }
 
 
-void UCustomizableInstancePrivateData::BuildOrCopyClothingData(const TSharedPtr<FMutableOperationData>&OperationData, USkeletalMesh * SkeletalMesh, const USkeletalMesh* LastUpdateSkeletalMesh, UCustomizableObjectInstance * CustomizableObjectInstance, int32 ComponentIndex)
+void UCustomizableInstancePrivateData::BuildOrCopyClothingData(const TSharedRef<FUpdateContextPrivate>&OperationData, USkeletalMesh * SkeletalMesh, const USkeletalMesh* LastUpdateSkeletalMesh, UCustomizableObjectInstance * CustomizableObjectInstance, int32 ComponentIndex)
 {
 	MUTABLE_CPUPROFILER_SCOPE(UCustomizableInstancePrivateData::BuildOrCopyClothingData);
 
@@ -4536,7 +4556,7 @@ void UCustomizableInstancePrivateData::BuildOrCopyClothingData(const TSharedPtr<
 }
 
 
-bool UCustomizableInstancePrivateData::BuildOrCopyRenderData(const TSharedPtr<FMutableOperationData>& OperationData, USkeletalMesh* SkeletalMesh, const USkeletalMesh* LastUpdateSkeletalMesh, UCustomizableObjectInstance* Public, int32 ComponentIndex)
+bool UCustomizableInstancePrivateData::BuildOrCopyRenderData(const TSharedRef<FUpdateContextPrivate>& OperationData, USkeletalMesh* SkeletalMesh, const USkeletalMesh* LastUpdateSkeletalMesh, UCustomizableObjectInstance* Public, int32 ComponentIndex)
 {
 	MUTABLE_CPUPROFILER_SCOPE(UCustomizableInstancePrivateData::BuildOrCopyRenderData);
 
@@ -4706,7 +4726,7 @@ FAutoConsoleVariableRef CVarMutableHighPriorityLoading(
 	TEXT("If enabled, the request to load additional assets will have high priority."));
 
 
-FGraphEventRef UCustomizableInstancePrivateData::LoadAdditionalAssetsAsync(const TSharedPtr<FMutableOperationData>& OperationData, UCustomizableObjectInstance* Public, FStreamableManager& StreamableManager)
+FGraphEventRef UCustomizableInstancePrivateData::LoadAdditionalAssetsAsync(const TSharedRef<FUpdateContextPrivate>& OperationData, UCustomizableObjectInstance* Public, FStreamableManager& StreamableManager)
 {
 	MUTABLE_CPUPROFILER_SCOPE(UCustomizableInstancePrivateData::LoadAdditionalAssetsAsync);
 
@@ -5256,7 +5276,7 @@ FAutoConsoleVariableRef CVarMutableReuseMaterialInstances(
 	bReuseMaterialInstances,
 	TEXT("If true, allow reuse of MaterialInstances between updates."));
 
-void UCustomizableInstancePrivateData::BuildMaterials(const TSharedPtr<FMutableOperationData>& OperationData, UCustomizableObjectInstance* Public)
+void UCustomizableInstancePrivateData::BuildMaterials(const TSharedRef<FUpdateContextPrivate>& OperationData, UCustomizableObjectInstance* Public)
 {
 	MUTABLE_CPUPROFILER_SCOPE(UCustomizableInstancePrivateData::BuildMaterials)
 
@@ -5280,7 +5300,7 @@ void UCustomizableInstancePrivateData::BuildMaterials(const TSharedPtr<FMutableO
 	// Prepare the data to store in order to regenerate resources for this instance (usually texture mips).
 	TSharedPtr<FMutableUpdateContext> UpdateContext = MakeShared<FMutableUpdateContext>(UCustomizableObjectSystem::GetInstance()->GetPrivate()->MutableSystem,
 		CustomizableObject->GetModel(),
-		OperationData->MutableParameters,
+		OperationData->Parameters,
 	    OperationData->State);
 	
 	const bool bReuseTextures = OperationData->bReuseInstanceTextures;
@@ -5899,7 +5919,7 @@ void UCustomizableInstancePrivateData::BuildMaterials(const TSharedPtr<FMutableO
 }
 
 
-void UCustomizableInstancePrivateData::ProcessTextureCoverageQueries(const TSharedPtr<FMutableOperationData>& OperationData, UCustomizableObject* CustomizableObject, const FString& ImageKeyName, FTexturePlatformData *PlatformData, UMaterialInterface* Material)
+void UCustomizableInstancePrivateData::ProcessTextureCoverageQueries(const TSharedRef<FUpdateContextPrivate>& OperationData, UCustomizableObject* CustomizableObject, const FString& ImageKeyName, FTexturePlatformData *PlatformData, UMaterialInterface* Material)
 {
 	if (OperationData->TextureCoverageQueries_MutableThreadParams.Num() && PlatformData->Mips.Num())
 	{
diff --git a/Engine/Plugins/Experimental/Mutable/Source/CustomizableObject/Private/MuCO/CustomizableObjectSystem.cpp b/Engine/Plugins/Experimental/Mutable/Source/CustomizableObject/Private/MuCO/CustomizableObjectSystem.cpp
--- a/Engine/Plugins/Experimental/Mutable/Source/CustomizableObject/Private/MuCO/CustomizableObjectSystem.cpp
+++ b/Engine/Plugins/Experimental/Mutable/Source/CustomizableObject/Private/MuCO/CustomizableObjectSystem.cpp
@@ -47,6 +47,10 @@
 class AActor;
 class UAnimInstance;
 
+void CacheTexturesParameters(const TArray<FName>& TextureParameters);
+void UnCacheTexturesParameters(const TArray<FName>& TextureParameters);
+
+
 DECLARE_CYCLE_STAT(TEXT("MutablePendingRelease Time"), STAT_MutablePendingRelease, STATGROUP_Game);
 DECLARE_CYCLE_STAT(TEXT("MutableTask"), STAT_MutableTask, STATGROUP_Game);
 
@@ -100,6 +104,85 @@ static void CVarMutableSinkFunction()
 static FAutoConsoleVariableSink CVarMutableSink(FConsoleCommandDelegate::CreateStatic(&CVarMutableSinkFunction));
 
 
+FUpdateContextPrivate::FUpdateContextPrivate(UCustomizableObjectInstance& InInstance)
+{
+	check(InInstance.GetPrivate());
+	check(InInstance.GetCustomizableObject());
+
+	Instance = &InInstance;
+	InstanceDescriptorRuntimeHash = FDescriptorRuntimeHash(Instance->GetDescriptor());
+	InInstance.GetPrivate()->UpdateDescriptorRuntimeHash = InstanceDescriptorRuntimeHash; // TODO GMTFuture Remove on MTBL-1409
+	State = InInstance.GetState();
+	bBuildParameterRelevancy = InInstance.GetBuildParameterRelevancy();
+	Parameters = InInstance.GetDescriptor().GetParameters();
+	TextureParameters = InInstance.GetPrivate()->UpdateTextureParameters;
+	
+	InInstance.GetCustomizableObject()->ApplyStateForcedValuesToParameters(State, Parameters.get());
+
+	CacheTexturesParameters(TextureParameters);
+}
+
+
+FUpdateContextPrivate::~FUpdateContextPrivate()
+{
+	UnCacheTexturesParameters(TextureParameters);
+}
+
+
+FMutablePendingInstanceUpdate::FMutablePendingInstanceUpdate(const TSharedRef<FUpdateContextPrivate>& InContext) :
+	Context(InContext)
+{
+	SecondsAtUpdate = FPlatformTime::Seconds();
+}
+
+
+bool FMutablePendingInstanceUpdate::operator==(const FMutablePendingInstanceUpdate& Other) const
+{
+	return Context->Instance.HasSameIndexAndSerialNumber(Other.Context->Instance);
+}
+
+
+bool FMutablePendingInstanceUpdate::operator<(const FMutablePendingInstanceUpdate& Other) const
+{
+	if (Context->PriorityType < Other.Context->PriorityType)
+	{
+		return true;
+	}
+	else if (Context->PriorityType > Other.Context->PriorityType)
+	{
+		return false;
+	}
+	else
+	{
+		return SecondsAtUpdate < Other.SecondsAtUpdate;
+	}
+}
+
+
+uint32 GetTypeHash(const FMutablePendingInstanceUpdate& Update)
+{
+	return GetTypeHash(Update.Context->Instance.GetWeakPtrTypeHash());
+}
+
+
+TWeakObjectPtr<const UCustomizableObjectInstance> FPendingInstanceUpdateKeyFuncs::GetSetKey(const FMutablePendingInstanceUpdate& PendingUpdate)
+{
+	return PendingUpdate.Context->Instance;
+}
+
+
+bool FPendingInstanceUpdateKeyFuncs::Matches(const TWeakObjectPtr<const UCustomizableObjectInstance>& A, const TWeakObjectPtr<const UCustomizableObjectInstance>& B)
+{
+	return A.HasSameIndexAndSerialNumber(B);
+}
+
+
+uint32 FPendingInstanceUpdateKeyFuncs::GetKeyHash(const TWeakObjectPtr<const UCustomizableObjectInstance>& Identifier)
+{
+	return GetTypeHash(Identifier.GetWeakPtrTypeHash());
+}
+
+
 bool FMutablePendingInstanceWork::ArePendingUpdatesEmpty() const
 {
 	return PendingInstanceUpdates.Num() == 0;
@@ -120,17 +203,16 @@ void FMutablePendingInstanceWork::SetLODUpdatesLastTick(int32 NumLODUpdates)
 
 void FMutablePendingInstanceWork::AddUpdate(const FMutablePendingInstanceUpdate& UpdateToAdd)
 {
-	if (const FMutablePendingInstanceUpdate* ExistingUpdate = PendingInstanceUpdates.Find(UpdateToAdd.CustomizableObjectInstance))
+	if (const FMutablePendingInstanceUpdate* ExistingUpdate = PendingInstanceUpdates.Find(UpdateToAdd.Context->Instance))
 	{
-		FMutablePendingInstanceUpdate TaskToEnqueue = UpdateToAdd;
-
-		FInstanceUpdateDelegate* Callback = ExistingUpdate->Callback;
-		FinishUpdateGlobal(ExistingUpdate->CustomizableObjectInstance.Get(), EUpdateResult::ErrorReplaced, Callback);
+		ExistingUpdate->Context->UpdateResult = EUpdateResult::ErrorReplaced;
+		FinishUpdateGlobal(ExistingUpdate->Context);
 
-		TaskToEnqueue.PriorityType = FMath::Min(ExistingUpdate->PriorityType, UpdateToAdd.PriorityType);
+		FMutablePendingInstanceUpdate TaskToEnqueue = UpdateToAdd;
+		TaskToEnqueue.Context->PriorityType = FMath::Min(ExistingUpdate->Context->PriorityType, UpdateToAdd.Context->PriorityType);
 		TaskToEnqueue.SecondsAtUpdate = FMath::Min(ExistingUpdate->SecondsAtUpdate, UpdateToAdd.SecondsAtUpdate);
 		
-		PendingInstanceUpdates.Remove(ExistingUpdate->CustomizableObjectInstance);
+		PendingInstanceUpdates.Remove(ExistingUpdate->Context->Instance);
 		PendingInstanceUpdates.Add(TaskToEnqueue);
 	}
 	else
@@ -138,9 +220,10 @@ void FMutablePendingInstanceWork::AddUpdate(const FMutablePendingInstanceUpdate&
 		PendingInstanceUpdates.Add(UpdateToAdd);
 	}
 
-	if (const FMutablePendingInstanceDiscard* ExistingDiscard = PendingInstanceDiscards.Find(UpdateToAdd.CustomizableObjectInstance))
+	if (const FMutablePendingInstanceDiscard* ExistingDiscard = PendingInstanceDiscards.Find(UpdateToAdd.Context->Instance))
 	{
-		FinishUpdateGlobal(ExistingDiscard->CustomizableObjectInstance.Get(), EUpdateResult::ErrorReplaced, nullptr);
+		UpdateToAdd.Context->UpdateResult = EUpdateResult::ErrorReplaced;
+		FinishUpdateGlobal(UpdateToAdd.Context);
 
 		PendingInstanceDiscards.Remove(ExistingDiscard->CustomizableObjectInstance);
 	}
@@ -163,8 +246,9 @@ void FMutablePendingInstanceWork::AddDiscard(const FMutablePendingInstanceDiscar
 {
 	if (const FMutablePendingInstanceUpdate* ExistingUpdate = PendingInstanceUpdates.Find(TaskToEnqueue.CustomizableObjectInstance.Get()))
 	{
-		FinishUpdateGlobal(ExistingUpdate->CustomizableObjectInstance.Get(), EUpdateResult::ErrorReplaced, ExistingUpdate->Callback);
-		PendingInstanceUpdates.Remove(ExistingUpdate->CustomizableObjectInstance);
+		ExistingUpdate->Context->UpdateResult = EUpdateResult::ErrorReplaced;
+		FinishUpdateGlobal(ExistingUpdate->Context);
+		PendingInstanceUpdates.Remove(ExistingUpdate->Context->Instance);
 	}
 
 	PendingInstanceDiscards.Add(TaskToEnqueue);
@@ -678,20 +762,22 @@ static FAutoConsoleVariableRef CVarApplyFixPrepareSkeletons(
 	TEXT("If true, Fix missing SkeletonsData when FirstLODToGenerate is greater than 0. If false, There may be a crash when generating meshes on platform that skip LODs."),
 	ECVF_Default);
 
-void FinishUpdateGlobal(UCustomizableObjectInstance* Instance, EUpdateResult UpdateResult, const FInstanceUpdateDelegate* UpdateCallback, const FDescriptorRuntimeHash InUpdatedHash)
+void FinishUpdateGlobal(const TSharedRef<FUpdateContextPrivate>& Context)
 {
 	check(IsInGameThread())
 
+	UCustomizableObjectInstance* Instance = Context->Instance.Get();
+
 	if (Instance)
 	{
 		UCustomizableInstancePrivateData* PrivateInstance = Instance->GetPrivate();
 		
-		switch (UpdateResult)
+		switch (Context->UpdateResult)
 		{
 		case EUpdateResult::Success:
 			PrivateInstance->SetSkeletalMeshStatus(ESkeletalMeshStatus::Success);
 
-			PrivateInstance->DescriptorRuntimeHash = InUpdatedHash;
+			PrivateInstance->DescriptorRuntimeHash = Context->InstanceDescriptorRuntimeHash;
 
 			// Delegates must be called only after updating the Instance flags.
 			Instance->UpdatedDelegate.Broadcast(Instance);
@@ -717,7 +803,7 @@ void FinishUpdateGlobal(UCustomizableObjectInstance* Instance, EUpdateResult Upd
 		}
 	}
 
-	if (UpdateResult == EUpdateResult::Success)
+	if (Context->UpdateResult == EUpdateResult::Success)
 	{
 		// Call Customizable Skeletal Components updated callbacks.
 		for (TObjectIterator<UCustomizableSkeletalComponent> It; It; ++It) // Since iterating objects is expensive, for now CustomizableSkeletalComponent does not have a FinishUpdate function.
@@ -738,13 +824,10 @@ void FinishUpdateGlobal(UCustomizableObjectInstance* Instance, EUpdateResult Upd
 		}
 	}
 
-	if (UpdateCallback)
-	{
-		FUpdateContext Context;
-		Context.UpdateResult = UpdateResult;
+	FUpdateContext ContextPublic;
+	ContextPublic.UpdateResult = Context->UpdateResult;
 		
-		UpdateCallback->ExecuteIfBound(Context);
-	}
+	Context->UpdateCallback.ExecuteIfBound(ContextPublic);
 
 	UCustomizableObjectSystem::GetInstance()->GetPrivate()->MutableTaskGraph.AllowLaunchingMutableTaskLowPriority(true, false);
 
@@ -754,15 +837,18 @@ void FinishUpdateGlobal(UCustomizableObjectInstance* Instance, EUpdateResult Upd
 
 
 /** Update the given Instance Skeletal Meshes and call its callbacks. */
-void UpdateSkeletalMesh(UCustomizableObjectInstance& CustomizableObjectInstance, const FDescriptorRuntimeHash& UpdatedDescriptorRuntimeHash, EUpdateResult UpdateResult, FInstanceUpdateDelegate* UpdateCallback)
+void UpdateSkeletalMesh(const TSharedRef<FUpdateContextPrivate>& Context)
 {
 	MUTABLE_CPUPROFILER_SCOPE(UpdateSkeletalMesh);
 
 	check(IsInGameThread());
 
-	for (int32 ComponentIndex = 0; ComponentIndex < CustomizableObjectInstance.SkeletalMeshes.Num(); ++ComponentIndex)
+	UCustomizableObjectInstance* CustomizableObjectInstance = Context->Instance.Get();
+	check(CustomizableObjectInstance);
+	
+	for (int32 ComponentIndex = 0; ComponentIndex < CustomizableObjectInstance->SkeletalMeshes.Num(); ++ComponentIndex)
 	{
-		if (TObjectPtr<USkeletalMesh> SkeletalMesh = CustomizableObjectInstance.SkeletalMeshes[ComponentIndex])
+		if (TObjectPtr<USkeletalMesh> SkeletalMesh = CustomizableObjectInstance->SkeletalMeshes[ComponentIndex])
 		{
 #if WITH_EDITOR
 			UCustomizableInstancePrivateData::RegenerateImportedModel(SkeletalMesh);
@@ -772,7 +858,7 @@ void UpdateSkeletalMesh(UCustomizableObjectInstance& CustomizableObjectInstance,
 		}
 	}
 
-	UCustomizableInstancePrivateData* CustomizableObjectInstancePrivateData = CustomizableObjectInstance.GetPrivate();
+	UCustomizableInstancePrivateData* CustomizableObjectInstancePrivateData = CustomizableObjectInstance->GetPrivate();
 	check(CustomizableObjectInstancePrivateData != nullptr);
 	for (TObjectIterator<UCustomizableSkeletalComponent> It; It; ++It)
 	{
@@ -786,20 +872,20 @@ void UpdateSkeletalMesh(UCustomizableObjectInstance& CustomizableObjectInstance,
 #endif
 
 		if (CustomizableSkeletalComponent &&
-			(CustomizableSkeletalComponent->CustomizableObjectInstance == &CustomizableObjectInstance) &&
-			CustomizableObjectInstance.SkeletalMeshes.IsValidIndex(CustomizableSkeletalComponent->ComponentIndex)
+			(CustomizableSkeletalComponent->CustomizableObjectInstance == CustomizableObjectInstance) &&
+			CustomizableObjectInstance->SkeletalMeshes.IsValidIndex(CustomizableSkeletalComponent->ComponentIndex)
 		   )
 		{
 			MUTABLE_CPUPROFILER_SCOPE(UpdateSkeletalMesh_SetSkeletalMesh);
 
 			const bool bIsCreatingSkeletalMesh = CustomizableObjectInstancePrivateData->HasCOInstanceFlags(CreatingSkeletalMesh); //TODO MTBL-391: Review
-			CustomizableSkeletalComponent->SetSkeletalMesh(CustomizableObjectInstance.SkeletalMeshes[CustomizableSkeletalComponent->ComponentIndex], false, bIsCreatingSkeletalMesh);
+			CustomizableSkeletalComponent->SetSkeletalMesh(CustomizableObjectInstance->SkeletalMeshes[CustomizableSkeletalComponent->ComponentIndex], false, bIsCreatingSkeletalMesh);
 
 			if (CustomizableObjectInstancePrivateData->HasCOInstanceFlags(ReplacePhysicsAssets))
 			{
 				CustomizableSkeletalComponent->SetPhysicsAsset(
-					CustomizableObjectInstance.SkeletalMeshes[CustomizableSkeletalComponent->ComponentIndex] ? 
-					CustomizableObjectInstance.SkeletalMeshes[CustomizableSkeletalComponent->ComponentIndex]->GetPhysicsAsset() : nullptr);
+					CustomizableObjectInstance->SkeletalMeshes[CustomizableSkeletalComponent->ComponentIndex] ? 
+					CustomizableObjectInstance->SkeletalMeshes[CustomizableSkeletalComponent->ComponentIndex]->GetPhysicsAsset() : nullptr);
 			}
 		}
 	}
@@ -807,9 +893,9 @@ void UpdateSkeletalMesh(UCustomizableObjectInstance& CustomizableObjectInstance,
 	CustomizableObjectInstancePrivateData->SetCOInstanceFlags(Generated);
 	CustomizableObjectInstancePrivateData->ClearCOInstanceFlags(CreatingSkeletalMesh);
 
-	CustomizableObjectInstance.bEditorPropertyChanged = false;
+	CustomizableObjectInstance->bEditorPropertyChanged = false;
 
-	FinishUpdateGlobal(&CustomizableObjectInstance, UpdateResult, UpdateCallback, UpdatedDescriptorRuntimeHash);
+	FinishUpdateGlobal(Context);
 }
 
 
@@ -996,68 +1082,74 @@ EQueuePriorityType FCustomizableObjectSystemPrivate::GetUpdatePriority(const UCu
 }
 
 
-void FCustomizableObjectSystemPrivate::EnqueueUpdateSkeletalMesh(UCustomizableObjectInstance& Instance, bool bOnlyUpdateIfNotGenerated, bool bIgnoreCloseDist, bool bForceHighPriority, const EUpdateRequired* OptionalUpdateRequired, FInstanceUpdateDelegate* UpdateCallback)
+void FCustomizableObjectSystemPrivate::EnqueueUpdateSkeletalMesh(const TSharedRef<FUpdateContextPrivate>& Context)
 {
 	MUTABLE_CPUPROFILER_SCOPE(FCustomizableObjectSystemPrivate::EnqueueUpdateSkeletalMesh);
 	check(IsInGameThread());
 
-	UCustomizableInstancePrivateData* InstancePrivate = Instance.GetPrivate();
+	UCustomizableObjectInstance* Instance = Context->Instance.Get();
+	check(Instance);
+	
+	UCustomizableInstancePrivateData* InstancePrivate = Instance->GetPrivate();
 	
-	if (!Instance.CanUpdateInstance())
+	if (!Instance->CanUpdateInstance())
 	{
-		FinishUpdateGlobal(&Instance, EUpdateResult::Error, UpdateCallback);
+		Context->UpdateResult = EUpdateResult::Error;
+		FinishUpdateGlobal(Context);
 		return;
 	}
 
 	UCustomizableObjectSystem* System = UCustomizableObjectSystem::GetInstance();
 
 	
-	const EUpdateRequired UpdateRequired = IsUpdateRequired(Instance, bOnlyUpdateIfNotGenerated, bIgnoreCloseDist);
+	const EUpdateRequired UpdateRequired = IsUpdateRequired(*Instance, Context->bOnlyUpdateIfNotGenerated, Context->bIgnoreCloseDist);
 	switch (UpdateRequired)
 	{
 	case EUpdateRequired::NoUpdate:
 	{	
-		FinishUpdateGlobal(&Instance, EUpdateResult::Error, UpdateCallback);
+		Context->UpdateResult = EUpdateResult::Error;
+		FinishUpdateGlobal(Context);
 		break;
 	}		
 	case EUpdateRequired::Update:
 	{
-		EQueuePriorityType Priority = GetUpdatePriority(Instance, bForceHighPriority);
+		EQueuePriorityType Priority = GetUpdatePriority(*Instance, Context->bForceHighPriority);
 
-		const uint32 InstanceId = Instance.GetUniqueID();
+		const uint32 InstanceId = Instance->GetUniqueID();
 		const float Distance = FMath::Sqrt(InstancePrivate->LastMinSquareDistFromComponentToPlayer);
 		const bool bIsPlayerOrNearIt = InstancePrivate->HasCOInstanceFlags(UsedByPlayerOrNearIt);
 		UE_LOG(LogMutable, Log, TEXT("Enqueued UpdateSkeletalMesh Async. of Instance %d with priority %d at dist %f bIsPlayerOrNearIt=%d, frame=%d"), InstanceId, static_cast<int32>(Priority), Distance, bIsPlayerOrNearIt, GFrameNumber);				
 
 		if (InstancePrivate->HasCOInstanceFlags(PendingLODsUpdate))
 		{
-			UE_LOG(LogMutable, Verbose, TEXT("LOD change: %d, %d -> %d, %d"), Instance.GetCurrentMinLOD(), Instance.GetCurrentMaxLOD(), Instance.GetMinLODToLoad(), Instance.GetMaxLODToLoad());
+			UE_LOG(LogMutable, Verbose, TEXT("LOD change: %d, %d -> %d, %d"), Instance->GetCurrentMinLOD(), Instance->GetCurrentMaxLOD(), Instance->GetMinLODToLoad(), Instance->GetMaxLODToLoad());
 		}
 
-		InstancePrivate->UpdateDescriptorRuntimeHash = FDescriptorRuntimeHash(Instance.GetDescriptor());
-
-		if (const FMutablePendingInstanceUpdate* QueueElem = MutablePendingInstanceWork.GetUpdate(&Instance))
+		if (const FMutablePendingInstanceUpdate* QueueElem = MutablePendingInstanceWork.GetUpdate(Instance))
 		{
-			if (InstancePrivate->UpdateDescriptorRuntimeHash.IsSubset(FDescriptorRuntimeHash(QueueElem->InstanceDescriptor)))
+			if (InstancePrivate->UpdateDescriptorRuntimeHash.IsSubset(FDescriptorRuntimeHash(QueueElem->Context->InstanceDescriptorRuntimeHash)))
 			{
-				FinishUpdateGlobal(&Instance, EUpdateResult::ErrorOptimized, UpdateCallback);			
+				Context->UpdateResult = EUpdateResult::ErrorOptimized;
+				FinishUpdateGlobal(Context);			
 				return; // The the requested update is equal to the last enqueued update.
 			}
 		}	
 
 		if (CurrentMutableOperation &&
-			&Instance == CurrentMutableOperation->CustomizableObjectInstance &&
+			Instance == CurrentMutableOperation->Instance &&
 			InstancePrivate->UpdateDescriptorRuntimeHash.IsSubset(CurrentMutableOperation->InstanceDescriptorRuntimeHash))
 		{
-			FinishUpdateGlobal(&Instance, EUpdateResult::ErrorOptimized, UpdateCallback);			
+			Context->UpdateResult = EUpdateResult::ErrorOptimized;
+			FinishUpdateGlobal(Context);
 			return; // The requested update is equal to the running update.
 		}
 	
 		if (InstancePrivate->UpdateDescriptorRuntimeHash.IsSubset(InstancePrivate->DescriptorRuntimeHash) &&
 			!(CurrentMutableOperation &&
-			&Instance == CurrentMutableOperation->CustomizableObjectInstance)) // This condition is necessary because even if the descriptor is a subset, it will be replaced by the CurrentMutableOperation
+			Instance == CurrentMutableOperation->Instance)) // This condition is necessary because even if the descriptor is a subset, it will be replaced by the CurrentMutableOperation
 		{
-			UpdateSkeletalMesh(Instance, Instance.GetDescriptorRuntimeHash(), EUpdateResult::Success, UpdateCallback);
+			Context->UpdateResult = EUpdateResult::Success;
+			UpdateSkeletalMesh(Context);
 		}
 		else
 		{
@@ -1065,7 +1157,7 @@ void FCustomizableObjectSystemPrivate::EnqueueUpdateSkeletalMesh(UCustomizableOb
 			check(ImageProvider);
 
 			// Cache new Texture Parameters
-			for (const FCustomizableObjectTextureParameterValue& TextureParameters : Instance.GetDescriptor().GetTextureParameters())
+			for (const FCustomizableObjectTextureParameterValue& TextureParameters : Instance->GetDescriptor().GetTextureParameters())
 			{
 				ImageProvider->CacheImage(TextureParameters.ParameterValue, false);
 
@@ -1083,7 +1175,7 @@ void FCustomizableObjectSystemPrivate::EnqueueUpdateSkeletalMesh(UCustomizableOb
 
 			// Update which ones are currently are being used
 			InstancePrivate->UpdateTextureParameters.Reset();
-			for (const FCustomizableObjectTextureParameterValue& TextureParameters : Instance.GetDescriptor().GetTextureParameters())
+			for (const FCustomizableObjectTextureParameterValue& TextureParameters : Instance->GetDescriptor().GetTextureParameters())
 			{
 				InstancePrivate->UpdateTextureParameters.Add(TextureParameters.ParameterValue);
 
@@ -1093,7 +1185,7 @@ void FCustomizableObjectSystemPrivate::EnqueueUpdateSkeletalMesh(UCustomizableOb
 				}
 			}
 
-			const FMutablePendingInstanceUpdate InstanceUpdate(&Instance, Priority, UpdateCallback);
+			const FMutablePendingInstanceUpdate InstanceUpdate(Context);
 			MutablePendingInstanceWork.AddUpdate(InstanceUpdate);
 		}
 
@@ -1102,9 +1194,10 @@ void FCustomizableObjectSystemPrivate::EnqueueUpdateSkeletalMesh(UCustomizableOb
 
 	case EUpdateRequired::Discard:
 	{
-		System->GetPrivate()->InitDiscardResourcesSkeletalMesh(&Instance);
-		
-		FinishUpdateGlobal(&Instance, EUpdateResult::ErrorDiscarded, UpdateCallback);
+		System->GetPrivate()->InitDiscardResourcesSkeletalMesh(Instance);
+
+		Context->UpdateResult = EUpdateResult::ErrorDiscarded;
+		FinishUpdateGlobal(Context);
 		break;
 	}
 
@@ -1418,12 +1511,11 @@ void FCustomizableObjectSystemPrivate::UpdateMemoryLimit()
 namespace impl
 {
 
-	void Subtask_Mutable_UpdateParameterRelevancy(const TSharedPtr<FMutableOperationData>& OperationData)
+	void Subtask_Mutable_UpdateParameterRelevancy(const TSharedRef<FUpdateContextPrivate>& OperationData)
 	{
 		MUTABLE_CPUPROFILER_SCOPE(Subtask_Mutable_UpdateParameterRelevancy)
 
-		check(OperationData);
-		check(OperationData->MutableParameters);
+		check(OperationData->Parameters);
 		check(OperationData->InstanceID != 0);
 
 		OperationData->RelevantParametersInProgress.Empty();
@@ -1437,11 +1529,11 @@ namespace impl
 		{
 			MUTABLE_CPUPROFILER_SCOPE(ParameterRelevancy)
 
-			const int32 NumParameters = OperationData->MutableParameters->GetCount();
+			const int32 NumParameters = OperationData->Parameters->GetCount();
 
 			TArray<bool> Relevant;
 			Relevant.SetNumZeroed(NumParameters);
-			MutableSystem->GetParameterRelevancy(OperationData->InstanceID, OperationData->MutableParameters, Relevant.GetData());
+			MutableSystem->GetParameterRelevancy(OperationData->InstanceID, OperationData->Parameters, Relevant.GetData());
 
 			for (int32 ParamIndex = 0; ParamIndex < NumParameters; ++ParamIndex)
 			{
@@ -1455,12 +1547,11 @@ namespace impl
 
 
 	// This runs in the mutable thread.
-	void Subtask_Mutable_BeginUpdate_GetMesh(const TSharedPtr<FMutableOperationData>& OperationData, TSharedPtr<mu::Model> Model)
+	void Subtask_Mutable_BeginUpdate_GetMesh(const TSharedRef<FUpdateContextPrivate>& OperationData, TSharedPtr<mu::Model> Model)
 	{
 		MUTABLE_CPUPROFILER_SCOPE(Subtask_Mutable_BeginUpdate_GetMesh)
 
-		check(OperationData->MutableParameters);
-		check(OperationData);
+		check(OperationData->Parameters);
 		OperationData->InstanceUpdateData.Clear();
 
 		check(UCustomizableObjectSystem::GetInstance() != nullptr);
@@ -1505,7 +1596,7 @@ namespace impl
 
 		// Main instance generation step
 		// LOD mask, set to all ones to build  all LODs
-		const mu::Instance* Instance = System->BeginUpdate(OperationData->InstanceID, OperationData->MutableParameters, OperationData->State, mu::System::AllLODs);
+		const mu::Instance* Instance = System->BeginUpdate(OperationData->InstanceID, OperationData->Parameters, OperationData->State, mu::System::AllLODs);
 		if (!Instance)
 		{
 			UE_LOG(LogMutable, Warning, TEXT("An Instace update has failed."));
@@ -1748,12 +1839,10 @@ namespace impl
 
 
 	// This runs in the mutable thread.
-	void Subtask_Mutable_GetImages(const TSharedPtr<FMutableOperationData>& OperationData)
+	void Subtask_Mutable_GetImages(const TSharedRef<FUpdateContextPrivate>& OperationData)
 	{
 		MUTABLE_CPUPROFILER_SCOPE(Subtask_Mutable_GetImages)
 
-		check(OperationData);
-
 		const FCustomizableObjectSystemPrivate* CustomizableObjectSystemPrivateData = UCustomizableObjectSystem::GetInstanceChecked()->GetPrivateChecked();
 		mu::System* System = CustomizableObjectSystemPrivateData->MutableSystem.get();
 		check(System != nullptr);
@@ -1867,11 +1956,10 @@ namespace impl
 	
 
 	// This runs in a worker thread
-	void Subtask_Mutable_PrepareTextures(const TSharedPtr<FMutableOperationData>& OperationData)
+	void Subtask_Mutable_PrepareTextures(const TSharedRef<FUpdateContextPrivate>& OperationData)
 	{
 		MUTABLE_CPUPROFILER_SCOPE(Subtask_Mutable_PrepareTextures)
 
-		check(OperationData);
 		for (const FInstanceUpdateData::FSurface& Surface : OperationData->InstanceUpdateData.Surfaces)
 		{
 			for (int32 ImageIndex = 0; ImageIndex<Surface.ImageCount; ++ImageIndex)
@@ -1906,7 +1994,7 @@ namespace impl
 	
 
 	// This runs in a worker thread
-	void Subtask_Mutable_PrepareSkeletonData(const TSharedPtr<FMutableOperationData>& OperationData)
+	void Subtask_Mutable_PrepareSkeletonData(const TSharedRef<FUpdateContextPrivate>& OperationData)
 	{
 		MUTABLE_CPUPROFILER_SCOPE(Subtask_Mutable_PrepareSkeletonData)
 
@@ -1965,7 +2053,6 @@ namespace impl
 		}
 
 
-		check(OperationData);
 		const int32 LODCount = OperationData->InstanceUpdateData.LODs.Num();
 		const FInstanceUpdateData::FLOD& MinLOD = OperationData->InstanceUpdateData.LODs[OperationData->CurrentMinLOD];
 		const int32 ComponentCount = MinLOD.ComponentCount;
@@ -2045,12 +2132,10 @@ namespace impl
 
 
 	// This runs in a worker thread.
-	void Task_Mutable_Update_GetMesh(const TSharedPtr<FMutableOperationData>& OperationData, const TSharedPtr<mu::Model>& Model)
+	void Task_Mutable_Update_GetMesh(const TSharedRef<FUpdateContextPrivate>& OperationData, const TSharedPtr<mu::Model>& Model)
 	{
 		MUTABLE_CPUPROFILER_SCOPE(Task_Mutable_Update_GetMesh)
 
-		check(OperationData);
-
 #if WITH_EDITOR
 		const uint32 StartCycles = FPlatformTime::Cycles();
 #endif
@@ -2077,12 +2162,10 @@ namespace impl
 
 	
 	// This runs in a worker thread.
-	void Task_Mutable_Update_GetImages(const TSharedPtr<FMutableOperationData>& OperationData)
+	void Task_Mutable_Update_GetImages(const TSharedRef<FUpdateContextPrivate>& OperationData)
 	{
 		MUTABLE_CPUPROFILER_SCOPE(Task_Mutable_GetImages)
 
-		check(OperationData);
-
 #if WITH_EDITOR
 		uint32 StartCycles = FPlatformTime::Cycles();
 #endif
@@ -2100,11 +2183,10 @@ namespace impl
 
 
 	// This runs in a worker thread.
-	void Task_Mutable_ReleaseInstance(const TSharedPtr<FMutableOperationData>& OperationData, mu::Ptr<mu::System> MutableSystem)
+	void Task_Mutable_ReleaseInstance(const TSharedRef<FUpdateContextPrivate>& OperationData, mu::Ptr<mu::System> MutableSystem)
 	{
 		MUTABLE_CPUPROFILER_SCOPE(Task_Mutable_ReleaseInstance)
 
-		check(OperationData);
 		check(MutableSystem);
 
 		if (OperationData->InstanceID > 0)
@@ -2132,16 +2214,13 @@ namespace impl
 
 
 	// This runs in a worker thread.
-	void Task_Mutable_ReleaseInstanceID(const TSharedPtr<FMutableOperationData>& OperationData, const mu::Ptr<mu::System>& MutableSystem)
+	void Task_Mutable_ReleaseInstanceID(const mu::Instance::ID InstanceID, const mu::Ptr<mu::System>& MutableSystem)
 	{
 		MUTABLE_CPUPROFILER_SCOPE(Task_Mutable_ReleaseInstanceID)
 
-		check(OperationData);
-
-		if (OperationData->InstanceID > 0)
+		if (InstanceID > 0)
 		{
-			MutableSystem->ReleaseInstance(OperationData->InstanceID);
-			OperationData->InstanceID = 0;
+			MutableSystem->ReleaseInstance(InstanceID);
 		}
 
 		if (CVarClearWorkingMemoryOnUpdateEnd.GetValueOnAnyThread())
@@ -2166,17 +2245,17 @@ namespace impl
 	}
 
 	
-	void Task_Game_Callbacks(const TSharedPtr<FMutableOperationData>& OperationData)
+	void Task_Game_Callbacks(const TSharedRef<FUpdateContextPrivate>& OperationData)
 	{
 		MUTABLE_CPUPROFILER_SCOPE(Task_Game_Callbacks)
 
 		check(IsInGameThread());
-		check(OperationData);
 
 		UCustomizableObjectSystem* System = UCustomizableObjectSystem::GetInstance();
 		if (!System || !System->IsValidLowLevel() || System->HasAnyFlags(RF_BeginDestroyed))
 		{
-			FinishUpdateGlobal(OperationData->Instance.Get(), EUpdateResult::Error, &OperationData->UpdateCallback);
+			OperationData->UpdateResult = EUpdateResult::Error;
+			FinishUpdateGlobal(OperationData);
 			return;
 		}
 
@@ -2186,7 +2265,8 @@ namespace impl
 		if (!CustomizableObjectInstance || !CustomizableObjectInstance->IsValidLowLevel() )
 		{
 			System->ClearCurrentMutableOperation();
-			FinishUpdateGlobal(CustomizableObjectInstance, EUpdateResult::Error, &OperationData->UpdateCallback);
+			OperationData->UpdateResult = EUpdateResult::Error;
+			FinishUpdateGlobal(OperationData);
 			return;
 		}
 
@@ -2194,7 +2274,7 @@ namespace impl
 
 		// Actual work
 		// TODO MTBL-391: Review This hotfix
-		UpdateSkeletalMesh(*CustomizableObjectInstance, CustomizableObjectSystemPrivateData->CurrentMutableOperation->InstanceDescriptorRuntimeHash, OperationData->UpdateResult, &OperationData->UpdateCallback);
+		UpdateSkeletalMesh(OperationData);
 
 		// All work is done, release unused textures.
 		if (CustomizableObjectSystemPrivateData->bReleaseTexturesImmediately)
@@ -2225,17 +2305,17 @@ namespace impl
 	}
 
 
-	void Task_Game_ConvertResources(TSharedPtr<FMutableOperationData> OperationData)
+	void Task_Game_ConvertResources(const TSharedRef<FUpdateContextPrivate>& OperationData)
 	{
 		MUTABLE_CPUPROFILER_SCOPE(Task_Game_ConvertResources)
 
 		check(IsInGameThread());
-		check(OperationData);
 
 		UCustomizableObjectSystem* System = UCustomizableObjectSystem::GetInstance();
 		if (!System || !System->IsValidLowLevel() || System->HasAnyFlags(RF_BeginDestroyed))
 		{
-			FinishUpdateGlobal(OperationData->Instance.Get(), EUpdateResult::Error, &OperationData->UpdateCallback);
+			OperationData->UpdateResult = EUpdateResult::Error;
+			FinishUpdateGlobal(OperationData);
 			return;
 		}
 
@@ -2365,18 +2445,18 @@ namespace impl
 		}
 		else
 		{
-			FinishUpdateGlobal(CustomizableObjectInstance, EUpdateResult::Error, &OperationData->UpdateCallback);
+			OperationData->UpdateResult = EUpdateResult::Error;
+			FinishUpdateGlobal(OperationData);
 		}
 	}
 
 
 	/** Lock Cached Resources. */
-	void Task_Game_LockCache(TSharedPtr<FMutableOperationData> OperationData)
+	void Task_Game_LockCache(const TSharedRef<FUpdateContextPrivate>& OperationData)
 	{
 		MUTABLE_CPUPROFILER_SCOPE(Task_Game_LockCache)
 
 		check(IsInGameThread());
-		check(OperationData);
 
 		UCustomizableObjectSystem* System = UCustomizableObjectSystem::GetInstance();
 		if (!System)
@@ -2388,7 +2468,9 @@ namespace impl
 		if (!ObjectInstance)
 		{
 			System->ClearCurrentMutableOperation();
-			FinishUpdateGlobal(ObjectInstance, EUpdateResult::Error, &OperationData->UpdateCallback);
+			
+			OperationData->UpdateResult = EUpdateResult::Error;
+			FinishUpdateGlobal(OperationData);
 			return;
 		}
 
@@ -2410,7 +2492,9 @@ namespace impl
 		if (!CustomizableObject)
 		{
 			System->ClearCurrentMutableOperation();
-			FinishUpdateGlobal(ObjectInstance, EUpdateResult::Error, &OperationData->UpdateCallback);
+			
+			OperationData->UpdateResult = EUpdateResult::Error;
+			FinishUpdateGlobal(OperationData);
 			return;
 		}
 		
@@ -2515,48 +2599,46 @@ namespace impl
 
 		// Task: Release Instance ID
 		//-------------------------------------------------------------
-		TSharedPtr<FMutableOperationData> CurrentOperationData = MakeShared<FMutableOperationData>();
-		check(CurrentOperationData);
-		CurrentOperationData->InstanceID = IDToRelease;
-
 		{
 			// Task inputs
 			SystemPrivateData->MutableTaskGraph.AddMutableThreadTask(
 				TEXT("Task_Mutable_ReleaseInstanceID"),
-				[CurrentOperationData, MutableSystem]()
+				[IDToRelease, MutableSystem]()
 				{
-					impl::Task_Mutable_ReleaseInstanceID(CurrentOperationData, MutableSystem);
+					impl::Task_Mutable_ReleaseInstanceID(IDToRelease, MutableSystem);
 				});
 		}
 	}
 
 
 	/** "Start Update" */
-	void Task_Game_StartUpdate(TSharedPtr<FMutableOperation> Operation)
+	void Task_Game_StartUpdate(const TSharedRef<FUpdateContextPrivate>& Operation)
 	{
 		MUTABLE_CPUPROFILER_SCOPE(Task_Game_StartUpdate)
 
-		check(Operation);
-		
 		UCustomizableObjectSystem::GetInstance()->GetPrivate()->MutableTaskGraph.AllowLaunchingMutableTaskLowPriority(false, false);
 		
 		UCustomizableObjectSystem* System = UCustomizableObjectSystem::GetInstance();
 		check(System != nullptr);
 
-		if (!Operation->CustomizableObjectInstance.IsValid() || !Operation->CustomizableObjectInstance->IsValidLowLevel()) // Only start if it hasn't been already destroyed (i.e. GC after finish PIE)
+		if (!Operation->Instance.IsValid() || !Operation->Instance->IsValidLowLevel()) // Only start if it hasn't been already destroyed (i.e. GC after finish PIE)
 		{
 			System->ClearCurrentMutableOperation();
-			FinishUpdateGlobal(nullptr, EUpdateResult::Error, &Operation->UpdateCallback);
+
+			Operation->UpdateResult = EUpdateResult::Error;
+			FinishUpdateGlobal(Operation);
 			return;
 		}
 
-		UCustomizableObjectInstance* CandidateInstance = Operation->CustomizableObjectInstance.Get();
+		UCustomizableObjectInstance* CandidateInstance = Operation->Instance.Get();
 		
 		UCustomizableInstancePrivateData* CandidateInstancePrivateData = CandidateInstance->GetPrivate();
 		if (!CandidateInstancePrivateData)
 		{
 			System->ClearCurrentMutableOperation();
-			FinishUpdateGlobal(nullptr, EUpdateResult::Error, &Operation->UpdateCallback);
+			
+			Operation->UpdateResult = EUpdateResult::Error;
+			FinishUpdateGlobal(Operation);
 			return;
 		}
 
@@ -2571,7 +2653,9 @@ namespace impl
 		if (Operation->InstanceDescriptorRuntimeHash.IsSubset(CandidateInstance->GetDescriptorRuntimeHash()))
 		{
 			System->ClearCurrentMutableOperation();
-			UpdateSkeletalMesh(*CandidateInstance, CandidateInstance->GetDescriptorRuntimeHash(), EUpdateResult::ErrorOptimized, &Operation->UpdateCallback);
+
+			Operation->UpdateResult = EUpdateResult::ErrorOptimized;
+			UpdateSkeletalMesh(Operation);
 			return;
 		}
 
@@ -2602,7 +2686,7 @@ namespace impl
 			bCancel = true;
 		}
 
-		mu::Ptr<const mu::Parameters> Parameters = Operation->GetParameters();
+		mu::Ptr<const mu::Parameters> Parameters = Operation->Parameters;
 		if (!Parameters)
 		{
 			bCancel = true;
@@ -2612,7 +2696,8 @@ namespace impl
 		{
 			System->ClearCurrentMutableOperation();
 
-			FinishUpdateGlobal(CandidateInstance, EUpdateResult::Error, &Operation->UpdateCallback);
+			Operation->UpdateResult = EUpdateResult::Error;
+			FinishUpdateGlobal(Operation);
 			return;
 		}
 
@@ -2637,39 +2722,39 @@ namespace impl
 		FString StateName = CandidateInstance->GetCustomizableObject()->GetStateName(CandidateInstance->GetState());
 		const FParameterUIData* StateData = CandidateInstance->GetCustomizableObject()->StateUIDataMap.Find(StateName);
 
-		bool bLiveUpdateMode = false;
+		Operation->bLiveUpdateMode = false;
 
 		if (SystemPrivateData->EnableMutableLiveUpdate)
 		{
-			bLiveUpdateMode = StateData ? StateData->bLiveUpdateMode : false;
+			Operation->bLiveUpdateMode = StateData ? StateData->bLiveUpdateMode : false;
 		}
 
-		bool bNeverStream = false;
-		int32 MipsToSkip = 0;
+		Operation->bNeverStream = false;
+		Operation->MipsToSkip = 0;
 
-		SystemPrivateData->GetMipStreamingConfig(*CandidateInstance, bNeverStream, MipsToSkip);
+		SystemPrivateData->GetMipStreamingConfig(*CandidateInstance, Operation->bNeverStream, Operation->MipsToSkip);
 		
-		if (bLiveUpdateMode && (!bNeverStream || MipsToSkip > 0))
+		if (Operation->bLiveUpdateMode && (!Operation->bNeverStream || Operation->MipsToSkip > 0))
 		{
 			UE_LOG(LogMutable, Warning, TEXT("Instance LiveUpdateMode does not yet support progressive streaming of Mutable textures. Disabling LiveUpdateMode for this update."));
-			bLiveUpdateMode = false;
+			Operation->bLiveUpdateMode = false;
 		}
 
-		bool bReuseInstanceTextures = false;
+		Operation->bReuseInstanceTextures = false;
 
 		if (SystemPrivateData->EnableReuseInstanceTextures)
 		{
-			bReuseInstanceTextures = StateData ? StateData->bReuseInstanceTextures : false;
-			bReuseInstanceTextures |= CandidateInstancePrivateData->HasCOInstanceFlags(ReuseTextures);
+			Operation->bReuseInstanceTextures = StateData ? StateData->bReuseInstanceTextures : false;
+			Operation->bReuseInstanceTextures |= CandidateInstancePrivateData->HasCOInstanceFlags(ReuseTextures);
 			
-			if (bReuseInstanceTextures && !bNeverStream)
+			if (Operation->bReuseInstanceTextures && !Operation->bNeverStream)
 			{
 				UE_LOG(LogMutable, Warning, TEXT("Instance texture reuse requires that the current Mutable state is in non-streaming mode. Change it in the Mutable graph base node in the state definition."));
-				bReuseInstanceTextures = false;
+				Operation->bReuseInstanceTextures = false;
 			}
 		}
 
-		if (!bLiveUpdateMode && CandidateInstancePrivateData->LiveUpdateModeInstanceID != 0)
+		if (!Operation->bLiveUpdateMode && CandidateInstancePrivateData->LiveUpdateModeInstanceID != 0)
 		{
 			// The instance was in live update mode last update, but now it's not. So the Id and resources have to be released.
 			// Enqueue a new mutable task to release them
@@ -2680,30 +2765,17 @@ namespace impl
 		
 		// Task: Mutable Update and GetMesh
 		//-------------------------------------------------------------
-		TSharedPtr<FMutableOperationData> CurrentOperationData = MakeShared<FMutableOperationData>();
-		check(CurrentOperationData);
-		CurrentOperationData->Instance = Operation->CustomizableObjectInstance;
-		CurrentOperationData->TextureCoverageQueries_MutableThreadParams = CandidateInstancePrivateData->TextureCoverageQueries;
-		CurrentOperationData->TextureCoverageQueries_MutableThreadResults.Empty();
-		CurrentOperationData->CurrentMinLOD = Operation->InstanceDescriptorRuntimeHash.GetMinLOD();
-		CurrentOperationData->CurrentMaxLOD = Operation->InstanceDescriptorRuntimeHash.GetMaxLOD();
-		CurrentOperationData->bNeverStream = bNeverStream;
-		CurrentOperationData->bLiveUpdateMode = bLiveUpdateMode;
-		CurrentOperationData->bReuseInstanceTextures = bReuseInstanceTextures;
-		CurrentOperationData->InstanceID = bLiveUpdateMode ? CandidateInstancePrivateData->LiveUpdateModeInstanceID : 0;
-		CurrentOperationData->MipsToSkip = MipsToSkip;
-		CurrentOperationData->MutableParameters = Parameters;
-		CurrentOperationData->State = Operation->GetState();
-		CurrentOperationData->UpdateResult = EUpdateResult::Success;
-		CurrentOperationData->UpdateCallback = Operation->UpdateCallback;
-		CurrentOperationData->bBuildParameterRelevancy = Operation->IsBuildParameterRelevancy();
+		Operation->TextureCoverageQueries_MutableThreadParams = CandidateInstancePrivateData->TextureCoverageQueries;
+		Operation->CurrentMinLOD = Operation->InstanceDescriptorRuntimeHash.GetMinLOD();
+		Operation->CurrentMaxLOD = Operation->InstanceDescriptorRuntimeHash.GetMaxLOD();
+		Operation->InstanceID = Operation->bLiveUpdateMode ? CandidateInstancePrivateData->LiveUpdateModeInstanceID : 0;
 #if WITH_EDITOR
-		CurrentOperationData->PixelFormatOverride = SystemPrivateData->ImageFormatOverrideFunc;
+		Operation->PixelFormatOverride = SystemPrivateData->ImageFormatOverrideFunc;
 #endif
 
 		if (!CandidateInstancePrivateData->HasCOInstanceFlags(ForceGenerateMipTail))
 		{
-			CustomizableObject->GetLowPriorityTextureNames(CurrentOperationData->LowPriorityTextures);
+			CustomizableObject->GetLowPriorityTextureNames(Operation->LowPriorityTextures);
 		}
 
 		bool bIsInEditorViewport = false;
@@ -2741,7 +2813,7 @@ namespace impl
 		if (System->IsOnlyGenerateRequestedLODsEnabled() && System->CurrentInstanceLODManagement->IsOnlyGenerateRequestedLODLevelsEnabled() && 
 			!bIsInEditorViewport)
 		{
-			CurrentOperationData->RequestedLODs = Operation->InstanceDescriptorRuntimeHash.GetRequestedLODs();
+			Operation->RequestedLODs = Operation->InstanceDescriptorRuntimeHash.GetRequestedLODs();
 		}
 
 #ifdef MUTABLE_USE_NEW_TASKGRAPH
@@ -2755,9 +2827,9 @@ namespace impl
 
 			Mutable_GetMeshTask = SystemPrivateData->MutableTaskGraph.AddMutableThreadTask(
 				TEXT("Task_Mutable_Update_GetMesh"),
-				[CurrentOperationData, Model]()
+				[Operation, Model]()
 				{
-					impl::Task_Mutable_Update_GetMesh(CurrentOperationData, Model);
+					impl::Task_Mutable_Update_GetMesh(Operation, Model);
 				});
 		}
 
@@ -2769,9 +2841,9 @@ namespace impl
 			SystemPrivateData->AddGameThreadTask(
 				{
 				FMutableTaskDelegate::CreateLambda(
-					[CurrentOperationData]()
+					[Operation]()
 					{
-						impl::Task_Game_LockCache(CurrentOperationData);
+						impl::Task_Game_LockCache(Operation);
 					}),
 				Mutable_GetMeshTask
 				});
@@ -2818,7 +2890,7 @@ void UCustomizableObjectSystem::AdvanceCurrentOperation()
 		MUTABLE_CPUPROFILER_SCOPE(OperationUpdate);
 
 		// Start the first task of the update process. See namespace impl comments above.
-		impl::Task_Game_StartUpdate(Private->CurrentMutableOperation);
+		impl::Task_Game_StartUpdate(Private->CurrentMutableOperation.ToSharedRef());
 	}
 }
 
@@ -2915,13 +2987,13 @@ bool UCustomizableObjectSystem::Tick(float DeltaTime)
 			{
 				FMutablePendingInstanceUpdate& PendingUpdate = *Iterator;
 
-				if (PendingUpdate.CustomizableObjectInstance.IsValid())
+				if (PendingUpdate.Context->Instance.IsValid())
 				{
-					const EQueuePriorityType PriorityType = Private->GetUpdatePriority(*PendingUpdate.CustomizableObjectInstance, false);
+					const EQueuePriorityType PriorityType = Private->GetUpdatePriority(*PendingUpdate.Context->Instance, false);
 					
-					if (PendingUpdate.PriorityType <= MaxPriorityFound)
+					if (PendingUpdate.Context->PriorityType <= MaxPriorityFound)
 					{
-						const double MinSquareDistFromComponentToPlayer = PendingUpdate.CustomizableObjectInstance->GetPrivate()->MinSquareDistFromComponentToPlayer;
+						const double MinSquareDistFromComponentToPlayer = PendingUpdate.Context->Instance->GetPrivate()->MinSquareDistFromComponentToPlayer;
 						
 						if (MinSquareDistFromComponentToPlayer < MaxSquareDistanceFound ||
 							(MinSquareDistFromComponentToPlayer == MaxSquareDistanceFound && PendingUpdate.SecondsAtUpdate < MinTimeFound))
@@ -2982,7 +3054,7 @@ bool UCustomizableObjectSystem::Tick(float DeltaTime)
 			{
 				check(!LODUpdateCandidateFound);
 
-				UCustomizableObjectInstance* PendingInstance = PendingInstanceUpdateFound->CustomizableObjectInstance.Get();
+				UCustomizableObjectInstance* PendingInstance = PendingInstanceUpdateFound->Context->Instance.Get();
 				check(PendingInstance);
 
 				// Maybe there's a LODUpdate that has the same instance, merge both updates as an optimization
@@ -2990,23 +3062,21 @@ bool UCustomizableObjectSystem::Tick(float DeltaTime)
 				
 				if (LODUpdateWithSameInstance)
 				{
-					LODUpdateWithSameInstance->ApplyLODUpdateParamsToInstance();
+					LODUpdateWithSameInstance->ApplyLODUpdateParamsToInstance(&PendingInstanceUpdateFound->Context.Get());
 				}
 
-				const TSharedPtr<FMutableOperation> Operation = MakeShared<FMutableOperation>(FMutableOperation::CreateInstanceUpdate(*PendingInstance, PendingInstanceUpdateFound->Callback));
-				Private->StartUpdateSkeletalMesh(Operation);
 				TRACE_BEGIN_REGION(UE_MUTABLE_UPDATE_REGION);
-
-				Private->MutablePendingInstanceWork.RemoveUpdate(PendingInstanceUpdateFound->CustomizableObjectInstance);
+				Private->StartUpdateSkeletalMesh(PendingInstanceUpdateFound->Context);				
+				Private->MutablePendingInstanceWork.RemoveUpdate(PendingInstanceUpdateFound->Context->Instance);
 			}
 			else if (LODUpdateCandidateFound)
 			{
 				// Commit the LOD changes
 				LODUpdateCandidateFound->ApplyLODUpdateParamsToInstance();
 
-				const TSharedRef<FMutableOperation> Operation = MakeShared<FMutableOperation>(FMutableOperation::CreateInstanceUpdate(*LODUpdateCandidateFound->CustomizableObjectInstance, nullptr));
-				Private->StartUpdateSkeletalMesh(Operation);
 				TRACE_BEGIN_REGION(UE_MUTABLE_UPDATE_REGION);
+				const TSharedRef<FUpdateContextPrivate> Context = MakeShared<FUpdateContextPrivate>(*LODUpdateCandidateFound->CustomizableObjectInstance);
+				Private->StartUpdateSkeletalMesh(Context);
 			}
 		}
 
@@ -3069,7 +3139,7 @@ void UCustomizableObjectSystem::DiscardInstances()
 
 		UCustomizableObjectInstance* COI = Iterator->CustomizableObjectInstance.Get();
 
-		const bool bUpdating = Private->CurrentMutableOperation && Private->CurrentMutableOperation->CustomizableObjectInstance != Iterator->CustomizableObjectInstance;
+		const bool bUpdating = Private->CurrentMutableOperation && Private->CurrentMutableOperation->Instance != Iterator->CustomizableObjectInstance;
 		if (COI && !bUpdating)
 		{
 			UCustomizableInstancePrivateData* COIPrivateData = COI ? COI->GetPrivate() : nullptr;
@@ -3158,11 +3228,11 @@ void UCustomizableObjectSystem::UnregisterImageProvider(UCustomizableSystemImage
 	GetPrivateChecked()->GetImageProviderChecked()->ImageProviders.Remove(Provider);
 }
 
-
 static bool bRevertCacheTextureParameters = false;
 static FAutoConsoleVariableRef CVarRevertCacheTextureParameters(
 	TEXT("mutable.RevertCacheTextureParameters"), bRevertCacheTextureParameters,
-	TEXT("If true, FMutableOperation will not cache/uncache texture parameters. If false, FMutableOperation will add an additional reference to the TextureParameters being used in the update."));
+	TEXT("If true, FUpdateContextPrivate will not cache/uncache texture parameters. If false, FUpdateContextPrivate will add an additional reference to the TextureParameters being used in the update."));
+
 
 void CacheTexturesParameters(const TArray<FName>& TextureParameters)
 {
@@ -3204,74 +3274,6 @@ void UnCacheTexturesParameters(const TArray<FName>& TextureParameters)
 }
 
 
-FMutableOperation::FMutableOperation(const FMutableOperation& Other)
-{
-	CustomizableObjectInstance = Other.CustomizableObjectInstance;
-	InstanceDescriptorRuntimeHash = Other.InstanceDescriptorRuntimeHash;
-	bBuildParameterRelevancy = Other.bBuildParameterRelevancy;
-	Parameters = Other.Parameters;
-	TextureParameters = Other.TextureParameters;
-	UpdateCallback = Other.UpdateCallback;
-	State = Other.State;
-
-	CacheTexturesParameters(TextureParameters);
-}
-
-
-FMutableOperation& FMutableOperation::operator=(const FMutableOperation& Other)
-{
-	CustomizableObjectInstance = Other.CustomizableObjectInstance;
-	InstanceDescriptorRuntimeHash = Other.InstanceDescriptorRuntimeHash;
-	bBuildParameterRelevancy = Other.bBuildParameterRelevancy;
-	Parameters = Other.Parameters;
-	TextureParameters = Other.TextureParameters;
-	UpdateCallback = Other.UpdateCallback;
-
-	CacheTexturesParameters(TextureParameters);
-
-	return *this;
-}
-
-
-FMutableOperation::~FMutableOperation()
-{
-	// Uncache Texture Parameters
-	UnCacheTexturesParameters(TextureParameters);
-}
-
-
-FMutableOperation FMutableOperation::CreateInstanceUpdate(UCustomizableObjectInstance& InCustomizableObjectInstance, const FInstanceUpdateDelegate* UpdateCallback)
-{
-	check(InCustomizableObjectInstance.GetPrivate() != nullptr);
-	check(InCustomizableObjectInstance.GetCustomizableObject() != nullptr);
-
-	FMutableOperation Op;
-	Op.CustomizableObjectInstance = &InCustomizableObjectInstance;
-	Op.InstanceDescriptorRuntimeHash = InCustomizableObjectInstance.GetUpdateDescriptorRuntimeHash();
-	Op.State = InCustomizableObjectInstance.GetState();
-	Op.bBuildParameterRelevancy = InCustomizableObjectInstance.GetBuildParameterRelevancy();
-	Op.Parameters = InCustomizableObjectInstance.GetDescriptor().GetParameters();
-	Op.TextureParameters = InCustomizableObjectInstance.GetPrivate()->UpdateTextureParameters;
-
-	if (UpdateCallback)
-	{
-		Op.UpdateCallback = *UpdateCallback;		
-	}
-	
-	InCustomizableObjectInstance.GetCustomizableObject()->ApplyStateForcedValuesToParameters(Op.State, Op.Parameters.get());
-
-	if (!Op.Parameters)
-	{
-		// Cancel the update because the parameters aren't valid, probably because the object is not compiled
-		Op.CustomizableObjectInstance = nullptr;
-	}
-
-	CacheTexturesParameters(Op.TextureParameters);
-
-	return Op;
-}
-
-
 int32 UCustomizableObjectSystem::GetNumInstances() const
 {
 	return GetPrivateChecked()->MutableStats.NumInstances;
@@ -3647,18 +3649,18 @@ FUnrealMutableImageProvider* FCustomizableObjectSystemPrivate::GetImageProviderC
 }
 
 
-void FCustomizableObjectSystemPrivate::StartUpdateSkeletalMesh(const TSharedPtr<FMutableOperation>& Operation)
+void FCustomizableObjectSystemPrivate::StartUpdateSkeletalMesh(const TSharedRef<FUpdateContextPrivate>& Context)
 {
 	check(!CurrentMutableOperation); // Can not start an update if there is already another in progress
-	check(Operation->CustomizableObjectInstance.IsValid()) // The instance has to be alive to start the update
+	check(Context->Instance.IsValid()) // The instance has to be alive to start the update
 		
-	CurrentMutableOperation = Operation;
+	CurrentMutableOperation = Context;
 }
 
 
 bool FCustomizableObjectSystemPrivate::IsUpdating(const UCustomizableObjectInstance& Instance) const
 {
-	if (CurrentMutableOperation && CurrentMutableOperation->CustomizableObjectInstance.Get() == &Instance)
+	if (CurrentMutableOperation && CurrentMutableOperation->Instance.Get() == &Instance)
 	{
 		return true;
 	}
diff --git a/Engine/Plugins/Experimental/Mutable/Source/CustomizableObject/Private/MuCO/CustomizableObjectSystemPrivate.h b/Engine/Plugins/Experimental/Mutable/Source/CustomizableObject/Private/MuCO/CustomizableObjectSystemPrivate.h
--- a/Engine/Plugins/Experimental/Mutable/Source/CustomizableObject/Private/MuCO/CustomizableObjectSystemPrivate.h
+++ b/Engine/Plugins/Experimental/Mutable/Source/CustomizableObject/Private/MuCO/CustomizableObjectSystemPrivate.h
@@ -32,136 +32,31 @@ class UCustomizableObjectSystem;
 namespace LowLevelTasks { enum class ETaskPriority : int8; }
 struct FTexturePlatformData;
 
-//! An operation to be performed by Mutable. This could be creating or updating an instance, releasing resources, changing an LOD, etc.
-//! Operations may be done in several tasks in several thread or across frames.
-class FMutableOperation 
-{
-	/** Instance parameters at the time of the operation request. */
-	mu::ParametersPtr Parameters; 
-	
-	TArray<FName> TextureParameters;
-
-	bool bBuildParameterRelevancy = false;
-
-	/** Instance optimization state. */
-	int32 State = 0;
-	
-	FMutableOperation() = default;
-	
-public:
-	FMutableOperation(const FMutableOperation&);
-	FMutableOperation(FMutableOperation&&) = default;
-	FMutableOperation& operator=(const FMutableOperation&);
-	FMutableOperation& operator=(FMutableOperation&&) = default;
-
-	~FMutableOperation();
-
-	static FMutableOperation CreateInstanceUpdate(UCustomizableObjectInstance& COInstance, const FInstanceUpdateDelegate* UpdateCallback);
-
-	// Weak reference to the instance we are operating on.
-	// It is weak because we don't want to lock it in case it becomes irrelevant in the game while operations are pending and it needs to be destroyed.
-	TWeakObjectPtr<UCustomizableObjectInstance> CustomizableObjectInstance;
-
-	/** Hash of the UCustomizableObjectInstance::Descriptor at the time of the update request. */
-	FDescriptorRuntimeHash InstanceDescriptorRuntimeHash;
-
-	//! This is used to calculate stats.
-	double StartUpdateTime = 0.0;
-
-	FInstanceUpdateDelegate UpdateCallback;
-	
-	bool IsBuildParameterRelevancy() const
-	{
-		return bBuildParameterRelevancy;
-	}
-
-	/** Read-only access to the mutable instance parameters for this operation. */
-	mu::ParametersPtrConst GetParameters() const
-	{
-		return Parameters;
-	}
-
-	int32 GetState() const
-	{
-		return State;
-	}
-};
-
 
 struct FMutablePendingInstanceUpdate
 {
-	EQueuePriorityType PriorityType = EQueuePriorityType::Low;
-
-	TWeakObjectPtr<UCustomizableObjectInstance> CustomizableObjectInstance;
-	FCustomizableObjectInstanceDescriptor InstanceDescriptor;
+	TSharedRef<FUpdateContextPrivate> Context;
 
 	double SecondsAtUpdate = 0;
-	FInstanceUpdateDelegate* Callback = nullptr;
 
-	FMutablePendingInstanceUpdate(UCustomizableObjectInstance* InCustomizableObjectInstance)
-	{
-		check(InCustomizableObjectInstance);
-		CustomizableObjectInstance = InCustomizableObjectInstance;
-	}
+	FMutablePendingInstanceUpdate(const TSharedRef<FUpdateContextPrivate>& InContext);
 
-	FMutablePendingInstanceUpdate(UCustomizableObjectInstance* InCustomizableObjectInstance, EQueuePriorityType NewPriorityType, 
-		                          FInstanceUpdateDelegate* InCallback)
-	{
-		PriorityType = NewPriorityType;
+	bool operator==(const FMutablePendingInstanceUpdate& Other) const;
 
-		check(InCustomizableObjectInstance);
-		CustomizableObjectInstance = InCustomizableObjectInstance;
-		InstanceDescriptor = CustomizableObjectInstance->GetDescriptor();
-
-		SecondsAtUpdate = FPlatformTime::Seconds();
-		Callback = InCallback;
-	}
-
-	friend bool operator ==(const FMutablePendingInstanceUpdate& A, const FMutablePendingInstanceUpdate& B)
-	{
-		return A.CustomizableObjectInstance.HasSameIndexAndSerialNumber(B.CustomizableObjectInstance);
-	}
-
-	friend bool operator <(const FMutablePendingInstanceUpdate& A, const FMutablePendingInstanceUpdate& B)
-	{
-		if (A.PriorityType < B.PriorityType)
-		{
-			return true;
-		}
-		else if (A.PriorityType > B.PriorityType)
-		{
-			return false;
-		}
-		else
-		{
-			return A.SecondsAtUpdate < B.SecondsAtUpdate;
-		}
-	}
+	bool operator<(const FMutablePendingInstanceUpdate& Other) const;
 };
 
 
-inline uint32 GetTypeHash(const FMutablePendingInstanceUpdate& Update)
-{
-	return GetTypeHash(Update.CustomizableObjectInstance.GetWeakPtrTypeHash());
-}
+inline uint32 GetTypeHash(const FMutablePendingInstanceUpdate& Update);
 
 
 struct FPendingInstanceUpdateKeyFuncs : BaseKeyFuncs<FMutablePendingInstanceUpdate, TWeakObjectPtr<const UCustomizableObjectInstance>>
 {
-	FORCEINLINE static TWeakObjectPtr<const UCustomizableObjectInstance> GetSetKey(const FMutablePendingInstanceUpdate& PendingUpdate)
-	{
-		return PendingUpdate.CustomizableObjectInstance;
-	}
+	FORCEINLINE static TWeakObjectPtr<const UCustomizableObjectInstance> GetSetKey(const FMutablePendingInstanceUpdate& PendingUpdate);
 
-	FORCEINLINE static bool Matches(const TWeakObjectPtr<const UCustomizableObjectInstance>& A, const TWeakObjectPtr<const UCustomizableObjectInstance>& B)
-	{
-		return A.HasSameIndexAndSerialNumber(B);
-	}
+	FORCEINLINE static bool Matches(const TWeakObjectPtr<const UCustomizableObjectInstance>& A, const TWeakObjectPtr<const UCustomizableObjectInstance>& B);
 
-	FORCEINLINE static uint32 GetKeyHash(const TWeakObjectPtr<const UCustomizableObjectInstance>& Identifier)
-	{
-		return GetTypeHash(Identifier.GetWeakPtrTypeHash());
-	}
+	FORCEINLINE static uint32 GetKeyHash(const TWeakObjectPtr<const UCustomizableObjectInstance>& Identifier);
 };
 
 
@@ -501,21 +396,52 @@ struct FPendingTextureCoverageQuery
 };
 
 
-/** Runtime data used during a mutable instance update */
-struct FMutableOperationData
+/** Update Context.
+ *
+ * Alive from the start to the end of the update (both API and LOD update). */
+class FUpdateContextPrivate
 {
+public:
+	FUpdateContextPrivate(UCustomizableObjectInstance& InInstance);
+	~FUpdateContextPrivate();
+
+	EQueuePriorityType PriorityType = EQueuePriorityType::Low;
+	
+	FInstanceUpdateDelegate UpdateCallback;
+
+	/** Weak reference to the instance we are operating on.
+	 *It is weak because we don't want to lock it in case it becomes irrelevant in the game while operations are pending and it needs to be destroyed. */
 	TWeakObjectPtr<UCustomizableObjectInstance> Instance;
 
+	/** Hash of the UCustomizableObjectInstance::Descriptor at the time of the update request. */
+	FDescriptorRuntimeHash InstanceDescriptorRuntimeHash;
+			
+	/** Instance parameters at the time of the operation request. */
+	mu::ParametersPtr Parameters; 
+	
+	TArray<FName> TextureParameters;
+
+	bool bBuildParameterRelevancy = false;
+
+	/** Instance state. */
+	int32 State = 0;
+
+	bool bOnlyUpdateIfNotGenerated = false;
+	bool bIgnoreCloseDist = false;
+	bool bForceHighPriority = false;
+	
 	FInstanceUpdateData InstanceUpdateData;
-	TArray<int> RelevantParametersInProgress;
+	TArray<int32> RelevantParametersInProgress;
 
 	TArray<FString> LowPriorityTextures;
 
 	/** This option comes from the operation request */
 	bool bNeverStream = false;
+	
 	/** When this option is enabled it will reuse the Mutable core instance and its temp data between updates.  */
 	bool bLiveUpdateMode = false;
 	bool bReuseInstanceTextures = false;
+	
 	/** This option comes from the operation request. It is used to reduce the number of mipmaps that mutable must generate for images.  */
 	int32 MipsToSkip = 0;
 
@@ -534,17 +460,13 @@ struct FMutableOperationData
 	/** This list of queries is generated in the update mutable task, and consumed later in the game thread. */
 	TArray<FPendingTextureCoverageQuery> PendingTextureCoverageQueries;
 
-	mu::Ptr<const mu::Parameters> MutableParameters;
-
-	int32 State = 0;
-
-	bool bBuildParameterRelevancy = false;
-
-	EUpdateResult UpdateResult;
-	FInstanceUpdateDelegate UpdateCallback;
+	EUpdateResult UpdateResult = EUpdateResult::Success;
 
 	mu::FImageOperator::FImagePixelFormatFunc PixelFormatOverride;
 
+	// Update stats
+	double StartUpdateTime = 0.0;
+	
 #if WITH_EDITOR
 	/** Used for profiling in the editor. */
 	uint32 MutableRuntimeCycles = 0;
@@ -664,7 +586,7 @@ public:
 
 	EQueuePriorityType GetUpdatePriority(const UCustomizableObjectInstance& Instance, bool bForceHighPriority) const;
 
-	void EnqueueUpdateSkeletalMesh(UCustomizableObjectInstance& Instance, bool bOnlyUpdateIfNotGenerated, bool bIgnoreCloseDist, bool bForceHighPriority, const EUpdateRequired* OptionalUpdateRequired, FInstanceUpdateDelegate* UpdateCallback);
+	void EnqueueUpdateSkeletalMesh(const TSharedRef<FUpdateContextPrivate>& Context);
 		
 	// Init an async and safe release of the UE and Mutable resources used by the instance without actually destroying the instance, for example if it's very far away
 	void InitDiscardResourcesSkeletalMesh(UCustomizableObjectInstance* InCustomizableObjectInstance);
@@ -695,7 +617,7 @@ public:
 	// This is protected from GC by AddReferencedObjects
 	TObjectPtr<UCustomizableObjectInstance> CurrentInstanceBeingUpdated = nullptr;
 
-	TSharedPtr<FMutableOperation> CurrentMutableOperation = nullptr;
+	TSharedPtr<FUpdateContextPrivate> CurrentMutableOperation = nullptr;
 
 	// Handle to the registered TickDelegate.
 	FTSTicker::FDelegateHandle TickDelegateHandle;
@@ -709,7 +631,7 @@ public:
 	FUnrealMutableImageProvider* GetImageProviderChecked() const;
 
 	/** Start the actual work of Update Skeletal Mesh process (Update Skeletal Mesh without the queue). */
-	void StartUpdateSkeletalMesh(const TSharedPtr<FMutableOperation>& Operation);
+	void StartUpdateSkeletalMesh(const TSharedRef<FUpdateContextPrivate>& Context);
 
 	/** See UCustomizableObjectInstance::IsUpdating. */
 	bool IsUpdating(const UCustomizableObjectInstance& Instance) const;
