[Mutable] Fix ImageTable columns ignoring MaxTextureSize.

#rnx

diff --git a/Engine/Plugins/Experimental/Mutable/Source/CustomizableObjectEditor/Private/MuCOE/GenerateMutableSource/GenerateMutableSourceImage.cpp b/Engine/Plugins/Experimental/Mutable/Source/CustomizableObjectEditor/Private/MuCOE/GenerateMutableSource/GenerateMutableSourceImage.cpp
--- a/Engine/Plugins/Experimental/Mutable/Source/CustomizableObjectEditor/Private/MuCOE/GenerateMutableSource/GenerateMutableSourceImage.cpp
+++ b/Engine/Plugins/Experimental/Mutable/Source/CustomizableObjectEditor/Private/MuCOE/GenerateMutableSource/GenerateMutableSourceImage.cpp
@@ -1113,6 +1113,7 @@ mu::NodeImagePtr GenerateMutableSourceImage(const UEdGraphPin* Pin, FMutableGrap
 							ImageTableNode->SetTable(Table);
 							ImageTableNode->SetColumn(StringCast<ANSICHAR>(*ColumnName).Get());
 							ImageTableNode->SetParameterName(StringCast<ANSICHAR>(*TypedNodeTable->ParameterName).Get());
+							ImageTableNode->SetMaxTextureSize(MaxTextureSize);
 
 							GenerationContext.AddParameterNameUnique(Node, TypedNodeTable->ParameterName);
 						}
diff --git a/Engine/Plugins/Experimental/Mutable/Source/MutableTools/Private/MuT/CodeGenerator_Image.cpp b/Engine/Plugins/Experimental/Mutable/Source/MutableTools/Private/MuT/CodeGenerator_Image.cpp
--- a/Engine/Plugins/Experimental/Mutable/Source/MutableTools/Private/MuT/CodeGenerator_Image.cpp
+++ b/Engine/Plugins/Experimental/Mutable/Source/MutableTools/Private/MuT/CodeGenerator_Image.cpp
@@ -1790,6 +1790,16 @@ namespace mu
 					}
 					else
 					{
+						if (node.MaxTextureSize > 0 && (node.MaxTextureSize < pImage->GetSizeX() || node.MaxTextureSize < pImage->GetSizeY()))
+						{
+							float Factor = FMath::Min(node.MaxTextureSize / (float)(pImage->GetSizeX()), node.MaxTextureSize / (float)(pImage->GetSizeY()));
+
+							FImageOperator ImOp = FImageOperator::GetDefault(m_compilerOptions->ImageFormatFunc);
+							Ptr<Image> ResizedImage = new Image(pImage->GetSizeX() * Factor, pImage->GetSizeY() * Factor, pImage->GetLODCount(), pImage->GetFormat(), EInitializationType::NotInitialized);
+							ImOp.ImageResizeLinear(ResizedImage.get(), m_compilerOptions->ImageCompressionQuality, pImage.get());
+							pImage = ResizedImage;
+						}
+
 						NodeImageConstantPtr ImageConst = new NodeImageConstant();
 						ImageConst->SetValue(pImage);
 
diff --git a/Engine/Plugins/Experimental/Mutable/Source/MutableTools/Private/MuT/NodeImageTable.cpp b/Engine/Plugins/Experimental/Mutable/Source/MutableTools/Private/MuT/NodeImageTable.cpp
--- a/Engine/Plugins/Experimental/Mutable/Source/MutableTools/Private/MuT/NodeImageTable.cpp
+++ b/Engine/Plugins/Experimental/Mutable/Source/MutableTools/Private/MuT/NodeImageTable.cpp
@@ -84,6 +84,19 @@ namespace mu
 	}
 
 
+	//---------------------------------------------------------------------------------------------
+	void NodeImageTable::SetMaxTextureSize(uint16 MaxTextureSize)
+	{
+		m_pD->MaxTextureSize = MaxTextureSize;
+	}
+
+
+	//---------------------------------------------------------------------------------------------
+	uint16 NodeImageTable::GetMaxTextureSize()
+	{
+		return m_pD->MaxTextureSize;
+	}
+
 }
 
 
diff --git a/Engine/Plugins/Experimental/Mutable/Source/MutableTools/Private/MuT/NodeImageTablePrivate.h b/Engine/Plugins/Experimental/Mutable/Source/MutableTools/Private/MuT/NodeImageTablePrivate.h
--- a/Engine/Plugins/Experimental/Mutable/Source/MutableTools/Private/MuT/NodeImageTablePrivate.h
+++ b/Engine/Plugins/Experimental/Mutable/Source/MutableTools/Private/MuT/NodeImageTablePrivate.h
@@ -25,16 +25,18 @@ namespace mu
 		FString m_parameterName;
 		TablePtr m_pTable;
 		FString m_columnName;
+		uint16 MaxTextureSize = 0;
 
 		//!
 		void Serialise( OutputArchive& arch ) const
 		{
-            uint32_t ver = 2;
+            uint32_t ver = 3;
 			arch << ver;
 
 			arch << m_parameterName;
 			arch << m_pTable;
 			arch << m_columnName;
+			arch << MaxTextureSize;
 		}
 
 		//!
@@ -42,7 +44,7 @@ namespace mu
 		{
             uint32_t ver;
 			arch >> ver;
-			check(ver>=1 && ver<= 2);
+			check(ver>=1 && ver<= 3);
 
 			if (ver == 1)
 			{
@@ -67,8 +69,12 @@ namespace mu
 			{
 				arch >> m_columnName;
 			}
-		}
 
+			if (ver >= 3)
+			{
+				arch >> MaxTextureSize;
+			}
+		}
 	};
 
 }
diff --git a/Engine/Plugins/Experimental/Mutable/Source/MutableTools/Public/MuT/NodeImageTable.h b/Engine/Plugins/Experimental/Mutable/Source/MutableTools/Public/MuT/NodeImageTable.h
--- a/Engine/Plugins/Experimental/Mutable/Source/MutableTools/Public/MuT/NodeImageTable.h
+++ b/Engine/Plugins/Experimental/Mutable/Source/MutableTools/Public/MuT/NodeImageTable.h
@@ -68,6 +68,10 @@ namespace mu
 		//!
 		void SetColumn( const FString& strName );
 
+		//!
+		void SetMaxTextureSize(uint16 Size);
+		uint16 GetMaxTextureSize();
+
 		//-----------------------------------------------------------------------------------------
 		// Interface pattern
 		//-----------------------------------------------------------------------------------------
