[Mutable] Removed Participating Objects which avoided to save COs.

#rnx

diff --git a/Engine/Plugins/Experimental/Mutable/Source/CustomizableObjectEditor/Private/MuCOE/GenerateMutableSource/GenerateMutableSourceMesh.cpp b/Engine/Plugins/Experimental/Mutable/Source/CustomizableObjectEditor/Private/MuCOE/GenerateMutableSource/GenerateMutableSourceMesh.cpp
--- a/Engine/Plugins/Experimental/Mutable/Source/CustomizableObjectEditor/Private/MuCOE/GenerateMutableSource/GenerateMutableSourceMesh.cpp
+++ b/Engine/Plugins/Experimental/Mutable/Source/CustomizableObjectEditor/Private/MuCOE/GenerateMutableSource/GenerateMutableSourceMesh.cpp
@@ -2888,8 +2888,6 @@ mu::NodeMeshPtr GenerateMutableSourceMesh(const UEdGraphPin* Pin,
 				{
 					for (const UAssetUserData* AssetUserData : *AssetUserDataArray)
 					{
-						GenerationContext.AddParticipatingObject(*AssetUserData);
-
 						FString AuxString = AssetUserData->GetPathName();
 						GenerationContext.AssetUserDataAssetsMap.Add(AuxString, TSoftObjectPtr<UAssetUserData>(AssetUserData));
 
diff --git a/Engine/Plugins/Experimental/Mutable/Source/CustomizableObjectEditor/Private/MuCOE/GenerateMutableSource/GenerateMutableSourceTable.cpp b/Engine/Plugins/Experimental/Mutable/Source/CustomizableObjectEditor/Private/MuCOE/GenerateMutableSource/GenerateMutableSourceTable.cpp
--- a/Engine/Plugins/Experimental/Mutable/Source/CustomizableObjectEditor/Private/MuCOE/GenerateMutableSource/GenerateMutableSourceTable.cpp
+++ b/Engine/Plugins/Experimental/Mutable/Source/CustomizableObjectEditor/Private/MuCOE/GenerateMutableSource/GenerateMutableSourceTable.cpp
@@ -274,8 +274,6 @@ bool FillTableColumn(const UCustomizableObjectNodeTable* TableNode,	mu::TablePtr
 					{
 						for (const UAssetUserData* AssetUserData : *AssetUserDataArray)
 						{
-							GenerationContext.AddParticipatingObject(*AssetUserData);
-
 							FString AuxString = AssetUserData->GetPathName();
 							GenerationContext.AssetUserDataAssetsMap.Add(AuxString, TSoftObjectPtr<UAssetUserData>(AssetUserData));
 
