[mutable] Undo texture block size optimization because it can break layouts in some cases.
#rnx

diff --git a/Engine/Plugins/Experimental/Mutable/Source/MutableTools/Private/MuT/CodeGenerator.cpp b/Engine/Plugins/Experimental/Mutable/Source/MutableTools/Private/MuT/CodeGenerator.cpp
--- a/Engine/Plugins/Experimental/Mutable/Source/MutableTools/Private/MuT/CodeGenerator.cpp
+++ b/Engine/Plugins/Experimental/Mutable/Source/MutableTools/Private/MuT/CodeGenerator.cpp
@@ -1338,6 +1338,7 @@ namespace mu
 							// Image
 							//-------------------------------------
 
+							// Size of a layout block in pixels
 							FIntPoint grid = pLayout->GetGridSize();
 
 							check(desc.m_format != EImageFormat::IF_NONE);
@@ -1401,18 +1402,18 @@ namespace mu
 								);
 
 								// Transform to pixels
-								box< UE::Math::TIntVector2<uint16> > BlockRectPixels = rectInCells;
-								BlockRectPixels.min[0] *= blockSizeX;
-								BlockRectPixels.min[1] *= blockSizeY;
-								BlockRectPixels.size[0] *= blockSizeX;
-								BlockRectPixels.size[1] *= blockSizeY;
+								box< UE::Math::TIntVector2<uint16> > rect = rectInCells;
+								rect.min[0] *= blockSizeX;
+								rect.min[1] *= blockSizeY;
+								rect.size[0] *= blockSizeX;
+								rect.size[1] *= blockSizeY;
 
 								// Generate the image
 								FImageGenerationOptions ImageOptions;
 								ImageOptions.CurrentStateIndex = m_currentStateIndex;
 								ImageOptions.ImageLayoutStrategy = ImageLayoutStrategy;
 								ImageOptions.ActiveTags = node.m_tags;
-								ImageOptions.RectSize = UE::Math::TIntVector2<int32>(BlockRectPixels.size);
+								ImageOptions.RectSize = UE::Math::TIntVector2<int32>(rect.size);
 								ImageOptions.LayoutToApply = pLayout;
 								ImageOptions.LayoutBlockId = pLayout->m_blocks[b].m_id;
 								FImageGenerationResult Result;
@@ -1507,16 +1508,12 @@ namespace mu
 												rect.size[0] = (blockRect.size[0] * extendDesc.m_size[0]) / extlayout[0];
 												rect.size[1] = (blockRect.size[1] * extendDesc.m_size[1]) / extlayout[1];
 
-												UE::Math::TIntVector2<int32> ExpectedBlockPixels;
-												ExpectedBlockPixels[0] = blockSizeX * blockRect.size[0];
-												ExpectedBlockPixels[1] = blockSizeY * blockRect.size[1];
-
 												// Generate the image block
 												FImageGenerationOptions ImageOptions;
 												ImageOptions.CurrentStateIndex = m_currentStateIndex;
 												ImageOptions.ImageLayoutStrategy = ImageLayoutStrategy;
 												ImageOptions.ActiveTags = node.m_tags;
-												ImageOptions.RectSize = ExpectedBlockPixels;
+												ImageOptions.RectSize = UE::Math::TIntVector2<int32>(extendDesc.m_size);
 												ImageOptions.LayoutToApply = pExtendLayout;
 												ImageOptions.LayoutBlockId = pExtendLayout->m_blocks[b].m_id;
 												FImageGenerationResult ExtendResult;
@@ -1524,13 +1521,15 @@ namespace mu
 												Ptr<ASTOp> fragmentAd = ExtendResult.op;
 
 												// Adjust the format and size of the block to be added
-												// \TODO: review if this is still needed considering a higher format op, and the RectSize in the ImageOptions
-												Ptr<ASTOp> formatted = GenerateImageFormat(fragmentAd, desc.m_format);
-												formatted = GenerateImageSize(formatted, ExpectedBlockPixels);
+												Ptr<ASTOp> formatted = GenerateImageFormat(fragmentAd, GetUncompressedFormat(desc.m_format));
+												UE::Math::TIntVector2<int32> expectedSize;
+												expectedSize[0] = blockSizeX * blockRect.size[0];
+												expectedSize[1] = blockSizeY * blockRect.size[1];
+												formatted = GenerateImageSize(formatted, expectedSize);
 												fragmentAd = formatted;
 
 												// Apply tiling to avoid generating chunks of image that are too big.
-												fragmentAd = ApplyTiling(fragmentAd, ExpectedBlockPixels, desc.m_format);
+												fragmentAd = ApplyTiling(fragmentAd, expectedSize, desc.m_format);
 
 												// Compose operation
 												Ptr<ASTOpImageCompose> composeOp = new ASTOpImageCompose();
diff --git a/Engine/Plugins/Experimental/Mutable/Source/MutableTools/Private/MuT/CodeGenerator_Image.cpp b/Engine/Plugins/Experimental/Mutable/Source/MutableTools/Private/MuT/CodeGenerator_Image.cpp
--- a/Engine/Plugins/Experimental/Mutable/Source/MutableTools/Private/MuT/CodeGenerator_Image.cpp
+++ b/Engine/Plugins/Experimental/Mutable/Source/MutableTools/Private/MuT/CodeGenerator_Image.cpp
@@ -193,7 +193,7 @@ namespace mu
             m_pErrorLog->GetPrivate()->Add( "Constant image not set.", ELMT_WARNING, node.m_errorContext );
         }
 
-		if (Options.ImageLayoutStrategy!=CompilerOptions::TextureLayoutStrategy::None && Options.LayoutToApply)
+		if (Options.ImageLayoutStrategy!= CompilerOptions::TextureLayoutStrategy::None && Options.LayoutToApply)
 		{
 			// We want to generate only a block from the image.
 
@@ -216,38 +216,42 @@ namespace mu
 			grid[1] = FMath::Max(1, grid[1]);
 
 			// Transform to pixels
-			box< UE::Math::TIntVector2<int32> > SourceRectPixels;
-			SourceRectPixels.min[0]  = (RectInCells.min[0]  * SourceImageSize[0]) / grid[0];
-			SourceRectPixels.min[1]  = (RectInCells.min[1]  * SourceImageSize[1]) / grid[1];
-			SourceRectPixels.size[0] = (RectInCells.size[0] * SourceImageSize[0]) / grid[0];
-			SourceRectPixels.size[1] = (RectInCells.size[1] * SourceImageSize[1]) / grid[1];
-
-			FImageOperator ImOp = FImageOperator::GetDefault(m_compilerOptions->ImageFormatFunc);
+			box< UE::Math::TIntVector2<int32> > rect;
+			rect.min[0]  = (RectInCells.min[0]  * SourceImageSize[0]) / grid[0];
+			rect.min[1]  = (RectInCells.min[1]  * SourceImageSize[1]) / grid[1];
+			rect.size[0] = (RectInCells.size[0] * SourceImageSize[0]) / grid[0];
+			rect.size[1] = (RectInCells.size[1] * SourceImageSize[1]) / grid[1];
 
 			// Do we need to crop?
-			if (SourceRectPixels.min[0]!=0 || SourceRectPixels.min[1]!=0 || pImage->GetSizeX() != SourceRectPixels.size[0] || pImage->GetSizeY() != SourceRectPixels.size[1])
+			if (rect.min[0]!=0 || rect.min[1]!=0 || pImage->GetSizeX() != rect.size[0] || pImage->GetSizeY() != rect.size[1])
 			{
 				// Crop now
-				Ptr<Image> pCropped = new Image(SourceRectPixels.size[0], SourceRectPixels.size[1], 1, pImage->GetFormat(), EInitializationType::NotInitialized);
-				ImOp.ImageCrop(pCropped.get(), m_compilerOptions->ImageCompressionQuality, pImage.get(), SourceRectPixels);
-				pImage = pCropped;
-			}
+				FImageOperator ImOp = FImageOperator::GetDefault(m_compilerOptions->ImageFormatFunc);
+
+				Ptr<Image> pCropped = new Image(rect.size[0], rect.size[1], 1, pImage->GetFormat(), EInitializationType::NotInitialized);
+				ImOp.ImageCrop(pCropped.get(), m_compilerOptions->ImageCompressionQuality, pImage.get(), rect);
 
-			// If actual size we need this image in is going to be smaller, resize now too
-			bool bBothAxisEqualOrSmaller = (Options.RectSize[0] < SourceRectPixels.size[0]) && (Options.RectSize[1] < SourceRectPixels.size[1]);
-			bool bBothAxisEqual = (Options.RectSize[0] == SourceRectPixels.size[0]) && (Options.RectSize[1] == SourceRectPixels.size[1]);
-			if (bBothAxisEqualOrSmaller && !bBothAxisEqual)
+				Ptr<ASTOpConstantResource> op = new ASTOpConstantResource();
+				op->type = OP_TYPE::IM_CONSTANT;
+				op->SetValue(pCropped, m_compilerOptions->OptimisationOptions.bUseDiskCache);
+				Result.op = op;
+			}
+			else
 			{
-				Ptr<Image> ResizedImage = new Image(Options.RectSize[0], Options.RectSize[1],1,pImage->GetFormat(),EInitializationType::NotInitialized);
-				ImOp.ImageResizeLinear(ResizedImage.get(), m_compilerOptions->ImageCompressionQuality, pImage.get() );
-				pImage = ResizedImage;
+				// No need to crop.
+				Ptr<ASTOpConstantResource> op = new ASTOpConstantResource();
+				op->type = OP_TYPE::IM_CONSTANT;
+				op->SetValue(pImage, m_compilerOptions->OptimisationOptions.bUseDiskCache);
+				Result.op = op;
 			}
 		}
- 
-		Ptr<ASTOpConstantResource> op = new ASTOpConstantResource();
-		op->type = OP_TYPE::IM_CONSTANT;
-		op->SetValue( pImage, m_compilerOptions->OptimisationOptions.bUseDiskCache );
-		Result.op = op;
+        else
+        {
+			Ptr<ASTOpConstantResource> op = new ASTOpConstantResource();
+			op->type = OP_TYPE::IM_CONSTANT;
+			op->SetValue( pImage, m_compilerOptions->OptimisationOptions.bUseDiskCache );
+			Result.op = op;
+		}
     }
 
 
