[Mutable] Added Descriptor Debug Print.

Allows to see the state of the parameters when enqueuing an update. Useful when analyzing why there are more updates than they should. CVar toggable.

#rnx

diff --git a/Engine/Plugins/Experimental/Mutable/Source/CustomizableObject/Private/MuCO/CustomizableObjectInstanceDescriptor.cpp b/Engine/Plugins/Experimental/Mutable/Source/CustomizableObject/Private/MuCO/CustomizableObjectInstanceDescriptor.cpp
--- a/Engine/Plugins/Experimental/Mutable/Source/CustomizableObject/Private/MuCO/CustomizableObjectInstanceDescriptor.cpp
+++ b/Engine/Plugins/Experimental/Mutable/Source/CustomizableObject/Private/MuCO/CustomizableObjectInstanceDescriptor.cpp
@@ -651,6 +651,89 @@ mu::ParametersPtr FCustomizableObjectInstanceDescriptor::GetParameters() const
 }
 
 
+FString FCustomizableObjectInstanceDescriptor::ToString() const
+{
+	TStringBuilder<3000> Builder;
+	
+	Builder.Appendf(TEXT("CustomizableObject=%s\n"), *CustomizableObject.GetFullName());
+	Builder.Appendf(TEXT("State=%i\n"), State);
+	Builder.Appendf(TEXT("BuildParameterRelevancy=%i\n"), bBuildParameterRelevancy);
+	Builder.Appendf(TEXT("MinLOD=%i\n"), MinLOD);
+	Builder.Appendf(TEXT("MaxLOD=%i\n"), MaxLOD);
+
+	for (int32 Index = 0; Index < RequestedLODLevels.Num(); ++Index)
+	{
+		Builder.Appendf(TEXT("RequestredLODLevels [%i]=%i\n"), Index, RequestedLODLevels[Index]);
+	}
+
+	Builder.Appendf(TEXT("BOOL PARAMETERS:\n"));
+	for (const FCustomizableObjectBoolParameterValue& Parameter : BoolParameters)
+	{
+		Builder.Appendf(TEXT("%s=%i\n"), *Parameter.ParameterName, Parameter.ParameterValue);
+	}
+
+	Builder.Appendf(TEXT("INT PARAMETERS:\n"));
+	for (const FCustomizableObjectIntParameterValue& Parameter : IntParameters)
+	{
+		Builder.Appendf(TEXT("%s=\"%s\"\n"), *Parameter.ParameterName, *Parameter.ParameterValueName);
+
+		for (int32 Index = 0; Index < Parameter.ParameterRangeValueNames.Num(); ++Index)
+		{
+			Builder.Appendf(TEXT("%s [%i]=\"%s\"\n"), *Parameter.ParameterName, Index, *Parameter.ParameterRangeValueNames[Index]);			
+		}
+	}
+
+	Builder.Appendf(TEXT("FLOAT PARAMETETS:\n"));
+	for (const FCustomizableObjectFloatParameterValue& Parameter : FloatParameters)
+	{
+		Builder.Appendf(TEXT("%s=%f\n"), *Parameter.ParameterName, Parameter.ParameterValue);
+
+		for (int32 Index = 0; Index < Parameter.ParameterRangeValues.Num(); ++Index)
+		{
+			Builder.Appendf(TEXT("%s [%i]=%f\n"), *Parameter.ParameterName, Index, Parameter.ParameterRangeValues[Index]);			
+		}
+	}
+
+	Builder.Appendf(TEXT("TEXTURE PARAMETERS:\n"));
+	for (const FCustomizableObjectTextureParameterValue& Parameter : TextureParameters)
+	{
+		Builder.Appendf(TEXT("%s=%s\n"), *Parameter.ParameterName, *Parameter.ParameterValue.ToString());
+
+		for (int32 Index = 0; Index < Parameter.ParameterRangeValues.Num(); ++Index)
+		{
+			Builder.Appendf(TEXT("%s [%i]=%s\n"), *Parameter.ParameterName, Index, *Parameter.ParameterRangeValues[Index].ToString());			
+		}
+	}
+
+	Builder.Appendf(TEXT("VECTOR PARAMETERS:\n"));
+	for (const FCustomizableObjectVectorParameterValue& Parameter : VectorParameters)
+	{
+		Builder.Appendf(TEXT("%s=%s\n"), *Parameter.ParameterName, *Parameter.ParameterValue.ToString());			
+	}
+
+	Builder.Appendf(TEXT("PROJECTOR PARAMETERS:\n"));
+	for (const FCustomizableObjectProjectorParameterValue& Parameter : ProjectorParameters)
+	{
+		auto StructToString = [](const UScriptStruct& ScriptStruct, const void* Data) {
+			FString String;
+			ScriptStruct.ExportText(String, &Data, &Data, nullptr, PPF_ExternalEditor | PPF_IncludeTransient, nullptr);
+			return String;
+		};
+
+		Builder.Appendf(TEXT("%s=%s\n"), *Parameter.ParameterName, *StructToString(*Parameter.Value.StaticStruct(), &Parameter.Value));
+
+		for (int32 Index = 0; Index < Parameter.RangeValues.Num(); ++Index)
+		{
+			const FCustomizableObjectProjector& RangeValue = Parameter.RangeValues[Index];
+			
+			Builder.Appendf(TEXT("%s [%i]=%s\n"), *Parameter.ParameterName, Index, *StructToString(*RangeValue.StaticStruct(), &RangeValue));			
+		}
+	}
+
+	return Builder.ToString();
+}
+
+
 void FCustomizableObjectInstanceDescriptor::ReloadParameters()
 {
 	if (!CustomizableObject)
@@ -2420,12 +2503,6 @@ bool FDescriptorHash::operator<(const FDescriptorHash& Other) const
 }
 
 
-FString FDescriptorHash::ToString() const
-{
-	return FString::Printf(TEXT("%u"), Hash);
-}
-
-
 FDescriptorRuntimeHash::FDescriptorRuntimeHash(const FCustomizableObjectInstanceDescriptor& Descriptor) :
 	FDescriptorHash(Descriptor)
 {
@@ -2496,4 +2573,25 @@ const TArray<uint16>& FDescriptorRuntimeHash::GetRequestedLODs() const
 	return RequestedLODsPerComponent;
 }
 
-#undef RETURN_ON_CO_UNCOMPILED
\ No newline at end of file
+
+FString FDescriptorRuntimeHash::ToString() const
+{
+	TStringBuilder<150> Builder;
+
+	Builder.Appendf(TEXT("(Hash=%u,"), Hash);
+	Builder.Appendf(TEXT("MinLOD=%i,"), MinLOD);
+	Builder.Appendf(TEXT("MaxLOD=%i,"), MaxLOD);
+	Builder.Appendf(TEXT("RequestredLODLevels=["));
+
+	for (const uint16 RequestedLODs : RequestedLODsPerComponent)
+	{
+		Builder.Appendf(TEXT("%i,"), RequestedLODs);
+	}
+	
+	Builder.Appendf(TEXT("])"));
+
+	return Builder.ToString();
+}
+
+
+#undef RETURN_ON_CO_UNCOMPILED
diff --git a/Engine/Plugins/Experimental/Mutable/Source/CustomizableObject/Private/MuCO/CustomizableObjectSystem.cpp b/Engine/Plugins/Experimental/Mutable/Source/CustomizableObject/Private/MuCO/CustomizableObjectSystem.cpp
--- a/Engine/Plugins/Experimental/Mutable/Source/CustomizableObject/Private/MuCO/CustomizableObjectSystem.cpp
+++ b/Engine/Plugins/Experimental/Mutable/Source/CustomizableObject/Private/MuCO/CustomizableObjectSystem.cpp
@@ -762,6 +762,13 @@ static FAutoConsoleVariableRef CVarApplyFixPrepareSkeletons(
 	TEXT("If true, Fix missing SkeletonsData when FirstLODToGenerate is greater than 0. If false, There may be a crash when generating meshes on platform that skip LODs."),
 	ECVF_Default);
 
+static FAutoConsoleVariable CVarDescriptorDebugPrint(
+	TEXT("mutable.DescriptorDebugPrint"),
+	false,
+	TEXT("If true, each time an update is enqueued, print its captured parameters."),
+	ECVF_Default);
+
+
 void FinishUpdateGlobal(const TSharedRef<FUpdateContextPrivate>& Context)
 {
 	check(IsInGameThread())
@@ -1185,6 +1192,17 @@ void FCustomizableObjectSystemPrivate::EnqueueUpdateSkeletalMesh(const TSharedRe
 				}
 			}
 
+			if (CVarDescriptorDebugPrint->GetBool())
+			{
+				FString String = TEXT("DESCRIPTOR DEBUG PRINT\n");
+				String += "================================\n";				
+				String += FString::Printf(TEXT("=== DESCRIPTOR HASH ===\n%s\n"), *InstancePrivate->UpdateDescriptorRuntimeHash.ToString());
+				String += FString::Printf(TEXT("=== DESCRIPTOR ===\n%s"), *Instance->GetDescriptor().ToString());
+				String += "================================";
+				
+				UE_LOG(LogMutable, Log, TEXT("%s"), *String);
+			}
+
 			const FMutablePendingInstanceUpdate InstanceUpdate(Context);
 			MutablePendingInstanceWork.AddUpdate(InstanceUpdate);
 		}
diff --git a/Engine/Plugins/Experimental/Mutable/Source/CustomizableObject/Public/MuCO/CustomizableObjectInstanceDescriptor.h b/Engine/Plugins/Experimental/Mutable/Source/CustomizableObject/Public/MuCO/CustomizableObjectInstanceDescriptor.h
--- a/Engine/Plugins/Experimental/Mutable/Source/CustomizableObject/Public/MuCO/CustomizableObjectInstanceDescriptor.h
+++ b/Engine/Plugins/Experimental/Mutable/Source/CustomizableObject/Public/MuCO/CustomizableObjectInstanceDescriptor.h
@@ -347,6 +347,8 @@ struct CUSTOMIZABLEOBJECT_API FCustomizableObjectInstanceDescriptor
 	/** Return a Mutable Core object containing all parameters. */
 	mu::Ptr<mu::Parameters> GetParameters() const;
 
+	FString ToString() const;
+	
 private:
 
 	UPROPERTY()
@@ -416,9 +418,7 @@ public:
 
 	bool operator<(const FDescriptorHash& Other) const;
 
-	FString ToString() const;
-	
-private:
+protected:
 	uint32 Hash = 0;
 };
 
@@ -444,6 +444,8 @@ public:
 
 	const TArray<uint16>& GetRequestedLODs() const;
 
+	FString ToString() const;
+
 private:
 	int32 MinLOD = 0;
 	int32 MaxLOD = INT32_MAX;
