[Mutable] Fix incompatible skeleton warning not showing.

#rnx

diff --git a/Engine/Plugins/Experimental/Mutable/Source/CustomizableObjectEditor/Private/MuCOE/GenerateMutableSource/GenerateMutableSourceMesh.cpp b/Engine/Plugins/Experimental/Mutable/Source/CustomizableObjectEditor/Private/MuCOE/GenerateMutableSource/GenerateMutableSourceMesh.cpp
--- a/Engine/Plugins/Experimental/Mutable/Source/CustomizableObjectEditor/Private/MuCOE/GenerateMutableSource/GenerateMutableSourceMesh.cpp
+++ b/Engine/Plugins/Experimental/Mutable/Source/CustomizableObjectEditor/Private/MuCOE/GenerateMutableSource/GenerateMutableSourceMesh.cpp
@@ -360,7 +360,7 @@ bool IsSkeletalMeshCompatibleWithRefSkeleton(FMutableComponentInfo& ComponentInf
 		if (const uint32* RefSMBonePathHash = RefMeshBoneNamesToPathHash.Find(Bone.Name); RefSMBonePathHash && *RefSMBonePathHash != BonePathHash)
 		{
 			// Different skeletons can't be used if they are incompatible with the reference skeleton.
-			FString Msg = FString::Printf(
+			OutErrorMessage = FString::Printf(
 				TEXT("The SkeletalMesh [%s] with Skeleton [%s] is incompatible with the reference mesh [%s] which has [%s]. "
 					"Bone [%s] has a differnt parent on the Skeleton from the reference mesh."),
 				*InSkeletalMesh->GetName(), *Skeleton->GetName(),
@@ -695,7 +695,7 @@ mu::MeshPtr ConvertSkeletalMeshToMutable(const USkeletalMesh* InSkeletalMesh, co
 
 			if (!bCompatible)
 			{
-				if (ErrorMessage.IsEmpty())
+				if (!ErrorMessage.IsEmpty())
 				{
 					GenerationContext.Compiler->CompilerLog(FText::FromString(ErrorMessage), CurrentNode, EMessageSeverity::Warning);
 				}
