[Mutable] Fix UVLayout mode incorrect deserialization due to a change in the default.

- Added missing copy of a pin property.

#rnx

diff --git a/Engine/Plugins/Experimental/Mutable/Source/CustomizableObjectEditor/Classes/MuCOE/Nodes/CustomizableObjectNodeMaterial.h b/Engine/Plugins/Experimental/Mutable/Source/CustomizableObjectEditor/Classes/MuCOE/Nodes/CustomizableObjectNodeMaterial.h
--- a/Engine/Plugins/Experimental/Mutable/Source/CustomizableObjectEditor/Classes/MuCOE/Nodes/CustomizableObjectNodeMaterial.h
+++ b/Engine/Plugins/Experimental/Mutable/Source/CustomizableObjectEditor/Classes/MuCOE/Nodes/CustomizableObjectNodeMaterial.h
@@ -348,8 +348,8 @@ public:
 	EUVLayoutMode UVLayoutMode = EUVLayoutMode::FromMaterial;
 	
 	/** Index of the UV channel that will be used with this image.It is necessary to apply the proper layout transformations to it. */
-	UPROPERTY(EditAnywhere, Category = NoCategory, meta = (ClampMin = 0, EditCondition = "UVLayoutMode == EUVLayoutMode::Index", EditConditionHides))
-	int32 UVLayout = 0;
+	UPROPERTY(EditAnywhere, Category = NoCategory, meta = (EditCondition = "UVLayoutMode == EUVLayoutMode::Index", EditConditionHides))
+	int32 UVLayout = -2;
 
 	/** Reference Texture used to decide the texture properties of the mutable-generated textures
 	* connected to this material. If null, it will try to be guessed at compile time from
diff --git a/Engine/Plugins/Experimental/Mutable/Source/CustomizableObjectEditor/Private/MuCOE/Nodes/CustomizableObjectNodeMaterial.cpp b/Engine/Plugins/Experimental/Mutable/Source/CustomizableObjectEditor/Private/MuCOE/Nodes/CustomizableObjectNodeMaterial.cpp
--- a/Engine/Plugins/Experimental/Mutable/Source/CustomizableObjectEditor/Private/MuCOE/Nodes/CustomizableObjectNodeMaterial.cpp
+++ b/Engine/Plugins/Experimental/Mutable/Source/CustomizableObjectEditor/Private/MuCOE/Nodes/CustomizableObjectNodeMaterial.cpp
@@ -122,6 +122,7 @@ void UCustomizableObjectNodeMaterialPinDataImage::Copy(const UCustomizableObject
 	if (const UCustomizableObjectNodeMaterialPinDataImage* PinDataOldPin = Cast<UCustomizableObjectNodeMaterialPinDataImage>(&Other))
 	{
 		PinMode = PinDataOldPin->PinMode;
+		UVLayoutMode = PinDataOldPin->UVLayoutMode;
 		UVLayout = PinDataOldPin->UVLayout;
 		ReferenceTexture = PinDataOldPin->ReferenceTexture;
 
