[Mutable] Ensure vertex indices stay within range when converting SkeletalMeshes to Mutable.

#rnx

diff --git a/Engine/Plugins/Experimental/Mutable/Source/CustomizableObjectEditor/Private/MuCOE/GenerateMutableSource/GenerateMutableSourceMesh.cpp b/Engine/Plugins/Experimental/Mutable/Source/CustomizableObjectEditor/Private/MuCOE/GenerateMutableSource/GenerateMutableSourceMesh.cpp
--- a/Engine/Plugins/Experimental/Mutable/Source/CustomizableObjectEditor/Private/MuCOE/GenerateMutableSource/GenerateMutableSourceMesh.cpp
+++ b/Engine/Plugins/Experimental/Mutable/Source/CustomizableObjectEditor/Private/MuCOE/GenerateMutableSource/GenerateMutableSourceMesh.cpp
@@ -1516,10 +1516,20 @@ mu::MeshPtr ConvertSkeletalMeshToMutable(const USkeletalMesh* InSkeletalMesh, co
 			uint32* pDest = reinterpret_cast<uint32*>(MutableMesh->GetIndexBuffers().GetBufferData(0));
 			const uint16* pSource = reinterpret_cast<const uint16*>(IndexDataPointer);
 
+			int32 VertexIndex = 0;
 			for (int i = 0; i < IndexCount; ++i)
 			{
-				*pDest = *pSource - VertexStart;
-				check(*pDest < uint32(VertexCount));
+				VertexIndex = *pSource - VertexStart;
+				if (ensureMsgf(VertexIndex < VertexCount, TEXT("Mutable: VertexIndex >= VertexCount. VI [%d], VC [%d], VS [%d]. SKM [%s] LOD [%d] Section [%d]."),
+					VertexIndex, VertexCount, VertexStart,
+					*GetNameSafe(InSkeletalMesh), LODIndex, SectionIndex))
+				{
+					*pDest = (uint32)VertexIndex;
+				}
+				else
+				{
+					*pDest = 0;
+				}
 				++pDest;
 				++pSource;
 			}
