[Mutable] Fix Reference Skeletal Mesh not respecting the Placeholder setting on PIE.

#rnx

diff --git a/Engine/Plugins/Experimental/Mutable/Source/CustomizableObject/Private/MuCO/CustomizableSkeletalComponent.cpp b/Engine/Plugins/Experimental/Mutable/Source/CustomizableObject/Private/MuCO/CustomizableSkeletalComponent.cpp
--- a/Engine/Plugins/Experimental/Mutable/Source/CustomizableObject/Private/MuCO/CustomizableSkeletalComponent.cpp
+++ b/Engine/Plugins/Experimental/Mutable/Source/CustomizableObject/Private/MuCO/CustomizableSkeletalComponent.cpp
@@ -321,7 +321,7 @@ void UCustomizableSkeletalComponent::TickComponent(float DeltaTime, ELevelTick T
 		USkeletalMesh* SkeletalMesh = CustomizableObjectInstance->GetSkeletalMesh(ComponentIndex);
 
 		// If not generated yet, conditionally set the SkeletalMesh of reference
-		if (!bInstanceGenerated && !bSkipSetReferenceSkeletalMesh)
+		if (!bInstanceGenerated && !bSkipSetReferenceSkeletalMesh && CustomizableObject->bEnableUseRefSkeletalMeshAsPlaceholder)
 		{
 			// Can be nullptr
 			SkeletalMesh = CustomizableObject->GetRefSkeletalMesh(ComponentIndex);
