[Mutable] Fix missing include in UCustomizableSkeletalComponent.

#test Local NoUnityBuld
#rnx

diff --git a/Engine/Plugins/Experimental/Mutable/Source/CustomizableObject/Private/MuCO/CustomizableSkeletalComponent.cpp b/Engine/Plugins/Experimental/Mutable/Source/CustomizableObject/Private/MuCO/CustomizableSkeletalComponent.cpp
--- a/Engine/Plugins/Experimental/Mutable/Source/CustomizableObject/Private/MuCO/CustomizableSkeletalComponent.cpp
+++ b/Engine/Plugins/Experimental/Mutable/Source/CustomizableObject/Private/MuCO/CustomizableSkeletalComponent.cpp
@@ -15,6 +15,7 @@
 #include "Rendering/SkeletalMeshRenderData.h"
 #include "UObject/ObjectSaveContext.h"
 #include "MuCO/UnrealPortabilityHelpers.h"
+#include "MuCO/CustomizableObject.h"
 
 #include UE_INLINE_GENERATED_CPP_BY_NAME(CustomizableSkeletalComponent)
 
