Fix deprecation warning on PerQuality property

diff --git a/Engine/Plugins/Experimental/Mutable/Source/CustomizableObject/Private/MuCO/CustomizableObject.cpp b/Engine/Plugins/Experimental/Mutable/Source/CustomizableObject/Private/MuCO/CustomizableObject.cpp
--- a/Engine/Plugins/Experimental/Mutable/Source/CustomizableObject/Private/MuCO/CustomizableObject.cpp
+++ b/Engine/Plugins/Experimental/Mutable/Source/CustomizableObject/Private/MuCO/CustomizableObject.cpp
@@ -55,7 +55,7 @@ UCustomizableObject::UCustomizableObject()
 #if WITH_EDITORONLY_DATA
 	const FString CVarName = TEXT("r.SkeletalMesh.MinLodQualityLevel");
 	const FString ScalabilitySectionName = TEXT("ViewDistanceQuality");
-	LODSettings.MinQualityLevelLOD.Init(*CVarName, *ScalabilitySectionName);
+	LODSettings.MinQualityLevelLOD.SetQualityLevelCVarForCooking(*CVarName, *ScalabilitySectionName);
 #endif
 }
 
