[Mutable] Copy indices from the LODModel when converting skeletal meshes to Mutable.

#rnx

diff --git a/Engine/Plugins/Experimental/Mutable/Source/CustomizableObjectEditor/Private/MuCOE/GenerateMutableSource/GenerateMutableSourceMesh.cpp b/Engine/Plugins/Experimental/Mutable/Source/CustomizableObjectEditor/Private/MuCOE/GenerateMutableSource/GenerateMutableSourceMesh.cpp
--- a/Engine/Plugins/Experimental/Mutable/Source/CustomizableObjectEditor/Private/MuCOE/GenerateMutableSource/GenerateMutableSourceMesh.cpp
+++ b/Engine/Plugins/Experimental/Mutable/Source/CustomizableObjectEditor/Private/MuCOE/GenerateMutableSource/GenerateMutableSourceMesh.cpp
@@ -632,7 +632,8 @@ mu::MeshPtr ConvertSkeletalMeshToMutable(const USkeletalMesh* InSkeletalMesh, co
 		}
 	}
 
-	if (!ImportedModel->LODModels[LODIndex].Sections.IsValidIndex(SectionIndex))
+	const FSkeletalMeshLODModel& LODModel = ImportedModel->LODModels[LODIndex];
+	if (!LODModel.Sections.IsValidIndex(SectionIndex))
 	{
 		if (GenerationContext.CurrentAutoLODStrategy == ECustomizableObjectAutomaticLODStrategy::AutomaticFromMesh && 
 			SectionIndex != SectionIndexConnected) // If we are using automatic LODs and not generating the base LOD (the connected one) is not an error.
@@ -647,13 +648,15 @@ mu::MeshPtr ConvertSkeletalMeshToMutable(const USkeletalMesh* InSkeletalMesh, co
 				LODIndex + 1,
 				ImportedModel->LODModels.Num(),
 				SectionIndex + 1,
-				ImportedModel->LODModels[LODIndex].Sections.Num());
+				LODModel.Sections.Num());
 			GenerationContext.Compiler->CompilerLog(FText::FromString(Msg), CurrentNode);
 
 			return nullptr;
 		}
 	}
 
+	const FSkelMeshSection& MeshSection = LODModel.Sections[SectionIndex];
+
 	// Get the mesh generation flags to use
 	const EMutableMeshConversionFlags CurrentFlags = GenerationContext.MeshGenerationFlags.Last();
 	const bool bIgnoreSkeleton = EnumHasAnyFlags(CurrentFlags, EMutableMeshConversionFlags::IgnoreSkinning);
@@ -745,7 +748,7 @@ mu::MeshPtr ConvertSkeletalMeshToMutable(const USkeletalMesh* InSkeletalMesh, co
 			}
 		}
 
-		const TArray<uint16>& SourceRequiredBones = ImportedModel->LODModels[LODIndex].RequiredBones;
+		const TArray<uint16>& SourceRequiredBones = LODModel.RequiredBones;
 
 		// Remove bones and build an array to remap indices of the BoneMap
 		TArray<FBoneIndexType> RemappedBones;
@@ -761,7 +764,7 @@ mu::MeshPtr ConvertSkeletalMeshToMutable(const USkeletalMesh* InSkeletalMesh, co
 		}
 
 		// Rebuild BoneMap
-		const TArray<uint16>& SourceBoneMap = ImportedModel->LODModels[LODIndex].Sections[SectionIndex].BoneMap;
+		const TArray<uint16>& SourceBoneMap = MeshSection.BoneMap;
 		const int32 NumBonesInBoneMap = SourceBoneMap.Num();
 		const int32 NumRemappedBones = RemappedBones.Num();
 
@@ -829,9 +832,9 @@ mu::MeshPtr ConvertSkeletalMeshToMutable(const USkeletalMesh* InSkeletalMesh, co
 
 	// Vertices
 	TArray<FSoftSkinVertex> Vertices;
-	ImportedModel->LODModels[LODIndex].GetVertices(Vertices);
-	int32 VertexStart = ImportedModel->LODModels[LODIndex].Sections[SectionIndex].GetVertexBufferIndex();
-	int32 VertexCount = ImportedModel->LODModels[LODIndex].Sections[SectionIndex].GetNumVertices();
+	LODModel.GetVertices(Vertices);
+	int32 VertexStart = MeshSection.GetVertexBufferIndex();
+	int32 VertexCount = MeshSection.GetNumVertices();
 
 	MutableMesh->GetVertexBuffers().SetElementCount(VertexCount);
 
@@ -841,7 +844,7 @@ mu::MeshPtr ConvertSkeletalMeshToMutable(const USkeletalMesh* InSkeletalMesh, co
 
 	MutableMesh->GetVertexBuffers().SetBufferCount(VertexBuffersCount);
 
-	const int32 MaxSectionInfluences = ImportedModel->LODModels[LODIndex].Sections[SectionIndex].MaxBoneInfluences;
+	const int32 MaxSectionInfluences = MeshSection.MaxBoneInfluences;
 	const bool bUseUnlimitedInfluences = FGPUBaseSkinVertexFactory::UseUnlimitedBoneInfluences(MaxSectionInfluences, GenerationContext.Options.TargetPlatform);
 
 	if (bIgnoreSkeleton)
@@ -1288,14 +1291,13 @@ mu::MeshPtr ConvertSkeletalMeshToMutable(const USkeletalMesh* InSkeletalMesh, co
 				FirstElem.SourceAssetIndex = INDEX_NONE;
 			}
 
-			const FSkelMeshSection& SectionData = ImportedModel->LODModels[LODIndex].Sections[SectionIndex];
-			const TArray<FMeshToMeshVertData>& ClothMappingData = SectionData.ClothMappingDataLODs[0];
+			const TArray<FMeshToMeshVertData>& ClothMappingData = MeshSection.ClothMappingDataLODs[0];
 
 
 			// Similar test as the one used on FSkeletalMeshObjectGPUSkin::FVertexFactoryData::InitAPEXClothVertexFactories
 			// Here should work as expexted, but in the reference code I'm not sure it always works. It is worth investigate
 			// in that direction if at some point multiple influences don't work as expected.
-			const bool bUseMutlipleInfluences = ClothMappingData.Num() > SectionData.NumVertices;
+			const bool bUseMutlipleInfluences = ClothMappingData.Num() > MeshSection.NumVertices;
 
 			// Constant defined in ClothMeshUtils.cpp with the following comment:
 			// // This must match NUM_INFLUENCES_PER_VERTEX in GpuSkinCacheComputeShader.usf and GpuSkinVertexFactory.ush
@@ -1367,12 +1369,12 @@ mu::MeshPtr ConvertSkeletalMeshToMutable(const USkeletalMesh* InSkeletalMesh, co
 		const int32 BoneWeightsSize = MutableBonesPerVertex * BoneWeightTypeSizeBytes;
 		const int32 SkinWeightProfileVertexSize = sizeof(int32) + BoneIndicesSize + BoneWeightsSize;
 
-		const int32 MaxSectionBoneMapIndex = ImportedModel->LODModels[LODIndex].Sections[SectionIndex].BoneMap.Num();
+		const int32 MaxSectionBoneMapIndex = MeshSection.BoneMap.Num();
 
 		const TArray<FSkinWeightProfileInfo>& SkinWeightProfilesInfo = InSkeletalMesh->GetSkinWeightProfiles();
 		for (const FSkinWeightProfileInfo& Profile : SkinWeightProfilesInfo)
 		{
-			const FImportedSkinWeightProfileData* ImportedProfileData = ImportedModel->LODModels[LODIndex].SkinWeightProfiles.Find(Profile.Name);
+			const FImportedSkinWeightProfileData* ImportedProfileData = LODModel.SkinWeightProfiles.Find(Profile.Name);
 			if (!ImportedProfileData)
 			{
 				continue;
@@ -1457,87 +1459,49 @@ mu::MeshPtr ConvertSkeletalMeshToMutable(const USkeletalMesh* InSkeletalMesh, co
 
 	// Indices
 	{
-		int IndexStart = ImportedModel->LODModels[LODIndex].Sections[SectionIndex].BaseIndex;
-		int IndexCount = ImportedModel->LODModels[LODIndex].Sections[SectionIndex].NumTriangles * 3;
+		const uint32 IndexStart = MeshSection.BaseIndex;
+		const uint32 IndexCount = MeshSection.NumTriangles * 3;
 		MutableMesh->GetIndexBuffers().SetBufferCount(1);
 		MutableMesh->GetIndexBuffers().SetElementCount(IndexCount);
-		MutableMesh->GetFaceBuffers().SetElementCount(IndexCount / 3);
+		MutableMesh->GetFaceBuffers().SetElementCount(MeshSection.NumTriangles);
 
 		using namespace mu;
-		// For some reason, the indices in 4.25 (and 4.24) are in different order in the Imported and Rendering data structures. The strange thing 
-		// is actually that the vertices in the imported model seem to match the rendering model indices. Maybe there is some mapping that
-		// we are missing, but for now this will do:
-		const int ElementSize = InSkeletalMesh->GetResourceForRendering()->LODRenderData[LODIndex].MultiSizeIndexContainer.GetDataTypeSize();
-		void* IndexDataPointer = InSkeletalMesh->GetResourceForRendering()->LODRenderData[LODIndex].MultiSizeIndexContainer.GetIndexBuffer()->GetPointerTo(IndexStart);
-		const int FinalElementSize = sizeof(uint32_t);
-		const int ChannelCount = 1;
+
+		check(LODModel.IndexBuffer.IsValidIndex(IndexStart) && LODModel.IndexBuffer.IsValidIndex(IndexStart + IndexCount - 1));
+		const uint32* IndexDataPtr = &LODModel.IndexBuffer[IndexStart];
+
+		const int32 FinalElementSize = sizeof(uint32_t);
+		const int32 ChannelCount = 1;
 		const MESH_BUFFER_SEMANTIC Semantics[ChannelCount] = { MBS_VERTEXINDEX };
-		const int SemanticIndices[ChannelCount] = { 0 };
+		const int32 SemanticIndices[ChannelCount] = { 0 };
 		// We force 32 bit indices, since merging meshes may create vertex buffers bigger than the initial mesh
 		// and for now the mutable runtime doesn't handle it.
 		// \TODO: go back to 16-bit indices when possible.
 		MESH_BUFFER_FORMAT Formats[ChannelCount] = { MBF_UINT32 };
-		const int Components[ChannelCount] = { 1 };
-		const int Offsets[ChannelCount] = { 0 };
+		const int32 Components[ChannelCount] = { 1 };
+		const int32 Offsets[ChannelCount] = { 0 };
 
 		MutableMesh->GetIndexBuffers().SetBuffer(0, FinalElementSize, ChannelCount, Semantics, SemanticIndices, Formats, Components, Offsets);
 
+		uint32* pDest = reinterpret_cast<uint32*>(MutableMesh->GetIndexBuffers().GetBufferData(0));
+
 		// 32-bit to 32-bit
-		if (ElementSize == 4)
+		uint32 VertexIndex = 0;
+		for (uint32 Index = 0; Index < IndexCount; ++Index)
 		{
-			uint32* pDest = reinterpret_cast<uint32*>(MutableMesh->GetIndexBuffers().GetBufferData(0));
-			const uint32* pSource = reinterpret_cast<const uint32*>(IndexDataPointer);
-
-			for (int i = 0; i < IndexCount; ++i)
+			VertexIndex = *IndexDataPtr - VertexStart;
+			if (ensureMsgf(VertexIndex < (uint32)VertexCount, TEXT("Mutable: VertexIndex >= VertexCount. VI [%d], VC [%d], VS [%d]. SKM [%s] LOD [%d] Section [%d]."),
+				VertexIndex, VertexCount, VertexStart,
+				*GetNameSafe(InSkeletalMesh), LODIndex, SectionIndex))
 			{
-				*pDest = *pSource - VertexStart;
-				check(*pDest < uint32(VertexCount));
-				++pDest;
-				++pSource;
+				*pDest = VertexIndex;
 			}
-		}
-		// 16-bit to 16-bit
-		else if (ElementSize == 2 && Formats[0] == MBF_UINT16)
-		{
-			uint16* pDest = reinterpret_cast<uint16*>(MutableMesh->GetIndexBuffers().GetBufferData(0));
-			const uint16* pSource = reinterpret_cast<const uint16*>(IndexDataPointer);
-
-			for (int i = 0; i < IndexCount; ++i)
-			{
-				*pDest = *pSource - VertexStart;
-				check(*pDest < uint32(VertexCount));
-				++pDest;
-				++pSource;
-			}
-		}
-		// 16-bit to 32-bit
-		else if (ElementSize == 2 && Formats[0] == MBF_UINT32)
-		{
-			uint32* pDest = reinterpret_cast<uint32*>(MutableMesh->GetIndexBuffers().GetBufferData(0));
-			const uint16* pSource = reinterpret_cast<const uint16*>(IndexDataPointer);
-
-			int32 VertexIndex = 0;
-			for (int i = 0; i < IndexCount; ++i)
+			else
 			{
-				VertexIndex = *pSource - VertexStart;
-				if (ensureMsgf(VertexIndex < VertexCount, TEXT("Mutable: VertexIndex >= VertexCount. VI [%d], VC [%d], VS [%d]. SKM [%s] LOD [%d] Section [%d]."),
-					VertexIndex, VertexCount, VertexStart,
-					*GetNameSafe(InSkeletalMesh), LODIndex, SectionIndex))
-				{
-					*pDest = (uint32)VertexIndex;
-				}
-				else
-				{
-					*pDest = 0;
-				}
-				++pDest;
-				++pSource;
+				*pDest = 0;
 			}
-		}
-		else
-		{
-			// Unsupported case!
-			check(false);
+			++pDest;
+			++IndexDataPtr;
 		}
 	}
 
