[Mutable] Fix Reference Skeletal Mesh not respecting the Placeholder setting on PIE.

#rnx

diff --git a/Engine/Plugins/Experimental/Mutable/Source/CustomizableObject/Private/MuCO/CustomizableSkeletalComponent.cpp b/Engine/Plugins/Experimental/Mutable/Source/CustomizableObject/Private/MuCO/CustomizableSkeletalComponent.cpp
--- a/Engine/Plugins/Experimental/Mutable/Source/CustomizableObject/Private/MuCO/CustomizableSkeletalComponent.cpp
+++ b/Engine/Plugins/Experimental/Mutable/Source/CustomizableObject/Private/MuCO/CustomizableSkeletalComponent.cpp
@@ -15,6 +15,7 @@
 #include "Rendering/SkeletalMeshRenderData.h"
 #include "UObject/ObjectSaveContext.h"
 #include "MuCO/UnrealPortabilityHelpers.h"
+#include "MuCO/CustomizableObject.h"
 
 #include UE_INLINE_GENERATED_CPP_BY_NAME(CustomizableSkeletalComponent)
 
@@ -321,7 +322,11 @@ void UCustomizableSkeletalComponent::TickComponent(float DeltaTime, ELevelTick T
 		USkeletalMesh* SkeletalMesh = CustomizableObjectInstance->GetSkeletalMesh(ComponentIndex);
 
 		// If not generated yet, conditionally set the SkeletalMesh of reference
-		if (!bInstanceGenerated && !bSkipSetReferenceSkeletalMesh)
+		if (!bInstanceGenerated && !bSkipSetReferenceSkeletalMesh
+#if WITH_EDITORONLY_DATA
+			&& CustomizableObject->bEnableUseRefSkeletalMeshAsPlaceholder
+#endif
+			)
 		{
 			// Can be nullptr
 			SkeletalMesh = CustomizableObject->GetRefSkeletalMesh(ComponentIndex);
