[Mutable] Allow adding modifiers to lower LODs when using automatic LODs Strategy.
Generate SharedSurfaceIDs for all NodeSurfaceNew and group them using a hash of modifiers.
Enable reuse materials between LODs by default. It can be turned off per NodeMaterial.
Fix material parameters order, preventing the reuse of materials.

#rnx

diff --git a/Engine/Plugins/Experimental/Mutable/Source/CustomizableObject/Private/MuCO/CustomizableInstancePrivateData.h b/Engine/Plugins/Experimental/Mutable/Source/CustomizableObject/Private/MuCO/CustomizableInstancePrivateData.h
--- a/Engine/Plugins/Experimental/Mutable/Source/CustomizableObject/Private/MuCO/CustomizableInstancePrivateData.h
+++ b/Engine/Plugins/Experimental/Mutable/Source/CustomizableObject/Private/MuCO/CustomizableInstancePrivateData.h
@@ -399,6 +399,11 @@ private:
 
 			FMutableMaterialPlaceHolderParam(const FName& InParamName, const int32 InLayerIndex, const FGeneratedTexture& InTexture)
 				: ParamName(InParamName), LayerIndex(InLayerIndex), Texture(InTexture), Type(EPlaceHolderParamType::Texture) {}
+
+			bool operator<(const FMutableMaterialPlaceHolderParam& Other) const
+			{
+				return Type < Other.Type || ParamName.CompareIndexes(Other.ParamName);
+			}
 		};
 
 		UMaterialInterface* ParentMaterial;
@@ -407,34 +412,8 @@ private:
 
 		void AddParam(const FMutableMaterialPlaceHolderParam& NewParam) { Params.Add(NewParam); }
 
-		FString GetSerialization() const
-		{
-			FString Serialization = ParentMaterial ? FString::FromInt(ParentMaterial->GetUniqueID()) : FString("null");
-
-			for (const FMutableMaterialPlaceHolderParam& Param : Params)
-			{
-				Serialization += FString("-");
-				Serialization += FString::FromInt((int32)Param.Type) + FString("_") + Param.ParamName.ToString() + FString("_")
-					+ FString::FromInt(Param.LayerIndex) + FString("_");
-
-				switch (Param.Type)
-				{
-				case EPlaceHolderParamType::Vector:
-					Serialization += Param.Vector.ToString();
-					break;
-
-				case EPlaceHolderParamType::Scalar:
-					Serialization += FString::Printf(TEXT("%f"), Param.Scalar);
-					break;
-
-				case EPlaceHolderParamType::Texture:
-					Serialization += FString::FromInt(Param.Texture.Texture->GetUniqueID());
-					break;
-				}
-			}
-
-			return Serialization;
-		}
+		// Return a hash of the material and its parameters
+		uint32 GetHash();
 	};
 };
 
diff --git a/Engine/Plugins/Experimental/Mutable/Source/CustomizableObject/Private/MuCO/CustomizableObjectInstance.cpp b/Engine/Plugins/Experimental/Mutable/Source/CustomizableObject/Private/MuCO/CustomizableObjectInstance.cpp
--- a/Engine/Plugins/Experimental/Mutable/Source/CustomizableObject/Private/MuCO/CustomizableObjectInstance.cpp
+++ b/Engine/Plugins/Experimental/Mutable/Source/CustomizableObject/Private/MuCO/CustomizableObjectInstance.cpp
@@ -5337,10 +5337,10 @@ void UCustomizableInstancePrivateData::BuildMaterials(const TSharedRef<FUpdateCo
 		SkeletalMesh->GetMaterials().Reset();
 
 		// Maps serializations of FMutableMaterialPlaceholder to Created Dynamic Material instances, used to reuse materials across LODs
-		TMap<FString, TSharedPtr<FMutableMaterialPlaceholder>> ReuseMaterialCache;
+		TMap<uint32, TSharedPtr<FMutableMaterialPlaceholder>> ReuseMaterialCache;
 
 		// Map of SharedSurfaceId to FMutableMaterialPlaceholder serialization key 
-		TMap<int32, FString> SharedSurfacesCache;
+		TMap<int32, uint32> SharedSurfacesCache;
 
 		MUTABLE_CPUPROFILER_SCOPE(BuildMaterials_LODLoop);
 
@@ -5378,7 +5378,7 @@ void UCustomizableInstancePrivateData::BuildMaterials(const TSharedRef<FUpdateCo
 				}
 
 				// Reuse surface between LODs when using AutomaticLODs from mesh.
-				if (const FString* SharedSurfaceSerialization = SharedSurfacesCache.Find(Surface.SurfaceId))
+				if (const uint32* SharedSurfaceSerialization = SharedSurfacesCache.Find(Surface.SurfaceId))
 				{
 					TSharedPtr<FMutableMaterialPlaceholder>* FoundMaterialPlaceholder = ReuseMaterialCache.Find(*SharedSurfaceSerialization);
 					check(FoundMaterialPlaceholder);
@@ -5741,12 +5741,11 @@ void UCustomizableInstancePrivateData::BuildMaterials(const TSharedRef<FUpdateCo
 
 				MutableMaterialPlaceholder.ParentMaterial = MaterialTemplate;
 
-				const FString MaterialPlaceholderSerialization = MutableMaterialPlaceholder.GetSerialization();
-				check(MaterialPlaceholderSerialization != FString("null"));
+				const uint32 MaterialParameterHash = MutableMaterialPlaceholder.GetHash();
 
 				int32 MatIndex = INDEX_NONE;
 				
-				if (TSharedPtr<FMutableMaterialPlaceholder>* FoundMaterialPlaceholder = ReuseMaterialCache.Find(MaterialPlaceholderSerialization))
+				if (TSharedPtr<FMutableMaterialPlaceholder>* FoundMaterialPlaceholder = ReuseMaterialCache.Find(MaterialParameterHash))
 				{
 					MatIndex = (*FoundMaterialPlaceholder)->MatIndex;
 				}
@@ -5754,9 +5753,9 @@ void UCustomizableInstancePrivateData::BuildMaterials(const TSharedRef<FUpdateCo
 				{
 					MUTABLE_CPUPROFILER_SCOPE(BuildMaterials_CreateMaterial);
 
-					ReuseMaterialCache.Add(MaterialPlaceholderSerialization, MutableMaterialPlaceholderPtr);
+					ReuseMaterialCache.Add(MaterialParameterHash, MutableMaterialPlaceholderPtr);
 
-					SharedSurfacesCache.Add(Surface.SurfaceId, MaterialPlaceholderSerialization);
+					SharedSurfacesCache.Add(Surface.SurfaceId, MaterialParameterHash);
 
 					FGeneratedMaterial Material;
 					Material.SurfaceId = Surface.SurfaceId;
@@ -6599,6 +6598,41 @@ TSet<UAssetUserData*> UCustomizableObjectInstance::GetMergedAssetUserData(int32
 }
 
 
+uint32 UCustomizableInstancePrivateData::FMutableMaterialPlaceholder::GetHash()
+{
+
+	uint32 Hash = ParentMaterial ? ParentMaterial->GetUniqueID() : 0;
+
+	// Sort parameters before building the hash.
+	Params.Sort();
+
+	for (const FMutableMaterialPlaceHolderParam& Param : Params)
+	{
+		uint32 ParamHash = GetTypeHash(Param.ParamName);
+		ParamHash = HashCombineFast(ParamHash, (uint32)Param.LayerIndex);
+		ParamHash = HashCombineFast(ParamHash, (uint32)Param.Type);
+
+		switch (Param.Type)
+		{
+		case EPlaceHolderParamType::Vector:
+			ParamHash = HashCombineFast(ParamHash, GetTypeHash(Param.Vector));
+			break;
+
+		case EPlaceHolderParamType::Scalar:
+			ParamHash = HashCombineFast(ParamHash, GetTypeHash(Param.Scalar));
+			break;
+
+		case EPlaceHolderParamType::Texture:
+			ParamHash = HashCombineFast(ParamHash, Param.Texture.Texture->GetUniqueID());
+			break;
+		}
+
+		Hash = HashCombineFast(Hash, ParamHash);
+	}
+
+	return Hash;
+}
+
 #if WITH_EDITORONLY_DATA
 
 void UCustomizableInstancePrivateData::RegenerateImportedModel(USkeletalMesh* SkeletalMesh)
diff --git a/Engine/Plugins/Experimental/Mutable/Source/CustomizableObjectEditor/Classes/MuCOE/Nodes/CustomizableObjectNodeMaterial.h b/Engine/Plugins/Experimental/Mutable/Source/CustomizableObjectEditor/Classes/MuCOE/Nodes/CustomizableObjectNodeMaterial.h
--- a/Engine/Plugins/Experimental/Mutable/Source/CustomizableObjectEditor/Classes/MuCOE/Nodes/CustomizableObjectNodeMaterial.h
+++ b/Engine/Plugins/Experimental/Mutable/Source/CustomizableObjectEditor/Classes/MuCOE/Nodes/CustomizableObjectNodeMaterial.h
@@ -119,10 +119,9 @@ public:
 	UPROPERTY(EditAnywhere, BlueprintReadWrite, Category = CustomizableObject, meta = (ClampMin = "0"))
 	int32 MeshComponentIndex = 0;
 
-	/** Reuse material between LODs when using the LODStrategy - Automatic from mesh. If the material is extended or edited by a CO using
-	* custom LODs (LODStrategy), on that given LOD, reuse materials will not be applied and the material and its textures will be unique. */
+	/** Materials will be reused between LODs when possible unless explicitly disabled. */
 	UPROPERTY(EditAnywhere, Category = CustomizableObject)
-	bool bReuseMaterialBetweenLODs = false;
+	bool bReuseMaterialBetweenLODs = true;
 	
 	// UObject interface.
 	virtual void Serialize(FArchive& Ar) override;
diff --git a/Engine/Plugins/Experimental/Mutable/Source/CustomizableObjectEditor/Private/MuCOE/CustomizableObjectCompiler.cpp b/Engine/Plugins/Experimental/Mutable/Source/CustomizableObjectEditor/Private/MuCOE/CustomizableObjectCompiler.cpp
--- a/Engine/Plugins/Experimental/Mutable/Source/CustomizableObjectEditor/Private/MuCOE/CustomizableObjectCompiler.cpp
+++ b/Engine/Plugins/Experimental/Mutable/Source/CustomizableObjectEditor/Private/MuCOE/CustomizableObjectCompiler.cpp
@@ -696,6 +696,8 @@ mu::NodeObjectPtr FCustomizableObjectCompiler::GenerateMutableRoot(
     ActualRoot->RealTimeMorphSelectionOverrides = GenerationContext.RealTimeMorphTargetsOverrides;
 	GenerationContext.GenerateClippingCOInternalTags();
 
+	GenerationContext.GenerateSharedSurfacesUniqueIds();
+
 	// Generate ReferenceSkeletalMeshes data;
 	PopulateReferenceSkeletalMeshesData(GenerationContext);
 	
diff --git a/Engine/Plugins/Experimental/Mutable/Source/CustomizableObjectEditor/Private/MuCOE/GenerateMutableSource/GenerateMutableSource.cpp b/Engine/Plugins/Experimental/Mutable/Source/CustomizableObjectEditor/Private/MuCOE/GenerateMutableSource/GenerateMutableSource.cpp
--- a/Engine/Plugins/Experimental/Mutable/Source/CustomizableObjectEditor/Private/MuCOE/GenerateMutableSource/GenerateMutableSource.cpp
+++ b/Engine/Plugins/Experimental/Mutable/Source/CustomizableObjectEditor/Private/MuCOE/GenerateMutableSource/GenerateMutableSource.cpp
@@ -243,6 +243,53 @@ void FMutableGraphGenerationContext::GenerateClippingCOInternalTags()
 }
 
 
+void FMutableGraphGenerationContext::GenerateSharedSurfacesUniqueIds()
+{
+	int32 UniqueId = 0;
+
+	TArray<TArray<FSharedSurface>> NodeToSharedSurfaces;
+	SharedSurfaceIds.GenerateValueArray(NodeToSharedSurfaces);
+
+	TArray<bool> VisitedSurfaces;
+	for (TArray<FSharedSurface>& SharedSurfaces : NodeToSharedSurfaces)
+	{
+		const int32 NumSurfaces = SharedSurfaces.Num();
+		VisitedSurfaces.Init(false, NumSurfaces);
+		
+		// Iterate all surfaces for a given NodeMaterial and set the same SharedSurfaceId to those that are equal.
+		for (int32 SurfaceIndex = 0; SurfaceIndex < NumSurfaces; ++SurfaceIndex)
+		{
+			if (VisitedSurfaces[SurfaceIndex])
+			{
+				continue;
+			}
+
+			FSharedSurface& CurrentSharedSurface = SharedSurfaces[SurfaceIndex];
+			CurrentSharedSurface.NodeSurfaceNew->SetSharedSurfaceId(UniqueId);
+			VisitedSurfaces[SurfaceIndex] = true;
+
+			for (int32 AuxSurfaceIndex = SurfaceIndex; AuxSurfaceIndex < NumSurfaces && !CurrentSharedSurface.bMakeUnique; ++AuxSurfaceIndex)
+			{
+				if (VisitedSurfaces[AuxSurfaceIndex])
+				{
+					continue;
+				}
+
+				if (SharedSurfaces[AuxSurfaceIndex].NodeModifierIDs != CurrentSharedSurface.NodeModifierIDs)
+				{
+					continue;
+				}
+
+				SharedSurfaces[AuxSurfaceIndex].NodeSurfaceNew->SetSharedSurfaceId(UniqueId);
+				VisitedSurfaces[AuxSurfaceIndex] = true;
+			}
+
+			++UniqueId;
+		}
+	}
+}
+
+
 //void FMutableGraphGenerationContext::CheckPhysicsAssetInSkeletalMesh(const USkeletalMesh* SkeletalMesh)
 //{
 //	if (SkeletalMesh && SkeletalMesh->GetPhysicsAsset() && !DiscartedPhysicsAssetMap.Find(SkeletalMesh->GetPhysicsAsset()))
@@ -616,6 +663,7 @@ mu::NodeMeshApplyPosePtr CreateNodeMeshApplyPose(FMutableGraphGenerationContext&
 	return NodeMeshApplyPose;
 }
 
+
 // Convert a CustomizableObject Source Graph into a mutable source graph  
 mu::NodeObjectPtr GenerateMutableSource(const UEdGraphPin * Pin, FMutableGraphGenerationContext & GenerationContext, bool bPartialCompilation)
 {
@@ -803,14 +851,14 @@ mu::NodeObjectPtr GenerateMutableSource(const UEdGraphPin * Pin, FMutableGraphGe
 				.Init(nullptr, NumMeshComponentsInRoot);
 		}
 
-		int32 LastDefinedLOD = -1;
+		int32 FirstLOD = -1;
 		ObjectNode->SetLODCount(NumLODsInRoot);
-		for (int32 LODIndex = 0; LODIndex < NumLODsInRoot; ++LODIndex)
+		for (int32 CurrentLOD = 0; CurrentLOD < NumLODsInRoot; ++CurrentLOD)
 		{
-			GenerationContext.CurrentLOD = LODIndex;
+			GenerationContext.CurrentLOD = CurrentLOD;
 
 			mu::NodeLODPtr LODNode = new mu::NodeLOD();
-			ObjectNode->SetLOD(LODIndex, LODNode);
+			ObjectNode->SetLOD(CurrentLOD, LODNode);
 
 			LODNode->SetMessageContext(Node);
 			LODNode->SetComponentCount(GenerationContext.NumMeshComponentsInRoot);
@@ -825,10 +873,10 @@ mu::NodeObjectPtr GenerateMutableSource(const UEdGraphPin * Pin, FMutableGraphGe
 			{
 				mu::NodeComponentPtr& ComponentNode = ComponentNodes[MeshComponentIndex];
 
-				ComponentNode = Invoke([&GenerationContext, LODIndex, MeshComponentIndex]() -> mu::NodeComponentPtr
+				ComponentNode = Invoke([&GenerationContext, CurrentLOD, MeshComponentIndex]() -> mu::NodeComponentPtr
 				{
 					mu::NodeComponentNewPtr& ParentComponent =
-							GenerationContext.ComponentNewNode.Last().ComponentsPerLOD[LODIndex][MeshComponentIndex];
+							GenerationContext.ComponentNewNode.Last().ComponentsPerLOD[CurrentLOD][MeshComponentIndex];
 
 					if (!ParentComponent)
 					{
@@ -850,11 +898,9 @@ mu::NodeObjectPtr GenerateMutableSource(const UEdGraphPin * Pin, FMutableGraphGe
 
 			const bool bUseAutomaticLods = 
 					GenerationContext.CurrentAutoLODStrategy == ECustomizableObjectAutomaticLODStrategy::AutomaticFromMesh;
-			LastDefinedLOD = (LODIndex < NumLODs) && (LastDefinedLOD == INDEX_NONE || !bUseAutomaticLods) ? LODIndex : LastDefinedLOD;
+			FirstLOD = (CurrentLOD < NumLODs) && (FirstLOD == INDEX_NONE || !bUseAutomaticLods) ? CurrentLOD : FirstLOD;
 
-			// It turns out LODToGenerate is always LastDefinedLOD.
-			const int32& LODToGenerate = LastDefinedLOD;
-			if (LODToGenerate < 0)
+			if (FirstLOD < 0)
 			{
 				continue;
 			}
@@ -864,54 +910,66 @@ mu::NodeObjectPtr GenerateMutableSource(const UEdGraphPin * Pin, FMutableGraphGe
 				continue;
 			} 
 			
-			TArray<UEdGraphPin*> ConnectedLODPins = FollowInputPinArray(*TypedNodeObj->LODPin(LODToGenerate));
-
-			// Proccess non modifier material nodes.
-			for (int32 MeshComponentIndex = 0; MeshComponentIndex < NumMeshComponentsInRoot; ++MeshComponentIndex)
+			// Generate all relevant LODs for this object up until the current LODIndex.
+			for (int32 LODIndex = FirstLOD; LODIndex <= CurrentLOD; ++LODIndex)
 			{
-				GenerationContext.CurrentMeshComponent = MeshComponentIndex;
+				const UEdGraphPin* LODPin = TypedNodeObj->LODPin(LODIndex);
+				if (!LODPin)
+				{
+					continue;
+				}
 
-				for (UEdGraphPin* const ChildNodePin : ConnectedLODPins)
+				GenerationContext.FromLOD = LODIndex;
+
+				TArray<UEdGraphPin*> ConnectedLODPins = FollowInputPinArray(*LODPin);
+
+				// Proccess non modifier material nodes.
+				for (int32 MeshComponentIndex = 0; MeshComponentIndex < NumMeshComponentsInRoot; ++MeshComponentIndex)
 				{
-					// Modifiers are shared for all components and are processed per LOD and not component.
-					if (Cast<UCustomizableObjectNodeModifierBase>(ChildNodePin->GetOwningNode()))
-					{
-						continue;
-					}
+					GenerationContext.CurrentMeshComponent = MeshComponentIndex;
 
-					if (!AffectsCurrentComponent(ChildNodePin, GenerationContext))
+					for (UEdGraphPin* const ChildNodePin : ConnectedLODPins)
 					{
-						continue;
-					}
+						// Modifiers are shared for all components and are processed per LOD and not component.
+						if (Cast<UCustomizableObjectNodeModifierBase>(ChildNodePin->GetOwningNode()))
+						{
+							continue;
+						}
 
-					FMutableGraphSurfaceGenerationData DummySurfaceData;
-					mu::NodeSurfacePtr SurfaceNode = GenerateMutableSourceSurface(ChildNodePin, GenerationContext, DummySurfaceData);
+						if (!AffectsCurrentComponent(ChildNodePin, GenerationContext))
+						{
+							continue;
+						}
 
-					mu::NodeComponentPtr& ComponentNode = ComponentNodes[MeshComponentIndex];
+						FMutableGraphSurfaceGenerationData DummySurfaceData;
+						mu::NodeSurfacePtr SurfaceNode = GenerateMutableSourceSurface(ChildNodePin, GenerationContext, DummySurfaceData);
 
-					const int32 SurfaceCount = ComponentNode->GetSurfaceCount();
-					ComponentNode->SetSurfaceCount(SurfaceCount + 1);
-					ComponentNode->SetSurface(SurfaceCount, SurfaceNode.get());
-					ComponentNode->SetMessageContext(Node);
-				}
-			}
+						mu::NodeComponentPtr& ComponentNode = ComponentNodes[MeshComponentIndex];
 
-			// Process modfiers. Those are shared between different components in a lod.	
-			for (UEdGraphPin* const ChildNodePin : ConnectedLODPins)
-			{
-				// Set it to -1 to indicate we don't care about component id.
-				GenerationContext.CurrentMeshComponent = -1;
+						const int32 SurfaceCount = ComponentNode->GetSurfaceCount();
+						ComponentNode->SetSurfaceCount(SurfaceCount + 1);
+						ComponentNode->SetSurface(SurfaceCount, SurfaceNode.get());
+						ComponentNode->SetMessageContext(Node);
+					}
+				}
 
-				if (!Cast<UCustomizableObjectNodeModifierBase>(ChildNodePin->GetOwningNode()))
+				// Process modfiers. Those are shared between different components in a lod.	
+				for (UEdGraphPin* const ChildNodePin : ConnectedLODPins)
 				{
-					continue;
-				}
+					// Set it to -1 to indicate we don't care about component id.
+					GenerationContext.CurrentMeshComponent = -1;
+
+					if (!Cast<UCustomizableObjectNodeModifierBase>(ChildNodePin->GetOwningNode()))
+					{
+						continue;
+					}
 
-				mu::NodeModifierPtr ModifierNode = GenerateMutableSourceModifier(ChildNodePin, GenerationContext);
+					mu::NodeModifierPtr ModifierNode = GenerateMutableSourceModifier(ChildNodePin, GenerationContext);
 
-				const int32 ModifierCount = LODNode->GetModifierCount();
-				LODNode->SetModifierCount(ModifierCount + 1);
-				LODNode->SetModifier(ModifierCount, ModifierNode.get());
+					const int32 ModifierCount = LODNode->GetModifierCount();
+					LODNode->SetModifierCount(ModifierCount + 1);
+					LODNode->SetModifier(ModifierCount, ModifierNode.get());
+				}
 			}
 		}
 
@@ -1625,3 +1683,14 @@ void FMutableComponentInfo::AccumulateBonesToRemovePerLOD(const FComponentSettin
 
 
 #undef LOCTEXT_NAMESPACE
+
+FMutableGraphGenerationContext::FSharedSurface::FSharedSurface(uint8 InLOD, const mu::NodeSurfaceNewPtr& InNodeSurfaceNew)
+{
+	LOD = InLOD;
+	NodeSurfaceNew = InNodeSurfaceNew;
+}
+
+bool FMutableGraphGenerationContext::FSharedSurface::operator==(const FSharedSurface& o) const
+{
+	return NodeModifierIDs == o.NodeModifierIDs;
+}
diff --git a/Engine/Plugins/Experimental/Mutable/Source/CustomizableObjectEditor/Private/MuCOE/GenerateMutableSource/GenerateMutableSource.h b/Engine/Plugins/Experimental/Mutable/Source/CustomizableObjectEditor/Private/MuCOE/GenerateMutableSource/GenerateMutableSource.h
--- a/Engine/Plugins/Experimental/Mutable/Source/CustomizableObjectEditor/Private/MuCOE/GenerateMutableSource/GenerateMutableSource.h
+++ b/Engine/Plugins/Experimental/Mutable/Source/CustomizableObjectEditor/Private/MuCOE/GenerateMutableSource/GenerateMutableSource.h
@@ -626,6 +626,9 @@ struct FMutableGraphGenerationContext
 	* @return nothing */
 	void GenerateClippingCOInternalTags();
 
+	/** Generates shared surface IDs for all surface nodes. If one or more nodes are equal, they will use the same SharedSurfaceId */
+	void GenerateSharedSurfacesUniqueIds();
+
 	/** Check if the PhysicsAsset of a given SkeletalMesh has any SkeletalBodySetup with BoneNames not present in the
 	* InSkeletalMesh's RefSkeleton, if so, adds the PhysicsAsset to the DiscartedPhysicsAssetMap to display a warning later on */
 	//void CheckPhysicsAssetInSkeletalMesh(const USkeletalMesh* InSkeletalMesh);
@@ -688,6 +691,7 @@ struct FMutableGraphGenerationContext
 	};
 	TArray< ObjectParent > ComponentNewNode;
 
+	int32 FromLOD = 0; // LOD to append to the CurrentLOD when using AutomaticLODs. 
 	int32 CurrentLOD = 0;
 	int32 NumLODsInRoot = 0;
 	int32 CurrentMeshComponent = 0;
@@ -768,24 +772,21 @@ struct FMutableGraphGenerationContext
 	// Stores the parameters generated in the node tables
 	TMap<const class UCustomizableObjectNodeTable*, TArray<FGuid>> GeneratedParametersInTables;
 
-	struct FSharedSurfaces
+	struct FSharedSurface
 	{
-		FSharedSurfaces(int32 InSurfaceId, mu::NodeSurfaceNewPtr InNodeSurface)
-		{
-			check(InSurfaceId != INDEX_NONE);
-			check(InNodeSurface);
-			SharedSurfaceId = InSurfaceId;
-			NodeSurface = InNodeSurface;
-		}
+		FSharedSurface(uint8 InLOD, const mu::NodeSurfaceNewPtr& InNodeSurfaceNew);
+
+		bool operator==(const FSharedSurface& o) const;
 
-		int32 SharedSurfaceId = INDEX_NONE;
+		uint8 LOD = 0;
+		mu::NodeSurfaceNewPtr NodeSurfaceNew;
 
-		// NodeSurface of the current LOD 
-		mu::NodeSurfaceNewPtr NodeSurface;
+		bool bMakeUnique = false;
+		TArray<SIZE_T> NodeModifierIDs;
 	};
 
 	// UCustomizableObjectNodeMaterial material to SharedSurfaceId
-	TMap<UCustomizableObjectNodeMaterial*, FSharedSurfaces> SharedSurfaceIds;
+	TMap<UCustomizableObjectNodeMaterial*, TArray<FSharedSurface>> SharedSurfaceIds;
 
 	/** Extension Data constants are collected here */
 	FExtensionDataCompilerInterface ExtensionDataCompilerInterface;
diff --git a/Engine/Plugins/Experimental/Mutable/Source/CustomizableObjectEditor/Private/MuCOE/GenerateMutableSource/GenerateMutableSourceMesh.cpp b/Engine/Plugins/Experimental/Mutable/Source/CustomizableObjectEditor/Private/MuCOE/GenerateMutableSource/GenerateMutableSourceMesh.cpp
--- a/Engine/Plugins/Experimental/Mutable/Source/CustomizableObjectEditor/Private/MuCOE/GenerateMutableSource/GenerateMutableSourceMesh.cpp
+++ b/Engine/Plugins/Experimental/Mutable/Source/CustomizableObjectEditor/Private/MuCOE/GenerateMutableSource/GenerateMutableSourceMesh.cpp
@@ -61,9 +61,8 @@ void GetLODAndSectionForAutomaticLODs(const FMutableGraphGenerationContext& Cont
 		return;
 	}
 	
-	const int32 CompilingLODIndex = LODIndexConnected + Context.CurrentLOD;
-	
-	if (CompilingLODIndex == LODIndexConnected)
+	// When processing pins of the current LOD, indices will remain the same.
+	if (Context.CurrentLOD == Context.FromLOD)
 	{
 		return;
 	}
@@ -80,8 +79,13 @@ void GetLODAndSectionForAutomaticLODs(const FMutableGraphGenerationContext& Cont
 		return;
 	}
 	
-	const int32 SearchLODMaterialIndex = ImportedModel->LODModels[LODIndexConnected].Sections[SectionIndexConnected].MaterialIndex; // Material Index of the connected pin
+	const FSkelMeshSection& FromSection = ImportedModel->LODModels[LODIndexConnected].Sections[SectionIndexConnected];
+	const TArray<int32>& FromMaterialMap = SkeletalMesh.GetLODInfoArray()[LODIndexConnected].LODMaterialMap;
+	
+	// Material Index of the connected pin
+	const int32 SearchLODMaterialIndex = FromMaterialMap.IsValidIndex(SectionIndexConnected) ? FromMaterialMap[SectionIndexConnected] : FromSection.MaterialIndex;
 
+	const int32 CompilingLODIndex = LODIndexConnected + (Context.CurrentLOD - Context.FromLOD);
 	if (!ImportedModel->LODModels.IsValidIndex(CompilingLODIndex))
 	{
 		OutLODIndex = -1;
diff --git a/Engine/Plugins/Experimental/Mutable/Source/CustomizableObjectEditor/Private/MuCOE/GenerateMutableSource/GenerateMutableSourceSurface.cpp b/Engine/Plugins/Experimental/Mutable/Source/CustomizableObjectEditor/Private/MuCOE/GenerateMutableSource/GenerateMutableSourceSurface.cpp
--- a/Engine/Plugins/Experimental/Mutable/Source/CustomizableObjectEditor/Private/MuCOE/GenerateMutableSource/GenerateMutableSourceSurface.cpp
+++ b/Engine/Plugins/Experimental/Mutable/Source/CustomizableObjectEditor/Private/MuCOE/GenerateMutableSource/GenerateMutableSourceSurface.cpp
@@ -186,29 +186,25 @@ void SetSurfaceFormat( FMutableGraphGenerationContext& GenerationContext,
 }
 
 
-void UpdateSharedSurfaceId(FMutableGraphGenerationContext& GenerationContext, UCustomizableObjectNodeMaterial* NodeMaterial, mu::NodeSurfaceNewPtr InNodeSurface)
+void AddModifierToSharedSurface(FMutableGraphGenerationContext& GenerationContext, UCustomizableObjectNodeMaterial* NodeMaterial, const UCustomizableObjectNode& NodeModifier)
 {
-	// \TODO: Reusability is actually per-texture: if a texture is set to not use layouts (setting LayoutIndex to -1) then it can always be reused.
-	// Reuse materials between LODs when using automatics LODs, if texture layout management is disabled or if bReuseMaterials is enabled in the node material.
-	const bool bCanShareSurface = GenerationContext.CurrentAutoLODStrategy == ECustomizableObjectAutomaticLODStrategy::AutomaticFromMesh
-		&& NodeMaterial->bReuseMaterialBetweenLODs;
-
-	// Set shared surface Id to reuse materials and textures between LODs if automatic LODs from mesh is being used
-	if (bCanShareSurface && InNodeSurface)
+	if (!NodeMaterial || !NodeMaterial->bReuseMaterialBetweenLODs)
 	{
-		FMutableGraphGenerationContext::FSharedSurfaces& SharedSurface = GenerationContext.SharedSurfaceIds.FindOrAdd(NodeMaterial, { GenerationContext.SharedSurfaceIds.Num(), InNodeSurface });
-		SharedSurface.NodeSurface = InNodeSurface;
+		return;
+	}
 
-		InNodeSurface->SetSharedSurfaceId(SharedSurface.SharedSurfaceId);
+	TArray<FMutableGraphGenerationContext::FSharedSurface>* SharedSurfaces = GenerationContext.SharedSurfaceIds.Find(NodeMaterial);
+	if (!SharedSurfaces)
+	{
+		return;
 	}
 
-	// Invalidate the shared surface Id. 
-	else if (!bCanShareSurface && !InNodeSurface)
+	// Add the modifier to the key of modifiers per surface. Only add it if there are image operations involved.
+	for (FMutableGraphGenerationContext::FSharedSurface& SharedSurface : *SharedSurfaces)
 	{
-		FMutableGraphGenerationContext::FSharedSurfaces* SharedSurface = GenerationContext.SharedSurfaceIds.Find(NodeMaterial);
-		if (SharedSurface && SharedSurface->NodeSurface)
+		if (SharedSurface.LOD == GenerationContext.CurrentLOD)
 		{
-			SharedSurface->NodeSurface->SetSharedSurfaceId(INDEX_NONE);
+			SharedSurface.NodeModifierIDs.Add(reinterpret_cast<SIZE_T>(&NodeModifier));
 		}
 	}
 }
@@ -295,6 +291,11 @@ mu::NodeSurfacePtr GenerateMutableSourceSurface(const UEdGraphPin * Pin, FMutabl
 		mu::NodeSurfaceNewPtr SurfNode = new mu::NodeSurfaceNew();
 		Result = SurfNode;
 
+		// Add to the list of surfaces that could be reused between LODs for this NodeMaterial.
+		TArray<FMutableGraphGenerationContext::FSharedSurface>& SharedSurfaces = GenerationContext.SharedSurfaceIds.FindOrAdd(TypedNodeMat, {});
+		FMutableGraphGenerationContext::FSharedSurface& SharedSurface = SharedSurfaces.Add_GetRef(FMutableGraphGenerationContext::FSharedSurface(GenerationContext.CurrentLOD, SurfNode));
+		SharedSurface.bMakeUnique = !TypedNodeMat->bReuseMaterialBetweenLODs;
+
 		int32 ReferencedMaterialsIndex = -1;
 		if (TypedNodeMat->Material)
 		{
@@ -304,9 +305,6 @@ mu::NodeSurfacePtr GenerateMutableSourceSurface(const UEdGraphPin * Pin, FMutabl
 			ReferencedMaterialsIndex = GenerationContext.ReferencedMaterials.AddUnique(TypedNodeMat->Material);
 			// Used ReferencedMaterialsIndex instead of TypedNodeMat->Material->GetName() to prevent material name collisions
 
-			// Set shared surface Id to reuse materials and textures between LODs if automatic LODs from mesh is being used
-			UpdateSharedSurfaceId(GenerationContext, TypedNodeMat, SurfNode);
-
 			// Take slot name from skeletal mesh if one can be found, else leave empty.
 			// Keep Referenced Materials and Materail Slot Names synchronized even if no material name can be found.
 			const bool IsNewSlotName = GenerationContext.ReferencedMaterialSlotNames.Num() == ReferencedMaterialsIndex;
@@ -1139,9 +1137,8 @@ mu::NodeSurfacePtr GenerateMutableSourceSurface(const UEdGraphPin * Pin, FMutabl
 					SurfNode->SetImage(ImageIndex, ImageNode);
 				}
 
-				// Validate if the ParentMaterialNode can be shared between LODs.
-				UpdateSharedSurfaceId(GenerationContext, ParentMaterialNode, nullptr);
 			}
+			AddModifierToSharedSurface(GenerationContext, ParentMaterialNode, *TypedNodeExt);
 		
 			for (const FString& Tag : TypedNodeExt->Tags)
 			{
@@ -1183,8 +1180,7 @@ mu::NodeSurfacePtr GenerateMutableSourceSurface(const UEdGraphPin * Pin, FMutabl
 				MeshPatch->SetMessageContext(Node);
 			}
 
-			// Validate if the ParentMaterialNode can be shared between LODs.
-			UpdateSharedSurfaceId(GenerationContext, ParentMaterialNode, nullptr);
+			AddModifierToSharedSurface(GenerationContext, ParentMaterialNode, *TypedNodeRem);
 		}
 	}
 
@@ -1268,8 +1264,7 @@ mu::NodeSurfacePtr GenerateMutableSourceSurface(const UEdGraphPin * Pin, FMutabl
 				MeshPatch->SetMessageContext(Node);
 			}
 
-			// Validate if the ParentMaterialNode can be shared between LODs.
-			UpdateSharedSurfaceId(GenerationContext, ParentMaterialNode, nullptr);
+			AddModifierToSharedSurface(GenerationContext, ParentMaterialNode, *TypedNodeRemBlocks);
 		}
 	}
 
@@ -1383,8 +1378,7 @@ mu::NodeSurfacePtr GenerateMutableSourceSurface(const UEdGraphPin * Pin, FMutabl
 					SurfNode->SetPatch(ImageIndex, ImagePatchNode);
 				}
 
-				// Validate if the ParentMaterialNode can be shared between LODs.
-				UpdateSharedSurfaceId(GenerationContext, ParentMaterialNode, nullptr);
+				AddModifierToSharedSurface(GenerationContext, ParentMaterialNode, *TypedNodeEdit);
 			}
 		}
 	}
@@ -1467,8 +1461,7 @@ mu::NodeSurfacePtr GenerateMutableSourceSurface(const UEdGraphPin * Pin, FMutabl
 				}
 			}
 
-			// Validate if the ParentMaterialNode can be shared between LODs.
-			UpdateSharedSurfaceId(GenerationContext, ParentMaterialNode, nullptr);
+			AddModifierToSharedSurface(GenerationContext, ParentMaterialNode, *TypedNodeMorph);
 		}
 	}
 
diff --git a/Engine/Plugins/Experimental/Mutable/Source/CustomizableObjectEditor/Private/MuCOE/Nodes/CustomizableObjectNodeParentedMaterial.cpp b/Engine/Plugins/Experimental/Mutable/Source/CustomizableObjectEditor/Private/MuCOE/Nodes/CustomizableObjectNodeParentedMaterial.cpp
--- a/Engine/Plugins/Experimental/Mutable/Source/CustomizableObjectEditor/Private/MuCOE/Nodes/CustomizableObjectNodeParentedMaterial.cpp
+++ b/Engine/Plugins/Experimental/Mutable/Source/CustomizableObjectEditor/Private/MuCOE/Nodes/CustomizableObjectNodeParentedMaterial.cpp
@@ -31,22 +31,40 @@ TArray<UCustomizableObjectNodeMaterial*> FCustomizableObjectNodeParentedMaterial
 	}
 	
 	TArray<UCustomizableObjectNodeObject*> ParentObjectNodes = Node.GetParentObjectNodes(LOD);
-	for (int32 Index = 1; Index < ParentObjectNodes.Num(); ++Index) // Skip first Object node since we do not support directly material siblings.
+
+	ECustomizableObjectAutomaticLODStrategy LODStrategy = ECustomizableObjectAutomaticLODStrategy::Inherited;
+
+	// Iterate backwards, from the Root CO to the parent CO, to propagate the LODStrategy. 
+	for (int32 Index = ParentObjectNodes.Num() - 1; Index > 0; --Index) // Skip first Object node (Index == 0) since we do not support directly material siblings.
 	{
 		const UCustomizableObjectNodeObject* ParentObjectNode = ParentObjectNodes[Index];
-	
-		for (UCustomizableObjectNodeMaterial* MaterialNode : ParentObjectNode->GetMaterialNodes(LOD))
+		
+		if (ParentObjectNode->AutoLODStrategy != ECustomizableObjectAutomaticLODStrategy::Inherited)
+		{
+			LODStrategy = ParentObjectNode->AutoLODStrategy;
+		}
+		
+		// When using AutomaticFromMesh find all materials within range [0..LOD].
+		int32 LODIndex = LODStrategy == ECustomizableObjectAutomaticLODStrategy::AutomaticFromMesh ? 0 : LOD;
+
+		// If LODStrategy is set to AutomaticFromMesh, find MaterialNodes belonging to lower LODs.
+		for (; LODIndex <= LOD; ++LODIndex)
 		{
-			if (UCustomizableObjectNodeCopyMaterial* TypedCopyMaterialNode = Cast<UCustomizableObjectNodeCopyMaterial>(MaterialNode))
+			const TArray<UCustomizableObjectNodeMaterial*> MaterialNodes = ParentObjectNode->GetMaterialNodes(LODIndex);
+
+			for (UCustomizableObjectNodeMaterial* MaterialNode : MaterialNodes)
 			{
-				if (TypedCopyMaterialNode->GetMaterialNode())
+				if (UCustomizableObjectNodeCopyMaterial* TypedCopyMaterialNode = Cast<UCustomizableObjectNodeCopyMaterial>(MaterialNode))
 				{
-					Result.Add(TypedCopyMaterialNode);
+					if (TypedCopyMaterialNode->GetMaterialNode())
+					{
+						Result.Add(TypedCopyMaterialNode);
+					}
+				}
+				else if (UCustomizableObjectNodeMaterial* TypedMaterialNode = Cast<UCustomizableObjectNodeMaterial>(MaterialNode))
+				{
+					Result.Add(TypedMaterialNode);
 				}
-			}
-			else if (UCustomizableObjectNodeMaterial* TypedMaterialNode = Cast<UCustomizableObjectNodeMaterial>(MaterialNode))
-			{
-				Result.Add(TypedMaterialNode);
 			}
 		}
 	}
diff --git a/Engine/Plugins/Experimental/Mutable/Source/MutableTools/Private/MuT/CodeGenerator.cpp b/Engine/Plugins/Experimental/Mutable/Source/MutableTools/Private/MuT/CodeGenerator.cpp
--- a/Engine/Plugins/Experimental/Mutable/Source/MutableTools/Private/MuT/CodeGenerator.cpp
+++ b/Engine/Plugins/Experimental/Mutable/Source/MutableTools/Private/MuT/CodeGenerator.cpp
@@ -835,15 +835,15 @@ namespace mu
 				MeshOptions.OverrideLayouts = SharedMeshResults->GeneratedLayouts;
 			}
 
-			// Ensure UV islands remain within their main layout block on lower LODs to avoid unexpected reordering 
-			// of the layout blocks when reusing a surface between LODs. Used to fix small displacements on vertices
-			// that may cause them to fall on a different block.
-			MeshOptions.bClampUVIslands = bShareSurface;
-
 			// Normalize UVs if we're going to work with images and layouts.
 			const bool bNormalizeUVs = node.m_images.Num() && edits.Num();
 			MeshOptions.bNormalizeUVs = bNormalizeUVs;
 
+			// Ensure UV islands remain within their main layout block on lower LODs to avoid unexpected reordering 
+			// of the layout blocks when reusing a surface between LODs. Used to fix small displacements on vertices
+			// that may cause them to fall on a different block.
+			MeshOptions.bClampUVIslands = bShareSurface && bNormalizeUVs;
+
             GenerateMesh(MeshOptions, meshResults, pMesh );
             lastMeshOp = meshResults.meshOp;
             meshResults.extraMeshLayouts.SetNum( edits.Num() );
@@ -870,7 +870,7 @@ namespace mu
 						FMeshGenerationOptions MergedMeshOptions;
 						MergedMeshOptions.bUniqueVertexIDs = true;
 						MergedMeshOptions.bLayouts = true;
-						MergedMeshOptions.bClampUVIslands = bShareSurface;
+						MergedMeshOptions.bClampUVIslands = bShareSurface && bNormalizeUVs;
 						MergedMeshOptions.bNormalizeUVs = bNormalizeUVs;
 						MergedMeshOptions.State = m_currentStateIndex;
 						MergedMeshOptions.ActiveTags = e.node->m_tags;
