[Mutable] Temporal fix for a crash when using the new fast compilation mode.

#rnx

diff --git a/Engine/Plugins/Experimental/Mutable/Source/CustomizableObject/Private/MuCO/UnrealMutableImageProvider.cpp b/Engine/Plugins/Experimental/Mutable/Source/CustomizableObject/Private/MuCO/UnrealMutableImageProvider.cpp
--- a/Engine/Plugins/Experimental/Mutable/Source/CustomizableObject/Private/MuCO/UnrealMutableImageProvider.cpp
+++ b/Engine/Plugins/Experimental/Mutable/Source/CustomizableObject/Private/MuCO/UnrealMutableImageProvider.cpp
@@ -134,7 +134,7 @@ bool FUnrealMutableImageProvider::Tick(float DeltaTime)
 	TSoftObjectPtr<UTexture> TexturePtr = CO->ReferencedPassThroughTextures[Request->Id];
 
 	// This can cause a stall because of loading the asset.
-	UTexture2D* Texture = Cast<UTexture2D>( TexturePtr.Get() );
+	UTexture2D* Texture = Cast<UTexture2D>( TexturePtr.LoadSynchronous() );
 	if (!Texture)
 	{
 		// Failed to load the texture
