[mutable] New modifier node to clip meshes using a UV mask.
Tidy up the mesh clip code to put it in a cpp.
#rnx

diff --git a/Engine/Plugins/Experimental/Mutable/Source/CustomizableObject/Classes/MuCO/CustomizableObject.h b/Engine/Plugins/Experimental/Mutable/Source/CustomizableObject/Classes/MuCO/CustomizableObject.h
--- a/Engine/Plugins/Experimental/Mutable/Source/CustomizableObject/Classes/MuCO/CustomizableObject.h
+++ b/Engine/Plugins/Experimental/Mutable/Source/CustomizableObject/Classes/MuCO/CustomizableObject.h
@@ -1477,7 +1477,7 @@ private:
 	// This is a manual version number for the binary blobs in this asset.
 	// Increasing it invalidates all the previously compiled models.
 	// Warning: If while merging code both versions have changed, take the highest+1.
-	static const int32 CurrentSupportedVersion = 407;
+	static const int32 CurrentSupportedVersion = 408;
 
 public:
 
diff --git a/Engine/Plugins/Experimental/Mutable/Source/CustomizableObjectEditor/Classes/MuCOE/Nodes/CustomizableObjectNodeModifierClipWithUVMask.h b/Engine/Plugins/Experimental/Mutable/Source/CustomizableObjectEditor/Classes/MuCOE/Nodes/CustomizableObjectNodeModifierClipWithUVMask.h
new file mode 100644
--- /dev/null
+++ b/Engine/Plugins/Experimental/Mutable/Source/CustomizableObjectEditor/Classes/MuCOE/Nodes/CustomizableObjectNodeModifierClipWithUVMask.h
@@ -0,0 +1,44 @@
+// Copyright Epic Games, Inc. All Rights Reserved.
+
+#pragma once
+
+#include "MuCOE/Nodes/CustomizableObjectNodeModifierBase.h"
+
+#include "CustomizableObjectNodeModifierClipWithUVMask.generated.h"
+
+namespace ENodeTitleType { enum Type : int; }
+
+class UCustomizableObject;
+class UCustomizableObjectNodeRemapPins;
+class UEdGraphPin;
+
+UCLASS()
+class CUSTOMIZABLEOBJECTEDITOR_API UCustomizableObjectNodeModifierClipWithUVMask : public UCustomizableObjectNodeModifierBase
+{
+public:
+	GENERATED_BODY()
+
+	/** Materials in all other objects that activate this tags will be clipped with this UV mask. */
+	UPROPERTY(EditAnywhere, BlueprintReadWrite, Category = MeshToClip)
+	TArray<FString> Tags;
+
+	/** UV channel index that will be used to get the UVs to apply the clipping mask to. */
+	UPROPERTY(EditAnywhere, BlueprintReadWrite, Category = MeshToClip)
+	int32 UVChannelForMask = 0;
+
+	// EdGraphNode interface
+	virtual FText GetNodeTitle(ENodeTitleType::Type TitleType) const override;
+	virtual FLinearColor GetNodeTitleColor() const override;
+	virtual FText GetTooltipText() const override;
+
+	// UCustomizableObjectNode interface
+	virtual void AllocateDefaultPins(UCustomizableObjectNodeRemapPins* RemapPins) override;
+	virtual void PinConnectionListChanged(UEdGraphPin* Pin) override;
+
+	// UCustomizableObjectNodeMaterialBase interface
+	virtual UEdGraphPin* OutputPin() const override;
+
+	// Own interface
+	UEdGraphPin* ClipMaskPin() const;
+};
+
diff --git a/Engine/Plugins/Experimental/Mutable/Source/CustomizableObjectEditor/Private/MuCOE/EdGraphSchema_CustomizableObject.cpp b/Engine/Plugins/Experimental/Mutable/Source/CustomizableObjectEditor/Private/MuCOE/EdGraphSchema_CustomizableObject.cpp
--- a/Engine/Plugins/Experimental/Mutable/Source/CustomizableObjectEditor/Private/MuCOE/EdGraphSchema_CustomizableObject.cpp
+++ b/Engine/Plugins/Experimental/Mutable/Source/CustomizableObjectEditor/Private/MuCOE/EdGraphSchema_CustomizableObject.cpp
@@ -36,6 +36,7 @@
 #include "MuCOE/Nodes/CustomizableObjectNodeMeshClipDeform.h"
 #include "MuCOE/Nodes/CustomizableObjectNodeMeshClipMorph.h"
 #include "MuCOE/Nodes/CustomizableObjectNodeMeshClipWithMesh.h"
+#include "MuCOE/Nodes/CustomizableObjectNodeModifierClipWithUVMask.h"
 #include "MuCOE/Nodes/CustomizableObjectNodeMeshGeometryOperation.h"
 #include "MuCOE/Nodes/CustomizableObjectNodeMeshMorph.h"
 #include "MuCOE/Nodes/CustomizableObjectNodeMeshMorphStackApplication.h"
@@ -378,6 +379,7 @@ void UEdGraphSchema_CustomizableObject::GetGraphContextActions(FGraphContextMenu
 			ContextMenuBuilder.CreateTemplateNode<UCustomizableObjectNodeMeshClipMorph>(),
 			ContextMenuBuilder.CreateTemplateNode<UCustomizableObjectNodeMeshClipWithMesh>(),
 			ContextMenuBuilder.CreateTemplateNode<UCustomizableObjectNodeMeshClipDeform>(),
+			ContextMenuBuilder.CreateTemplateNode<UCustomizableObjectNodeModifierClipWithUVMask>(),
 			ContextMenuBuilder.CreateTemplateNode<UCustomizableObjectNodeTable>(),
 		};
 
diff --git a/Engine/Plugins/Experimental/Mutable/Source/CustomizableObjectEditor/Private/MuCOE/GenerateMutableSource/GenerateMutableSource.h b/Engine/Plugins/Experimental/Mutable/Source/CustomizableObjectEditor/Private/MuCOE/GenerateMutableSource/GenerateMutableSource.h
--- a/Engine/Plugins/Experimental/Mutable/Source/CustomizableObjectEditor/Private/MuCOE/GenerateMutableSource/GenerateMutableSource.h
+++ b/Engine/Plugins/Experimental/Mutable/Source/CustomizableObjectEditor/Private/MuCOE/GenerateMutableSource/GenerateMutableSource.h
@@ -669,7 +669,7 @@ struct FMutableGraphGenerationContext
 	TMap<class UCustomizableObjectNodeMaterial*, TArray<mu::NodeSurfaceNewPtr>> MapMaterialNodeToMutableSurfaceNodeArray;
 
 	// Map with pairs (Unreal Mutable clip mesh node, array with its corresponding Mutable mesh modifier nodes) to add tags
-	TMap<class UCustomizableObjectNodeMeshClipWithMesh*, TArray<mu::NodeModifierMeshClipWithMeshPtr>> MapClipMeshNodeToMutableClipMeshNodeArray;
+	TMap<class UCustomizableObjectNodeMeshClipWithMesh*, TArray<mu::Ptr<mu::NodeModifierMeshClipWithMesh>>> MapClipMeshNodeToMutableClipMeshNodeArray;
 
 	// Data used for MorphTarget reconstruction.
 	TArray<FMorphTargetInfo> ContributingMorphTargetsInfo;
diff --git a/Engine/Plugins/Experimental/Mutable/Source/CustomizableObjectEditor/Private/MuCOE/GenerateMutableSource/GenerateMutableSourceMesh.h b/Engine/Plugins/Experimental/Mutable/Source/CustomizableObjectEditor/Private/MuCOE/GenerateMutableSource/GenerateMutableSourceMesh.h
--- a/Engine/Plugins/Experimental/Mutable/Source/CustomizableObjectEditor/Private/MuCOE/GenerateMutableSource/GenerateMutableSourceMesh.h
+++ b/Engine/Plugins/Experimental/Mutable/Source/CustomizableObjectEditor/Private/MuCOE/GenerateMutableSource/GenerateMutableSourceMesh.h
@@ -89,5 +89,4 @@ mu::Ptr<mu::Mesh> BuildMorphedMutableMesh(const UEdGraphPin* BaseSourcePin, cons
  * @param bLinkedToExtendMaterial 
  * @param bOnlyConnectedLOD Corrected LOD and Section will unconditionally always be the connected ones.
  * @return  Mutable Mesh Node. */
-mu::NodeMeshPtr GenerateMutableSourceMesh(const UEdGraphPin* Pin, FMutableGraphGenerationContext& GenerationContext,
-	FMutableGraphMeshGenerationData& MeshData, bool bLinkedToExtendMaterial, bool bOnlyConnectedLOD);
+mu::Ptr<mu::NodeMesh> GenerateMutableSourceMesh(const UEdGraphPin* Pin, FMutableGraphGenerationContext& GenerationContext, FMutableGraphMeshGenerationData& MeshData, bool bLinkedToExtendMaterial, bool bOnlyConnectedLOD);
diff --git a/Engine/Plugins/Experimental/Mutable/Source/CustomizableObjectEditor/Private/MuCOE/GenerateMutableSource/GenerateMutableSourceModifier.cpp b/Engine/Plugins/Experimental/Mutable/Source/CustomizableObjectEditor/Private/MuCOE/GenerateMutableSource/GenerateMutableSourceModifier.cpp
--- a/Engine/Plugins/Experimental/Mutable/Source/CustomizableObjectEditor/Private/MuCOE/GenerateMutableSource/GenerateMutableSourceModifier.cpp
+++ b/Engine/Plugins/Experimental/Mutable/Source/CustomizableObjectEditor/Private/MuCOE/GenerateMutableSource/GenerateMutableSourceModifier.cpp
@@ -5,16 +5,19 @@
 #include "Engine/StaticMesh.h"
 #include "MuCOE/CustomizableObjectCompiler.h"
 #include "MuCOE/GenerateMutableSource/GenerateMutableSourceMesh.h"
+#include "MuCOE/GenerateMutableSource/GenerateMutableSourceImage.h"
 #include "MuCOE/GraphTraversal.h"
 #include "MuCOE/MutableUtils.h"
 #include "MuCOE/Nodes/CustomizableObjectNodeMeshClipDeform.h"
 #include "MuCOE/Nodes/CustomizableObjectNodeMeshClipMorph.h"
 #include "MuCOE/Nodes/CustomizableObjectNodeMeshClipWithMesh.h"
+#include "MuCOE/Nodes/CustomizableObjectNodeModifierClipWithUVMask.h"
 #include "MuCOE/Nodes/CustomizableObjectNodeStaticMesh.h"
 #include "MuR/Mesh.h"
 #include "MuT/NodeMeshTransform.h"
 #include "MuT/NodeModifierMeshClipDeform.h"
 #include "MuT/NodeModifierMeshClipMorphPlane.h"
+#include "MuT/NodeModifierMeshClipWithUVMask.h"
 
 #define LOCTEXT_NAMESPACE "CustomizableObjectEditor"
 
@@ -110,7 +113,7 @@ mu::NodeModifierPtr GenerateMutableSourceModifier(const UEdGraphPin * Pin, FMuta
 		// needs to be different for each object. If it were added to the Generated cache, all the objects would get the same.
 		bDoNotAddToGeneratedCache = true;
 
-		mu::NodeModifierMeshClipWithMeshPtr ClipNode = new mu::NodeModifierMeshClipWithMesh();
+		mu::Ptr<mu::NodeModifierMeshClipWithMesh> ClipNode = new mu::NodeModifierMeshClipWithMesh();
 		Result = ClipNode;
 
 		if (const UEdGraphPin* ConnectedPin = FollowInputPin(*TypedNodeClipMesh->ClipMeshPin()))
@@ -133,6 +136,8 @@ mu::NodeModifierPtr GenerateMutableSourceModifier(const UEdGraphPin * Pin, FMuta
 				}
 				else
 				{
+					// TODO: We support the clip mesh not being constant. This message is not precise enough. It should say that it hasn't been 
+					// possible to check if the mesh is closed or not.
 					GenerationContext.Compiler->CompilerLog(LOCTEXT("UnimplementedNode", "Node type not implemented yet."), MeshData.Node);
 				}
 
@@ -165,16 +170,16 @@ mu::NodeModifierPtr GenerateMutableSourceModifier(const UEdGraphPin * Pin, FMuta
 
 		for (const FString& Tag : TypedNodeClipMesh->Tags)
 		{
-			 ClipNode->AddTag(StringCast<ANSICHAR>(*Tag).Get());
+			 ClipNode->AddTag(Tag);
 		}
 
 		if (TypedNodeClipMesh->CustomizableObjectToClipWith != nullptr)
 		{
-			TArray<mu::NodeModifierMeshClipWithMeshPtr>* ArrayDataPtr = GenerationContext.MapClipMeshNodeToMutableClipMeshNodeArray.Find(Cast<UCustomizableObjectNodeMeshClipWithMesh>(Pin->GetOwningNode()));
+			TArray<mu::Ptr<mu::NodeModifierMeshClipWithMesh>>* ArrayDataPtr = GenerationContext.MapClipMeshNodeToMutableClipMeshNodeArray.Find(Cast<UCustomizableObjectNodeMeshClipWithMesh>(Pin->GetOwningNode()));
 
 			if (ArrayDataPtr == nullptr)
 			{
-				TArray<mu::NodeModifierMeshClipWithMeshPtr> ArrayData;
+				TArray<mu::Ptr<mu::NodeModifierMeshClipWithMesh>> ArrayData;
 				ArrayData.Add(ClipNode);
 				UCustomizableObjectNodeMeshClipWithMesh* CastedNode = Cast<UCustomizableObjectNodeMeshClipWithMesh>(Pin->GetOwningNode());
 				GenerationContext.MapClipMeshNodeToMutableClipMeshNodeArray.Add(CastedNode, ArrayData);
@@ -186,6 +191,32 @@ mu::NodeModifierPtr GenerateMutableSourceModifier(const UEdGraphPin * Pin, FMuta
 		}
 	}
 
+	else if (const UCustomizableObjectNodeModifierClipWithUVMask* TypedNodeClipUVMask = Cast<UCustomizableObjectNodeModifierClipWithUVMask>(Node))
+	{
+		// This modifier can be connected to multiple objects, so the compiled node
+		// needs to be different for each object. If it were added to the Generated cache, all the objects would get the same.
+		bDoNotAddToGeneratedCache = true;
+
+		mu::Ptr<mu::NodeModifierMeshClipWithUVMask> ClipNode = new mu::NodeModifierMeshClipWithUVMask();
+		Result = ClipNode;
+
+		if (const UEdGraphPin* ConnectedPin = FollowInputPin(*TypedNodeClipUVMask->ClipMaskPin()))
+		{
+			FMutableGraphMeshGenerationData DummyMeshData;
+
+			mu::Ptr<mu::NodeImage> ClipMask = GenerateMutableSourceImage(ConnectedPin, GenerationContext, 0);
+
+			ClipNode->SetClipMask(ClipMask.get());
+		}
+
+		ClipNode->SetLayoutIndex(TypedNodeClipUVMask->UVChannelForMask);
+
+		for (const FString& Tag : TypedNodeClipUVMask->Tags)
+		{
+			ClipNode->AddTag(Tag);
+		}
+	}
+
 	else
 	{
 		GenerationContext.Compiler->CompilerLog(LOCTEXT("UnimplementedNode", "Node type not implemented yet."), Node);
diff --git a/Engine/Plugins/Experimental/Mutable/Source/CustomizableObjectEditor/Private/MuCOE/Nodes/CustomizableObjectNodeModifierClipWithUVMask.cpp b/Engine/Plugins/Experimental/Mutable/Source/CustomizableObjectEditor/Private/MuCOE/Nodes/CustomizableObjectNodeModifierClipWithUVMask.cpp
new file mode 100644
--- /dev/null
+++ b/Engine/Plugins/Experimental/Mutable/Source/CustomizableObjectEditor/Private/MuCOE/Nodes/CustomizableObjectNodeModifierClipWithUVMask.cpp
@@ -0,0 +1,69 @@
+// Copyright Epic Games, Inc. All Rights Reserved.
+
+#include "MuCOE/Nodes/CustomizableObjectNodeModifierClipWithUVMask.h"
+
+#include "MuCOE/CustomizableObjectGraph.h"
+#include "MuCOE/EdGraphSchema_CustomizableObject.h"
+#include "MuCOE/ICustomizableObjectEditor.h"
+
+class UCustomizableObjectNodeRemapPins;
+
+#define LOCTEXT_NAMESPACE "CustomizableObjectEditor"
+
+
+
+void UCustomizableObjectNodeModifierClipWithUVMask::AllocateDefaultPins(UCustomizableObjectNodeRemapPins* RemapPins)
+{
+	const UEdGraphSchema_CustomizableObject* Schema = GetDefault<UEdGraphSchema_CustomizableObject>();
+
+	UEdGraphPin* ClipMaskPin = CustomCreatePin(EGPD_Input, Schema->PC_Image, FName("Clip Mask"));
+	ClipMaskPin->bDefaultValueIsIgnored = true;
+
+	UEdGraphPin* OutputPin_p = CustomCreatePin(EGPD_Output, Schema->PC_Material, FName("Modifier"));
+}
+
+
+FText UCustomizableObjectNodeModifierClipWithUVMask::GetNodeTitle(ENodeTitleType::Type TitleType) const
+{
+	return LOCTEXT("Clip_With_UV_Mask", "Clip With UV Mask");
+}
+
+
+FLinearColor UCustomizableObjectNodeModifierClipWithUVMask::GetNodeTitleColor() const
+{
+	const UEdGraphSchema_CustomizableObject* Schema = GetDefault<UEdGraphSchema_CustomizableObject>();
+	return Schema->GetPinTypeColor(Schema->PC_Material);
+}
+
+
+void UCustomizableObjectNodeModifierClipWithUVMask::PinConnectionListChanged(UEdGraphPin* Pin)
+{
+	Super::PinConnectionListChanged(Pin);
+
+	TSharedPtr<ICustomizableObjectEditor> Editor = GetGraphEditor();
+
+	if (Editor.IsValid())
+	{
+		Editor->UpdateGraphNodeProperties();
+	}
+}
+
+
+UEdGraphPin* UCustomizableObjectNodeModifierClipWithUVMask::OutputPin() const
+{
+	return FindPin(TEXT("Modifier"));
+}
+
+
+UEdGraphPin* UCustomizableObjectNodeModifierClipWithUVMask::ClipMaskPin() const
+{
+	return FindPin(TEXT("Clip Mask"));
+}
+
+
+FText UCustomizableObjectNodeModifierClipWithUVMask::GetTooltipText() const
+{
+	return LOCTEXT("Clip_Mask_Tooltip", "Removes the part of a material that has a UV layout inside a mask defined with a texture.\nIt only removes the faces that fall completely inside the mask, along with the vertices and edges that define only faces that are deleted.");
+}
+
+#undef LOCTEXT_NAMESPACE
diff --git a/Engine/Plugins/Experimental/Mutable/Source/CustomizableObjectEditor/Private/MuCOE/SCustomizableObjectEditorTagExplorer.cpp b/Engine/Plugins/Experimental/Mutable/Source/CustomizableObjectEditor/Private/MuCOE/SCustomizableObjectEditorTagExplorer.cpp
--- a/Engine/Plugins/Experimental/Mutable/Source/CustomizableObjectEditor/Private/MuCOE/SCustomizableObjectEditorTagExplorer.cpp
+++ b/Engine/Plugins/Experimental/Mutable/Source/CustomizableObjectEditor/Private/MuCOE/SCustomizableObjectEditorTagExplorer.cpp
@@ -14,6 +14,7 @@
 #include "MuCOE/Nodes/CustomizableObjectNodeMeshClipDeform.h"
 #include "MuCOE/Nodes/CustomizableObjectNodeMeshClipMorph.h"
 #include "MuCOE/Nodes/CustomizableObjectNodeMeshClipWithMesh.h"
+#include "MuCOE/Nodes/CustomizableObjectNodeModifierClipWithUVMask.h"
 #include "Widgets/Input/SButton.h"
 #include "Widgets/Input/SComboButton.h"
 #include "Widgets/Views/SListView.h"
@@ -298,6 +299,18 @@ void SCustomizableObjectEditorTagExplorer::FillTagInformation(UCustomizableObjec
 				}
 			}
 
+			if (UCustomizableObjectNodeModifierClipWithUVMask* TypedNodeClipMask = Cast<UCustomizableObjectNodeModifierClipWithUVMask>(Node))
+			{
+				for (int32 i = 0; i < TypedNodeClipMask->Tags.Num(); ++i)
+				{
+					NodeTags.Add(TypedNodeClipMask->Tags[i], TypedNodeClipMask);
+					if (Tags.Find(TypedNodeClipMask->Tags[i]) == INDEX_NONE)
+					{
+						Tags.Add(TypedNodeClipMask->Tags[i]);
+					}
+				}
+			}
+
 			if (UCustomizableObjectNodeExtendMaterial* TypedNodeExtend = Cast<UCustomizableObjectNodeExtendMaterial>(Node))
 			{
 				for (int32 i = 0; i < TypedNodeExtend->Tags.Num(); ++i)
diff --git a/Engine/Plugins/Experimental/Mutable/Source/CustomizableObjectEditor/Private/MuCOE/SMutableCodeViewer.h b/Engine/Plugins/Experimental/Mutable/Source/CustomizableObjectEditor/Private/MuCOE/SMutableCodeViewer.h
--- a/Engine/Plugins/Experimental/Mutable/Source/CustomizableObjectEditor/Private/MuCOE/SMutableCodeViewer.h
+++ b/Engine/Plugins/Experimental/Mutable/Source/CustomizableObjectEditor/Private/MuCOE/SMutableCodeViewer.h
@@ -216,6 +216,7 @@ private:
 	{
 		mu::OP_TYPE::ME_BINDSHAPE,
 		mu::OP_TYPE::ME_MASKCLIPMESH,
+		mu::OP_TYPE::ME_MASKCLIPUVMASK,
 		mu::OP_TYPE::ME_FORMAT,
 		mu::OP_TYPE::ME_DIFFERENCE,
 		mu::OP_TYPE::IM_MAKEGROWMAP,
diff --git a/Engine/Plugins/Experimental/Mutable/Source/MutableRuntime/Private/MuR/CodeRunner.cpp b/Engine/Plugins/Experimental/Mutable/Source/MutableRuntime/Private/MuR/CodeRunner.cpp
--- a/Engine/Plugins/Experimental/Mutable/Source/MutableRuntime/Private/MuR/CodeRunner.cpp
+++ b/Engine/Plugins/Experimental/Mutable/Source/MutableRuntime/Private/MuR/CodeRunner.cpp
@@ -1645,6 +1645,62 @@ namespace mu
             break;
         }
 
+		case OP_TYPE::ME_MASKCLIPUVMASK:
+		{
+			OP::MeshMaskClipUVMaskArgs args = pModel->GetPrivate()->m_program.GetOpArgs<OP::MeshMaskClipUVMaskArgs>(item.At);
+			switch (item.Stage)
+			{
+			case 0:
+			{
+				AddOp(FScheduledOp(item.At, item, 1),
+					FScheduledOp(args.Source, item),
+					FScheduledOp(args.Mask, item));
+				break;
+			}
+			case 1:
+			{
+				MUTABLE_CPUPROFILER_SCOPE(ME_MASKCLIPUVMASK_1)
+
+				Ptr<const Mesh> Source = LoadMesh(FCacheAddress(args.Source, item));
+				Ptr<const Image> Mask = LoadImage(FCacheAddress(args.Mask, item));
+
+				// Only if both are valid.
+				if (Source.get() && Mask.get())
+				{
+					Ptr<Mesh> Result = CreateMesh();
+
+					bool bOutSuccess = false;
+					MeshMaskClipUVMask(Result.get(), Source.get(), Mask.get(), args.LayoutIndex, bOutSuccess);
+
+					Release(Source);
+					Release(Mask);
+					if (!bOutSuccess)
+					{
+						Release(Result);
+						StoreMesh(item, nullptr);
+					}
+					else
+					{
+						StoreMesh(item, Result);
+					}
+				}
+				else
+				{
+					Release(Source);
+					Release(Mask);
+					StoreMesh(item, nullptr);
+				}
+
+				break;
+			}
+
+			default:
+				check(false);
+			}
+
+			break;
+		}
+
         case OP_TYPE::ME_MASKDIFF:
         {
 			OP::MeshMaskDiffArgs args = pModel->GetPrivate()->m_program.GetOpArgs<OP::MeshMaskDiffArgs>(item.At);
diff --git a/Engine/Plugins/Experimental/Mutable/Source/MutableRuntime/Private/MuR/CodeRunnerTasked.cpp b/Engine/Plugins/Experimental/Mutable/Source/MutableRuntime/Private/MuR/CodeRunnerTasked.cpp
--- a/Engine/Plugins/Experimental/Mutable/Source/MutableRuntime/Private/MuR/CodeRunnerTasked.cpp
+++ b/Engine/Plugins/Experimental/Mutable/Source/MutableRuntime/Private/MuR/CodeRunnerTasked.cpp
@@ -183,6 +183,7 @@ namespace mu
 		{   0,   0,   0,   0 },	// ME_MERGE
 		{   0,   0, -20,   0 },	// ME_INTERPOLATE
 		{   0, -10,   0,   0 },	// ME_MASKCLIPMESH
+		{   0, -10,   0,   0 },	// ME_MASKCLIPUVMASK
 		{   0, -10,   0,   0 },	// ME_MASKDIFF
 		{   0,   0, -10,   0 },	// ME_REMOVEMASK
 		{   0,   0,   0,   0 },	// ME_FORMAT
diff --git a/Engine/Plugins/Experimental/Mutable/Source/MutableRuntime/Private/MuR/OpMeshClipWithMesh.cpp b/Engine/Plugins/Experimental/Mutable/Source/MutableRuntime/Private/MuR/OpMeshClipWithMesh.cpp
new file mode 100644
--- /dev/null
+++ b/Engine/Plugins/Experimental/Mutable/Source/MutableRuntime/Private/MuR/OpMeshClipWithMesh.cpp
@@ -0,0 +1,1114 @@
+// Copyright Epic Games, Inc. All Rights Reserved.
+
+#pragma once
+
+#include "MuR/OpMeshClipWithMesh.h"
+
+#include "MuR/MeshPrivate.h"
+#include "MuR/ImagePrivate.h"
+#include "MuR/ConvertData.h"
+#include "MuR/Platform.h"
+#include "MuR/MutableTrace.h"
+#include "Math/Ray.h"
+#include "TriangleTypes.h"
+#include "BoxTypes.h"
+#include "Intersection/IntrRay3Triangle3.h"
+#include "MathUtil.h"
+#include "IntVectorTypes.h"
+
+#include "Math/UnrealMathUtility.h"
+
+#include "Spatial/PointHashGrid3.h"
+
+namespace mu { namespace  {
+
+    //---------------------------------------------------------------------------------------------
+    //! Create a map from vertices into vertices, collapsing vertices that have the same position, 
+	//! This version uses UE Containers to return.	
+    //---------------------------------------------------------------------------------------------
+	void MeshCreateCollapsedVertexMap(const Mesh* pMesh, TArray<int32>& OutCollapsedVertices, TArray<FVector3f>& OutVertices)
+	{
+		MUTABLE_CPUPROFILER_SCOPE(MeshCreateCollapsedVertexMap);
+
+		const int32 NumVertices = pMesh->GetVertexCount();
+
+		// Used to speed up vertex comparison
+		UE::Geometry::TPointHashGrid3f<int32> VertHash(0.01f, INDEX_NONE);
+		VertHash.Reserve(NumVertices);
+
+		// Info to collect. Vertices and collapsed vertices
+		OutVertices.SetNumUninitialized(NumVertices);
+		OutCollapsedVertices.Init(INDEX_NONE, NumVertices);
+
+		// Get Vertices
+		mu::UntypedMeshBufferIteratorConst ItPosition = mu::UntypedMeshBufferIteratorConst(pMesh->GetVertexBuffers(), mu::MBS_POSITION);
+		FVector3f* VertexData = OutVertices.GetData();
+
+		for (int32 VertexIndex = 0; VertexIndex < NumVertices; ++VertexIndex)
+		{
+			*VertexData = ItPosition.GetAsVec3f();
+			VertHash.InsertPointUnsafe(VertexIndex, *VertexData);
+
+			++ItPosition;
+			++VertexData;
+		}
+
+		// Find unique vertices
+		TArray<int32> NearbyVertices;
+		for (int32 VertexIndex = 0; VertexIndex < NumVertices; ++VertexIndex)
+		{
+			if (OutCollapsedVertices[VertexIndex] != INDEX_NONE)
+			{
+				continue;
+			}
+
+			const FVector3f& Vertex = OutVertices[VertexIndex];
+
+			NearbyVertices.Reset();
+			VertHash.FindPointsInBall(Vertex, TMathUtilConstants<float>::ZeroTolerance,
+				[&Vertex, &OutVertices](const int32& Other) -> float {return FVector3f::DistSquared(OutVertices[Other], Vertex); },
+				NearbyVertices);
+
+			for (int32 NearbyVertexIndex : NearbyVertices)
+			{
+				OutCollapsedVertices[NearbyVertexIndex] = VertexIndex;
+			}
+		}
+	}
+
+
+    //---------------------------------------------------------------------------------------------
+    //! Return true if a mesh is closed
+    //---------------------------------------------------------------------------------------------
+  //  bool MeshIsClosed( const Mesh* pMesh )
+  //  {
+  //      MUTABLE_CPUPROFILER_SCOPE(MeshIsClosed);
+
+  //      int vcount = pMesh->GetVertexCount();
+  //      int fcount = pMesh->GetFaceCount();
+
+  //      TArray<FVector3f> vertices;
+
+  //      // Map in ClipMesh from vertices to the one they are collapsed to because they are very similar, if they aren't collapsed then they are mapped to themselves
+		//TArray<int> collapsedVertexMap;
+
+  //      MeshCreateCollapsedVertexMap( pMesh, collapsedVertexMap, vertices );
+
+  //      // Acumulate edges
+  //      TMap< TPair<int,int>, int > faceCountPerEdge;
+
+  //      UntypedMeshBufferIteratorConst itClipMesh(pMesh->GetIndexBuffers(), MBS_VERTEXINDEX);
+  //      for (int f = 0; f < fcount; ++f)
+  //      {
+  //          int face[3];
+  //          face[0] = collapsedVertexMap[itClipMesh.GetAsUINT32()]; ++itClipMesh;
+  //          face[1] = collapsedVertexMap[itClipMesh.GetAsUINT32()]; ++itClipMesh;
+  //          face[2] = collapsedVertexMap[itClipMesh.GetAsUINT32()]; ++itClipMesh;
+
+  //          for (int e=0; e<3; ++e)
+  //          {
+  //              int v0 = face[e];
+  //              int v1 = face[(e+1)%3];
+
+  //              if(v0==v1)
+  //              {
+  //                  // Degenerated mesh
+  //                  return false;
+  //              }
+
+  //              TPair<int,int> edge;
+  //              edge.Key = FMath::Min( v0, v1 );
+  //              edge.Value = FMath::Max( v0, v1 );
+
+		//		if (faceCountPerEdge.Contains(edge))
+		//		{
+		//			faceCountPerEdge[edge]++;
+		//		}
+		//		else
+		//		{
+		//			faceCountPerEdge.Add(edge,1);
+		//		}
+  //          }
+  //      }
+
+
+  //      // See if every edge has 2 faces
+  //      for( const auto& e: faceCountPerEdge)
+  //      {
+  //          if (e.Value!=2)
+  //          {
+  //              return false;
+  //          }
+  //      }
+
+  //      return true;
+  //  }
+
+
+    //---------------------------------------------------------------------------------------------
+    //! Remove all unused vertices from a mesh, and fix its index buffers.
+    //---------------------------------------------------------------------------------------------
+    void MeshRemoveUnusedVertices( Mesh* pMesh )
+    {
+        MUTABLE_CPUPROFILER_SCOPE(MeshRemoveUnusedVertices);
+
+        // Mark used vertices
+		TArray<uint8> used;
+		used.SetNumZeroed(pMesh->GetVertexCount());
+        UntypedMeshBufferIteratorConst iti(pMesh->GetIndexBuffers(), MBS_VERTEXINDEX);
+        int IndexCount = pMesh->GetIndexCount();
+        for (int i = 0; i < IndexCount; ++i)
+        {
+            uint32 Index = iti.GetAsUINT32();
+            ++iti;
+            used[Index] = true;
+        }
+
+        // Build vertex map
+		TArray<int32> oldToNewVertex;
+		oldToNewVertex.SetNumUninitialized(pMesh->GetVertexCount());
+        int32 totalNewVertices = 0;
+        for (int32 v = 0; v<pMesh->GetVertexCount(); ++v)
+        {
+            if (used[v])
+            {
+                oldToNewVertex[v] = totalNewVertices;
+                ++totalNewVertices;
+            }
+            else
+            {
+                oldToNewVertex[v] = -1;
+            }
+        }
+
+        // Remove unused vertices and build index map
+        for (int b = 0; b<pMesh->GetVertexBuffers().GetBufferCount(); ++b)
+        {
+            int elemSize = pMesh->GetVertexBuffers().GetElementSize(b);
+            const uint8* pSourceData = pMesh->GetVertexBuffers().GetBufferData(b);
+            uint8* pData = pMesh->GetVertexBuffers().GetBufferData(b);
+            for (int v = 0; v<pMesh->GetVertexCount(); ++v)
+            {
+                if (oldToNewVertex[v]!=-1)
+                {
+                    uint8* pElemData = pData + elemSize*oldToNewVertex[v];
+                    const uint8* pElemSourceData = pSourceData + elemSize*v;
+                    // Avoid warning for overlapping memcpy in valgrind
+                    if (pElemData != pElemSourceData)
+                    {
+                        memcpy(pElemData, pElemSourceData, elemSize);
+                    }
+                }
+            }
+        }
+        pMesh->GetVertexBuffers().SetElementCount(totalNewVertices);
+
+        // Update indices
+        UntypedMeshBufferIteratorConst ito(pMesh->GetIndexBuffers(), MBS_VERTEXINDEX);
+        if (ito.GetFormat() == MBF_UINT32)
+        {
+            for (int i = 0; i < IndexCount; ++i)
+            {
+                uint32 Index = *(uint32*)ito.ptr();
+                int32 NewIndex = oldToNewVertex[Index];
+                check(NewIndex >= 0);
+                *(uint32*)ito.ptr() = (uint32)NewIndex;
+                ++ito;
+            }
+        }
+        else if (ito.GetFormat() == MBF_UINT16)
+        {
+            for (int i = 0; i < IndexCount; ++i)
+            {
+                uint16 Index = *(uint16*)ito.ptr();
+                int32 NewIndex = oldToNewVertex[Index];
+                check(NewIndex >= 0);
+                *(uint16*)ito.ptr() = (uint16)NewIndex;
+                ++ito;
+            }
+        }
+        else
+        {
+            checkf(false, TEXT("Format not implemented.") );
+        }
+
+
+        // \todo: face buffers?
+    }
+
+
+    //---------------------------------------------------------------------------------------------
+    //! Reference version
+    //---------------------------------------------------------------------------------------------
+  //  inline int get_num_intersections( const FVector3f& vertex,
+  //                                    const FVector3f& ray,
+  //                                    const TArray<FVector3f>& vertices,
+  //                                    const TArray<uint32>& faces,
+  //                                    TArray<int>& collapsedVertexMap,
+		//							  TArray<uint8>& vertex_already_intersected,
+  //                                    TMap<TPair<int, int>, bool>& edge_already_intersected,
+		//							  float dynamic_epsilon
+  //                                    )
+  //  {
+  //  	
+  //      MUTABLE_CPUPROFILER_SCOPE(get_num_intersections);
+  //  	
+  //      int num_intersections = 0;
+		//FVector3f intersection;
+
+		//FMemory::Memzero( vertex_already_intersected.GetData(), vertex_already_intersected.Num());
+
+  //      edge_already_intersected.Empty();
+
+  //  #define get_collapsed_vertex(vertex_idx) (vertices[collapsedVertexMap[vertex_idx]])
+
+  //      size_t fcount = faces.Num()/3;
+
+  //      // Check vertex against all ClipMorph faces
+  //      for (uint32 face = 0; face < fcount; ++face)
+  //      {
+  //          uint32 vertex_indexs[3] = { faces[3 * face], faces[3 * face + 1], faces[3 * face + 2] };
+
+		//	FVector3f v0 = get_collapsed_vertex(vertex_indexs[0]);
+		//	FVector3f v1 = get_collapsed_vertex(vertex_indexs[1]);
+		//	FVector3f v2 = get_collapsed_vertex(vertex_indexs[2]);
+
+  //          int out_intersected_vert, out_intersected_edge_v0, out_intersected_edge_v1;
+
+  //          if (rayIntersectsFace(vertex, ray, v0, v1, v2, intersection, out_intersected_vert, out_intersected_edge_v0, out_intersected_edge_v1, dynamic_epsilon))
+  //          {
+  //              bool vertex_not_intersected_before = true;
+  //              bool edge_not_intersected_before = true;
+
+  //              if (out_intersected_vert >= 0)
+  //              {
+  //                  int collapsed_vert_index = collapsedVertexMap[vertex_indexs[out_intersected_vert]];
+  //                  vertex_not_intersected_before = !vertex_already_intersected[collapsed_vert_index];
+
+  //                  vertex_already_intersected[collapsed_vert_index] = true;
+  //              }
+
+  //              if (out_intersected_edge_v0 >= 0)
+  //              {
+  //                  int collapsed_edge_vert_index0 = collapsedVertexMap[vertex_indexs[out_intersected_edge_v0]];
+  //                  int collapsed_edge_vert_index1 = collapsedVertexMap[vertex_indexs[out_intersected_edge_v1]];
+
+  //                  TPair<int, int> key;
+  //                  key.Key = collapsed_edge_vert_index0 <= collapsed_edge_vert_index1 ? collapsed_edge_vert_index0 : collapsed_edge_vert_index1;
+  //                  key.Value = collapsed_edge_vert_index1 > collapsed_edge_vert_index0 ? collapsed_edge_vert_index1 : collapsed_edge_vert_index0;
+  //                  edge_not_intersected_before = edge_already_intersected.Find(key) != nullptr;
+
+  //                  if (edge_not_intersected_before)
+  //                  {
+  //                      edge_already_intersected.Add( key, true );
+  //                  }
+  //              }
+
+  //              if (vertex_not_intersected_before && edge_not_intersected_before)
+  //              {
+  //                  num_intersections++;
+  //              }
+  //          }
+  //      }
+
+  //      return num_intersections;
+  //  }
+
+
+    //---------------------------------------------------------------------------------------------
+    //! Core Geometry version
+    //---------------------------------------------------------------------------------------------
+    inline int GetNumIntersections( const FRay3f& Ray,
+                                    const TArray<FVector3f>& Vertices,
+                                    const TArray<uint32>& Faces,
+                                    const TArray<int32>& CollapsedVertexMap,
+                                    TArray<uint8>& InOutVertexAlreadyIntersected,
+                                    TSet<uint64>& InOutEdgeAlreadyIntersected,
+									const float DynamicEpsilon )
+    {
+		using namespace UE::Geometry;    
+
+        MUTABLE_CPUPROFILER_SCOPE(GetNumIntersections);
+    	
+        int32 NumIntersections = 0;
+
+		FMemory::Memzero( InOutVertexAlreadyIntersected.GetData(), InOutVertexAlreadyIntersected.Num() );
+		InOutEdgeAlreadyIntersected.Empty();
+
+		auto GetCollapsedVertex = [&CollapsedVertexMap, &Vertices]( const uint32 V ) -> const FVector3f&
+		{
+			return Vertices[ CollapsedVertexMap[V] ];
+		};
+
+        int32 FaceCount = Faces.Num() / 3;
+
+		UE::Geometry::FIntrRay3Triangle3f Intersector(Ray, UE::Geometry::FTriangle3f());
+
+        // Check vertex against all ClipMorph faces
+        for (int32 Face = 0; Face < FaceCount; ++Face)
+        {
+            const uint32 VertexIndexs[3] = { Faces[3 * Face], Faces[3 * Face + 1], Faces[3 * Face + 2] };
+
+            const FVector3f& V0 = GetCollapsedVertex(VertexIndexs[0]);
+            const FVector3f& V1 = GetCollapsedVertex(VertexIndexs[1]);
+            const FVector3f& V2 = GetCollapsedVertex(VertexIndexs[2]);
+
+			Intersector.Triangle = UE::Geometry::FTriangle3f( V0, V1, V2 );
+
+			if ( Intersector.Find() )
+            {
+				// Find if close to edge using barycentric coordinates form intersector.
+	
+				// Is the Dynamic Epsilon needed?. Intersector bary coords are in the range [0.0, 1.0], 
+				// furthermore, IntrRay3Triangle3f::TriangleBaryCoords are double pressison, not sure 
+				// is intentional or is a bug ( is double even when the FReal type is float ) 
+				const bool IntersectsEdge01 = FMath::IsNearlyZero( Intersector.TriangleBaryCoords.Z, DynamicEpsilon );
+				const bool IntersectsEdge02 = FMath::IsNearlyZero( Intersector.TriangleBaryCoords.Y, DynamicEpsilon );
+				const bool IntersectsEdge12 = FMath::IsNearlyZero( Intersector.TriangleBaryCoords.X, DynamicEpsilon );
+
+				int32 IntersectedTriangleVertexId = -1;
+				IntersectedTriangleVertexId =  (IntersectsEdge01 & IntersectsEdge02) ? 0 : IntersectedTriangleVertexId;
+				IntersectedTriangleVertexId =  (IntersectsEdge01 & IntersectsEdge12) ? 1 : IntersectedTriangleVertexId;
+				IntersectedTriangleVertexId =  (IntersectsEdge02 & IntersectsEdge12) ? 2 : IntersectedTriangleVertexId;
+
+				bool bIsAlreadyIntersected = false;
+
+                if ( IntersectedTriangleVertexId >= 0 )
+                {
+                    const int32 CollapsedVertIndex = CollapsedVertexMap[VertexIndexs[IntersectedTriangleVertexId]];
+                    bIsAlreadyIntersected = InOutVertexAlreadyIntersected[CollapsedVertIndex] == 0;
+
+                    InOutVertexAlreadyIntersected[CollapsedVertIndex] = true;
+                }
+				else if ( IntersectsEdge01 | IntersectsEdge02 | IntersectsEdge12 )
+                {
+					int32 EdgeV0 = (IntersectsEdge01 | IntersectsEdge02) ? 0 : 1;
+					int32 EdgeV1 = IntersectsEdge01 ? 1 : 2;				
+
+                    const int32 CollapsedEdgeVertIndex0 = CollapsedVertexMap[VertexIndexs[EdgeV0]];
+                    const int32 CollapsedEdgeVertIndex1 = CollapsedVertexMap[VertexIndexs[EdgeV1]];
+
+                    const uint64 EdgeKey = ( ( (uint64)FMath::Max( CollapsedEdgeVertIndex0, CollapsedEdgeVertIndex1 ) ) << 32 ) | 
+							                   (uint64)FMath::Min( CollapsedEdgeVertIndex0, CollapsedEdgeVertIndex1 );
+	
+					InOutEdgeAlreadyIntersected.FindOrAdd( EdgeKey, &bIsAlreadyIntersected );
+                }
+
+                NumIntersections += (int32)(!bIsAlreadyIntersected);
+            }
+        }
+
+        return NumIntersections;
+    }
+
+    /** Core Geometry version */
+    void MeshClipMeshClassifyVertices(TArray<uint8>& VertexInClipMesh, const Mesh* pBase, const Mesh* pClipMesh)
+    {
+		using namespace UE::Geometry;
+
+        MUTABLE_CPUPROFILER_SCOPE(MeshClipMeshClassifyVertices);
+    	
+        // TODO: this fails in bro for some object. do it at graph construction time to report it 
+        // nicely.
+        //check( MeshIsClosed(pClipMesh) );
+
+        const int32 VCount = pClipMesh->GetVertexBuffers().GetElementCount();
+        const int32 FCount = pClipMesh->GetFaceCount();
+
+        int32 OrigVertCount = pBase->GetVertexBuffers().GetElementCount();
+
+        // Stores whether each vertex in the original mesh in the clip mesh volume
+        VertexInClipMesh.AddZeroed(OrigVertCount);
+
+        if (VCount == 0)
+        {
+            return;
+        }
+
+        TArray<FVector3f> Vertices; // ClipMesh vertex cache
+		Vertices.SetNumUninitialized(VCount);
+
+        TArray<uint32> Faces; // ClipMesh face cache
+		Faces.SetNumUninitialized(FCount * 3); 
+
+        // Map in ClipMesh from vertices to the one they are collapsed to because they are very 
+        // similar, if they aren't collapsed then they are mapped to themselves
+        TArray<int32> CollapsedVertexMap;
+		CollapsedVertexMap.AddUninitialized(VCount);
+
+        MeshCreateCollapsedVertexMap( pClipMesh, CollapsedVertexMap, Vertices );
+
+        // Create cache of the faces
+        UntypedMeshBufferIteratorConst ItClipMesh(pClipMesh->GetIndexBuffers(), MBS_VERTEXINDEX);
+
+        for ( int32 F = 0; F < FCount; ++F )
+        {
+            Faces[3 * F]     = ItClipMesh.GetAsUINT32(); ++ItClipMesh;
+            Faces[3 * F + 1] = ItClipMesh.GetAsUINT32(); ++ItClipMesh;
+            Faces[3 * F + 2] = ItClipMesh.GetAsUINT32(); ++ItClipMesh;
+        }
+
+
+        // Create a bounding box of the clip mesh
+		UE::Geometry::FAxisAlignedBox3f ClipMeshBoundingBox = UE::Geometry::FAxisAlignedBox3f::Empty();
+
+        for ( const FVector3f& Vert : Vertices )
+        {
+            ClipMeshBoundingBox.Contain( Vert );
+        }
+
+		// Dynamic distance epsilon to support different engines
+		const float MaxDimensionBoundingBox = ClipMeshBoundingBox.DiagonalLength();
+		// 0.000001 is the value that helps to achieve the dynamic epsilon, do not change it
+		const float DynamicEpsilon = 0.000001f * MaxDimensionBoundingBox * (MaxDimensionBoundingBox < 1.0f ? MaxDimensionBoundingBox : 1.0f);
+
+        // Create an acceleration grid to avoid testing all clip-mesh triangles.
+        // This assumes that the testing ray direction is Z
+        constexpr int32 GRID_SIZE = 8;
+        TUniquePtr<TArray<uint32>[]> GridFaces( new TArray<uint32>[GRID_SIZE*GRID_SIZE] );
+        const FVector2f GridCellSize = FVector2f( ClipMeshBoundingBox.Width(), ClipMeshBoundingBox.Height() ) / (float)GRID_SIZE;
+
+        for ( int32 I = 0; I < GRID_SIZE; ++I )
+        {
+            for ( int32 J = 0; J < GRID_SIZE; ++J )
+            {
+				const FVector2f BBoxMin = FVector2f(ClipMeshBoundingBox.Min.X, ClipMeshBoundingBox.Min.Y) + GridCellSize * FVector2f(I, J);
+
+                FAxisAlignedBox2f CellBox( BBoxMin, BBoxMin + GridCellSize );
+
+                TArray<uint32>& CellFaces = GridFaces[I+J*GRID_SIZE];
+                CellFaces.Empty(FCount/GRID_SIZE);
+                for (int32 F = 0; F < FCount; ++F)
+                {
+                    // Imprecise, conservative classification of faces.
+					const FVector3f& V0 = Vertices[ Faces[3*F + 0] ];
+					const FVector3f& V1 = Vertices[ Faces[3*F + 1] ];
+					const FVector3f& V2 = Vertices[ Faces[3*F + 2] ];
+
+					FAxisAlignedBox2f FaceBox;
+					FaceBox.Contain( FVector2f( V0.X, V0.Y ) );
+                    FaceBox.Contain( FVector2f( V1.X, V1.Y ) );
+                    FaceBox.Contain( FVector2f( V2.X, V2.Y ) );
+
+                    if (CellBox.Intersects(FaceBox))
+                    {
+                        CellFaces.Add( Faces[3*F + 0] );
+                        CellFaces.Add( Faces[3*F + 1] );
+                        CellFaces.Add( Faces[3*F + 2] );
+                    }
+                }
+            }
+        }	
+
+        // Now go through all vertices in the mesh and record whether they are inside or outside of the ClipMesh
+        uint32 DestVertexCount = pBase->GetVertexCount();
+
+        const FMeshBufferSet& MBSPriv2 = pBase->GetVertexBuffers();
+        for (int32 b = 0; b < MBSPriv2.m_buffers.Num(); ++b)
+        {
+            for (int32 c = 0; c < MBSPriv2.m_buffers[b].m_channels.Num(); ++c)
+            {
+                MESH_BUFFER_SEMANTIC Sem = MBSPriv2.m_buffers[b].m_channels[c].m_semantic;
+                int32 SemIndex = MBSPriv2.m_buffers[b].m_channels[c].m_semanticIndex;
+
+                UntypedMeshBufferIteratorConst It(pBase->GetVertexBuffers(), Sem, SemIndex);
+
+                TArray<uint8> VertexAlreadyIntersected;
+				VertexAlreadyIntersected.AddZeroed(VCount);
+
+                TSet<uint64> EdgeAlreadyIntersected;
+
+                switch ( Sem )
+                {
+                case MBS_POSITION:
+                    for (uint32 V = 0; V < DestVertexCount; ++V)
+                    {
+                        FVector3f Vertex(0.0f, 0.0f, 0.0f);
+                        for (int32 Offset = 0; Offset < 3; ++Offset)
+                        {
+                            ConvertData(Offset, &Vertex[0], MBF_FLOAT32, It.ptr(), It.GetFormat());
+                        }
+						
+						const FVector2f ClipBBoxMin = FVector2f( ClipMeshBoundingBox.Min.X, ClipMeshBoundingBox.Min.Y);
+						const FVector2f ClipBBoxSize = FVector2f( ClipMeshBoundingBox.Width(), ClipMeshBoundingBox.Height() );
+
+						const FVector2i HPos = FVector2i( ( ( FVector2f( Vertex.X, Vertex.Y ) - ClipBBoxMin ) / ClipBBoxSize ) * (float)GRID_SIZE );   
+						const FVector2i CellCoord = FVector2i( FMath::Clamp( HPos.X, 0, GRID_SIZE - 1 ), 
+														       FMath::Clamp( HPos.Y, 0, GRID_SIZE - 1 ) );
+
+                        // Early discard test: if the vertex is not inside the bounding box of the clip mesh, it won't be clipped.
+                        const bool bContainsVertex = ClipMeshBoundingBox.Contains( Vertex );
+
+                        if ( bContainsVertex )
+                        {
+                            // Optimised test	
+							
+							// Z-direction. Don't change this without reviewing the acceleration structure.
+							FVector3f RayDir = FVector3f(0.0f, 0.0f, 1.0f);	
+                            const int32 CellIndex = CellCoord.X + CellCoord.Y*GRID_SIZE;
+
+                            int32 NumIntersections = GetNumIntersections( 
+									FRay3f( Vertex, RayDir ), 
+									Vertices,
+									GridFaces[CellIndex],
+									CollapsedVertexMap, 
+									VertexAlreadyIntersected, 
+									EdgeAlreadyIntersected, 
+								    DynamicEpsilon);
+
+                            // Full test BLEH, debug
+//                            int32 FullNumIntersections = GetNumIntersections(
+//									FRay3f( Vertex, RayDir ), 
+//									Vertices,
+//									Faces,
+//									CollapsedVertexMap, 
+//									VertexAlreadyIntersected, 
+//									EdgeAlreadyIntersected, 
+//								    DynamicEpsilon);
+
+							VertexInClipMesh[V] = NumIntersections % 2 == 1;
+
+                            // This may be used to debug degenerated cases if the conditional above is also removed.
+                            // \todo: make sure it works well
+//                          if (!bContainsVertex && VertexInClipMesh[V])
+//                          {
+//                              assert(false);
+//                          }
+                        }
+
+                        ++It;
+                    }
+                    break;
+
+                default:
+                    break;
+                }
+            }
+        }
+    }	
+
+
+	void MeshUVMaskClassifyVertices(TArray<uint8>& VertexClipped, const Mesh* Base, const Image* Mask, const uint8 LayoutIndex)
+	{
+		using namespace UE::Geometry;
+
+		MUTABLE_CPUPROFILER_SCOPE(MeshUVMaskClassifyVertices);
+
+		int32 OrigVertCount = Base->GetVertexBuffers().GetElementCount();
+
+		// Stores whether each vertex in the original mesh in the clip mesh volume
+		VertexClipped.SetNum(0,false);
+		VertexClipped.SetNumZeroed(OrigVertCount);
+
+		// Now go through all vertices in the mesh and record whether they are inside or outside of the ClipMesh
+		uint32 DestVertexCount = Base->GetVertexCount();
+
+		const FMeshBufferSet& MBSPriv = Base->GetVertexBuffers();
+		for (int32 b = 0; b < MBSPriv.m_buffers.Num(); ++b)
+		{
+			const TArray<mu::MESH_BUFFER_CHANNEL>& Channels = MBSPriv.m_buffers[b].m_channels;
+
+			for (int32 c = 0; c < Channels.Num(); ++c)
+			{
+				MESH_BUFFER_SEMANTIC Sem = Channels[c].m_semantic;
+				if (Sem != MBS_TEXCOORDS)
+				{
+					continue;
+				}
+
+				int32 SemIndex = Channels[c].m_semanticIndex;
+				if (SemIndex != LayoutIndex)
+				{
+					continue;
+				}
+
+				UntypedMeshBufferIteratorConst It(Base->GetVertexBuffers(), Sem, SemIndex);
+				for (uint32 V = 0; V < DestVertexCount; ++V)
+				{
+					// \TODO: This could be optimized.
+					FVector2f UV = It.GetAsVec2f();
+
+					// \TODO: This could also be optimized
+					FVector4f Color = Mask->Sample(UV);
+
+					VertexClipped[V] = Color[0] >= 0.5f;
+
+					++It;
+				}
+
+				break;
+			}
+		}
+	}
+
+
+
+    //---------------------------------------------------------------------------------------------
+    //! Reference version
+    //---------------------------------------------------------------------------------------------
+  //  inline void MeshClipWithMesh_Reference(Mesh* Result, const Mesh* pBase, const Mesh* pClipMesh, bool& bOutSuccess)
+  //  {
+  //      MUTABLE_CPUPROFILER_SCOPE(MeshClipWithMesh_Reference);
+  //      
+		//bOutSuccess = true;
+
+  //      uint32 vcount = pClipMesh->GetVertexBuffers().GetElementCount();
+  //      if (!vcount)
+  //      {
+		//	bOutSuccess = false;
+		//	return; // Since there's nothing to clip against return a copy of the original base mesh
+  //      }
+
+  //      TArray<uint8> vertex_in_clip_mesh;  // Stores whether each vertex in the original mesh in in the clip mesh volume
+  //      MeshClipMeshClassifyVertices(vertex_in_clip_mesh, pBase, pClipMesh);
+
+		//Result->CopyFrom(*pBase);
+  //      // Now remove all the faces from the result mesh that have all the vertices outside the clip volume
+  //      UntypedMeshBufferIteratorConst itBase(Result->GetIndexBuffers(), MBS_VERTEXINDEX);
+  //      UntypedMeshBufferIterator itDest(Result->GetIndexBuffers(), MBS_VERTEXINDEX);
+  //      int aFaceCount = Result->GetFaceCount();
+
+  //      UntypedMeshBufferIteratorConst ito(Result->GetIndexBuffers(), MBS_VERTEXINDEX);
+  //      for (int f = 0; f < aFaceCount; ++f)
+  //      {
+  //          vec3<uint32> ov;
+  //          ov[0] = ito.GetAsUINT32(); ++ito;
+  //          ov[1] = ito.GetAsUINT32(); ++ito;
+  //          ov[2] = ito.GetAsUINT32(); ++ito;
+
+  //          bool all_verts_in = vertex_in_clip_mesh[ov[0]] && vertex_in_clip_mesh[ov[1]] && vertex_in_clip_mesh[ov[2]];
+
+  //          if (!all_verts_in)
+  //          {
+  //              if (itDest.ptr() != itBase.ptr())
+  //              {
+		//			FMemory::Memcpy(itDest.ptr(), itBase.ptr(), itBase.GetElementSize() * 3);
+  //              }
+
+  //              itDest += 3;
+  //          }
+
+  //          itBase += 3;
+  //      }
+
+  //      SIZE_T removedIndices = itBase - itDest;
+  //      check(removedIndices % 3 == 0);
+
+  //      Result->GetFaceBuffers().SetElementCount(aFaceCount - (int)removedIndices / 3);
+  //      Result->GetIndexBuffers().SetElementCount(aFaceCount * 3 - (int)removedIndices);
+
+  //      // TODO: Should redo/reorder the face buffer before SetElementCount since some deleted faces could be left and some remaining faces deleted.
+
+  //      //if (Result->GetFaceBuffers().GetElementCount() == 0)
+  //      //{
+  //      //	Result->CopyFrom(pBase); // If all faces have been discarded, return a copy of the unmodified mesh because unreal doesn't like empty meshes
+  //      //}
+
+  //      // [jordi] Remove unused vertices. This is necessary to avoid returning a mesh with vertices and no faces, which screws some
+  //      // engines like Unreal Engine 4.
+  //      MeshRemoveUnusedVertices(Result);
+  //  }
+
+
+	/** Make a mask with the indices of the vertices with 0 in the IncludedVertices array. */
+	void CreateMask(Mesh* Result, const Mesh* Base, const TArray<uint8>& IncludedVertices)
+	{
+		int32 MaskVertexCount = 0;
+		for (uint8 b : IncludedVertices)
+		{
+			if (!b)
+			{
+				++MaskVertexCount;
+			}
+		}
+
+		// Create the vertex buffer
+		{
+			Result->GetVertexBuffers().SetElementCount(MaskVertexCount);
+			Result->GetVertexBuffers().SetBufferCount(1);
+
+			// Vertex index channel
+			MESH_BUFFER_SEMANTIC Semantic = MBS_VERTEXINDEX;
+			int32 SemanticIndex = 0;
+			MESH_BUFFER_FORMAT Format = MBF_UINT32;
+			int32 Components = 1;
+			int32 Offsets = 0;
+
+			Result->GetVertexBuffers().SetBuffer
+			(
+				0,
+				sizeof(uint32),
+				1,
+				&Semantic,
+				&SemanticIndex,
+				&Format,
+				&Components,
+				&Offsets
+			);
+		}
+
+		// Fill the buffer
+		MeshBufferIterator<MBF_UINT32, uint32, 1> itMask(Result->GetVertexBuffers(), MBS_VERTEXINDEX, 0);
+		UntypedMeshBufferIteratorConst itBase(Base->GetVertexBuffers(), MBS_VERTEXINDEX, 0);
+		for (int32 v = 0; v<IncludedVertices.Num(); ++v)
+		{
+			if (!IncludedVertices[v])
+			{
+				(*itMask)[0] = itBase.GetAsUINT32();
+				++itMask;
+			}
+
+			++itBase;
+		}
+	}
+
+}} // namespace mu::anonymous
+
+
+namespace mu
+{
+
+    void MeshClipWithMesh(Mesh* Result, const Mesh* pBase, const Mesh* pClipMesh, bool& bOutSuccess)
+    {
+        MUTABLE_CPUPROFILER_SCOPE(MeshClipWithMesh);
+		bOutSuccess = true;
+
+        uint32 VCount = pClipMesh->GetVertexBuffers().GetElementCount();
+        if (!VCount)
+        {
+			bOutSuccess = false;
+			return; // OutSuccess false indicates the pBase can be reused in this case. 
+        }
+
+		Result->CopyFrom(*pBase);
+
+        TArray<uint8> VertexInClipMesh;  // Stores whether each vertex in the original mesh in in the clip mesh volume
+        MeshClipMeshClassifyVertices(VertexInClipMesh, pBase, pClipMesh);
+
+        // Now remove all the faces from the result mesh that have all the vertices outside the clip volume
+        UntypedMeshBufferIteratorConst ItBase(Result->GetIndexBuffers(), MBS_VERTEXINDEX);
+        UntypedMeshBufferIterator ItDest(Result->GetIndexBuffers(), MBS_VERTEXINDEX);
+        int32 AFaceCount = Result->GetFaceCount();
+
+        UntypedMeshBufferIteratorConst Ito(Result->GetIndexBuffers(), MBS_VERTEXINDEX);
+        for (int32 F = 0; F < AFaceCount; ++F)
+        {
+			uint32 OV[3] = {0, 0, 0};
+
+            OV[0] = Ito.GetAsUINT32(); ++Ito;
+            OV[1] = Ito.GetAsUINT32(); ++Ito;
+            OV[2] = Ito.GetAsUINT32(); ++Ito;
+
+            bool AllVertsIn = VertexInClipMesh[OV[0]] && VertexInClipMesh[OV[1]] && VertexInClipMesh[OV[2]];
+
+            if (!AllVertsIn)
+            {
+                if (ItDest.ptr() != ItBase.ptr())
+                {
+					FMemory::Memcpy(ItDest.ptr(), ItBase.ptr(), ItBase.GetElementSize() * 3);
+                }
+
+                ItDest += 3;
+            }
+
+            ItBase += 3;
+        }
+
+        SIZE_T RemovedIndices = ItBase - ItDest;
+        check(RemovedIndices % 3 == 0);
+
+        Result->GetFaceBuffers().SetElementCount(AFaceCount - (int32)RemovedIndices / 3);
+        Result->GetIndexBuffers().SetElementCount(AFaceCount * 3 - (int32)RemovedIndices);
+
+        // TODO: Should redo/reorder the face buffer before SetElementCount since some deleted faces could be left and some remaining faces deleted.
+
+        //if (Result->GetFaceBuffers().GetElementCount() == 0)
+        //{
+        //	Result->CopyFrom(pBase); // If all faces have been discarded, return a copy of the unmodified mesh because unreal doesn't like empty meshes
+        //}
+
+        // [jordi] Remove unused vertices. This is necessary to avoid returning a mesh with vertices and no faces, which screws some
+        // engines like Unreal Engine 4.
+        MeshRemoveUnusedVertices(Result);
+    }
+
+
+	void MeshMaskClipMesh(Mesh* Result, const Mesh* pBase, const Mesh* pClipMesh, bool& bOutSuccess)
+    {
+        MUTABLE_CPUPROFILER_SCOPE(MeshMaskClipMesh);
+
+		bOutSuccess = true;
+
+        uint32 VCount = pClipMesh->GetVertexBuffers().GetElementCount();
+        if (!VCount)
+        {
+			bOutSuccess = false;
+			return;
+        }
+
+        TArray<uint8> VertexInClipMesh;  // Stores whether each vertex in the original mesh in in the clip mesh volume
+        MeshClipMeshClassifyVertices( VertexInClipMesh, pBase, pClipMesh );
+
+        // We only remove vertices if all their faces are clipped
+        TArray<uint8> VertexWithFaceNotClipped;
+		VertexWithFaceNotClipped.AddZeroed( VertexInClipMesh.Num() );
+
+        UntypedMeshBufferIteratorConst Ito(pBase->GetIndexBuffers(), MBS_VERTEXINDEX);
+        int32 AFaceCount = pBase->GetFaceCount();
+        for (int32 F = 0; F < AFaceCount; ++F)
+        {
+			uint32 OV[3] = { 0, 0, 0 };
+
+            OV[0] = Ito.GetAsUINT32(); ++Ito;
+            OV[1] = Ito.GetAsUINT32(); ++Ito;
+            OV[2] = Ito.GetAsUINT32(); ++Ito;
+
+            bool bFaceClipped =
+                    VertexInClipMesh[OV[0]] &&
+                    VertexInClipMesh[OV[1]] &&
+                    VertexInClipMesh[OV[2]];
+
+            if (!bFaceClipped)
+            {
+                VertexWithFaceNotClipped[OV[0]] = true;
+                VertexWithFaceNotClipped[OV[1]] = true;
+                VertexWithFaceNotClipped[OV[2]] = true;
+            }
+        }
+
+        CreateMask(Result, pBase, VertexWithFaceNotClipped);
+    }
+
+
+	void MeshMaskClipUVMask(Mesh* Result, const Mesh* Base, const Image* Mask, uint8 LayoutIndex, bool& bOutSuccess)
+	{
+		MUTABLE_CPUPROFILER_SCOPE(MeshMaskClipMesh);
+
+		check(Result && Base && Mask);
+
+		bOutSuccess = true;
+
+		// Stores whether each vertex in the original mesh in in the clip mesh volume
+		TArray<uint8> VertexClipped; 
+
+		MeshUVMaskClassifyVertices(VertexClipped, Base, Mask, LayoutIndex);
+
+		// We only remove vertices if all their faces are clipped
+		TArray<uint8> VertexWithFaceNotClipped;
+		VertexWithFaceNotClipped.AddZeroed(VertexClipped.Num());
+
+		UntypedMeshBufferIteratorConst Ito(Base->GetIndexBuffers(), MBS_VERTEXINDEX);
+		int32 FaceCount = Base->GetFaceCount();
+		for (int32 F = 0; F < FaceCount; ++F)
+		{
+			uint32 OV[3] = { 0, 0, 0 };
+
+			OV[0] = Ito.GetAsUINT32(); ++Ito;
+			OV[1] = Ito.GetAsUINT32(); ++Ito;
+			OV[2] = Ito.GetAsUINT32(); ++Ito;
+
+			bool bFaceClipped =
+				VertexClipped[OV[0]] &&
+				VertexClipped[OV[1]] &&
+				VertexClipped[OV[2]];
+
+			if (!bFaceClipped)
+			{
+				VertexWithFaceNotClipped[OV[0]] = true;
+				VertexWithFaceNotClipped[OV[1]] = true;
+				VertexWithFaceNotClipped[OV[2]] = true;
+			}
+		}
+
+		CreateMask(Result, Base, VertexWithFaceNotClipped);
+	}
+
+
+  //  void MeshMaskClipMesh_Reference(Mesh* Result, const Mesh* pBase, const Mesh* pClipMesh )
+  //  {
+  //      MUTABLE_CPUPROFILER_SCOPE(MeshMaskClipMesh_Reference);
+
+  //      uint32 vcount = pClipMesh->GetVertexBuffers().GetElementCount();
+  //      if (!vcount)
+  //      {
+  //          return; // Since there's nothing to clip against return a copy of the original base mesh
+  //      }
+
+  //      TArray<uint8> vertex_in_clip_mesh;  // Stores whether each vertex in the original mesh in in the clip mesh volume
+  //      MeshClipMeshClassifyVertices( vertex_in_clip_mesh, pBase, pClipMesh );
+
+  //      // We only remove vertices if all their faces are clipped
+		//TArray<uint8> vertex_with_face_not_clipped;
+		//vertex_with_face_not_clipped.SetNumZeroed(vertex_in_clip_mesh.Num());
+
+  //      UntypedMeshBufferIteratorConst ito(pBase->GetIndexBuffers(), MBS_VERTEXINDEX);
+  //      int aFaceCount = pBase->GetFaceCount();
+  //      for (int f = 0; f < aFaceCount; ++f)
+  //      {
+  //          vec3<uint32> ov;
+  //          ov[0] = ito.GetAsUINT32(); ++ito;
+  //          ov[1] = ito.GetAsUINT32(); ++ito;
+  //          ov[2] = ito.GetAsUINT32(); ++ito;
+
+  //          bool faceClipped =
+  //                  vertex_in_clip_mesh[ov[0]] &&
+  //                  vertex_in_clip_mesh[ov[1]] &&
+  //                  vertex_in_clip_mesh[ov[2]];
+
+  //          if (!faceClipped)
+  //          {
+  //              vertex_with_face_not_clipped[ov[0]] = true;
+  //              vertex_with_face_not_clipped[ov[1]] = true;
+  //              vertex_with_face_not_clipped[ov[2]] = true;
+  //          }
+  //      }
+
+		//CreateMask(Result, pBase, vertex_with_face_not_clipped );
+  //  }
+
+
+	void MeshMaskDiff(Mesh* Result, const Mesh* pBase, const Mesh* pFragment, bool& bOutSuccess)
+    {
+        MUTABLE_CPUPROFILER_SCOPE(MeshMaskDiff);
+		bOutSuccess = true;
+
+        uint32 vcount = pFragment->GetVertexBuffers().GetElementCount();
+        if (!vcount)
+        {
+			bOutSuccess = false;
+            return;
+        }
+
+        int sourceFaceCount = pBase->GetFaceCount();
+        int sourceVertexCount = pBase->GetVertexCount();
+        int fragmentFaceCount = pFragment->GetFaceCount();
+
+
+        // Make a tolerance proportional to the mesh bounding box size
+        // TODO: Use precomputed bounding box
+        box< vec3<float> > aabbox;
+        if ( fragmentFaceCount > 0 )
+        {
+            MeshBufferIteratorConst<MBF_FLOAT32,float,3> itp( pFragment->GetVertexBuffers(), MBS_POSITION );
+
+            aabbox.min = *itp;
+            ++itp;
+
+            for ( int v=1; v<pFragment->GetVertexBuffers().GetElementCount(); ++v )
+            {
+                aabbox.Bound( *itp );
+                ++itp;
+            }
+        }
+        float tolerance = 1e-5f * length(aabbox.size);
+        Mesh::VERTEX_MATCH_MAP vertexMap;
+        pFragment->GetVertexMap( *pBase, vertexMap, tolerance );
+
+
+        // Classify the target faces in buckets along the Y axis
+#define NUM_BUCKETS	128
+#define AXIS		1
+        TArray<int> buckets[ NUM_BUCKETS ];
+        float bucketStart = aabbox.min[AXIS];
+        float bucketSize = aabbox.size[AXIS] / NUM_BUCKETS;
+
+        float bucketThreshold = ( 4 * tolerance ) / bucketSize;
+        UntypedMeshBufferIteratorConst iti( pFragment->GetIndexBuffers(), MBS_VERTEXINDEX );
+        MeshBufferIteratorConst<MBF_FLOAT32,float,3> itp( pFragment->GetVertexBuffers(), MBS_POSITION );
+        for ( int tf=0; tf<fragmentFaceCount; tf++ )
+        {
+            uint32 index0 = iti.GetAsUINT32(); ++iti;
+            uint32 index1 = iti.GetAsUINT32(); ++iti;
+            uint32 index2 = iti.GetAsUINT32(); ++iti;
+            float y = ( (*(itp+index0))[AXIS] + (*(itp+index1))[AXIS] + (*(itp+index2))[AXIS] ) / 3;
+            float fbucket = (y-bucketStart) / bucketSize;
+            int bucket = FMath::Min( NUM_BUCKETS-1, FMath::Max( 0, (int)fbucket ) );
+            buckets[bucket].Add(tf);
+            int hibucket = FMath::Min( NUM_BUCKETS-1, FMath::Max( 0, (int)(fbucket+bucketThreshold) ) );
+            if (hibucket!=bucket)
+            {
+                buckets[hibucket].Add(tf);
+            }
+            int lobucket = FMath::Min( NUM_BUCKETS-1, FMath::Max( 0, (int)(fbucket-bucketThreshold) ) );
+            if (lobucket!=bucket)
+            {
+                buckets[lobucket].Add(tf);
+            }
+        }
+
+//		LogDebug("Box : min %.3f, %.3f, %.3f    size %.3f,%.3f,%.3f\n",
+//				aabbox.min[0], aabbox.min[1], aabbox.min[2],
+//				aabbox.size[0], aabbox.size[1], aabbox.size[2] );
+//		for ( int b=0; b<NUM_BUCKETS; ++b )
+//		{
+//			LogDebug("bucket : %d\n", buckets[b].size() );
+//		}
+
+		TArray<uint8> faceClipped;
+		faceClipped.SetNumZeroed(sourceFaceCount);
+
+        UntypedMeshBufferIteratorConst ito( pBase->GetIndexBuffers(), MBS_VERTEXINDEX );
+        MeshBufferIteratorConst<MBF_FLOAT32,float,3> itop( pBase->GetVertexBuffers(), MBS_POSITION );
+        UntypedMeshBufferIteratorConst itti( pFragment->GetIndexBuffers(), MBS_VERTEXINDEX );
+        for ( int f=0; f<sourceFaceCount; ++f )
+        {
+            bool hasFace = false;
+            vec3<uint32> ov;
+            ov[0] = ito.GetAsUINT32(); ++ito;
+            ov[1] = ito.GetAsUINT32(); ++ito;
+            ov[2] = ito.GetAsUINT32(); ++ito;
+
+            // find the bucket for this face
+            float y = ( (*(itop+ov[0]))[AXIS] + (*(itop+ov[1]))[AXIS] + (*(itop+ov[2]))[AXIS] ) / 3;
+            float fbucket = (y-bucketStart) / bucketSize;
+            int bucket = FMath::Min( NUM_BUCKETS-1, FMath::Max( 0, (int)fbucket ) );
+
+            for ( int32 btf=0; !hasFace && btf<buckets[bucket].Num(); btf++ )
+            {
+                int tf =  buckets[bucket][btf];
+
+                vec3<uint32> v;
+                v[0] = (itti+3*tf+0).GetAsUINT32();
+                v[1] = (itti+3*tf+1).GetAsUINT32();
+                v[2] = (itti+3*tf+2).GetAsUINT32();
+
+                hasFace = true;
+                for ( int vi=0; hasFace && vi<3; ++vi )
+                {
+                    hasFace = vertexMap.Matches(v[vi],ov[0])
+                         || vertexMap.Matches(v[vi],ov[1])
+                         || vertexMap.Matches(v[vi],ov[2]);
+                }
+            }
+
+            if ( hasFace )
+            {
+                faceClipped[f] = true;
+            }
+        }
+
+        // We only remove vertices if all their faces are clipped
+		TArray<uint8> vertex_with_face_not_clipped;
+		vertex_with_face_not_clipped.SetNumZeroed(sourceVertexCount);
+
+        UntypedMeshBufferIteratorConst itoi(pBase->GetIndexBuffers(), MBS_VERTEXINDEX);
+        int aFaceCount = pBase->GetFaceCount();
+        for (int f = 0; f < aFaceCount; ++f)
+        {
+            vec3<uint32> ov;
+            ov[0] = itoi.GetAsUINT32(); ++itoi;
+            ov[1] = itoi.GetAsUINT32(); ++itoi;
+            ov[2] = itoi.GetAsUINT32(); ++itoi;
+
+            if (!faceClipped[f])
+            {
+                vertex_with_face_not_clipped[ov[0]] = true;
+                vertex_with_face_not_clipped[ov[1]] = true;
+                vertex_with_face_not_clipped[ov[2]] = true;
+            }
+        }
+
+		CreateMask(Result, pBase, vertex_with_face_not_clipped);
+    }
+
+}
diff --git a/Engine/Plugins/Experimental/Mutable/Source/MutableRuntime/Private/MuR/OpMeshClipWithMesh.h b/Engine/Plugins/Experimental/Mutable/Source/MutableRuntime/Private/MuR/OpMeshClipWithMesh.h
--- a/Engine/Plugins/Experimental/Mutable/Source/MutableRuntime/Private/MuR/OpMeshClipWithMesh.h
+++ b/Engine/Plugins/Experimental/Mutable/Source/MutableRuntime/Private/MuR/OpMeshClipWithMesh.h
@@ -2,1023 +2,23 @@
 
 #pragma once
 
-#include "MuR/MeshPrivate.h"
-#include "MuR/ConvertData.h"
 #include "MuR/Platform.h"
-#include "MuR/MutableTrace.h"
-#include "Math/Ray.h"
-#include "TriangleTypes.h"
-#include "BoxTypes.h"
-#include "Intersection/IntrRay3Triangle3.h"
-#include "MathUtil.h"
-#include "IntVectorTypes.h"
-
-#include "Math/UnrealMathUtility.h"
-
-#include "Spatial/PointHashGrid3.h"
 
 namespace mu
 {
+	class Mesh;
+	class Image;
 
-    //---------------------------------------------------------------------------------------------
-    //! Create a map from vertices into vertices, collapsing vertices that have the same position, 
-	//! This version uses UE Containers to return.	
-    //---------------------------------------------------------------------------------------------
-	inline void MeshCreateCollapsedVertexMap(const Mesh* pMesh,	TArray<int32>& OutCollapsedVertices, TArray<FVector3f>& OutVertices)
-	{
-		MUTABLE_CPUPROFILER_SCOPE(MeshCreateCollapsedVertexMap);
-
-		const int32 NumVertices = pMesh->GetVertexCount();
-
-		// Used to speed up vertex comparison
-		UE::Geometry::TPointHashGrid3f<int32> VertHash(0.01f, INDEX_NONE);
-		VertHash.Reserve(NumVertices);
-
-		// Info to collect. Vertices and collapsed vertices
-		OutVertices.SetNumUninitialized(NumVertices);
-		OutCollapsedVertices.Init(INDEX_NONE, NumVertices);
-
-		// Get Vertices
-		mu::UntypedMeshBufferIteratorConst ItPosition = mu::UntypedMeshBufferIteratorConst(pMesh->GetVertexBuffers(), mu::MBS_POSITION);
-		FVector3f* VertexData = OutVertices.GetData();
-
-		for (int32 VertexIndex = 0; VertexIndex < NumVertices; ++VertexIndex)
-		{
-			*VertexData = ItPosition.GetAsVec3f();
-			VertHash.InsertPointUnsafe(VertexIndex, *VertexData);
-
-			++ItPosition;
-			++VertexData;
-		}
-
-		// Find unique vertices
-		TArray<int32> NearbyVertices;
-		for (int32 VertexIndex = 0; VertexIndex < NumVertices; ++VertexIndex)
-		{
-			if (OutCollapsedVertices[VertexIndex] != INDEX_NONE)
-			{
-				continue;
-			}
-
-			const FVector3f& Vertex = OutVertices[VertexIndex];
-
-			NearbyVertices.Reset();
-			VertHash.FindPointsInBall(Vertex, TMathUtilConstants<float>::ZeroTolerance,
-				[&Vertex, &OutVertices](const int32& Other) -> float {return FVector3f::DistSquared(OutVertices[Other], Vertex); },
-				NearbyVertices);
-
-			for (int32 NearbyVertexIndex : NearbyVertices)
-			{
-				OutCollapsedVertices[NearbyVertexIndex] = VertexIndex;
-			}
-		}
-	}
-
-
-    //---------------------------------------------------------------------------------------------
-    //! Return true if a mesh is closed
-    //---------------------------------------------------------------------------------------------
-    inline bool MeshIsClosed( const Mesh* pMesh )
-    {
-        MUTABLE_CPUPROFILER_SCOPE(MeshIsClosed);
-
-        int vcount = pMesh->GetVertexCount();
-        int fcount = pMesh->GetFaceCount();
-
-        TArray<FVector3f> vertices;
-
-        // Map in ClipMesh from vertices to the one they are collapsed to because they are very similar, if they aren't collapsed then they are mapped to themselves
-		TArray<int> collapsedVertexMap;
-
-        MeshCreateCollapsedVertexMap( pMesh, collapsedVertexMap, vertices );
-
-        // Acumulate edges
-        TMap< TPair<int,int>, int > faceCountPerEdge;
-
-        UntypedMeshBufferIteratorConst itClipMesh(pMesh->GetIndexBuffers(), MBS_VERTEXINDEX);
-        for (int f = 0; f < fcount; ++f)
-        {
-            int face[3];
-            face[0] = collapsedVertexMap[itClipMesh.GetAsUINT32()]; ++itClipMesh;
-            face[1] = collapsedVertexMap[itClipMesh.GetAsUINT32()]; ++itClipMesh;
-            face[2] = collapsedVertexMap[itClipMesh.GetAsUINT32()]; ++itClipMesh;
-
-            for (int e=0; e<3; ++e)
-            {
-                int v0 = face[e];
-                int v1 = face[(e+1)%3];
-
-                if(v0==v1)
-                {
-                    // Degenerated mesh
-                    return false;
-                }
-
-                TPair<int,int> edge;
-                edge.Key = FMath::Min( v0, v1 );
-                edge.Value = FMath::Max( v0, v1 );
-
-				if (faceCountPerEdge.Contains(edge))
-				{
-					faceCountPerEdge[edge]++;
-				}
-				else
-				{
-					faceCountPerEdge.Add(edge,1);
-				}
-            }
-        }
-
-
-        // See if every edge has 2 faces
-        for( const auto& e: faceCountPerEdge)
-        {
-            if (e.Value!=2)
-            {
-                return false;
-            }
-        }
-
-        return true;
-    }
-
-
-    //---------------------------------------------------------------------------------------------
-    //! Remove all unused vertices from a mesh, and fix its index buffers.
-    //---------------------------------------------------------------------------------------------
-    inline void MeshRemoveUnusedVertices( Mesh* pMesh )
-    {
-        MUTABLE_CPUPROFILER_SCOPE(MeshRemoveUnusedVertices);
-
-        // Mark used vertices
-		TArray<uint8> used;
-		used.SetNumZeroed(pMesh->GetVertexCount());
-        UntypedMeshBufferIteratorConst iti(pMesh->GetIndexBuffers(), MBS_VERTEXINDEX);
-        int IndexCount = pMesh->GetIndexCount();
-        for (int i = 0; i < IndexCount; ++i)
-        {
-            uint32 Index = iti.GetAsUINT32();
-            ++iti;
-            used[Index] = true;
-        }
-
-        // Build vertex map
-		TArray<int32> oldToNewVertex;
-		oldToNewVertex.SetNumUninitialized(pMesh->GetVertexCount());
-        int32 totalNewVertices = 0;
-        for (int32 v = 0; v<pMesh->GetVertexCount(); ++v)
-        {
-            if (used[v])
-            {
-                oldToNewVertex[v] = totalNewVertices;
-                ++totalNewVertices;
-            }
-            else
-            {
-                oldToNewVertex[v] = -1;
-            }
-        }
-
-        // Remove unused vertices and build index map
-        for (int b = 0; b<pMesh->GetVertexBuffers().GetBufferCount(); ++b)
-        {
-            int elemSize = pMesh->GetVertexBuffers().GetElementSize(b);
-            const uint8* pSourceData = pMesh->GetVertexBuffers().GetBufferData(b);
-            uint8* pData = pMesh->GetVertexBuffers().GetBufferData(b);
-            for (int v = 0; v<pMesh->GetVertexCount(); ++v)
-            {
-                if (oldToNewVertex[v]!=-1)
-                {
-                    uint8* pElemData = pData + elemSize*oldToNewVertex[v];
-                    const uint8* pElemSourceData = pSourceData + elemSize*v;
-                    // Avoid warning for overlapping memcpy in valgrind
-                    if (pElemData != pElemSourceData)
-                    {
-                        memcpy(pElemData, pElemSourceData, elemSize);
-                    }
-                }
-            }
-        }
-        pMesh->GetVertexBuffers().SetElementCount(totalNewVertices);
-
-        // Update indices
-        UntypedMeshBufferIteratorConst ito(pMesh->GetIndexBuffers(), MBS_VERTEXINDEX);
-        if (ito.GetFormat() == MBF_UINT32)
-        {
-            for (int i = 0; i < IndexCount; ++i)
-            {
-                uint32 Index = *(uint32*)ito.ptr();
-                int32 NewIndex = oldToNewVertex[Index];
-                check(NewIndex >= 0);
-                *(uint32*)ito.ptr() = (uint32)NewIndex;
-                ++ito;
-            }
-        }
-        else if (ito.GetFormat() == MBF_UINT16)
-        {
-            for (int i = 0; i < IndexCount; ++i)
-            {
-                uint16 Index = *(uint16*)ito.ptr();
-                int32 NewIndex = oldToNewVertex[Index];
-                check(NewIndex >= 0);
-                *(uint16*)ito.ptr() = (uint16)NewIndex;
-                ++ito;
-            }
-        }
-        else
-        {
-            checkf(false, TEXT("Format not implemented.") );
-        }
-
-
-        // \todo: face buffers?
-    }
-
-
-    //---------------------------------------------------------------------------------------------
-    //! Reference version
-    //---------------------------------------------------------------------------------------------
-    inline int get_num_intersections( const FVector3f& vertex,
-                                      const FVector3f& ray,
-                                      const TArray<FVector3f>& vertices,
-                                      const TArray<uint32>& faces,
-                                      TArray<int>& collapsedVertexMap,
-									  TArray<uint8>& vertex_already_intersected,
-                                      TMap<TPair<int, int>, bool>& edge_already_intersected,
-									  float dynamic_epsilon
-                                      )
-    {
-    	
-        MUTABLE_CPUPROFILER_SCOPE(get_num_intersections);
-    	
-        int num_intersections = 0;
-		FVector3f intersection;
-
-		FMemory::Memzero( vertex_already_intersected.GetData(), vertex_already_intersected.Num());
-
-        edge_already_intersected.Empty();
-
-    #define get_collapsed_vertex(vertex_idx) (vertices[collapsedVertexMap[vertex_idx]])
-
-        size_t fcount = faces.Num()/3;
-
-        // Check vertex against all ClipMorph faces
-        for (uint32 face = 0; face < fcount; ++face)
-        {
-            uint32 vertex_indexs[3] = { faces[3 * face], faces[3 * face + 1], faces[3 * face + 2] };
-
-			FVector3f v0 = get_collapsed_vertex(vertex_indexs[0]);
-			FVector3f v1 = get_collapsed_vertex(vertex_indexs[1]);
-			FVector3f v2 = get_collapsed_vertex(vertex_indexs[2]);
-
-            int out_intersected_vert, out_intersected_edge_v0, out_intersected_edge_v1;
-
-            if (rayIntersectsFace(vertex, ray, v0, v1, v2, intersection, out_intersected_vert, out_intersected_edge_v0, out_intersected_edge_v1, dynamic_epsilon))
-            {
-                bool vertex_not_intersected_before = true;
-                bool edge_not_intersected_before = true;
-
-                if (out_intersected_vert >= 0)
-                {
-                    int collapsed_vert_index = collapsedVertexMap[vertex_indexs[out_intersected_vert]];
-                    vertex_not_intersected_before = !vertex_already_intersected[collapsed_vert_index];
-
-                    vertex_already_intersected[collapsed_vert_index] = true;
-                }
-
-                if (out_intersected_edge_v0 >= 0)
-                {
-                    int collapsed_edge_vert_index0 = collapsedVertexMap[vertex_indexs[out_intersected_edge_v0]];
-                    int collapsed_edge_vert_index1 = collapsedVertexMap[vertex_indexs[out_intersected_edge_v1]];
-
-                    TPair<int, int> key;
-                    key.Key = collapsed_edge_vert_index0 <= collapsed_edge_vert_index1 ? collapsed_edge_vert_index0 : collapsed_edge_vert_index1;
-                    key.Value = collapsed_edge_vert_index1 > collapsed_edge_vert_index0 ? collapsed_edge_vert_index1 : collapsed_edge_vert_index0;
-                    edge_not_intersected_before = edge_already_intersected.Find(key) != nullptr;
-
-                    if (edge_not_intersected_before)
-                    {
-                        edge_already_intersected.Add( key, true );
-                    }
-                }
-
-                if (vertex_not_intersected_before && edge_not_intersected_before)
-                {
-                    num_intersections++;
-                }
-            }
-        }
-
-        return num_intersections;
-    }
-
-
-    //---------------------------------------------------------------------------------------------
-    //! Core Geometry version
-    //---------------------------------------------------------------------------------------------
-    inline int GetNumIntersections( const FRay3f& Ray,
-                                    const TArray<FVector3f>& Vertices,
-                                    const TArray<uint32>& Faces,
-                                    const TArray<int32>& CollapsedVertexMap,
-                                    TArray<uint8>& InOutVertexAlreadyIntersected,
-                                    TSet<uint64>& InOutEdgeAlreadyIntersected,
-									const float DynamicEpsilon )
-    {
-		using namespace UE::Geometry;    
-
-        MUTABLE_CPUPROFILER_SCOPE(GetNumIntersections);
-    	
-        int32 NumIntersections = 0;
-
-		FMemory::Memzero( InOutVertexAlreadyIntersected.GetData(), InOutVertexAlreadyIntersected.Num() );
-		InOutEdgeAlreadyIntersected.Empty();
-
-		auto GetCollapsedVertex = [&CollapsedVertexMap, &Vertices]( const uint32 V ) -> const FVector3f&
-		{
-			return Vertices[ CollapsedVertexMap[V] ];
-		};
-
-        int32 FaceCount = Faces.Num() / 3;
-
-		UE::Geometry::FIntrRay3Triangle3f Intersector(Ray, UE::Geometry::FTriangle3f());
-
-        // Check vertex against all ClipMorph faces
-        for (int32 Face = 0; Face < FaceCount; ++Face)
-        {
-            const uint32 VertexIndexs[3] = { Faces[3 * Face], Faces[3 * Face + 1], Faces[3 * Face + 2] };
-
-            const FVector3f& V0 = GetCollapsedVertex(VertexIndexs[0]);
-            const FVector3f& V1 = GetCollapsedVertex(VertexIndexs[1]);
-            const FVector3f& V2 = GetCollapsedVertex(VertexIndexs[2]);
-
-			Intersector.Triangle = UE::Geometry::FTriangle3f( V0, V1, V2 );
-
-			if ( Intersector.Find() )
-            {
-				// Find if close to edge using barycentric coordinates form intersector.
-	
-				// Is the Dynamic Epsilon needed?. Intersector bary coords are in the range [0.0, 1.0], 
-				// furthermore, IntrRay3Triangle3f::TriangleBaryCoords are double pressison, not sure 
-				// is intentional or is a bug ( is double even when the FReal type is float ) 
-				const bool IntersectsEdge01 = FMath::IsNearlyZero( Intersector.TriangleBaryCoords.Z, DynamicEpsilon );
-				const bool IntersectsEdge02 = FMath::IsNearlyZero( Intersector.TriangleBaryCoords.Y, DynamicEpsilon );
-				const bool IntersectsEdge12 = FMath::IsNearlyZero( Intersector.TriangleBaryCoords.X, DynamicEpsilon );
-
-				int32 IntersectedTriangleVertexId = -1;
-				IntersectedTriangleVertexId =  (IntersectsEdge01 & IntersectsEdge02) ? 0 : IntersectedTriangleVertexId;
-				IntersectedTriangleVertexId =  (IntersectsEdge01 & IntersectsEdge12) ? 1 : IntersectedTriangleVertexId;
-				IntersectedTriangleVertexId =  (IntersectsEdge02 & IntersectsEdge12) ? 2 : IntersectedTriangleVertexId;
-
-				bool bIsAlreadyIntersected = false;
-
-                if ( IntersectedTriangleVertexId >= 0 )
-                {
-                    const int32 CollapsedVertIndex = CollapsedVertexMap[VertexIndexs[IntersectedTriangleVertexId]];
-                    bIsAlreadyIntersected = InOutVertexAlreadyIntersected[CollapsedVertIndex] == 0;
-
-                    InOutVertexAlreadyIntersected[CollapsedVertIndex] = true;
-                }
-				else if ( IntersectsEdge01 | IntersectsEdge02 | IntersectsEdge12 )
-                {
-					int32 EdgeV0 = (IntersectsEdge01 | IntersectsEdge02) ? 0 : 1;
-					int32 EdgeV1 = IntersectsEdge01 ? 1 : 2;				
-
-                    const int32 CollapsedEdgeVertIndex0 = CollapsedVertexMap[VertexIndexs[EdgeV0]];
-                    const int32 CollapsedEdgeVertIndex1 = CollapsedVertexMap[VertexIndexs[EdgeV1]];
-
-                    const uint64 EdgeKey = ( ( (uint64)FMath::Max( CollapsedEdgeVertIndex0, CollapsedEdgeVertIndex1 ) ) << 32 ) | 
-							                   (uint64)FMath::Min( CollapsedEdgeVertIndex0, CollapsedEdgeVertIndex1 );
-	
-					InOutEdgeAlreadyIntersected.FindOrAdd( EdgeKey, &bIsAlreadyIntersected );
-                }
-
-                NumIntersections += (int32)(!bIsAlreadyIntersected);
-            }
-        }
-
-        return NumIntersections;
-    }
-
-
-    //---------------------------------------------------------------------------------------------
-    //! Core Geometry version
-    //---------------------------------------------------------------------------------------------
-    inline void MeshClipMeshClassifyVertices(TArray<uint8>& VertexInClipMesh,
-                                             const Mesh* pBase, const Mesh* pClipMesh)
-    {
-		using namespace UE::Geometry;
-
-        MUTABLE_CPUPROFILER_SCOPE(MeshClipMeshClassifyVertices);
-    	
-        // TODO: this fails in bro for some object. do it at graph construction time to report it 
-        // nicely.
-        //check( MeshIsClosed(pClipMesh) );
-
-        const int32 VCount = pClipMesh->GetVertexBuffers().GetElementCount();
-        const int32 FCount = pClipMesh->GetFaceCount();
-
-        int32 OrigVertCount = pBase->GetVertexBuffers().GetElementCount();
-
-        // Stores whether each vertex in the original mesh in the clip mesh volume
-        VertexInClipMesh.AddZeroed(OrigVertCount);
-
-        if (VCount == 0)
-        {
-            return;
-        }
-
-        TArray<FVector3f> Vertices; // ClipMesh vertex cache
-		Vertices.SetNumUninitialized(VCount);
-
-        TArray<uint32> Faces; // ClipMesh face cache
-		Faces.SetNumUninitialized(FCount * 3); 
-
-        // Map in ClipMesh from vertices to the one they are collapsed to because they are very 
-        // similar, if they aren't collapsed then they are mapped to themselves
-        TArray<int32> CollapsedVertexMap;
-		CollapsedVertexMap.AddUninitialized(VCount);
-
-        MeshCreateCollapsedVertexMap( pClipMesh, CollapsedVertexMap, Vertices );
-
-        // Create cache of the faces
-        UntypedMeshBufferIteratorConst ItClipMesh(pClipMesh->GetIndexBuffers(), MBS_VERTEXINDEX);
-
-        for ( int32 F = 0; F < FCount; ++F )
-        {
-            Faces[3 * F]     = ItClipMesh.GetAsUINT32(); ++ItClipMesh;
-            Faces[3 * F + 1] = ItClipMesh.GetAsUINT32(); ++ItClipMesh;
-            Faces[3 * F + 2] = ItClipMesh.GetAsUINT32(); ++ItClipMesh;
-        }
-
-
-        // Create a bounding box of the clip mesh
-		UE::Geometry::FAxisAlignedBox3f ClipMeshBoundingBox = UE::Geometry::FAxisAlignedBox3f::Empty();
-
-        for ( const FVector3f& Vert : Vertices )
-        {
-            ClipMeshBoundingBox.Contain( Vert );
-        }
-
-		// Dynamic distance epsilon to support different engines
-		const float MaxDimensionBoundingBox = ClipMeshBoundingBox.DiagonalLength();
-		// 0.000001 is the value that helps to achieve the dynamic epsilon, do not change it
-		const float DynamicEpsilon = 0.000001f * MaxDimensionBoundingBox * (MaxDimensionBoundingBox < 1.0f ? MaxDimensionBoundingBox : 1.0f);
-
-        // Create an acceleration grid to avoid testing all clip-mesh triangles.
-        // This assumes that the testing ray direction is Z
-        constexpr int32 GRID_SIZE = 8;
-        TUniquePtr<TArray<uint32>[]> GridFaces( new TArray<uint32>[GRID_SIZE*GRID_SIZE] );
-        const FVector2f GridCellSize = FVector2f( ClipMeshBoundingBox.Width(), ClipMeshBoundingBox.Height() ) / (float)GRID_SIZE;
-
-        for ( int32 I = 0; I < GRID_SIZE; ++I )
-        {
-            for ( int32 J = 0; J < GRID_SIZE; ++J )
-            {
-				const FVector2f BBoxMin = FVector2f(ClipMeshBoundingBox.Min.X, ClipMeshBoundingBox.Min.Y) + GridCellSize * FVector2f(I, J);
-
-                FAxisAlignedBox2f CellBox( BBoxMin, BBoxMin + GridCellSize );
-
-                TArray<uint32>& CellFaces = GridFaces[I+J*GRID_SIZE];
-                CellFaces.Empty(FCount/GRID_SIZE);
-                for (int32 F = 0; F < FCount; ++F)
-                {
-                    // Imprecise, conservative classification of faces.
-					const FVector3f& V0 = Vertices[ Faces[3*F + 0] ];
-					const FVector3f& V1 = Vertices[ Faces[3*F + 1] ];
-					const FVector3f& V2 = Vertices[ Faces[3*F + 2] ];
-
-					FAxisAlignedBox2f FaceBox;
-					FaceBox.Contain( FVector2f( V0.X, V0.Y ) );
-                    FaceBox.Contain( FVector2f( V1.X, V1.Y ) );
-                    FaceBox.Contain( FVector2f( V2.X, V2.Y ) );
-
-                    if (CellBox.Intersects(FaceBox))
-                    {
-                        CellFaces.Add( Faces[3*F + 0] );
-                        CellFaces.Add( Faces[3*F + 1] );
-                        CellFaces.Add( Faces[3*F + 2] );
-                    }
-                }
-            }
-        }	
-
-        // Now go through all vertices in the mesh and record whether they are inside or outside of the ClipMesh
-        uint32 DestVertexCount = pBase->GetVertexCount();
-
-        const FMeshBufferSet& MBSPriv2 = pBase->GetVertexBuffers();
-        for (int32 b = 0; b < MBSPriv2.m_buffers.Num(); ++b)
-        {
-            for (int32 c = 0; c < MBSPriv2.m_buffers[b].m_channels.Num(); ++c)
-            {
-                MESH_BUFFER_SEMANTIC Sem = MBSPriv2.m_buffers[b].m_channels[c].m_semantic;
-                int32 SemIndex = MBSPriv2.m_buffers[b].m_channels[c].m_semanticIndex;
-
-                UntypedMeshBufferIteratorConst It(pBase->GetVertexBuffers(), Sem, SemIndex);
-
-                TArray<uint8> VertexAlreadyIntersected;
-				VertexAlreadyIntersected.AddZeroed(VCount);
-
-                TSet<uint64> EdgeAlreadyIntersected;
-
-                switch ( Sem )
-                {
-                case MBS_POSITION:
-                    for (uint32 V = 0; V < DestVertexCount; ++V)
-                    {
-                        FVector3f Vertex(0.0f, 0.0f, 0.0f);
-                        for (int32 Offset = 0; Offset < 3; ++Offset)
-                        {
-                            ConvertData(Offset, &Vertex[0], MBF_FLOAT32, It.ptr(), It.GetFormat());
-                        }
-						
-						const FVector2f ClipBBoxMin = FVector2f( ClipMeshBoundingBox.Min.X, ClipMeshBoundingBox.Min.Y);
-						const FVector2f ClipBBoxSize = FVector2f( ClipMeshBoundingBox.Width(), ClipMeshBoundingBox.Height() );
-
-						const FVector2i HPos = FVector2i( ( ( FVector2f( Vertex.X, Vertex.Y ) - ClipBBoxMin ) / ClipBBoxSize ) * (float)GRID_SIZE );   
-						const FVector2i CellCoord = FVector2i( FMath::Clamp( HPos.X, 0, GRID_SIZE - 1 ), 
-														       FMath::Clamp( HPos.Y, 0, GRID_SIZE - 1 ) );
-
-                        // Early discard test: if the vertex is not inside the bounding box of the clip mesh, it won't be clipped.
-                        const bool bContainsVertex = ClipMeshBoundingBox.Contains( Vertex );
-
-                        if ( bContainsVertex )
-                        {
-                            // Optimised test	
-							
-							// Z-direction. Don't change this without reviewing the acceleration structure.
-							FVector3f RayDir = FVector3f(0.0f, 0.0f, 1.0f);	
-                            const int32 CellIndex = CellCoord.X + CellCoord.Y*GRID_SIZE;
-
-                            int32 NumIntersections = GetNumIntersections( 
-									FRay3f( Vertex, RayDir ), 
-									Vertices,
-									GridFaces[CellIndex],
-									CollapsedVertexMap, 
-									VertexAlreadyIntersected, 
-									EdgeAlreadyIntersected, 
-								    DynamicEpsilon);
-
-                            // Full test BLEH, debug
-//                            int32 FullNumIntersections = GetNumIntersections(
-//									FRay3f( Vertex, RayDir ), 
-//									Vertices,
-//									Faces,
-//									CollapsedVertexMap, 
-//									VertexAlreadyIntersected, 
-//									EdgeAlreadyIntersected, 
-//								    DynamicEpsilon);
-
-							VertexInClipMesh[V] = NumIntersections % 2 == 1;
-
-                            // This may be used to debug degenerated cases if the conditional above is also removed.
-                            // \todo: make sure it works well
-//                          if (!bContainsVertex && VertexInClipMesh[V])
-//                          {
-//                              assert(false);
-//                          }
-                        }
-
-                        ++It;
-                    }
-                    break;
-
-                default:
-                    break;
-                }
-            }
-        }
-    }	
-
-
-    //---------------------------------------------------------------------------------------------
-    //! Reference version
-    //---------------------------------------------------------------------------------------------
-    inline void MeshClipWithMesh_Reference(Mesh* Result, const Mesh* pBase, const Mesh* pClipMesh, bool& bOutSuccess)
-    {
-        MUTABLE_CPUPROFILER_SCOPE(MeshClipWithMesh_Reference);
-        
-		bOutSuccess = true;
-
-        uint32 vcount = pClipMesh->GetVertexBuffers().GetElementCount();
-        if (!vcount)
-        {
-			bOutSuccess = false;
-			return; // Since there's nothing to clip against return a copy of the original base mesh
-        }
-
-        TArray<uint8> vertex_in_clip_mesh;  // Stores whether each vertex in the original mesh in in the clip mesh volume
-        MeshClipMeshClassifyVertices(vertex_in_clip_mesh, pBase, pClipMesh);
-
-		Result->CopyFrom(*pBase);
-        // Now remove all the faces from the result mesh that have all the vertices outside the clip volume
-        UntypedMeshBufferIteratorConst itBase(Result->GetIndexBuffers(), MBS_VERTEXINDEX);
-        UntypedMeshBufferIterator itDest(Result->GetIndexBuffers(), MBS_VERTEXINDEX);
-        int aFaceCount = Result->GetFaceCount();
-
-        UntypedMeshBufferIteratorConst ito(Result->GetIndexBuffers(), MBS_VERTEXINDEX);
-        for (int f = 0; f < aFaceCount; ++f)
-        {
-            vec3<uint32> ov;
-            ov[0] = ito.GetAsUINT32(); ++ito;
-            ov[1] = ito.GetAsUINT32(); ++ito;
-            ov[2] = ito.GetAsUINT32(); ++ito;
-
-            bool all_verts_in = vertex_in_clip_mesh[ov[0]] && vertex_in_clip_mesh[ov[1]] && vertex_in_clip_mesh[ov[2]];
-
-            if (!all_verts_in)
-            {
-                if (itDest.ptr() != itBase.ptr())
-                {
-					FMemory::Memcpy(itDest.ptr(), itBase.ptr(), itBase.GetElementSize() * 3);
-                }
-
-                itDest += 3;
-            }
-
-            itBase += 3;
-        }
-
-        SIZE_T removedIndices = itBase - itDest;
-        check(removedIndices % 3 == 0);
-
-        Result->GetFaceBuffers().SetElementCount(aFaceCount - (int)removedIndices / 3);
-        Result->GetIndexBuffers().SetElementCount(aFaceCount * 3 - (int)removedIndices);
-
-        // TODO: Should redo/reorder the face buffer before SetElementCount since some deleted faces could be left and some remaining faces deleted.
-
-        //if (Result->GetFaceBuffers().GetElementCount() == 0)
-        //{
-        //	Result->CopyFrom(pBase); // If all faces have been discarded, return a copy of the unmodified mesh because unreal doesn't like empty meshes
-        //}
-
-        // [jordi] Remove unused vertices. This is necessary to avoid returning a mesh with vertices and no faces, which screws some
-        // engines like Unreal Engine 4.
-        MeshRemoveUnusedVertices(Result);
-    }
-
-
-    //---------------------------------------------------------------------------------------------
-    //! CoreGeometry version
-    //---------------------------------------------------------------------------------------------
-    inline void MeshClipWithMesh(Mesh* Result, const Mesh* pBase, const Mesh* pClipMesh, bool& bOutSuccess)
-    {
-        MUTABLE_CPUPROFILER_SCOPE(MeshClipWithMesh);
-		bOutSuccess = true;
-
-        uint32 VCount = pClipMesh->GetVertexBuffers().GetElementCount();
-        if (!VCount)
-        {
-			bOutSuccess = false;
-			return; // OutSuccess false indicates the pBase can be reused in this case. 
-        }
-
-		Result->CopyFrom(*pBase);
-
-        TArray<uint8> VertexInClipMesh;  // Stores whether each vertex in the original mesh in in the clip mesh volume
-        MeshClipMeshClassifyVertices(VertexInClipMesh, pBase, pClipMesh);
-
-        // Now remove all the faces from the result mesh that have all the vertices outside the clip volume
-        UntypedMeshBufferIteratorConst ItBase(Result->GetIndexBuffers(), MBS_VERTEXINDEX);
-        UntypedMeshBufferIterator ItDest(Result->GetIndexBuffers(), MBS_VERTEXINDEX);
-        int32 AFaceCount = Result->GetFaceCount();
-
-        UntypedMeshBufferIteratorConst Ito(Result->GetIndexBuffers(), MBS_VERTEXINDEX);
-        for (int32 F = 0; F < AFaceCount; ++F)
-        {
-			uint32 OV[3] = {0, 0, 0};
-
-            OV[0] = Ito.GetAsUINT32(); ++Ito;
-            OV[1] = Ito.GetAsUINT32(); ++Ito;
-            OV[2] = Ito.GetAsUINT32(); ++Ito;
-
-            bool AllVertsIn = VertexInClipMesh[OV[0]] && VertexInClipMesh[OV[1]] && VertexInClipMesh[OV[2]];
-
-            if (!AllVertsIn)
-            {
-                if (ItDest.ptr() != ItBase.ptr())
-                {
-					FMemory::Memcpy(ItDest.ptr(), ItBase.ptr(), ItBase.GetElementSize() * 3);
-                }
-
-                ItDest += 3;
-            }
-
-            ItBase += 3;
-        }
-
-        SIZE_T RemovedIndices = ItBase - ItDest;
-        check(RemovedIndices % 3 == 0);
-
-        Result->GetFaceBuffers().SetElementCount(AFaceCount - (int32)RemovedIndices / 3);
-        Result->GetIndexBuffers().SetElementCount(AFaceCount * 3 - (int32)RemovedIndices);
-
-        // TODO: Should redo/reorder the face buffer before SetElementCount since some deleted faces could be left and some remaining faces deleted.
-
-        //if (Result->GetFaceBuffers().GetElementCount() == 0)
-        //{
-        //	Result->CopyFrom(pBase); // If all faces have been discarded, return a copy of the unmodified mesh because unreal doesn't like empty meshes
-        //}
-
-        // [jordi] Remove unused vertices. This is necessary to avoid returning a mesh with vertices and no faces, which screws some
-        // engines like Unreal Engine 4.
-        MeshRemoveUnusedVertices(Result);
-    }
-
-
-    //---------------------------------------------------------------------------------------------
-    //!
-    //---------------------------------------------------------------------------------------------
-    inline void CreateMask(Mesh* Result, const Mesh* pBase, const TArray<uint8>& excludedVertices)
-    {
-
-        int maskVertexCount = 0;
-        for( uint8 b: excludedVertices )
-        {
-            if (!b) ++maskVertexCount;
-        }
-
-
-        // Create the vertex buffer
-        {
-            Result->GetVertexBuffers().SetElementCount( maskVertexCount );
-            Result->GetVertexBuffers().SetBufferCount( 1 );
-
-			// Vertex index channel
-			MESH_BUFFER_SEMANTIC semantic = MBS_VERTEXINDEX;
-			int semanticIndex = 0;
-			MESH_BUFFER_FORMAT format = MBF_UINT32;
-			int components = 1;
-			int offsets = 0;
-
-            Result->GetVertexBuffers().SetBuffer
-                (
-                    0,
-                    4,
-                    1,
-                    &semantic,
-                    &semanticIndex,
-                    &format,
-                    &components,
-                    &offsets
-                );
-        }
-
-        MeshBufferIterator<MBF_UINT32,uint32,1> itMask( Result->GetVertexBuffers(), MBS_VERTEXINDEX, 0 );
-        UntypedMeshBufferIteratorConst itBase( pBase->GetVertexBuffers(), MBS_VERTEXINDEX, 0 );
-
-        for (size_t v = 0; v<excludedVertices.Num(); ++v)
-        {
-            if (!excludedVertices[v])
-            {
-                (*itMask)[0] = itBase.GetAsUINT32();
-                ++itMask;
-            }
-
-            ++itBase;
-        }
-
-    }
-
-
-    //---------------------------------------------------------------------------------------------
-    //! Generate a mask mesh with the faces of the base mesh inside the clip mesh.
-    //---------------------------------------------------------------------------------------------
-    inline void MeshMaskClipMesh(Mesh* Result, const Mesh* pBase, const Mesh* pClipMesh, bool& bOutSuccess)
-    {
-        MUTABLE_CPUPROFILER_SCOPE(MeshMaskClipMesh);
-
-		bOutSuccess = true;
-
-        uint32 VCount = pClipMesh->GetVertexBuffers().GetElementCount();
-        if (!VCount)
-        {
-			bOutSuccess = false;
-			return;
-        }
-
-        TArray<uint8> VertexInClipMesh;  // Stores whether each vertex in the original mesh in in the clip mesh volume
-        MeshClipMeshClassifyVertices( VertexInClipMesh, pBase, pClipMesh );
-
-        // We only remove vertices if all their faces are clipped
-        TArray<uint8> VertexWithFaceNotClipped;
-		VertexWithFaceNotClipped.AddZeroed( VertexInClipMesh.Num() );
-
-        UntypedMeshBufferIteratorConst Ito(pBase->GetIndexBuffers(), MBS_VERTEXINDEX);
-        int32 AFaceCount = pBase->GetFaceCount();
-        for (int32 F = 0; F < AFaceCount; ++F)
-        {
-			uint32 OV[3] = { 0, 0, 0 };
-
-            OV[0] = Ito.GetAsUINT32(); ++Ito;
-            OV[1] = Ito.GetAsUINT32(); ++Ito;
-            OV[2] = Ito.GetAsUINT32(); ++Ito;
-
-            bool bFaceClipped =
-                    VertexInClipMesh[OV[0]] &&
-                    VertexInClipMesh[OV[1]] &&
-                    VertexInClipMesh[OV[2]];
-
-            if (!bFaceClipped)
-            {
-                VertexWithFaceNotClipped[OV[0]] = true;
-                VertexWithFaceNotClipped[OV[1]] = true;
-                VertexWithFaceNotClipped[OV[2]] = true;
-            }
-        }
-
-        CreateMask(Result, pBase, VertexWithFaceNotClipped);
-    }
-
-    //---------------------------------------------------------------------------------------------
-    //! Generate a mask mesh with the faces of the base mesh inside the clip mesh.
-    //---------------------------------------------------------------------------------------------
-    inline void MeshMaskClipMesh_Reference(Mesh* Result, const Mesh* pBase, const Mesh* pClipMesh )
-    {
-        MUTABLE_CPUPROFILER_SCOPE(MeshMaskClipMesh_Reference);
-
-        uint32 vcount = pClipMesh->GetVertexBuffers().GetElementCount();
-        if (!vcount)
-        {
-            return; // Since there's nothing to clip against return a copy of the original base mesh
-        }
-
-        TArray<uint8> vertex_in_clip_mesh;  // Stores whether each vertex in the original mesh in in the clip mesh volume
-        MeshClipMeshClassifyVertices( vertex_in_clip_mesh, pBase, pClipMesh );
-
-        // We only remove vertices if all their faces are clipped
-		TArray<uint8> vertex_with_face_not_clipped;
-		vertex_with_face_not_clipped.SetNumZeroed(vertex_in_clip_mesh.Num());
-
-        UntypedMeshBufferIteratorConst ito(pBase->GetIndexBuffers(), MBS_VERTEXINDEX);
-        int aFaceCount = pBase->GetFaceCount();
-        for (int f = 0; f < aFaceCount; ++f)
-        {
-            vec3<uint32> ov;
-            ov[0] = ito.GetAsUINT32(); ++ito;
-            ov[1] = ito.GetAsUINT32(); ++ito;
-            ov[2] = ito.GetAsUINT32(); ++ito;
-
-            bool faceClipped =
-                    vertex_in_clip_mesh[ov[0]] &&
-                    vertex_in_clip_mesh[ov[1]] &&
-                    vertex_in_clip_mesh[ov[2]];
-
-            if (!faceClipped)
-            {
-                vertex_with_face_not_clipped[ov[0]] = true;
-                vertex_with_face_not_clipped[ov[1]] = true;
-                vertex_with_face_not_clipped[ov[2]] = true;
-            }
-        }
-
-		CreateMask(Result, pBase, vertex_with_face_not_clipped );
-    }
-
-    //---------------------------------------------------------------------------------------------
-    //! Generate a mask mesh with the faces of the base mesh matching the fragment.
-    //---------------------------------------------------------------------------------------------
-    inline void MeshMaskDiff(Mesh* Result, const Mesh* pBase, const Mesh* pFragment, bool& bOutSuccess)
-    {
-        MUTABLE_CPUPROFILER_SCOPE(MeshMaskDiff);
-		bOutSuccess = true;
-
-        uint32 vcount = pFragment->GetVertexBuffers().GetElementCount();
-        if (!vcount)
-        {
-			bOutSuccess = false;
-            return;
-        }
-
-        int sourceFaceCount = pBase->GetFaceCount();
-        int sourceVertexCount = pBase->GetVertexCount();
-        int fragmentFaceCount = pFragment->GetFaceCount();
-
-
-        // Make a tolerance proportional to the mesh bounding box size
-        // TODO: Use precomputed bounding box
-        box< vec3<float> > aabbox;
-        if ( fragmentFaceCount > 0 )
-        {
-            MeshBufferIteratorConst<MBF_FLOAT32,float,3> itp( pFragment->GetVertexBuffers(), MBS_POSITION );
-
-            aabbox.min = *itp;
-            ++itp;
-
-            for ( int v=1; v<pFragment->GetVertexBuffers().GetElementCount(); ++v )
-            {
-                aabbox.Bound( *itp );
-                ++itp;
-            }
-        }
-        float tolerance = 1e-5f * length(aabbox.size);
-        Mesh::VERTEX_MATCH_MAP vertexMap;
-        pFragment->GetVertexMap( *pBase, vertexMap, tolerance );
-
-
-        // Classify the target faces in buckets along the Y axis
-#define NUM_BUCKETS	128
-#define AXIS		1
-        TArray<int> buckets[ NUM_BUCKETS ];
-        float bucketStart = aabbox.min[AXIS];
-        float bucketSize = aabbox.size[AXIS] / NUM_BUCKETS;
-
-        float bucketThreshold = ( 4 * tolerance ) / bucketSize;
-        UntypedMeshBufferIteratorConst iti( pFragment->GetIndexBuffers(), MBS_VERTEXINDEX );
-        MeshBufferIteratorConst<MBF_FLOAT32,float,3> itp( pFragment->GetVertexBuffers(), MBS_POSITION );
-        for ( int tf=0; tf<fragmentFaceCount; tf++ )
-        {
-            uint32 index0 = iti.GetAsUINT32(); ++iti;
-            uint32 index1 = iti.GetAsUINT32(); ++iti;
-            uint32 index2 = iti.GetAsUINT32(); ++iti;
-            float y = ( (*(itp+index0))[AXIS] + (*(itp+index1))[AXIS] + (*(itp+index2))[AXIS] ) / 3;
-            float fbucket = (y-bucketStart) / bucketSize;
-            int bucket = FMath::Min( NUM_BUCKETS-1, FMath::Max( 0, (int)fbucket ) );
-            buckets[bucket].Add(tf);
-            int hibucket = FMath::Min( NUM_BUCKETS-1, FMath::Max( 0, (int)(fbucket+bucketThreshold) ) );
-            if (hibucket!=bucket)
-            {
-                buckets[hibucket].Add(tf);
-            }
-            int lobucket = FMath::Min( NUM_BUCKETS-1, FMath::Max( 0, (int)(fbucket-bucketThreshold) ) );
-            if (lobucket!=bucket)
-            {
-                buckets[lobucket].Add(tf);
-            }
-        }
-
-//		LogDebug("Box : min %.3f, %.3f, %.3f    size %.3f,%.3f,%.3f\n",
-//				aabbox.min[0], aabbox.min[1], aabbox.min[2],
-//				aabbox.size[0], aabbox.size[1], aabbox.size[2] );
-//		for ( int b=0; b<NUM_BUCKETS; ++b )
-//		{
-//			LogDebug("bucket : %d\n", buckets[b].size() );
-//		}
-
-		TArray<uint8> faceClipped;
-		faceClipped.SetNumZeroed(sourceFaceCount);
-
-        UntypedMeshBufferIteratorConst ito( pBase->GetIndexBuffers(), MBS_VERTEXINDEX );
-        MeshBufferIteratorConst<MBF_FLOAT32,float,3> itop( pBase->GetVertexBuffers(), MBS_POSITION );
-        UntypedMeshBufferIteratorConst itti( pFragment->GetIndexBuffers(), MBS_VERTEXINDEX );
-        for ( int f=0; f<sourceFaceCount; ++f )
-        {
-            bool hasFace = false;
-            vec3<uint32> ov;
-            ov[0] = ito.GetAsUINT32(); ++ito;
-            ov[1] = ito.GetAsUINT32(); ++ito;
-            ov[2] = ito.GetAsUINT32(); ++ito;
-
-            // find the bucket for this face
-            float y = ( (*(itop+ov[0]))[AXIS] + (*(itop+ov[1]))[AXIS] + (*(itop+ov[2]))[AXIS] ) / 3;
-            float fbucket = (y-bucketStart) / bucketSize;
-            int bucket = FMath::Min( NUM_BUCKETS-1, FMath::Max( 0, (int)fbucket ) );
-
-            for ( int32 btf=0; !hasFace && btf<buckets[bucket].Num(); btf++ )
-            {
-                int tf =  buckets[bucket][btf];
-
-                vec3<uint32> v;
-                v[0] = (itti+3*tf+0).GetAsUINT32();
-                v[1] = (itti+3*tf+1).GetAsUINT32();
-                v[2] = (itti+3*tf+2).GetAsUINT32();
-
-                hasFace = true;
-                for ( int vi=0; hasFace && vi<3; ++vi )
-                {
-                    hasFace = vertexMap.Matches(v[vi],ov[0])
-                         || vertexMap.Matches(v[vi],ov[1])
-                         || vertexMap.Matches(v[vi],ov[2]);
-                }
-            }
-
-            if ( hasFace )
-            {
-                faceClipped[f] = true;
-            }
-        }
-
-        // We only remove vertices if all their faces are clipped
-		TArray<uint8> vertex_with_face_not_clipped;
-		vertex_with_face_not_clipped.SetNumZeroed(sourceVertexCount);
+    /**  */
+	extern void MeshClipWithMesh(Mesh* Result, const Mesh* pBase, const Mesh* pClipMesh, bool& bOutSuccess);
 
-        UntypedMeshBufferIteratorConst itoi(pBase->GetIndexBuffers(), MBS_VERTEXINDEX);
-        int aFaceCount = pBase->GetFaceCount();
-        for (int f = 0; f < aFaceCount; ++f)
-        {
-            vec3<uint32> ov;
-            ov[0] = itoi.GetAsUINT32(); ++itoi;
-            ov[1] = itoi.GetAsUINT32(); ++itoi;
-            ov[2] = itoi.GetAsUINT32(); ++itoi;
+    /** Generate a mask mesh with the faces of the base mesh inside the clip mesh. */
+	extern void MeshMaskClipMesh(Mesh* Result, const Mesh* pBase, const Mesh* pClipMesh, bool& bOutSuccess);
 
-            if (!faceClipped[f])
-            {
-                vertex_with_face_not_clipped[ov[0]] = true;
-                vertex_with_face_not_clipped[ov[1]] = true;
-                vertex_with_face_not_clipped[ov[2]] = true;
-            }
-        }
+	/** Generate a mask mesh with the faces of the base mesh that have all 3 vertices marked in the fiven mask. */
+	extern void MeshMaskClipUVMask(Mesh* Result, const Mesh* Base, const Image* Mask, uint8 LayoutIndex, bool& bOutSuccess);
 
-		CreateMask(Result, pBase, vertex_with_face_not_clipped);
-    }
+    /** Generate a mask mesh with the faces of the base mesh matching the fragment. */
+	extern void MeshMaskDiff(Mesh* Result, const Mesh* pBase, const Mesh* pFragment, bool& bOutSuccess);
 
 }
diff --git a/Engine/Plugins/Experimental/Mutable/Source/MutableRuntime/Private/MuR/Operations.cpp b/Engine/Plugins/Experimental/Mutable/Source/MutableRuntime/Private/MuR/Operations.cpp
--- a/Engine/Plugins/Experimental/Mutable/Source/MutableRuntime/Private/MuR/Operations.cpp
+++ b/Engine/Plugins/Experimental/Mutable/Source/MutableRuntime/Private/MuR/Operations.cpp
@@ -109,6 +109,7 @@ namespace mu
 		{ DT_MESH,			true,	{ 0, 0, 0, 0, 0, 0, 0, 0, 0, 0, 0, 0, 0, 0, 0, 0, 0, 0, 0, 0, 0, 0, 0, 0, 0, 0, 0, 0, 0 }		},	// ME_MERGE
 		{ DT_MESH,			true,	{ 0, 0, 0, 0, 0, 0, 0, 0, 0, 0, 0, 0, 0, 0, 0, 0, 0, 0, 0, 0, 0, 0, 0, 0, 0, 0, 0, 0, 0 }		},	// ME_INTERPOLATE
 		{ DT_MESH,			true,	{ 0, 0, 0, 0, 0, 0, 0, 0, 0, 0, 0, 0, 0, 0, 0, 0, 0, 0, 0, 0, 0, 0, 0, 0, 0, 0, 0, 0, 0 }		},	// ME_MASKCLIPMESH
+		{ DT_MESH,			true,	{ 0, 0, 0, 0, 0, 0, 0, 0, 0, 0, 0, 0, 0, 0, 0, 0, 0, 0, 0, 0, 0, 0, 0, 0, 0, 0, 0, 0, 0 }		},	// ME_MASKCLIPUVMASK
 		{ DT_MESH,			true,	{ 0, 0, 0, 0, 0, 0, 0, 0, 0, 0, 0, 0, 0, 0, 0, 0, 0, 0, 0, 0, 0, 0, 0, 0, 0, 0, 0, 0, 0 }		},	// ME_MASKDIFF
 		{ DT_MESH,			true,	{ 0, 0, 0, 0, 0, 0, 0, 0, 0, 0, 0, 0, 0, 0, 0, 0, 0, 0, 0, 0, 0, 0, 0, 0, 0, 0, 0, 0, 0 }		},	// ME_REMOVEMASK
 		{ DT_MESH,			true,	{ 0, 0, 0, 0, 0, 0, 0, 0, 0, 0, 0, 0, 0, 0, 0, 0, 0, 0, 0, 0, 0, 0, 0, 0, 0, 0, 0, 0, 0 }		},	// ME_FORMAT
@@ -834,13 +835,21 @@ namespace mu
             break;
         }
 
-        case OP_TYPE::ME_MASKCLIPMESH:
-        {
+		case OP_TYPE::ME_MASKCLIPMESH:
+		{
 			OP::MeshMaskClipMeshArgs args = program.GetOpArgs<OP::MeshMaskClipMeshArgs>(at);
-            f(args.source );
-            f(args.clip );
-            break;
-        }
+			f(args.source);
+			f(args.clip);
+			break;
+		}
+
+		case OP_TYPE::ME_MASKCLIPUVMASK:
+		{
+			OP::MeshMaskClipUVMaskArgs args = program.GetOpArgs<OP::MeshMaskClipUVMaskArgs>(at);
+			f(args.Source);
+			f(args.Mask);
+			break;
+		}
 
         case OP_TYPE::ME_MASKDIFF:
         {
diff --git a/Engine/Plugins/Experimental/Mutable/Source/MutableRuntime/Private/MuR/Operations.h b/Engine/Plugins/Experimental/Mutable/Source/MutableRuntime/Private/MuR/Operations.h
--- a/Engine/Plugins/Experimental/Mutable/Source/MutableRuntime/Private/MuR/Operations.h
+++ b/Engine/Plugins/Experimental/Mutable/Source/MutableRuntime/Private/MuR/Operations.h
@@ -237,6 +237,9 @@ namespace mu
         //! clip mesh.
         ME_MASKCLIPMESH,
 
+        /** Create a new mask mesh selecting the faces of a source that have UVs inside the region marked in an image mask. */
+		ME_MASKCLIPUVMASK,
+
         //! Create a new mask mesh selecting all the faces of a source that match another mesh.
         ME_MASKDIFF,
 
@@ -850,17 +853,18 @@ namespace mu
             ADDRESS targets[ MUTABLE_OP_MAX_INTERPOLATE_COUNT-1 ];
         };
 
-        struct MeshMaskClipMeshArgs
-        {
-            ADDRESS source;
-            ADDRESS clip;
-        };
+		struct MeshMaskClipMeshArgs
+		{
+			ADDRESS source;
+			ADDRESS clip;
+		};
 
-        struct MeshRemapIndicesArgs
-        {
-            ADDRESS source;
-            ADDRESS reference;
-        };
+		struct MeshMaskClipUVMaskArgs
+		{
+			ADDRESS Source = 0;
+			ADDRESS Mask = 0;
+			uint8 LayoutIndex = 0;
+		};
 
         struct MeshMaskDiffArgs
         {
diff --git a/Engine/Plugins/Experimental/Mutable/Source/MutableRuntime/Public/MuR/Model.h b/Engine/Plugins/Experimental/Mutable/Source/MutableRuntime/Public/MuR/Model.h
--- a/Engine/Plugins/Experimental/Mutable/Source/MutableRuntime/Public/MuR/Model.h
+++ b/Engine/Plugins/Experimental/Mutable/Source/MutableRuntime/Public/MuR/Model.h
@@ -13,7 +13,7 @@
 //! data structures. Compiled models are not necessarily compatible when the runtime is updated,
 //! so this version number can be used externally to verify this. It is not used internally, and
 //! serializing models from different versions than this runtime will probably result in a crash.
-#define MUTABLE_COMPILED_MODEL_CODE_VERSION		uint32( 80 )
+#define MUTABLE_COMPILED_MODEL_CODE_VERSION		uint32( 81 )
 #define MUTABLE_PARAMETERS_VERSION              uint32( 2 )
 
 
diff --git a/Engine/Plugins/Experimental/Mutable/Source/MutableTools/Private/MuT/ASTOpMeshMaskClipUVMask.cpp b/Engine/Plugins/Experimental/Mutable/Source/MutableTools/Private/MuT/ASTOpMeshMaskClipUVMask.cpp
new file mode 100644
--- /dev/null
+++ b/Engine/Plugins/Experimental/Mutable/Source/MutableTools/Private/MuT/ASTOpMeshMaskClipUVMask.cpp
@@ -0,0 +1,353 @@
+// Copyright Epic Games, Inc. All Rights Reserved.
+
+#include "MuT/ASTOpMeshMaskClipUVMask.h"
+
+#include "HAL/PlatformMath.h"
+#include "MuR/ModelPrivate.h"
+#include "MuR/RefCounted.h"
+#include "MuR/Types.h"
+#include "MuT/ASTOpConditional.h"
+#include "MuT/ASTOpMeshRemoveMask.h"
+#include "MuT/ASTOpSwitch.h"
+#include "MuT/StreamsPrivate.h"
+
+
+namespace mu
+{
+
+
+	//-------------------------------------------------------------------------------------------------
+	//-------------------------------------------------------------------------------------------------
+	//-------------------------------------------------------------------------------------------------
+	ASTOpMeshMaskClipUVMask::ASTOpMeshMaskClipUVMask()
+		: Source(this)
+		, Mask(this)
+	{
+	}
+
+
+	//-------------------------------------------------------------------------------------------------
+	ASTOpMeshMaskClipUVMask::~ASTOpMeshMaskClipUVMask()
+	{
+		// Explicit call needed to avoid recursive destruction
+		ASTOp::RemoveChildren();
+	}
+
+
+	//-------------------------------------------------------------------------------------------------
+	bool ASTOpMeshMaskClipUVMask::IsEqual(const ASTOp& otherUntyped) const
+	{
+		if (const ASTOpMeshMaskClipUVMask* other = dynamic_cast<const ASTOpMeshMaskClipUVMask*>(&otherUntyped))
+		{
+			return Source == other->Source && Mask == other->Mask && LayoutIndex==other->LayoutIndex;
+		}
+		return false;
+	}
+
+
+	//-------------------------------------------------------------------------------------------------
+	uint64 ASTOpMeshMaskClipUVMask::Hash() const
+	{
+		uint64 res = std::hash<void*>()(Source.child().get());
+		hash_combine(res, Mask.child());
+		return res;
+	}
+
+
+	//-------------------------------------------------------------------------------------------------
+	mu::Ptr<ASTOp> ASTOpMeshMaskClipUVMask::Clone(MapChildFuncRef mapChild) const
+	{
+		Ptr<ASTOpMeshMaskClipUVMask> n = new ASTOpMeshMaskClipUVMask();
+		n->Source = mapChild(Source.child());
+		n->Mask = mapChild(Mask.child());
+		n->LayoutIndex = LayoutIndex;
+		return n;
+	}
+
+
+	//-------------------------------------------------------------------------------------------------
+	void ASTOpMeshMaskClipUVMask::ForEachChild(const TFunctionRef<void(ASTChild&)> f)
+	{
+		f(Source);
+		f(Mask);
+	}
+
+
+	//-------------------------------------------------------------------------------------------------
+	void ASTOpMeshMaskClipUVMask::Link(FProgram& program, FLinkerOptions*)
+	{
+		// Already linked?
+		if (!linkedAddress)
+		{
+			OP::MeshMaskClipUVMaskArgs args;
+			FMemory::Memzero(args);
+
+			if (Source) args.Source = Source->linkedAddress;
+			if (Mask) args.Mask = Mask->linkedAddress;
+			args.LayoutIndex = LayoutIndex;
+
+			linkedAddress = (OP::ADDRESS)program.m_opAddress.Num();
+			program.m_opAddress.Add((uint32_t)program.m_byteCode.Num());
+			AppendCode(program.m_byteCode, OP_TYPE::ME_MASKCLIPUVMASK);
+			AppendCode(program.m_byteCode, args);
+		}
+	}
+
+
+	//-------------------------------------------------------------------------------------------------
+	//-------------------------------------------------------------------------------------------------
+	//-------------------------------------------------------------------------------------------------
+	namespace
+	{
+		class Sink_MeshMaskClipUVMaskSource
+		{
+		public:
+
+			// \TODO This is recursive and may cause stack overflows in big models.
+
+			mu::Ptr<ASTOp> Apply(const ASTOpMeshMaskClipUVMask* root)
+			{
+				m_root = root;
+				m_oldToNew.Empty();
+
+				m_initialSource = root->Source.child();
+				mu::Ptr<ASTOp> NewSource = Visit(m_initialSource);
+
+				// If there is any change, it is the new root.
+				if (NewSource != m_initialSource)
+				{
+					return NewSource;
+				}
+
+				return nullptr;
+			}
+
+		protected:
+
+			const ASTOpMeshMaskClipUVMask* m_root;
+			mu::Ptr<ASTOp> m_initialSource;
+			TMap<mu::Ptr<ASTOp>, mu::Ptr<ASTOp>> m_oldToNew;
+			TArray<mu::Ptr<ASTOp>> m_newOps;
+
+			mu::Ptr<ASTOp> Visit(const mu::Ptr<ASTOp>& at)
+			{
+				if (!at) return nullptr;
+
+				// Newly created?
+				if (m_newOps.Find(at) != INDEX_NONE)
+				{
+					return at;
+				}
+
+				// Already visited?
+				mu::Ptr<ASTOp>* cacheIt = m_oldToNew.Find(at);
+				if (cacheIt)
+				{
+					return *cacheIt;
+				}
+
+				mu::Ptr<ASTOp> newAt = at;
+				switch (at->GetOpType())
+				{
+
+				// This cannot be sunk since the result is different. Since the clipping is now correctly
+				// generated at the end of the chain when really necessary, this wrong optimisation is no 
+				// longer needed.
+				//case OP_TYPE::ME_MORPH2:
+		        //{
+		        //    break;
+		        //}
+
+				case OP_TYPE::ME_REMOVEMASK:
+				{
+					// Remove this op:
+					// This may lead to the mask being bigger than needed since it will include
+					// faces removed by the ignored removemask, but it is ok
+
+					// TODO: Swap instead of ignore, and implement removemask on a mask?
+					const ASTOpMeshRemoveMask* typedAt = dynamic_cast<const ASTOpMeshRemoveMask*>(at.get());
+					newAt = Visit(typedAt->source.child());
+					break;
+				}
+
+				case OP_TYPE::ME_CONDITIONAL:
+				{
+					// We move the mask creation down the two paths
+					// It always needs to be a clone because otherwise we could be modifying an
+					// instruction that shouldn't if the parent was a ME_REMOVEMASK above and we
+					// skipped the cloning for the parent.
+					Ptr<ASTOpConditional> newOp = mu::Clone<ASTOpConditional>(at);
+					newOp->yes = Visit(newOp->yes.child());
+					newOp->no = Visit(newOp->no.child());
+					newAt = newOp;
+					break;
+				}
+
+				case OP_TYPE::ME_SWITCH:
+				{
+					// We move the mask creation down all the paths
+					Ptr<ASTOpSwitch> newOp = mu::Clone<ASTOpSwitch>(at);
+					newOp->def = Visit(newOp->def.child());
+					for (ASTOpSwitch::FCase& c : newOp->cases)
+					{
+						c.branch = Visit(c.branch.child());
+					}
+					newAt = newOp;
+					break;
+				}
+
+				// This cannot be sunk since the result is different. Since the clipping is now correctly
+				// generated at the end of the chain when really necessary, this wrong optimisation is no 
+				// longer needed.
+				//case OP_TYPE::ME_CLIPMORPHPLANE:
+	   //         {
+	   //             // We move the mask creation down the Source
+	   //             auto typedAt = dynamic_cast<const ASTOpMeshClipMorphPlane*>(at.get());
+	   //             newAt = Visit(typedAt->Source.child());
+	   //             break;
+	   //         }
+
+				default:
+				{
+					//
+					if (at != m_initialSource)
+					{
+						Ptr<ASTOpMeshMaskClipUVMask> newOp = mu::Clone<ASTOpMeshMaskClipUVMask>(m_root);
+						newOp->Source = at;
+						newAt = newOp;
+					}
+					break;
+				}
+
+				}
+
+				m_oldToNew.Add(at, newAt);
+
+				return newAt;
+			}
+		};
+
+
+		class Sink_MeshMaskClipUVMaskClip
+		{
+		public:
+
+			// \TODO This is recursive and may cause stack overflows in big models.
+
+			mu::Ptr<ASTOp> Apply(const ASTOpMeshMaskClipUVMask* root)
+			{
+				m_root = root;
+				m_oldToNew.Empty();
+
+				m_initialClip = root->Mask.child();
+				mu::Ptr<ASTOp> NewClip = Visit(m_initialClip);
+
+				// If there is any change, it is the new root.
+				if (NewClip != m_initialClip)
+				{
+					return NewClip;
+				}
+
+				return nullptr;
+			}
+
+		protected:
+
+			const ASTOpMeshMaskClipUVMask* m_root;
+			mu::Ptr<ASTOp> m_initialClip;
+			TMap<mu::Ptr<ASTOp>, mu::Ptr<ASTOp>> m_oldToNew;
+			TArray<mu::Ptr<ASTOp>> m_newOps;
+
+			mu::Ptr<ASTOp> Visit(const mu::Ptr<ASTOp>& at)
+			{
+				if (!at) return nullptr;
+
+				// Newly created?
+				if (m_newOps.Contains(at))
+				{
+					return at;
+				}
+
+				// Already visited?
+				mu::Ptr<ASTOp>* cacheIt = m_oldToNew.Find(at);
+				if (cacheIt)
+				{
+					return *cacheIt;
+				}
+
+				mu::Ptr<ASTOp> newAt = at;
+				switch (at->GetOpType())
+				{
+
+				case OP_TYPE::ME_CONDITIONAL:
+				{
+					// We move the mask creation down the two paths
+					// It always needs to be a clone because otherwise we could be modifying an
+					// instruction that shouldn't if the parent was a ME_REMOVEMASK above and we
+					// skipped the cloning for the parent.
+					Ptr<ASTOpConditional> newOp = mu::Clone<ASTOpConditional>(at);
+					newOp->yes = Visit(newOp->yes.child());
+					newOp->no = Visit(newOp->no.child());
+					newAt = newOp;
+					break;
+				}
+
+				case OP_TYPE::ME_SWITCH:
+				{
+					// We move the mask creation down all the paths
+					Ptr<ASTOpSwitch> newOp = mu::Clone<ASTOpSwitch>(at);
+					newOp->def = Visit(newOp->def.child());
+					for (ASTOpSwitch::FCase& c : newOp->cases)
+					{
+						c.branch = Visit(c.branch.child());
+					}
+					newAt = newOp;
+					break;
+				}
+
+
+				default:
+				{
+					//
+					if (at != m_initialClip)
+					{
+						Ptr<ASTOpMeshMaskClipUVMask> newOp = mu::Clone<ASTOpMeshMaskClipUVMask>(m_root);
+						newOp->Mask = at;
+						newAt = newOp;
+					}
+					break;
+				}
+
+				}
+
+				m_oldToNew.Add(at, newAt);
+
+				return newAt;
+			}
+		};
+
+	}
+
+
+	//-------------------------------------------------------------------------------------------------
+	mu::Ptr<ASTOp> ASTOpMeshMaskClipUVMask::OptimiseSink(const FModelOptimizationOptions&, FOptimizeSinkContext&) const
+	{
+		// \TODO: Add logic state to the sinkers to avoid explosion with switches in both branches and similar cases.
+
+		// This will sink the operation down the Source
+		Sink_MeshMaskClipUVMaskSource SinkerSource;
+		mu::Ptr<ASTOp> at = SinkerSource.Apply(this);
+
+		// If we didn't sink it.
+		if (!at || at == this)
+		{
+			// This will sink the operation down the Mask child
+			Sink_MeshMaskClipUVMaskClip SinkerClip;
+			at = SinkerClip.Apply(this);
+		}
+
+		return at;
+	}
+
+}
+ 
\ No newline at end of file
diff --git a/Engine/Plugins/Experimental/Mutable/Source/MutableTools/Private/MuT/ASTOpMeshMaskClipUVMask.h b/Engine/Plugins/Experimental/Mutable/Source/MutableTools/Private/MuT/ASTOpMeshMaskClipUVMask.h
new file mode 100644
--- /dev/null
+++ b/Engine/Plugins/Experimental/Mutable/Source/MutableTools/Private/MuT/ASTOpMeshMaskClipUVMask.h
@@ -0,0 +1,44 @@
+// Copyright Epic Games, Inc. All Rights Reserved.
+
+#pragma once
+
+#include "HAL/Platform.h"
+#include "MuR/Operations.h"
+#include "MuR/Ptr.h"
+#include "MuT/AST.h"
+
+
+namespace mu
+{
+struct FProgram;
+
+
+	//---------------------------------------------------------------------------------------------
+	//!
+	//---------------------------------------------------------------------------------------------
+	class ASTOpMeshMaskClipUVMask : public ASTOp
+	{
+	public:
+
+		ASTChild Source;
+		ASTChild Mask;
+		uint8 LayoutIndex = 0;
+
+	public:
+
+		ASTOpMeshMaskClipUVMask();
+		ASTOpMeshMaskClipUVMask(const ASTOpMeshMaskClipUVMask&) = delete;
+		~ASTOpMeshMaskClipUVMask();
+
+		OP_TYPE GetOpType() const override { return OP_TYPE::ME_MASKCLIPUVMASK; }
+		uint64 Hash() const override;
+		bool IsEqual(const ASTOp& otherUntyped) const override;
+		Ptr<ASTOp> Clone(MapChildFuncRef mapChild) const override;
+		void ForEachChild(const TFunctionRef<void(ASTChild&)>) override;
+		void Link(FProgram& program, FLinkerOptions* Options) override;
+		Ptr<ASTOp> OptimiseSink(const FModelOptimizationOptions&, FOptimizeSinkContext&) const override;
+
+	};
+
+}
+
diff --git a/Engine/Plugins/Experimental/Mutable/Source/MutableTools/Private/MuT/CodeGenerator.cpp b/Engine/Plugins/Experimental/Mutable/Source/MutableTools/Private/MuT/CodeGenerator.cpp
--- a/Engine/Plugins/Experimental/Mutable/Source/MutableTools/Private/MuT/CodeGenerator.cpp
+++ b/Engine/Plugins/Experimental/Mutable/Source/MutableTools/Private/MuT/CodeGenerator.cpp
@@ -30,6 +30,7 @@
 #include "MuT/ASTOpMeshClipDeform.h"
 #include "MuT/ASTOpMeshClipMorphPlane.h"
 #include "MuT/ASTOpMeshMaskClipMesh.h"
+#include "MuT/ASTOpMeshMaskClipUVMask.h"
 #include "MuT/ASTOpMeshRemoveMask.h"
 #include "MuT/ASTOpMeshDifference.h"
 #include "MuT/ASTOpMeshMorph.h"
@@ -65,12 +66,10 @@
 #include "MuT/NodeMeshMorph.h"
 #include "MuT/NodeMeshReshape.h"
 #include "MuT/NodeModifier.h"
-#include "MuT/NodeModifierMeshClipDeform.h"
 #include "MuT/NodeModifierMeshClipDeformPrivate.h"
-#include "MuT/NodeModifierMeshClipMorphPlane.h"
 #include "MuT/NodeModifierMeshClipMorphPlanePrivate.h"
-#include "MuT/NodeModifierMeshClipWithMesh.h"
 #include "MuT/NodeModifierMeshClipWithMeshPrivate.h"
+#include "MuT/NodeModifierMeshClipWithUVMaskPrivate.h"
 #include "MuT/NodeModifierPrivate.h"
 #include "MuT/NodeObject.h"
 #include "MuT/NodeObjectGroupPrivate.h"
@@ -2160,13 +2159,12 @@ namespace mu
 				ClipOptions.State = m_currentStateIndex;
 
 				FMeshGenerationResult clipResult;
-				GenerateMesh(ClipOptions, clipResult, TypedClipNode->m_clipMesh);
+				GenerateMesh(ClipOptions, clipResult, TypedClipNode->ClipMesh);
 				op->clip = clipResult.meshOp;
 
 				if (!op->clip)
 				{
-					m_pErrorLog->GetPrivate()->Add
-					("Clip mesh has not been generated", ELMT_ERROR, errorContext);
+					m_pErrorLog->GetPrivate()->Add("Clip mesh has not been generated", ELMT_ERROR, errorContext);
 					continue;
 				}
 
@@ -2189,6 +2187,51 @@ namespace mu
 			}
 		}
 
+		// Process clip-with-mask modifiers
+		for (const FirstPassGenerator::FModifier& m : modifiers)
+		{
+			if (const NodeModifierMeshClipWithUVMask::Private* TypedClipNode = dynamic_cast<const NodeModifierMeshClipWithUVMask::Private*>(m.node))
+			{
+				Ptr<ASTOpMeshMaskClipUVMask> op = new ASTOpMeshMaskClipUVMask();
+				op->Source = preModifiersMesh;
+				op->LayoutIndex = TypedClipNode->LayoutIndex;
+
+				// Parameters
+				FImageGenerationOptions ClipOptions;
+				ClipOptions.ImageLayoutStrategy = CompilerOptions::TextureLayoutStrategy::None;
+				ClipOptions.LayoutBlockId = -1;
+				ClipOptions.CurrentStateIndex = m_currentStateIndex;
+
+				FImageGenerationResult ClipMaskResult;
+				GenerateImage(ClipOptions, ClipMaskResult, TypedClipNode->ClipMask);
+
+				// It could be IF_L_UBIT, but since this should be optimized out at compile time, leave the most cpu efficient.
+				op->Mask = GenerateImageFormat(ClipMaskResult.op, mu::EImageFormat::IF_L_UBYTE);
+
+				if (!op->Mask)
+				{
+					m_pErrorLog->GetPrivate()->Add("Clip UV mask has not been generated", ELMT_ERROR, errorContext);
+					continue;
+				}
+
+				Ptr<ASTOp> maskAt = op;
+
+				if (!removeOp)
+				{
+					removeOp = new ASTOpMeshRemoveMask();
+					removeOp->source = lastMeshOp;
+					lastMeshOp = removeOp;
+				}
+
+				Ptr<ASTOpFixed> surfCondOp = new ASTOpFixed();
+				surfCondOp->op.type = OP_TYPE::BO_AND;
+				surfCondOp->SetChild(surfCondOp->op.args.BoolBinary.a, m.objectCondition);
+				surfCondOp->SetChild(surfCondOp->op.args.BoolBinary.b, m.surfaceCondition);
+				Ptr<ASTOp> fullCondition = surfCondOp;
+
+				removeOp->AddRemove(fullCondition, maskAt);
+			}
+		}
 
 		// Process clip-morph-plane modifiers
 		for (const FirstPassGenerator::FModifier& m : modifiers)
diff --git a/Engine/Plugins/Experimental/Mutable/Source/MutableTools/Private/MuT/CodeGenerator_FirstPass.cpp b/Engine/Plugins/Experimental/Mutable/Source/MutableTools/Private/MuT/CodeGenerator_FirstPass.cpp
--- a/Engine/Plugins/Experimental/Mutable/Source/MutableTools/Private/MuT/CodeGenerator_FirstPass.cpp
+++ b/Engine/Plugins/Experimental/Mutable/Source/MutableTools/Private/MuT/CodeGenerator_FirstPass.cpp
@@ -22,6 +22,7 @@
 #include "MuT/NodeModifierMeshClipDeformPrivate.h"
 #include "MuT/NodeModifierMeshClipMorphPlanePrivate.h"
 #include "MuT/NodeModifierMeshClipWithMeshPrivate.h"
+#include "MuT/NodeModifierMeshClipWithUVMaskPrivate.h"
 #include "MuT/NodeObject.h"
 #include "MuT/NodeObjectGroupPrivate.h"
 #include "MuT/NodeObjectNewPrivate.h"
@@ -176,19 +177,36 @@ namespace mu
 
 
 	//---------------------------------------------------------------------------------------------
-    Ptr<ASTOp> FirstPassGenerator::Visit(const NodeModifierMeshClipWithMesh::Private& node)
+	Ptr<ASTOp> FirstPassGenerator::Visit(const NodeModifierMeshClipWithMesh::Private& node)
 	{
-        // Add the data about this modifier
+		// Add the data about this modifier
 		FModifier thisData;
 		thisData.node = &node;
-        thisData.objectCondition = m_currentCondition.Last().objectCondition;
-        thisData.stateCondition = m_currentStateCondition.Last();
-        thisData.lod = m_currentLOD;
-        thisData.positiveTags = m_currentPositiveTags;
-        thisData.negativeTags = m_currentNegativeTags;
-        modifiers.Add(thisData);
+		thisData.objectCondition = m_currentCondition.Last().objectCondition;
+		thisData.stateCondition = m_currentStateCondition.Last();
+		thisData.lod = m_currentLOD;
+		thisData.positiveTags = m_currentPositiveTags;
+		thisData.negativeTags = m_currentNegativeTags;
+		modifiers.Add(thisData);
+
+		return nullptr;
+	}
 
-        return nullptr;
+
+	//---------------------------------------------------------------------------------------------
+	Ptr<ASTOp> FirstPassGenerator::Visit(const NodeModifierMeshClipWithUVMask::Private& node)
+	{
+		// Add the data about this modifier
+		FModifier thisData;
+		thisData.node = &node;
+		thisData.objectCondition = m_currentCondition.Last().objectCondition;
+		thisData.stateCondition = m_currentStateCondition.Last();
+		thisData.lod = m_currentLOD;
+		thisData.positiveTags = m_currentPositiveTags;
+		thisData.negativeTags = m_currentNegativeTags;
+		modifiers.Add(thisData);
+
+		return nullptr;
 	}
 
 	
diff --git a/Engine/Plugins/Experimental/Mutable/Source/MutableTools/Private/MuT/CodeGenerator_FirstPass.h b/Engine/Plugins/Experimental/Mutable/Source/MutableTools/Private/MuT/CodeGenerator_FirstPass.h
--- a/Engine/Plugins/Experimental/Mutable/Source/MutableTools/Private/MuT/CodeGenerator_FirstPass.h
+++ b/Engine/Plugins/Experimental/Mutable/Source/MutableTools/Private/MuT/CodeGenerator_FirstPass.h
@@ -15,6 +15,7 @@
 #include "MuT/NodeModifierMeshClipDeform.h"
 #include "MuT/NodeModifierMeshClipMorphPlane.h"
 #include "MuT/NodeModifierMeshClipWithMesh.h"
+#include "MuT/NodeModifierMeshClipWithUVMask.h"
 #include "MuT/NodeObjectGroup.h"
 #include "MuT/NodeObjectNew.h"
 #include "MuT/NodePatchMesh.h"
@@ -52,8 +53,9 @@ namespace mu
         public Visitor<NodePatchMesh::Private, Ptr<ASTOp>, true>,
 
         public Visitor<NodeModifierMeshClipMorphPlane::Private, Ptr<ASTOp>, true>,
-        public Visitor<NodeModifierMeshClipWithMesh::Private, Ptr<ASTOp>, true>,
-        public Visitor<NodeModifierMeshClipDeform::Private, Ptr<ASTOp>, true>
+		public Visitor<NodeModifierMeshClipWithMesh::Private, Ptr<ASTOp>, true>,
+		public Visitor<NodeModifierMeshClipWithUVMask::Private, Ptr<ASTOp>, true>,
+		public Visitor<NodeModifierMeshClipDeform::Private, Ptr<ASTOp>, true>
 	{
 	public:
 
@@ -195,7 +197,8 @@ namespace mu
         Ptr<ASTOp> Visit(const NodeObjectGroup::Private&) override;
         Ptr<ASTOp> Visit(const NodeModifierMeshClipMorphPlane::Private&) override;
         Ptr<ASTOp> Visit(const NodeModifierMeshClipWithMesh::Private&) override;
-        Ptr<ASTOp> Visit(const NodeModifierMeshClipDeform::Private&) override;
+		Ptr<ASTOp> Visit(const NodeModifierMeshClipWithUVMask::Private&) override;
+		Ptr<ASTOp> Visit(const NodeModifierMeshClipDeform::Private&) override;
         Ptr<ASTOp> Visit(const NodePatchMesh::Private&) override;
 
 	private:
diff --git a/Engine/Plugins/Experimental/Mutable/Source/MutableTools/Private/MuT/ErrorLog.cpp b/Engine/Plugins/Experimental/Mutable/Source/MutableTools/Private/MuT/ErrorLog.cpp
--- a/Engine/Plugins/Experimental/Mutable/Source/MutableTools/Private/MuT/ErrorLog.cpp
+++ b/Engine/Plugins/Experimental/Mutable/Source/MutableTools/Private/MuT/ErrorLog.cpp
@@ -293,6 +293,7 @@ namespace mu
 		TEXT("ME_MERGE         "),
 		TEXT("ME_INTERPOLATE   "),
 		TEXT("ME_MASKCLIPMESH  "),
+		TEXT("ME_MASKCLIPUVMASK"),
 		TEXT("ME_MASKDIFF      "),
         TEXT("ME_REMOVEMASK    "),
         TEXT("ME_FORMAT        "),
diff --git a/Engine/Plugins/Experimental/Mutable/Source/MutableTools/Private/MuT/NodeModifier.cpp b/Engine/Plugins/Experimental/Mutable/Source/MutableTools/Private/MuT/NodeModifier.cpp
--- a/Engine/Plugins/Experimental/Mutable/Source/MutableTools/Private/MuT/NodeModifier.cpp
+++ b/Engine/Plugins/Experimental/Mutable/Source/MutableTools/Private/MuT/NodeModifier.cpp
@@ -9,6 +9,7 @@
 #include "MuT/NodeModifierMeshClipDeform.h"
 #include "MuT/NodeModifierMeshClipMorphPlane.h"
 #include "MuT/NodeModifierMeshClipWithMesh.h"
+#include "MuT/NodeModifierMeshClipWithUVMask.h"
 #include "MuT/NodeModifierPrivate.h"
 #include "MuT/NodePrivate.h"
 
@@ -56,6 +57,7 @@ namespace mu
 		if SERIALISE_CHILDREN(NodeModifierMeshClipMorphPlane, 			0 )
 		else if SERIALISE_CHILDREN(NodeModifierMeshClipWithMesh, 		1 )
 		else if SERIALISE_CHILDREN(NodeModifierMeshClipDeform, 		    2 )
+		else if SERIALISE_CHILDREN(NodeModifierMeshClipWithUVMask, 		3 )
 		else check(false);
 
 #undef SERIALISE_CHILDREN
@@ -76,7 +78,8 @@ namespace mu
 		{
 		case 0 :  return NodeModifierMeshClipMorphPlane::StaticUnserialise( arch ); break;
 		case 1 :  return NodeModifierMeshClipWithMesh::StaticUnserialise( arch ); break;
-		case 2 :  return NodeModifierMeshClipDeform::StaticUnserialise( arch ); break;
+		case 2 :  return NodeModifierMeshClipDeform::StaticUnserialise(arch); break;
+		case 3 :  return NodeModifierMeshClipWithUVMask::StaticUnserialise(arch); break;
 		default : check(false);
 		}
 
diff --git a/Engine/Plugins/Experimental/Mutable/Source/MutableTools/Private/MuT/NodeModifierMeshClipWithMesh.cpp b/Engine/Plugins/Experimental/Mutable/Source/MutableTools/Private/MuT/NodeModifierMeshClipWithMesh.cpp
--- a/Engine/Plugins/Experimental/Mutable/Source/MutableTools/Private/MuT/NodeModifierMeshClipWithMesh.cpp
+++ b/Engine/Plugins/Experimental/Mutable/Source/MutableTools/Private/MuT/NodeModifierMeshClipWithMesh.cpp
@@ -31,7 +31,7 @@ namespace mu
 	//---------------------------------------------------------------------------------------------
     int NodeModifierMeshClipWithMesh::GetInputCount() const
 	{
-		return 0;
+		return 1;
 	}
 
 
@@ -40,15 +40,17 @@ namespace mu
 	{
 		check( i>=0 && i< GetInputCount());
         (void)i;
-        return nullptr;
+        return m_pD->ClipMesh.get();
 	}
 
 
 	//---------------------------------------------------------------------------------------------
-    void NodeModifierMeshClipWithMesh::SetInputNode( int i, NodePtr )
+    void NodeModifierMeshClipWithMesh::SetInputNode( int i, Ptr<Node> Input )
 	{
 		check( i>=0 && i< GetInputCount());
         (void)i;
+		m_pD->ClipMesh = dynamic_cast<NodeMesh*>(Input.get());
+
     }
 
 
@@ -57,7 +59,7 @@ namespace mu
 	//---------------------------------------------------------------------------------------------
 	void NodeModifierMeshClipWithMesh::SetClipMesh(NodeMesh* clipMesh)
 	{
-		m_pD->m_clipMesh = clipMesh;
+		m_pD->ClipMesh = clipMesh;
 	}
 
 }
diff --git a/Engine/Plugins/Experimental/Mutable/Source/MutableTools/Private/MuT/NodeModifierMeshClipWithMeshPrivate.h b/Engine/Plugins/Experimental/Mutable/Source/MutableTools/Private/MuT/NodeModifierMeshClipWithMeshPrivate.h
--- a/Engine/Plugins/Experimental/Mutable/Source/MutableTools/Private/MuT/NodeModifierMeshClipWithMeshPrivate.h
+++ b/Engine/Plugins/Experimental/Mutable/Source/MutableTools/Private/MuT/NodeModifierMeshClipWithMeshPrivate.h
@@ -28,7 +28,7 @@ namespace mu
 		static NODE_TYPE s_type;
 
 		//! 
-		NodeMeshPtr m_clipMesh;
+		Ptr<NodeMesh> ClipMesh;
 
 		//!
 		void Serialise( OutputArchive& arch ) const
@@ -38,7 +38,7 @@ namespace mu
             uint32_t ver = 0;
 			arch << ver;
 
-			arch << m_clipMesh;
+			arch << ClipMesh;
 		}
 
 		//!
@@ -50,7 +50,7 @@ namespace mu
 			arch >> ver;
             check(ver<=0);
 
-			arch >> m_clipMesh;
+			arch >> ClipMesh;
 		}
 
 	};
diff --git a/Engine/Plugins/Experimental/Mutable/Source/MutableTools/Private/MuT/NodeModifierMeshClipWithUVMask.cpp b/Engine/Plugins/Experimental/Mutable/Source/MutableTools/Private/MuT/NodeModifierMeshClipWithUVMask.cpp
new file mode 100644
--- /dev/null
+++ b/Engine/Plugins/Experimental/Mutable/Source/MutableTools/Private/MuT/NodeModifierMeshClipWithUVMask.cpp
@@ -0,0 +1,71 @@
+// Copyright Epic Games, Inc. All Rights Reserved.
+
+
+#include "MuT/NodeModifierMeshClipWithUVMask.h"
+
+#include "Misc/AssertionMacros.h"
+#include "MuT/NodeImage.h"
+#include "MuT/NodeModifierMeshClipWithUVMaskPrivate.h"
+#include "MuT/NodePrivate.h"
+
+
+namespace mu
+{
+
+	//---------------------------------------------------------------------------------------------
+	// Static initialisation
+	//---------------------------------------------------------------------------------------------
+    NODE_TYPE NodeModifierMeshClipWithUVMask::Private::s_type =
+            NODE_TYPE( "NodeModifierMeshClipWithUVMask", NodeModifier::GetStaticType() );
+
+
+	//---------------------------------------------------------------------------------------------
+	//!
+	//---------------------------------------------------------------------------------------------
+
+    MUTABLE_IMPLEMENT_NODE( NodeModifierMeshClipWithUVMask, EType::MeshClipWithMesh, Node, Node::EType::Modifier)
+
+
+	//---------------------------------------------------------------------------------------------
+	// Node Interface
+	//---------------------------------------------------------------------------------------------
+    int NodeModifierMeshClipWithUVMask::GetInputCount() const
+	{
+		return 1;
+	}
+
+
+	//---------------------------------------------------------------------------------------------
+    Node* NodeModifierMeshClipWithUVMask::GetInputNode( int i ) const
+	{
+		check( i>=0 && i< GetInputCount());
+        (void)i;
+        return m_pD->ClipMask.get();
+	}
+
+
+	//---------------------------------------------------------------------------------------------
+    void NodeModifierMeshClipWithUVMask::SetInputNode( int i, Ptr<Node> Input )
+	{
+		check( i>=0 && i< GetInputCount());
+        (void)i;
+		m_pD->ClipMask = dynamic_cast<NodeImage*>(Input.get());
+    }
+
+
+	//---------------------------------------------------------------------------------------------
+	// Own Interface
+	//---------------------------------------------------------------------------------------------
+	void NodeModifierMeshClipWithUVMask::SetClipMask(NodeImage* InClipMask)
+	{
+		m_pD->ClipMask = InClipMask;
+	}
+
+	void NodeModifierMeshClipWithUVMask::SetLayoutIndex(uint8 LayoutIndex)
+	{
+		m_pD->LayoutIndex = LayoutIndex;
+	}
+
+	
+
+}
diff --git a/Engine/Plugins/Experimental/Mutable/Source/MutableTools/Private/MuT/NodeModifierMeshClipWithUVMaskPrivate.h b/Engine/Plugins/Experimental/Mutable/Source/MutableTools/Private/MuT/NodeModifierMeshClipWithUVMaskPrivate.h
new file mode 100644
--- /dev/null
+++ b/Engine/Plugins/Experimental/Mutable/Source/MutableTools/Private/MuT/NodeModifierMeshClipWithUVMaskPrivate.h
@@ -0,0 +1,60 @@
+// Copyright Epic Games, Inc. All Rights Reserved.
+
+#pragma once
+
+#include "MuT/NodeModifierPrivate.h"
+#include "MuT/NodeModifierMeshClipWithUVMask.h"
+#include "MuT/NodeImage.h"
+#include "MuT/AST.h"
+
+#include "MuR/MutableMath.h"
+
+
+namespace mu
+{
+
+    class NodeModifierMeshClipWithUVMask::Private : public NodeModifier::Private
+	{
+	public:
+
+		MUTABLE_DEFINE_CONST_VISITABLE()
+
+	public:
+
+		static NODE_TYPE s_type;
+
+		/** Image with the regions to remove. It will be interpreted as a bitmap. */
+		Ptr<NodeImage> ClipMask;
+
+		/** Layout index of the UVs to use inthe source mesh to ben clipped with the mask. */
+		uint8 LayoutIndex = 0;
+
+		//!
+		void Serialise( OutputArchive& arch ) const
+		{
+			NodeModifier::Private::Serialise(arch);
+
+            uint32 ver = 0;
+			arch << ver;
+
+			arch << ClipMask;
+			arch << LayoutIndex;
+		}
+
+		//!
+		void Unserialise( InputArchive& arch )
+		{
+			NodeModifier::Private::Unserialise( arch );
+			
+            uint32 ver;
+			arch >> ver;
+            check(ver<=0);
+
+			arch >> ClipMask;
+			arch >> LayoutIndex;
+		}
+
+	};
+
+
+}
diff --git a/Engine/Plugins/Experimental/Mutable/Source/MutableTools/Public/MuT/NodeModifierMeshClipWithMesh.h b/Engine/Plugins/Experimental/Mutable/Source/MutableTools/Public/MuT/NodeModifierMeshClipWithMesh.h
--- a/Engine/Plugins/Experimental/Mutable/Source/MutableTools/Public/MuT/NodeModifierMeshClipWithMesh.h
+++ b/Engine/Plugins/Experimental/Mutable/Source/MutableTools/Public/MuT/NodeModifierMeshClipWithMesh.h
@@ -12,13 +12,7 @@ namespace mu
 {
 
 	// Forward definitions
-	class NodeModifierMeshClipWithMesh;
-	typedef Ptr<NodeModifierMeshClipWithMesh> NodeModifierMeshClipWithMeshPtr;
-	typedef Ptr<const NodeModifierMeshClipWithMesh> NodeModifierMeshClipWithMeshPtrConst;
-
 	class NodeMesh;
-	typedef Ptr<NodeMesh> NodeMeshPtr;
-	typedef Ptr<const NodeMesh> NodeMeshPtrConst;
 
 
 	//! This node makes a new component from several meshes and images.
@@ -35,7 +29,7 @@ namespace mu
 
 		void SerialiseWrapper(OutputArchive& arch) const override;
 		static void Serialise( const NodeModifierMeshClipWithMesh* pNode, OutputArchive& arch );
-		static NodeModifierMeshClipWithMeshPtr StaticUnserialise( InputArchive& arch );
+		static Ptr<NodeModifierMeshClipWithMesh> StaticUnserialise( InputArchive& arch );
 
 		//-----------------------------------------------------------------------------------------
         // Node interface
@@ -48,7 +42,7 @@ namespace mu
 
         int GetInputCount() const override;
         Node* GetInputNode( int i ) const override;
-        void SetInputNode( int i, NodePtr pNode ) override;
+        void SetInputNode( int i, Ptr<Node> ) override;
 
 		//-----------------------------------------------------------------------------------------
         // Own interface
diff --git a/Engine/Plugins/Experimental/Mutable/Source/MutableTools/Public/MuT/NodeModifierMeshClipWithUVMask.h b/Engine/Plugins/Experimental/Mutable/Source/MutableTools/Public/MuT/NodeModifierMeshClipWithUVMask.h
new file mode 100644
--- /dev/null
+++ b/Engine/Plugins/Experimental/Mutable/Source/MutableTools/Public/MuT/NodeModifierMeshClipWithUVMask.h
@@ -0,0 +1,75 @@
+// Copyright Epic Games, Inc. All Rights Reserved.
+
+#pragma once
+
+#include "MuR/Ptr.h"
+#include "MuR/RefCounted.h"
+#include "MuT/Node.h"
+#include "MuT/NodeModifier.h"
+
+
+namespace mu
+{
+
+	// Forward definitions
+	class NodeImage;
+
+	//! This node makes a new component from several meshes and images.
+	//! \ingroup model
+	class MUTABLETOOLS_API NodeModifierMeshClipWithUVMask : public NodeModifier
+	{
+	public:
+
+		//-----------------------------------------------------------------------------------------
+		// Life cycle
+		//-----------------------------------------------------------------------------------------
+
+		NodeModifierMeshClipWithUVMask();
+
+		void SerialiseWrapper(OutputArchive& arch) const override;
+		static void Serialise( const NodeModifierMeshClipWithUVMask* pNode, OutputArchive& arch );
+		static Ptr<NodeModifierMeshClipWithUVMask> StaticUnserialise( InputArchive& arch );
+
+		//-----------------------------------------------------------------------------------------
+        // Node interface
+		//-----------------------------------------------------------------------------------------
+
+        
+
+        const NODE_TYPE* GetType() const override;
+		static const NODE_TYPE* GetStaticType();
+
+        int GetInputCount() const override;
+        Node* GetInputNode( int i ) const override;
+        void SetInputNode( int i, Ptr<Node> pNode ) override;
+
+		//-----------------------------------------------------------------------------------------
+        // Own interface
+		//-----------------------------------------------------------------------------------------
+
+		/** Set the expression generating the image witht he UV mask use to clip the mesh. */
+		void SetClipMask(NodeImage*);
+
+		/** Set the UV channel index for the UVs to check against the mask. */
+		void SetLayoutIndex(uint8 LayoutIndex);
+
+		//-----------------------------------------------------------------------------------------
+		// Interface pattern
+		//-----------------------------------------------------------------------------------------
+		class Private;
+		Private* GetPrivate() const;
+        Node::Private* GetBasePrivate() const override;
+
+	protected:
+
+		//! Forbidden. Manage with the Ptr<> template.
+		~NodeModifierMeshClipWithUVMask();
+
+	private:
+
+		Private* m_pD;
+
+	};
+
+
+}
