#rnx

diff --git a/Engine/Plugins/Experimental/Mutable/Source/CustomizableObject/Classes/MuCO/CustomizableObjectInstance.h b/Engine/Plugins/Experimental/Mutable/Source/CustomizableObject/Classes/MuCO/CustomizableObjectInstance.h
--- a/Engine/Plugins/Experimental/Mutable/Source/CustomizableObject/Classes/MuCO/CustomizableObjectInstance.h
+++ b/Engine/Plugins/Experimental/Mutable/Source/CustomizableObject/Classes/MuCO/CustomizableObjectInstance.h
@@ -15,9 +15,9 @@ class FProperty;
 class UAnimInstance;
 class UCustomizableInstancePrivateData; // This is used to hide Mutable SDK members in the public headers.
 class UCustomizableObject;
-class UCustomizableSkeletalComponent;
 class UTexture2D;
 class FUpdateContextPrivate;
+class UCustomizableObjectInstanceUsage;
 struct FFrame;
 struct FGameplayTagContainer;
 struct FPropertyChangedEvent;
@@ -648,7 +648,7 @@ public:
 	UCustomizableInstancePrivateData* GetPrivate() const;
 
 	// TEMP VARIABLE to check the Min desired LODs for this instance
-	TWeakObjectPtr<UCustomizableSkeletalComponent> NearestToActor;
+	TWeakObjectPtr<UCustomizableObjectInstanceUsage> NearestToActor;
 	TWeakObjectPtr<const AActor> NearestToViewCenter;
 
 #if WITH_EDITOR
diff --git a/Engine/Plugins/Experimental/Mutable/Source/CustomizableObject/Classes/MuCO/CustomizableObjectInstanceUsage.h b/Engine/Plugins/Experimental/Mutable/Source/CustomizableObject/Classes/MuCO/CustomizableObjectInstanceUsage.h
new file mode 100644
--- /dev/null
+++ b/Engine/Plugins/Experimental/Mutable/Source/CustomizableObject/Classes/MuCO/CustomizableObjectInstanceUsage.h
@@ -0,0 +1,117 @@
+// Copyright Epic Games, Inc. All Rights Reserved.
+
+#pragma once
+
+#include "MuCO/CustomizableObjectInstance.h"
+#include "Tickable.h"
+#include "Engine/EngineBaseTypes.h"
+
+#include "CustomizableObjectInstanceUsage.generated.h"
+
+class AActor;
+class FObjectPreSaveContext;
+class UObject;
+class USkeletalMesh;
+class UCustomizableSkeletalComponent;
+struct FFrame;
+enum class EUpdateResult : uint8;
+
+DECLARE_DELEGATE(FCustomizableObjectInstanceUsageUpdatedDelegate);
+
+// This class can be used instead of a UCustomizableComponent (for example for non-BP projects) to link a 
+// UCustomizableObjectInstance and a USkeletalComponent so that the CustomizableObjectSystem takes care of updating it and its LODs, 
+// streaming, etc. It's a UObject, so it will be much cheaper than a UCustomizableComponent as it won't have to refresh its transforms
+// every time it's moved.
+UCLASS(Blueprintable, BlueprintType, ClassGroup = (CustomizableObject), meta = (BlueprintSpawnableComponent))
+class CUSTOMIZABLEOBJECT_API UCustomizableObjectInstanceUsage : public UObject, public FTickableGameObject
+{
+public:
+	GENERATED_BODY()
+
+	UPROPERTY(Transient)
+	float SkippedLastRenderTime;
+
+	FCustomizableObjectInstanceUsageUpdatedDelegate UpdatedDelegate;
+
+	void SetCustomizableObjectInstance(UCustomizableObjectInstance* CustomizableObjectInstance);
+	UCustomizableObjectInstance* GetCustomizableObjectInstance() const;
+
+	void SetComponentIndex(int32 ComponentIndex);
+	int32 GetComponentIndex() const;
+
+	void AttachTo(USkeletalMeshComponent* SkeletalMeshComponent);
+	USkeletalMeshComponent* GetAttachParent() const;
+
+	/** Common end point of all updates. Even those which failed. */
+	void Callbacks() const;
+	
+	USkeletalMesh* GetSkeletalMesh() const;
+	USkeletalMesh* GetAttachedSkeletalMesh() const;
+
+	/** Update Skeletal Mesh asynchronously. */
+	UFUNCTION(BlueprintCallable, Category = CustomizableObject)
+	void UpdateSkeletalMeshAsync(bool bNeverSkipUpdate = false);
+
+	/** Update Skeletal Mesh asynchronously. Callback will be called once the update finishes, even if it fails. */
+	UFUNCTION(BlueprintCallable, Category = CustomizableObjectInstance)
+	void UpdateSkeletalMeshAsyncResult(FInstanceUpdateDelegate Callback, bool bIgnoreCloseDist = false, bool bForceHighPriority = false);
+
+	void SetSkeletalMesh(USkeletalMesh* SkeletalMesh, bool bReinitPose=true, bool bForceClothReset=true);
+	void SetPhysicsAsset(class UPhysicsAsset* PhysicsAsset);
+	void UpdateDistFromComponentToPlayer(const AActor* const Pawn, bool bForceEvenIfNotBegunPlay = false);
+
+	void SetVisibilityOfSkeletalMeshSectionWithMaterialName(bool bVisible, const FString& MaterialName, int32 LOD);
+
+	// Set to true to replace the SkeletalMesh of the parent component by the ReferenceSkeletalMesh or the generated SkeletalMesh 
+	void SetPendingSetSkeletalMesh(bool bIsActive);
+	bool GetPendingSetSkeletalMesh() const;
+
+	// Set to true to avoid replacing the SkeletalMesh of the parent component by the ReferenceSkeletalMesh if bPendingSetSkeletalMesh is true
+	void SetSkipSetReferenceSkeletalMesh(bool bIsActive);
+	bool GetSkipSetReferenceSkeletalMesh() const;
+
+	// USceneComponent interface
+	virtual void BeginDestroy() override;
+
+	// Returns true if the NetMode of the associated UCustomizableSkeletalComponent (or the associated SkeletalMeshComponent if the former does not exist) is equal to InNetMode
+	bool IsNetMode(ENetMode InNetMode) const;
+
+#if WITH_EDITOR
+	// Used to generate instances outside the CustomizableObject editor and PIE
+	void UpdateDistFromComponentToLevelEditorCamera(const FVector& CameraPosition);
+	void EditorUpdateComponent();
+#endif
+
+private:
+	// If this CustomizableSkeletalComponent is not null, it means this Usage was created by it, and all persistent properties should be obtained through it
+	TObjectPtr<UCustomizableSkeletalComponent> CustomizableSkeletalComponent;
+
+	// If no CustomizableSkeletalComponent is associated, this SkeletalComponent will be used
+	UPROPERTY()
+	TObjectPtr<USkeletalMeshComponent> UsedSkeletalMeshComponent;
+
+	// If no CustomizableSkeletalComponent is associated, this Instance will be used
+	UPROPERTY()
+	TObjectPtr<UCustomizableObjectInstance> UsedCustomizableObjectInstance;
+
+	// If no CustomizableSkeletalComponent is associated, this Index will be used
+	UPROPERTY()
+	int32 UsedComponentIndex;
+
+	// Used to replace the SkeletalMesh of the parent component by the ReferenceSkeletalMesh or the generated SkeletalMesh 
+	bool bUsedPendingSetSkeletalMesh = false;
+
+	// Used to avoid replacing the SkeletalMesh of the parent component by the ReferenceSkeletalMesh if bPendingSetSkeletalMesh is true
+	bool bUsedSkipSetReferenceSkeletalMesh = false;
+
+	// Begin FTickableGameObject
+	virtual void Tick(float DeltaTime) override;
+	virtual ETickableTickType GetTickableTickType() const override { return ETickableTickType::Always; }
+	virtual TStatId GetStatId() const override;
+	virtual bool IsTickableWhenPaused() const override { return true; }
+	virtual bool IsTickableInEditor() const override { return true; }
+	// End FTickableGameObject
+
+	friend UCustomizableSkeletalComponent;
+};
+
diff --git a/Engine/Plugins/Experimental/Mutable/Source/CustomizableObject/Classes/MuCO/CustomizableSkeletalComponent.h b/Engine/Plugins/Experimental/Mutable/Source/CustomizableObject/Classes/MuCO/CustomizableSkeletalComponent.h
--- a/Engine/Plugins/Experimental/Mutable/Source/CustomizableObject/Classes/MuCO/CustomizableSkeletalComponent.h
+++ b/Engine/Plugins/Experimental/Mutable/Source/CustomizableObject/Classes/MuCO/CustomizableSkeletalComponent.h
@@ -3,19 +3,12 @@
 #pragma once
 
 #include "Components/SceneComponent.h"
+
 #include "MuCO/CustomizableObjectInstance.h"
 
 #include "CustomizableSkeletalComponent.generated.h"
 
-class AActor;
-class FObjectPreSaveContext;
-class UCustomizableSkeletalComponent;
-class UObject;
-class USkeletalMesh;
-struct FFrame;
-enum class EUpdateResult : uint8;
 
-DECLARE_DELEGATE_TwoParams(FCustomizableSkeletalComponentPreUpdateDelegate, UCustomizableSkeletalComponent* Component, USkeletalMesh* NextMesh);
 DECLARE_DELEGATE(FCustomizableSkeletalComponentUpdatedDelegate);
 
 UCLASS(Blueprintable, BlueprintType, ClassGroup = (CustomizableObject), meta = (BlueprintSpawnableComponent))
@@ -24,7 +17,11 @@ class CUSTOMIZABLEOBJECT_API UCustomizableSkeletalComponent : public USceneCompo
 public:
 	GENERATED_BODY()
 
-	UCustomizableSkeletalComponent();
+	UPROPERTY(BlueprintReadWrite, EditAnywhere, Category = CustomizableSkeletalMesh)
+	TObjectPtr<UCustomizableObjectInstance> CustomizableObjectInstance;
+
+	UPROPERTY(BlueprintReadWrite, EditAnywhere, Category = CustomizableSkeletalMesh)
+	int32 ComponentIndex;
 
 	// Used to replace the SkeletalMesh of the parent component by the ReferenceSkeletalMesh or the generated SkeletalMesh 
 	bool bPendingSetSkeletalMesh = false;
@@ -32,16 +29,6 @@ public:
 	// Used to avoid replacing the SkeletalMesh of the parent component by the ReferenceSkeletalMesh if bPendingSetSkeletalMesh is true
 	bool bSkipSetReferenceSkeletalMesh = false;
 
-	UPROPERTY(Transient)
-	float SkippedLastRenderTime;
-
-	UPROPERTY(BlueprintReadWrite, EditAnywhere, Category = CustomizableSkeletalMesh)
-	TObjectPtr<UCustomizableObjectInstance> CustomizableObjectInstance;
-
-	UPROPERTY(BlueprintReadWrite, EditAnywhere, Category = CustomizableSkeletalMesh)
-	int32 ComponentIndex;
-
-	FCustomizableSkeletalComponentPreUpdateDelegate PreUpdateDelegate;
 	FCustomizableSkeletalComponentUpdatedDelegate UpdatedDelegate;
 
 	/** Common end point of all updates. Even those which failed. */
@@ -64,9 +51,6 @@ public:
 
 	void SetVisibilityOfSkeletalMeshSectionWithMaterialName(bool bVisible, const FString& MaterialName, int32 LOD);
 
-	// USceneComponent interface
-	virtual void BeginDestroy() override;
-
 #if WITH_EDITOR
 	// Used to generate instances outside the CustomizableObject editor and PIE
 	void UpdateDistFromComponentToLevelEditorCamera(const FVector& CameraPosition);
@@ -74,8 +58,13 @@ public:
 #endif
 
 private:
-	void TickComponent(float DeltaTime, enum ELevelTick TickType, FActorComponentTickFunction* ThisTickFunction) override;
+	UPROPERTY(Transient)
+	TObjectPtr<UCustomizableObjectInstanceUsage> CustomizableObjectInstanceUsage;
+
 	void OnAttachmentChanged() override;
+	void CreateCustomizableObjectInstanceUsage();
 
+protected:
+	virtual void PostInitProperties() override;
 };
 
diff --git a/Engine/Plugins/Experimental/Mutable/Source/CustomizableObject/Private/MuCO/CustomizableInstanceLODManagement.cpp b/Engine/Plugins/Experimental/Mutable/Source/CustomizableObject/Private/MuCO/CustomizableInstanceLODManagement.cpp
--- a/Engine/Plugins/Experimental/Mutable/Source/CustomizableObject/Private/MuCO/CustomizableInstanceLODManagement.cpp
+++ b/Engine/Plugins/Experimental/Mutable/Source/CustomizableObject/Private/MuCO/CustomizableInstanceLODManagement.cpp
@@ -11,8 +11,8 @@
 #include "MuCO/CustomizableObject.h"
 #include "MuCO/CustomizableObjectInstance.h"
 #include "MuCO/CustomizableObjectSystem.h"
-#include "MuCO/CustomizableSkeletalComponent.h"
 #include "MuCO/UnrealPortabilityHelpers.h"
+#include "MuCO/CustomizableObjectInstanceUsage.h"
 #include "UObject/UObjectIterator.h"
 #include "Components/SkeletalMeshComponent.h"
 
@@ -70,18 +70,18 @@ UCustomizableInstanceLODManagement::~UCustomizableInstanceLODManagement()
 // ViewCenter is the origin where the distances will be measured from.
 void UpdatePawnToInstancesDistances(const class UCustomizableObjectInstance* OnlyForInstance, const TWeakObjectPtr<const AActor> ViewCenter)
 {
-	for (TObjectIterator<UCustomizableSkeletalComponent> CustomizableSkeletalComponent; CustomizableSkeletalComponent; ++CustomizableSkeletalComponent)
+	for (TObjectIterator<UCustomizableObjectInstanceUsage> CustomizableObjectInstanceUsage; CustomizableObjectInstanceUsage; ++CustomizableObjectInstanceUsage)
 	{
 #if WITH_EDITOR
-		if (CustomizableSkeletalComponent && CustomizableSkeletalComponent->IsNetMode(NM_DedicatedServer))
+		if (CustomizableObjectInstanceUsage && CustomizableObjectInstanceUsage->IsNetMode(NM_DedicatedServer))
 		{
 			continue;
 		}
 #endif
 
-		if (CustomizableSkeletalComponent->IsValidLowLevel() && (OnlyForInstance == nullptr || CustomizableSkeletalComponent->CustomizableObjectInstance == OnlyForInstance))
+		if (CustomizableObjectInstanceUsage->IsValidLowLevel() && (OnlyForInstance == nullptr || CustomizableObjectInstanceUsage->GetCustomizableObjectInstance() == OnlyForInstance))
 		{
-			CustomizableSkeletalComponent->UpdateDistFromComponentToPlayer(ViewCenter.IsValid() ? ViewCenter.Get() : nullptr, OnlyForInstance != nullptr);
+			CustomizableObjectInstanceUsage->UpdateDistFromComponentToPlayer(ViewCenter.IsValid() ? ViewCenter.Get() : nullptr, OnlyForInstance != nullptr);
 		}
 	}
 }
@@ -90,17 +90,17 @@ void UpdatePawnToInstancesDistances(const class UCustomizableObjectInstance* Onl
 // Used to manually update instances in the level editor
 void UpdateCameraToInstancesDistance(const FVector CameraPosition)
 {
-	for (TObjectIterator<UCustomizableSkeletalComponent> CustomizableSkeletalComponent; CustomizableSkeletalComponent; ++CustomizableSkeletalComponent)
+	for (TObjectIterator<UCustomizableObjectInstanceUsage> CustomizableObjectInstanceUsage; CustomizableObjectInstanceUsage; ++CustomizableObjectInstanceUsage)
 	{
 #if WITH_EDITOR
-		if (CustomizableSkeletalComponent && CustomizableSkeletalComponent->IsNetMode(NM_DedicatedServer))
+		if (CustomizableObjectInstanceUsage && CustomizableObjectInstanceUsage->IsNetMode(NM_DedicatedServer))
 		{
 			continue;
 		}
 #endif
-		if (CustomizableSkeletalComponent->IsValidLowLevel() && !CustomizableSkeletalComponent->IsTemplate())
+		if (CustomizableObjectInstanceUsage->IsValidLowLevel() && !CustomizableObjectInstanceUsage->IsTemplate())
 		{
-			CustomizableSkeletalComponent->UpdateDistFromComponentToLevelEditorCamera(CameraPosition);
+			CustomizableObjectInstanceUsage->UpdateDistFromComponentToLevelEditorCamera(CameraPosition);
 		}
 	}
 }
@@ -116,23 +116,23 @@ void UCustomizableInstanceLODManagement::UpdateInstanceDistsAndLODs(FMutableInst
 
 		if (ViewCenters.Num() == 0) // Just use the first pawn
 		{
-			UCustomizableSkeletalComponent* FirstCustomizableSkeletalComponent = nullptr;
+			UCustomizableObjectInstanceUsage* FirstCustomizableObjectInstanceUsage = nullptr;
 
 	#if WITH_EDITOR
 			bool bLevelEditorInstancesUpdated = false;
 	#endif
 
-			for (TObjectIterator<UCustomizableSkeletalComponent> CustomizableSkeletalComponent; CustomizableSkeletalComponent; ++CustomizableSkeletalComponent)
+			for (TObjectIterator<UCustomizableObjectInstanceUsage> CustomizableObjectInstanceUsage; CustomizableObjectInstanceUsage; ++CustomizableObjectInstanceUsage)
 			{
 #if WITH_EDITOR
-				if (CustomizableSkeletalComponent && CustomizableSkeletalComponent->IsNetMode(NM_DedicatedServer))
+				if (CustomizableObjectInstanceUsage && CustomizableObjectInstanceUsage->IsNetMode(NM_DedicatedServer))
 				{
 					continue;
 				}
 #endif
-				if (CustomizableSkeletalComponent && !CustomizableSkeletalComponent->IsTemplate())
+				if (CustomizableObjectInstanceUsage && !CustomizableObjectInstanceUsage->IsTemplate())
 				{
-					UWorld* LocalWorld = CustomizableSkeletalComponent->GetWorld();
+					UWorld* LocalWorld = CustomizableObjectInstanceUsage->GetWorld();
 					APlayerController* Controller = LocalWorld ? LocalWorld->GetFirstPlayerController() : nullptr;
 					TWeakObjectPtr<const AActor> ViewCenter = Controller ? TWeakObjectPtr<const AActor>(Controller->GetPawn()) : nullptr;
 
@@ -143,9 +143,9 @@ void UCustomizableInstanceLODManagement::UpdateInstanceDistsAndLODs(FMutableInst
 					}
 
 	#if WITH_EDITOR
-					else if (CustomizableSkeletalComponent->GetWorld())
+					else if (CustomizableObjectInstanceUsage->GetWorld())
 					{
-						EWorldType::Type worldType = CustomizableSkeletalComponent->GetWorld()->WorldType;
+						EWorldType::Type worldType = CustomizableObjectInstanceUsage->GetWorld()->WorldType;
 					
 						// Level Editor Instances (non PIE)
 						if (!bLevelEditorInstancesUpdated && worldType == EWorldType::Editor)
@@ -164,7 +164,7 @@ void UCustomizableInstanceLODManagement::UpdateInstanceDistsAndLODs(FMutableInst
 						// Blueprint instances
 						else if (worldType == EWorldType::EditorPreview)
 						{
-							CustomizableSkeletalComponent->EditorUpdateComponent();
+							CustomizableObjectInstanceUsage->EditorUpdateComponent();
 						}
 					}
 	#endif // WITH_EDITOR
@@ -292,18 +292,18 @@ void UCustomizableInstanceLODManagement::UpdateInstanceDistsAndLODs(FMutableInst
 		TMap<TObjectPtr<UCustomizableObjectInstance>, FLODTracker> InstancesMinLOD;
 		InstancesMinLOD.Reserve(100);
 
-		for (TObjectIterator<UCustomizableSkeletalComponent> CustomizableSkeletalComponent; CustomizableSkeletalComponent; ++CustomizableSkeletalComponent)
+		for (TObjectIterator<UCustomizableObjectInstanceUsage> CustomizableObjectInstanceUsage; CustomizableObjectInstanceUsage; ++CustomizableObjectInstanceUsage)
 		{
 #if WITH_EDITOR
-			if (CustomizableSkeletalComponent && CustomizableSkeletalComponent->IsNetMode(NM_DedicatedServer))
+			if (CustomizableObjectInstanceUsage && CustomizableObjectInstanceUsage->IsNetMode(NM_DedicatedServer))
 			{
 				continue;
 			}
 #endif
 
-			if (CustomizableSkeletalComponent && !CustomizableSkeletalComponent->IsTemplate())
+			if (CustomizableObjectInstanceUsage && !CustomizableObjectInstanceUsage->IsTemplate())
 			{
-				UCustomizableObjectInstance* COI = CustomizableSkeletalComponent->CustomizableObjectInstance;
+				UCustomizableObjectInstance* COI = CustomizableObjectInstanceUsage->GetCustomizableObjectInstance();
 				if (!COI || !COI->GetCustomizableObject())
 				{
 					continue;
@@ -313,7 +313,7 @@ void UCustomizableInstanceLODManagement::UpdateInstanceDistsAndLODs(FMutableInst
 
 				EWorldType::Type WorldType = EWorldType::Type::None;
 
-				UWorld* World = CustomizableSkeletalComponent->GetWorld();
+				UWorld* World = CustomizableObjectInstanceUsage->GetWorld();
 				if (World)
 				{
 					WorldType = World->WorldType;
@@ -321,12 +321,12 @@ void UCustomizableInstanceLODManagement::UpdateInstanceDistsAndLODs(FMutableInst
 
 				// Blueprint instances and CO Editors
 				
-				const USceneComponent* const AttachParentComponent = CustomizableSkeletalComponent->GetAttachParent();
+				const USceneComponent* const AttachParentComponent = CustomizableObjectInstanceUsage->GetAttachParent();
 				bool bAttachParentActor = AttachParentComponent ? AttachParentComponent->GetOwner()!=nullptr : false;
 
 				if (WorldType == EWorldType::EditorPreview || (!World && !bAttachParentActor))
 				{
-					CustomizableSkeletalComponent->EditorUpdateComponent();
+					CustomizableObjectInstanceUsage->EditorUpdateComponent();
 					continue;
 				}
 
@@ -345,12 +345,12 @@ void UCustomizableInstanceLODManagement::UpdateInstanceDistsAndLODs(FMutableInst
 					LODTracker.bInitialized = true;
 				}
 
-				USkeletalMeshComponent* Parent = Cast<USkeletalMeshComponent>(CustomizableSkeletalComponent->GetAttachParent());
+				USkeletalMeshComponent* Parent = Cast<USkeletalMeshComponent>(CustomizableObjectInstanceUsage->GetAttachParent());
 				if (Parent)
 				{
 					COI->SetIsBeingUsedByComponentInPlay(true);
 
-					int32 ComponentIndex = CustomizableSkeletalComponent->ComponentIndex;
+					int32 ComponentIndex = CustomizableObjectInstanceUsage->GetComponentIndex();
 
 #if WITH_EDITOR
 					// If the instance is generated but the component doesn't have a mesh, set it.
diff --git a/Engine/Plugins/Experimental/Mutable/Source/CustomizableObject/Private/MuCO/CustomizableObjectDGGUI.cpp b/Engine/Plugins/Experimental/Mutable/Source/CustomizableObject/Private/MuCO/CustomizableObjectDGGUI.cpp
--- a/Engine/Plugins/Experimental/Mutable/Source/CustomizableObject/Private/MuCO/CustomizableObjectDGGUI.cpp
+++ b/Engine/Plugins/Experimental/Mutable/Source/CustomizableObject/Private/MuCO/CustomizableObjectDGGUI.cpp
@@ -7,7 +7,7 @@
 
 #include UE_INLINE_GENERATED_CPP_BY_NAME(CustomizableObjectDGGUI)
 
-void UDGGUI::OpenDGGUI(const int32 SlotID, UCustomizableSkeletalComponent* SelectedCustomizableSkeletalComponent, const UWorld* CurrentWorld, const int32 PlayerIndex)
+void UDGGUI::OpenDGGUI(const int32 SlotID, UCustomizableObjectInstanceUsage* SelectedCustomizableObjectInstanceUsage, const UWorld* CurrentWorld, const int32 PlayerIndex)
 {
 #if !UE_BUILD_SHIPPING
 	if (APlayerController* Player = UGameplayStatics::GetPlayerController(CurrentWorld, PlayerIndex))
@@ -18,7 +18,7 @@ void UDGGUI::OpenDGGUI(const int32 SlotID, UCustomizableSkeletalComponent* Selec
 			UDGGUI* WDGUI = CreateWidget<UDGGUI>(Player, DGUI);
 			if (WDGUI)
 			{
-				WDGUI->SetCustomizableSkeletalComponent(SelectedCustomizableSkeletalComponent);
+				WDGUI->SetCustomizableObjectInstanceUsage(SelectedCustomizableObjectInstanceUsage);
 				WDGUI->AddToViewport();
 				Player->SetShowMouseCursor(true);
 			}
@@ -35,9 +35,9 @@ bool UDGGUI::CloseExistingDGGUI(const UWorld* CurrentWorld)
 	{
 		if (PreviousGUI->IsValidLowLevel())
 		{
-			if (UCustomizableSkeletalComponent* CustomizableComponent = PreviousGUI->GetCustomizableSkeletalComponent())
+			if (UCustomizableObjectInstanceUsage* CustomizableObjectInstanceUsage = PreviousGUI->GetCustomizableObjectInstanceUsage())
 			{
-				PreviousGUI->SetCustomizableSkeletalComponent(nullptr);
+				PreviousGUI->SetCustomizableObjectInstanceUsage(nullptr);
 				bClosing = true;
 			}
 			PreviousGUI->RemoveFromParent();
diff --git a/Engine/Plugins/Experimental/Mutable/Source/CustomizableObject/Private/MuCO/CustomizableObjectInstance.cpp b/Engine/Plugins/Experimental/Mutable/Source/CustomizableObject/Private/MuCO/CustomizableObjectInstance.cpp
--- a/Engine/Plugins/Experimental/Mutable/Source/CustomizableObject/Private/MuCO/CustomizableObjectInstance.cpp
+++ b/Engine/Plugins/Experimental/Mutable/Source/CustomizableObject/Private/MuCO/CustomizableObjectInstance.cpp
@@ -24,7 +24,7 @@
 #include "MuCO/CustomizableObjectExtension.h"
 #include "MuCO/CustomizableObjectMipDataProvider.h"
 #include "MuCO/CustomizableObjectPrivate.h"
-#include "MuCO/CustomizableSkeletalComponent.h"
+#include "MuCO/CustomizableObjectInstanceUsage.h"
 #include "MuCO/DefaultImageProvider.h"
 #include "MuCO/ICustomizableObjectModule.h"
 #include "MuCO/UnrealConversionUtils.h"
@@ -2710,22 +2710,22 @@ void UCustomizableInstancePrivateData::DiscardResourcesAndSetReferenceSkeletalMe
 	Public->SkeletalMeshes.Reset();
 	DescriptorRuntimeHash = FDescriptorRuntimeHash();
 
-	for (TObjectIterator<UCustomizableSkeletalComponent> It; It; ++It)
+	for (TObjectIterator<UCustomizableObjectInstanceUsage> It; It; ++It)
 	{
-		UCustomizableSkeletalComponent* CustomizableSkeletalComponent = *It;
+		UCustomizableObjectInstanceUsage* CustomizableObjectInstanceUsage = *It;
 
 #if WITH_EDITOR
-		if (CustomizableSkeletalComponent && CustomizableSkeletalComponent->IsNetMode(NM_DedicatedServer))
+		if (CustomizableObjectInstanceUsage && CustomizableObjectInstanceUsage->IsNetMode(NM_DedicatedServer))
 		{
 			continue;
 		}
 #endif
 
-		if (CustomizableSkeletalComponent && CustomizableSkeletalComponent->CustomizableObjectInstance == Public)
+		if (CustomizableObjectInstanceUsage && CustomizableObjectInstanceUsage->GetCustomizableObjectInstance() == Public)
 		{
 			UCustomizableObject* CustomizableObject = Public->GetCustomizableObject();
 			bool bReplaceDiscardedWithReferenceMesh = UCustomizableObjectSystem::GetInstance()->GetPrivate()->IsReplaceDiscardedWithReferenceMeshEnabled();
-			CustomizableSkeletalComponent->SetSkeletalMesh(CustomizableObject && bReplaceDiscardedWithReferenceMesh ? CustomizableObject->GetRefSkeletalMesh(CustomizableSkeletalComponent->ComponentIndex) : nullptr);
+			CustomizableObjectInstanceUsage->SetSkeletalMesh(CustomizableObject && bReplaceDiscardedWithReferenceMesh ? CustomizableObject->GetRefSkeletalMesh(CustomizableObjectInstanceUsage->GetComponentIndex()) : nullptr);
 		}
 	}
 }
diff --git a/Engine/Plugins/Experimental/Mutable/Source/CustomizableObject/Private/MuCO/CustomizableObjectInstanceUsage.cpp b/Engine/Plugins/Experimental/Mutable/Source/CustomizableObject/Private/MuCO/CustomizableObjectInstanceUsage.cpp
new file mode 100644
--- /dev/null
+++ b/Engine/Plugins/Experimental/Mutable/Source/CustomizableObject/Private/MuCO/CustomizableObjectInstanceUsage.cpp
@@ -0,0 +1,537 @@
+// Copyright Epic Games, Inc. All Rights Reserved.
+
+#include "MuCO/CustomizableObjectInstanceUsage.h"
+#include "MuCO/CustomizableSkeletalComponent.h"
+#include "MuCO/CustomizableObjectSystem.h"
+#include "MuCO/CustomizableInstancePrivateData.h"
+#include "MuCO/UnrealPortabilityHelpers.h"
+
+#include "AnimationRuntime.h"
+#include "Components/SkeletalMeshComponent.h"
+#include "Engine/SkeletalMesh.h"
+#include "Engine/SkinnedAssetCommon.h"
+#include "Engine/World.h"
+#include "GameFramework/Pawn.h"
+#include "Materials/MaterialInstanceDynamic.h"
+#include "Rendering/SkeletalMeshRenderData.h"
+#include "UObject/ObjectSaveContext.h"
+#include "Stats/Stats.h"
+
+
+void UCustomizableObjectInstanceUsage::Callbacks() const
+{
+	if (CustomizableSkeletalComponent)
+	{
+		CustomizableSkeletalComponent->UpdatedDelegate.ExecuteIfBound();
+
+		if (UpdatedDelegate.IsBound() && CustomizableSkeletalComponent->UpdatedDelegate.IsBound())
+		{
+			UE_LOG(LogMutable, Error, TEXT("The UpdatedDelegate is bound both in the UCustomizableObjectInstanceUsage and in its parent CustomizableSkeletalComponent. Only one should be bound."));
+			ensure(false);
+		}
+	}
+	
+	UpdatedDelegate.ExecuteIfBound();
+}
+
+
+void UCustomizableObjectInstanceUsage::SetCustomizableObjectInstance(UCustomizableObjectInstance* CustomizableObjectInstance)
+{
+	if (CustomizableSkeletalComponent)
+	{
+		CustomizableSkeletalComponent->CustomizableObjectInstance = CustomizableObjectInstance;
+	}
+	else
+	{
+		UsedCustomizableObjectInstance = CustomizableObjectInstance;
+	}
+}
+
+
+UCustomizableObjectInstance* UCustomizableObjectInstanceUsage::GetCustomizableObjectInstance() const
+{
+	if (CustomizableSkeletalComponent)
+	{
+		return CustomizableSkeletalComponent->CustomizableObjectInstance;
+	}
+	else
+	{
+		return UsedCustomizableObjectInstance;
+	}
+}
+
+
+void UCustomizableObjectInstanceUsage::SetComponentIndex(int32 ComponentIndex)
+{
+	if (CustomizableSkeletalComponent)
+	{
+		CustomizableSkeletalComponent->ComponentIndex = ComponentIndex;
+	}
+	else
+	{
+		UsedComponentIndex = ComponentIndex;
+	}
+}
+
+
+int32 UCustomizableObjectInstanceUsage::GetComponentIndex() const
+{
+	if (CustomizableSkeletalComponent)
+	{
+		return CustomizableSkeletalComponent->ComponentIndex;
+	}
+	else
+	{
+		return UsedComponentIndex;
+	}
+}
+
+
+void UCustomizableObjectInstanceUsage::SetPendingSetSkeletalMesh(bool bIsActive)
+{
+	if (CustomizableSkeletalComponent)
+	{
+		CustomizableSkeletalComponent->bPendingSetSkeletalMesh = bIsActive;
+	}
+	else
+	{
+		bUsedPendingSetSkeletalMesh = bIsActive;
+	}
+}
+
+
+bool UCustomizableObjectInstanceUsage::GetPendingSetSkeletalMesh() const
+{
+	if (CustomizableSkeletalComponent)
+	{
+		return CustomizableSkeletalComponent->bPendingSetSkeletalMesh;
+	}
+	else
+	{
+		return bUsedPendingSetSkeletalMesh;
+	}
+}
+
+
+void UCustomizableObjectInstanceUsage::SetSkipSetReferenceSkeletalMesh(bool bIsActive)
+{
+	if (CustomizableSkeletalComponent)
+	{
+		CustomizableSkeletalComponent->bSkipSetReferenceSkeletalMesh = bIsActive;
+	}
+	else
+	{
+		bUsedSkipSetReferenceSkeletalMesh = bIsActive;
+	}
+}
+
+
+bool UCustomizableObjectInstanceUsage::GetSkipSetReferenceSkeletalMesh() const
+{
+	if (CustomizableSkeletalComponent)
+	{
+		return CustomizableSkeletalComponent->bSkipSetReferenceSkeletalMesh;
+	}
+	else
+	{
+		return bUsedSkipSetReferenceSkeletalMesh;
+	}
+}
+
+
+void UCustomizableObjectInstanceUsage::AttachTo(USkeletalMeshComponent* SkeletalMeshComponent)
+{
+	if (CustomizableSkeletalComponent)
+	{
+		UE_LOG(LogMutable, Error, TEXT("Cannot change the attachment of a UCustomizableObjectInstanceUsage that has been automatically created by a CustomizableSkeletalComponent. Reattach the CustomizableSkeletalComponent instead."));
+		ensure(false);
+	}
+	else
+	{
+		UsedSkeletalMeshComponent = SkeletalMeshComponent;
+	}
+}
+
+
+USkeletalMeshComponent* UCustomizableObjectInstanceUsage::GetAttachParent() const
+{
+	if (CustomizableSkeletalComponent)
+	{
+		return Cast<USkeletalMeshComponent>(CustomizableSkeletalComponent->GetAttachParent());
+	}
+	else
+	{
+		return UsedSkeletalMeshComponent;
+	}
+}
+
+
+USkeletalMesh* UCustomizableObjectInstanceUsage::GetSkeletalMesh() const
+{
+	UCustomizableObjectInstance* CustomizableObjectInstance = GetCustomizableObjectInstance();
+
+	return CustomizableObjectInstance ? CustomizableObjectInstance->GetSkeletalMesh(GetComponentIndex()) : nullptr;
+}
+
+
+void UCustomizableObjectInstanceUsage::SetSkeletalMesh(USkeletalMesh* SkeletalMesh, bool bReinitPose, bool bForceClothReset)
+{
+	USkeletalMeshComponent* Parent = Cast<USkeletalMeshComponent>(GetAttachParent());
+
+	// LINUX_PLATFORM needs a more aggressive workaround so morph and cloth glitches
+	// are not visible. 
+	// TODO: Try to find a better way of setting BP Morphs and clothing and remove
+	// the workaround.  
+	if (Parent)
+	{
+#if PLATFORM_LINUX
+		const bool bDisableClothSimulation = Parent->bDisableClothSimulation;
+		if (bForceClothReset)
+		{
+			Parent->bDisableClothSimulation = true;
+		}
+#endif
+
+		if(UE_MUTABLE_GETSKINNEDASSET(Parent) == SkeletalMesh)
+		{
+			Parent->RecreateRenderState_Concurrent();
+		}
+		
+		TMap<FName, float> MorphTargetCurves = Parent->GetMorphTargetCurves();
+		Parent->SetSkeletalMesh(SkeletalMesh, bReinitPose);
+		
+		// USkeletalMeshCompoent MorphTargetCurves are reset when SetSkeletalMesh is called.
+		// Re-enable them if not bReinitPose.
+		if (!bReinitPose && MorphTargetCurves.Num() > 0)
+		{
+			for (const TPair<FName, float>& MorphTarget : MorphTargetCurves)
+			{
+				Parent->SetMorphTarget(MorphTarget.Key, MorphTarget.Value);
+			}
+
+#if PLATFORM_LINUX
+			FRenderStateRecreator RenderStateRecreator(Parent);
+			FAnimationRuntime::AppendActiveMorphTargets(SkeletalMesh, Parent->GetMorphTargetCurves(), Parent->ActiveMorphTargets, Parent->MorphTargetWeights);		
+#endif
+		}
+
+		if (bForceClothReset)
+		{
+			Parent->ForceClothNextUpdateTeleportAndReset();
+
+#if PLATFORM_LINUX
+			Parent->bDisableClothSimulation = bDisableClothSimulation;
+#endif
+		}
+
+		if (Parent->GetNumOverrideMaterials() > 0)
+		{
+			// For some reason the reference skeletal mesh materials are added as override materials, clear them if necessary
+			Parent->EmptyOverrideMaterials();
+		}	
+	}
+}
+
+void UCustomizableObjectInstanceUsage::SetPhysicsAsset(UPhysicsAsset* PhysicsAsset)
+{
+	USkeletalMeshComponent* Parent = Cast<USkeletalMeshComponent>(GetAttachParent());
+
+	if (Parent && Parent->GetWorld())
+	{
+		Parent->SetPhysicsAsset(PhysicsAsset, true);
+	}
+}
+
+
+USkeletalMesh* UCustomizableObjectInstanceUsage::GetAttachedSkeletalMesh() const
+{
+	USkeletalMeshComponent* Parent = Cast<USkeletalMeshComponent>(GetAttachParent());
+
+	if (Parent)
+	{
+		return UE_MUTABLE_GETSKELETALMESHASSET(Parent);
+	}
+
+	return nullptr;
+}
+
+
+void UCustomizableObjectInstanceUsage::UpdateSkeletalMeshAsync(bool bNeverSkipUpdate)
+{
+	UCustomizableObjectInstance* CustomizableObjectInstance = GetCustomizableObjectInstance();
+
+	if (CustomizableObjectInstance)
+	{
+		CustomizableObjectInstance->UpdateSkeletalMeshAsync(false, false);
+	}
+}
+
+
+void UCustomizableObjectInstanceUsage::UpdateSkeletalMeshAsyncResult(FInstanceUpdateDelegate Callback, bool bIgnoreCloseDist, bool bForceHighPriority)
+{
+	UCustomizableObjectInstance* CustomizableObjectInstance = GetCustomizableObjectInstance();
+
+	if (CustomizableObjectInstance)
+	{
+		CustomizableObjectInstance->UpdateSkeletalMeshAsyncResult(Callback, false, false);
+	}
+}
+
+
+void UCustomizableObjectInstanceUsage::BeginDestroy()
+{
+	Super::BeginDestroy();
+}
+
+
+#if WITH_EDITOR
+
+void UCustomizableObjectInstanceUsage::UpdateDistFromComponentToLevelEditorCamera(const FVector& CameraPosition)
+{
+	// We want instances in the editor to be generated
+	if (!GetWorld() || GetWorld()->WorldType != EWorldType::Editor)
+	{
+		return;
+	}
+
+	UCustomizableObjectInstance* CustomizableObjectInstance = GetCustomizableObjectInstance();
+
+	if (CustomizableObjectInstance)
+	{
+		USkeletalMeshComponent* SkeletalMeshComponent = GetAttachParent();
+		AActor* ParentActor = SkeletalMeshComponent ? SkeletalMeshComponent->GetAttachmentRootActor() : nullptr;
+		if (ParentActor && ParentActor->IsValidLowLevel())
+		{
+			// update distance to camera and set the instance as being used by a component
+			CustomizableObjectInstance->GetPrivate()->SetCOInstanceFlags(UsedByComponent);
+
+			float SquareDist = FVector::DistSquared(CameraPosition, ParentActor->GetActorLocation());
+			CustomizableObjectInstance->GetPrivate()->MinSquareDistFromComponentToPlayer = 
+				FMath::Min(SquareDist, CustomizableObjectInstance->GetPrivate()->MinSquareDistFromComponentToPlayer);
+		}
+
+		USkeletalMesh* AttachedSkeletalMesh = GetAttachedSkeletalMesh();
+		int32 ComponentIndex = GetComponentIndex();
+		USkeletalMesh* GeneratedSKeletalMesh = CustomizableObjectInstance->GetSkeletalMesh(ComponentIndex);
+		if (!AttachedSkeletalMesh && !CustomizableObjectInstance->GetPrivate()->HasCOInstanceFlags(CreatingSkeletalMesh))
+		{
+			UCustomizableObject* Object = CustomizableObjectInstance->GetCustomizableObject(); 
+			USkeletalMesh* RefMesh = Object ? Object->GetRefSkeletalMesh(ComponentIndex) : nullptr;
+			SetSkeletalMesh(GeneratedSKeletalMesh ? GeneratedSKeletalMesh : RefMesh);
+		}
+		else if (GeneratedSKeletalMesh && AttachedSkeletalMesh != GeneratedSKeletalMesh)
+		{
+			SetSkeletalMesh(GeneratedSKeletalMesh);
+		}
+	}
+}
+
+
+void UCustomizableObjectInstanceUsage::EditorUpdateComponent()
+{
+	UCustomizableObjectInstance* CustomizableObjectInstance = GetCustomizableObjectInstance();
+
+	if (CustomizableObjectInstance)
+	{
+		CustomizableObjectInstance->GetPrivate()->SetCOInstanceFlags(UsedByComponent);
+
+		USkeletalMeshComponent* SkeletalMeshComponent = GetAttachParent();
+		AActor* ParentActor = SkeletalMeshComponent ? SkeletalMeshComponent->GetAttachmentRootActor() : nullptr;
+		
+		if (ParentActor)
+		{
+			USkeletalMesh* AttachedSkeletalMesh = GetAttachedSkeletalMesh();
+			int32 ComponentIndex = GetComponentIndex();
+			USkeletalMesh* GeneratedSKeletalMesh = CustomizableObjectInstance->GetSkeletalMesh(ComponentIndex);
+			if (!AttachedSkeletalMesh && !CustomizableObjectInstance->GetPrivate()->HasCOInstanceFlags(CreatingSkeletalMesh))
+			{
+				SetSkeletalMesh(GeneratedSKeletalMesh ? GeneratedSKeletalMesh : CustomizableObjectInstance->GetCustomizableObject()->GetRefSkeletalMesh(ComponentIndex));
+			}
+			else if (GeneratedSKeletalMesh && AttachedSkeletalMesh != GeneratedSKeletalMesh)
+			{
+				SetSkeletalMesh(GeneratedSKeletalMesh);
+			}
+		}
+	}
+}
+#endif
+
+
+void UCustomizableObjectInstanceUsage::SetVisibilityOfSkeletalMeshSectionWithMaterialName(bool bInVisible, const FString& MaterialName, int32 LOD)
+{
+	UCustomizableObjectInstance* CustomizableObjectInstance = GetCustomizableObjectInstance();
+	int32 ComponentIndex = GetComponentIndex();
+
+	USkeletalMesh* SkeletalMesh = CustomizableObjectInstance->GetSkeletalMesh(ComponentIndex);
+
+	if (!SkeletalMesh)
+	{
+		return;
+	}
+
+	for (int32 i = 0; i < SkeletalMesh->GetMaterials().Num(); ++i)
+	{
+		UMaterialInstanceDynamic* MatInst = Cast<UMaterialInstanceDynamic>(SkeletalMesh->GetMaterials()[i].MaterialInterface);
+
+		if (MatInst && MatInst->Parent->GetName() == MaterialName)
+		{
+			if (FSkeletalMeshRenderData* SkelResource = SkeletalMesh->GetResourceForRendering() )
+			{
+				if ( SkelResource->LODRenderData.IsValidIndex(LOD))
+				{
+					FSkeletalMeshLODRenderData& LodData = SkelResource->LODRenderData[LOD];
+
+					for (int32 j = 0; j < LodData.RenderSections.Num(); ++j)
+					{
+						FSkelMeshRenderSection& Section = LodData.RenderSections[j];
+
+						if (Section.MaterialIndex == i)
+						{
+							Section.bDisabled = !bInVisible;
+						}
+					}
+				}
+			}
+		}
+	}
+}
+
+
+void UCustomizableObjectInstanceUsage::UpdateDistFromComponentToPlayer(const AActor* ViewCenter, const bool bForceEvenIfNotBegunPlay)
+{
+	UCustomizableObjectInstance* CustomizableObjectInstance = GetCustomizableObjectInstance();
+
+	if (CustomizableObjectInstance)
+	{
+		USkeletalMeshComponent* SkeletalMeshComponent = GetAttachParent();
+		AActor* ParentActor = SkeletalMeshComponent ? SkeletalMeshComponent->GetAttachmentRootActor() : nullptr;
+		
+		CustomizableObjectInstance->SetIsPlayerOrNearIt(false);
+
+		if (ParentActor && ParentActor->IsValidLowLevel())
+		{
+			if (ParentActor->HasActorBegunPlay() || bForceEvenIfNotBegunPlay)
+			{
+				float SquareDist = FLT_MAX;
+
+				if (ViewCenter && ViewCenter->IsValidLowLevel())
+				{
+					APawn* Pawn = Cast<APawn>(ParentActor);
+					bool bIsPlayer = Pawn ? Pawn->IsPlayerControlled() : false;
+					CustomizableObjectInstance->SetIsPlayerOrNearIt(bIsPlayer);
+
+					if (bIsPlayer)
+					{
+						SquareDist = -0.01f; // Negative value to give the player character more priority than any other character
+					}
+					else
+					{
+						SquareDist = FVector::DistSquared(ViewCenter->GetActorLocation(), ParentActor->GetActorLocation());
+					}
+				}
+				else if (bForceEvenIfNotBegunPlay)
+				{
+					SquareDist = -0.01f; // This is a manual update before begin play and the creation of the pawn, so it should probably be high priority
+					CustomizableObjectInstance->GetPrivate()->LastMinSquareDistFromComponentToPlayer = FMath::Min(SquareDist, CustomizableObjectInstance->GetPrivate()->MinSquareDistFromComponentToPlayer);
+				}
+				else
+				{
+					SquareDist = 0.f; // This a mutable tick before begin play and the creation of the pawn, so it should have a definite and high priority but less than a manual update
+					CustomizableObjectInstance->GetPrivate()->LastMinSquareDistFromComponentToPlayer = FMath::Min(SquareDist, CustomizableObjectInstance->GetPrivate()->MinSquareDistFromComponentToPlayer);
+				}
+
+				CustomizableObjectInstance->GetPrivate()->MinSquareDistFromComponentToPlayer = FMath::Min(SquareDist, CustomizableObjectInstance->GetPrivate()->MinSquareDistFromComponentToPlayer);
+				CustomizableObjectInstance->SetIsBeingUsedByComponentInPlay(true);
+
+				if (CustomizableObjectInstance->GetPrivate()->MinSquareDistFromComponentToPlayer == SquareDist)
+				{
+					CustomizableObjectInstance->NearestToActor = this;
+					CustomizableObjectInstance->NearestToViewCenter = ViewCenter;
+				}
+			}
+		}
+
+		int32 ComponentIndex = GetComponentIndex();
+
+		if (ParentActor && GetAttachedSkeletalMesh() == nullptr && CustomizableObjectInstance->GetSkeletalMesh(ComponentIndex) && !CustomizableObjectInstance->GetPrivate()->HasCOInstanceFlags(CreatingSkeletalMesh))
+		{
+			SetSkeletalMesh(CustomizableObjectInstance->GetSkeletalMesh(ComponentIndex));
+		}
+	}
+}
+
+
+void UCustomizableObjectInstanceUsage::Tick(float DeltaTime)
+{
+	UCustomizableObjectInstance* CustomizableObjectInstance = GetCustomizableObjectInstance();
+
+	if (!GetPendingSetSkeletalMesh() || !CustomizableObjectInstance || !CustomizableObjectInstance->GetCustomizableObject())
+	{
+		return;
+	}
+
+	if (USkeletalMeshComponent* Parent = Cast<USkeletalMeshComponent>(GetAttachParent()))
+	{
+		UCustomizableObject* CustomizableObject = CustomizableObjectInstance->GetCustomizableObject();
+
+		// Hacky. Replace once we know if the instance has been generated
+		const bool bInstanceGenerated = CustomizableObjectInstance->HasAnySkeletalMesh();
+
+		int32 ComponentIndex = GetComponentIndex();
+
+		// Generated SkeletalMesh to set, can be null if the component is empty
+		USkeletalMesh* SkeletalMesh = CustomizableObjectInstance->GetSkeletalMesh(ComponentIndex);
+
+		// If not generated yet, conditionally set the SkeletalMesh of reference
+		if (!bInstanceGenerated && !GetSkipSetReferenceSkeletalMesh()
+#if WITH_EDITORONLY_DATA
+			&& CustomizableObject->bEnableUseRefSkeletalMeshAsPlaceholder
+#endif
+			)
+		{
+			// Can be nullptr
+			SkeletalMesh = CustomizableObject->GetRefSkeletalMesh(ComponentIndex);
+		}
+
+		// Set SkeletalMesh
+		if (bInstanceGenerated || SkeletalMesh)
+		{
+			Parent->SetSkeletalMesh(SkeletalMesh);
+
+			if (Parent->OverrideMaterials.Num() > 0)
+			{
+#if !(UE_BUILD_SHIPPING || UE_BUILD_TEST)
+				if (Parent->GetClass()->GetFName() != FName(TEXT("SkeletalMeshComponentBudgeted"))) // Reduce unnecessary logging
+				{
+					UE_LOG(LogMutable, Warning, TEXT("Attaching Customizable Skeletal Component to Skeletal Mesh Component with overriden materials! Deleting overrides."));
+				}
+#endif
+
+				Parent->EmptyOverrideMaterials();
+			}
+
+			SetPendingSetSkeletalMesh(false);
+		}
+	}
+}
+
+
+TStatId UCustomizableObjectInstanceUsage::GetStatId() const
+{
+	RETURN_QUICK_DECLARE_CYCLE_STAT(UCustomizableObjectInstanceUsage, STATGROUP_Tickables);
+}
+
+
+bool UCustomizableObjectInstanceUsage::IsNetMode(ENetMode InNetMode) const
+{
+	if (CustomizableSkeletalComponent)
+	{
+		return CustomizableSkeletalComponent->IsNetMode(InNetMode);
+	}
+	else if(UsedSkeletalMeshComponent)
+	{
+		return UsedSkeletalMeshComponent->IsNetMode(InNetMode);
+	}
+
+	return false;
+}
diff --git a/Engine/Plugins/Experimental/Mutable/Source/CustomizableObject/Private/MuCO/CustomizableObjectModule.cpp b/Engine/Plugins/Experimental/Mutable/Source/CustomizableObject/Private/MuCO/CustomizableObjectModule.cpp
--- a/Engine/Plugins/Experimental/Mutable/Source/CustomizableObject/Private/MuCO/CustomizableObjectModule.cpp
+++ b/Engine/Plugins/Experimental/Mutable/Source/CustomizableObject/Private/MuCO/CustomizableObjectModule.cpp
@@ -9,10 +9,11 @@
 #include "MuCO/CustomizableObject.h"
 #include "MuCO/CustomizableObjectDGGUI.h"
 #include "MuCO/CustomizableObjectExtension.h"
-#include "MuCO/CustomizableSkeletalComponent.h"
+#include "MuCO/CustomizableObjectInstanceUsage.h"
 #include "MuCO/ICustomizableObjectModule.h"
 #include "UObject/StrongObjectPtr.h"
 #include "GPUSkinPublicDefs.h"
+#include "Components/SkeletalMeshComponent.h"
 
 /**
  * Customizable Object module implementation (private)
@@ -216,29 +217,30 @@ void FCustomizableObjectModule::RefreshExtensionData()
 }
 
 
-UCustomizableSkeletalComponent* GetPlayerCustomizableSkeletalComponent(const int32 SlotID, const UWorld* CurrentWorld, const int32 PlayerIndex)
+UCustomizableObjectInstanceUsage* GetPlayerCustomizableObjectInstanceUsage(const int32 SlotID, const UWorld* CurrentWorld, const int32 PlayerIndex)
 {
 	// Get customizable skeletal component attached to player pawn
-	UCustomizableSkeletalComponent* SelectedCustomizableSkeletalComponent = nullptr;
+	UCustomizableObjectInstanceUsage* SelectedCustomizableObjectInstanceUsage = nullptr;
 	{
 		AActor* PlayerPawn = Cast<AActor>(UGameplayStatics::GetPlayerPawn(CurrentWorld, PlayerIndex));
 		int32 IndexFound = INDEX_NONE;
-		for (TObjectIterator<UCustomizableSkeletalComponent> CustomizableSkeletalComponent; CustomizableSkeletalComponent; ++CustomizableSkeletalComponent)
+		for (TObjectIterator<UCustomizableObjectInstanceUsage> CustomizableObjectInstanceUsage; CustomizableObjectInstanceUsage; ++CustomizableObjectInstanceUsage)
 		{
 #if WITH_EDITOR
-			if (CustomizableSkeletalComponent && CustomizableSkeletalComponent->IsNetMode(NM_DedicatedServer))
+			if (CustomizableObjectInstanceUsage && CustomizableObjectInstanceUsage->IsNetMode(NM_DedicatedServer))
 			{
 				continue;
 			}
 #endif
 
-			if (CustomizableSkeletalComponent->IsValidLowLevel() && !CustomizableSkeletalComponent->IsTemplate())
+			if (CustomizableObjectInstanceUsage->IsValidLowLevel() && !CustomizableObjectInstanceUsage->IsTemplate() 
+				&& CustomizableObjectInstanceUsage->GetAttachParent())
 			{
-				AActor* CustomizableActor = CustomizableSkeletalComponent->GetAttachmentRootActor();
+				AActor* CustomizableActor = CustomizableObjectInstanceUsage->GetAttachParent()->GetAttachmentRootActor();
 				if (CustomizableActor && PlayerPawn == CustomizableActor)
 				{
 					++IndexFound;
-					SelectedCustomizableSkeletalComponent = *CustomizableSkeletalComponent;
+					SelectedCustomizableObjectInstanceUsage = *CustomizableObjectInstanceUsage;
 					if (IndexFound == SlotID)
 					{
 						break;
@@ -250,23 +252,23 @@ UCustomizableSkeletalComponent* GetPlayerCustomizableSkeletalComponent(const int
 
 
 	// If none found, try getting a component without caring about the actor
-	if (!SelectedCustomizableSkeletalComponent)
+	if (!SelectedCustomizableObjectInstanceUsage)
 	{
 		AActor* PlayerPawn = Cast<AActor>(UGameplayStatics::GetPlayerPawn(CurrentWorld, PlayerIndex));
 		int32 IndexFound = INDEX_NONE;
-		for (TObjectIterator<UCustomizableSkeletalComponent> CustomizableSkeletalComponent; CustomizableSkeletalComponent; ++CustomizableSkeletalComponent)
+		for (TObjectIterator<UCustomizableObjectInstanceUsage> CustomizableObjectInstanceUsage; CustomizableObjectInstanceUsage; ++CustomizableObjectInstanceUsage)
 		{
 #if WITH_EDITOR
-			if (CustomizableSkeletalComponent && CustomizableSkeletalComponent->IsNetMode(NM_DedicatedServer))
+			if (CustomizableObjectInstanceUsage && CustomizableObjectInstanceUsage->IsNetMode(NM_DedicatedServer))
 			{
 				continue;
 			}
 #endif
 
-			if (CustomizableSkeletalComponent->IsValidLowLevel() && !CustomizableSkeletalComponent->IsTemplate())
+			if (CustomizableObjectInstanceUsage->IsValidLowLevel() && !CustomizableObjectInstanceUsage->IsTemplate())
 			{
 				++IndexFound;
-				SelectedCustomizableSkeletalComponent = *CustomizableSkeletalComponent;
+				SelectedCustomizableObjectInstanceUsage = *CustomizableObjectInstanceUsage;
 				if (IndexFound == SlotID)
 				{
 					break;
@@ -275,7 +277,7 @@ UCustomizableSkeletalComponent* GetPlayerCustomizableSkeletalComponent(const int
 		}
 	}
 
-	return SelectedCustomizableSkeletalComponent;
+	return SelectedCustomizableObjectInstanceUsage;
 }
 
 
@@ -314,8 +316,8 @@ void FCustomizableObjectModule::ToggleDGGUI(const TArray<FString>& Arguments)
 	{
 		return;
 	}
-	else if (UCustomizableSkeletalComponent* SelectedCustomizableSkeletalComponent = GetPlayerCustomizableSkeletalComponent(SlotID, CurrentWorld, PlayerIndex))
+	else if (UCustomizableObjectInstanceUsage* SelectedCustomizableObjectInstanceUsage = GetPlayerCustomizableObjectInstanceUsage(SlotID, CurrentWorld, PlayerIndex))
 	{
-		UDGGUI::OpenDGGUI(SlotID, SelectedCustomizableSkeletalComponent, CurrentWorld, PlayerIndex);
+		UDGGUI::OpenDGGUI(SlotID, SelectedCustomizableObjectInstanceUsage, CurrentWorld, PlayerIndex);
 	}
 }
\ No newline at end of file
diff --git a/Engine/Plugins/Experimental/Mutable/Source/CustomizableObject/Private/MuCO/CustomizableObjectSystem.cpp b/Engine/Plugins/Experimental/Mutable/Source/CustomizableObject/Private/MuCO/CustomizableObjectSystem.cpp
--- a/Engine/Plugins/Experimental/Mutable/Source/CustomizableObject/Private/MuCO/CustomizableObjectSystem.cpp
+++ b/Engine/Plugins/Experimental/Mutable/Source/CustomizableObject/Private/MuCO/CustomizableObjectSystem.cpp
@@ -15,7 +15,7 @@
 #include "MuCO/CustomizableInstancePrivateData.h"
 #include "MuCO/CustomizableObjectPrivate.h"
 #include "MuCO/DefaultImageProvider.h"
-#include "MuCO/CustomizableSkeletalComponent.h"
+#include "MuCO/CustomizableObjectInstanceUsage.h"
 #include "MuCO/ICustomizableObjectModule.h"
 #include "MuCO/LogBenchmarkUtil.h"
 #include "MuCO/LogInformationUtil.h"
@@ -29,6 +29,7 @@
 #include "Widgets/Notifications/SNotificationList.h"
 #include "ContentStreaming.h"
 #include "MuCO/EditorImageProvider.h"
+#include "Components/SkeletalMeshComponent.h"
 
 #if WITH_EDITOR
 #include "Editor.h"
@@ -325,24 +326,25 @@ void UCustomizableObjectSystem::LogShowData(bool bFullInfo, bool ShowMaterialInf
 
 	TArray<UCustomizableObjectInstance*> ArrayData;
 
-	for (TObjectIterator<UCustomizableSkeletalComponent> It; It; ++It)
+	for (TObjectIterator<UCustomizableObjectInstanceUsage> It; It; ++It)
 	{
-		const UCustomizableSkeletalComponent* CustomizableSkeletalComponent = *It;
+		const UCustomizableObjectInstanceUsage* CustomizableObjectInstanceUsage = *It;
 
 #if WITH_EDITOR
-		if (CustomizableSkeletalComponent && CustomizableSkeletalComponent->IsNetMode(NM_DedicatedServer))
+		if (CustomizableObjectInstanceUsage && CustomizableObjectInstanceUsage->IsNetMode(NM_DedicatedServer))
 		{
 			continue;
 		}
 #endif
 
-		if (CustomizableSkeletalComponent && CustomizableSkeletalComponent->CustomizableObjectInstance)
+		if (CustomizableObjectInstanceUsage && CustomizableObjectInstanceUsage->GetCustomizableObjectInstance() 
+			&& CustomizableObjectInstanceUsage->GetAttachParent())
 		{
-			const AActor* ParentActor = CustomizableSkeletalComponent->GetAttachmentRootActor();
+			const AActor* ParentActor = CustomizableObjectInstanceUsage->GetAttachParent()->GetAttachmentRootActor();
 
 			if (ParentActor != nullptr)
 			{
-				ArrayData.AddUnique(CustomizableSkeletalComponent->CustomizableObjectInstance);
+				ArrayData.AddUnique(CustomizableObjectInstanceUsage->GetCustomizableObjectInstance());
 			}
 		}
 	}
@@ -812,8 +814,8 @@ void FinishUpdateGlobal(const TSharedRef<FUpdateContextPrivate>& Context)
 
 	if (Context->UpdateResult == EUpdateResult::Success)
 	{
-		// Call Customizable Skeletal Components updated callbacks.
-		for (TObjectIterator<UCustomizableSkeletalComponent> It; It; ++It) // Since iterating objects is expensive, for now CustomizableSkeletalComponent does not have a FinishUpdate function.
+		// Call CustomizableObjectInstanceUsages updated callbacks.
+		for (TObjectIterator<UCustomizableObjectInstanceUsage> It; It; ++It) // Since iterating objects is expensive, for now CustomizableObjectInstanceUsage does not have a FinishUpdate function.
 		{
 #if WITH_EDITOR
 			if (It && It->IsNetMode(NM_DedicatedServer))
@@ -822,11 +824,11 @@ void FinishUpdateGlobal(const TSharedRef<FUpdateContextPrivate>& Context)
 			}
 #endif
 
-			if (const UCustomizableSkeletalComponent* CustomizableSkeletalComponent = *It;
-				CustomizableSkeletalComponent &&
-				CustomizableSkeletalComponent->CustomizableObjectInstance == Instance)
+			if (const UCustomizableObjectInstanceUsage* CustomizableObjectInstanceUsage = *It;
+				CustomizableObjectInstanceUsage &&
+				CustomizableObjectInstanceUsage->GetCustomizableObjectInstance() == Instance)
 			{
-				CustomizableSkeletalComponent->Callbacks();
+				CustomizableObjectInstanceUsage->Callbacks();
 			}
 		}
 	}
@@ -867,32 +869,32 @@ void UpdateSkeletalMesh(const TSharedRef<FUpdateContextPrivate>& Context)
 
 	UCustomizableInstancePrivateData* CustomizableObjectInstancePrivateData = CustomizableObjectInstance->GetPrivate();
 	check(CustomizableObjectInstancePrivateData != nullptr);
-	for (TObjectIterator<UCustomizableSkeletalComponent> It; It; ++It)
+	for (TObjectIterator<UCustomizableObjectInstanceUsage> It; It; ++It)
 	{
-		UCustomizableSkeletalComponent* CustomizableSkeletalComponent = *It;
+		UCustomizableObjectInstanceUsage* CustomizableObjectInstanceUsage = *It;
 
 #if WITH_EDITOR
-		if (CustomizableSkeletalComponent && CustomizableSkeletalComponent->IsNetMode(NM_DedicatedServer))
+		if (CustomizableObjectInstanceUsage && CustomizableObjectInstanceUsage->IsNetMode(NM_DedicatedServer))
 		{
 			continue;
 		}
 #endif
 
-		if (CustomizableSkeletalComponent &&
-			(CustomizableSkeletalComponent->CustomizableObjectInstance == CustomizableObjectInstance) &&
-			CustomizableObjectInstance->SkeletalMeshes.IsValidIndex(CustomizableSkeletalComponent->ComponentIndex)
+		if (CustomizableObjectInstanceUsage &&
+			(CustomizableObjectInstanceUsage->GetCustomizableObjectInstance() == CustomizableObjectInstance) &&
+			CustomizableObjectInstance->SkeletalMeshes.IsValidIndex(CustomizableObjectInstanceUsage->GetComponentIndex())
 		   )
 		{
 			MUTABLE_CPUPROFILER_SCOPE(UpdateSkeletalMesh_SetSkeletalMesh);
 
 			const bool bIsCreatingSkeletalMesh = CustomizableObjectInstancePrivateData->HasCOInstanceFlags(CreatingSkeletalMesh); //TODO MTBL-391: Review
-			CustomizableSkeletalComponent->SetSkeletalMesh(CustomizableObjectInstance->SkeletalMeshes[CustomizableSkeletalComponent->ComponentIndex], false, bIsCreatingSkeletalMesh);
+			CustomizableObjectInstanceUsage->SetSkeletalMesh(CustomizableObjectInstance->SkeletalMeshes[CustomizableObjectInstanceUsage->GetComponentIndex()], false, bIsCreatingSkeletalMesh);
 
 			if (CustomizableObjectInstancePrivateData->HasCOInstanceFlags(ReplacePhysicsAssets))
 			{
-				CustomizableSkeletalComponent->SetPhysicsAsset(
-					CustomizableObjectInstance->SkeletalMeshes[CustomizableSkeletalComponent->ComponentIndex] ? 
-					CustomizableObjectInstance->SkeletalMeshes[CustomizableSkeletalComponent->ComponentIndex]->GetPhysicsAsset() : nullptr);
+				CustomizableObjectInstanceUsage->SetPhysicsAsset(
+					CustomizableObjectInstance->SkeletalMeshes[CustomizableObjectInstanceUsage->GetComponentIndex()] ?
+					CustomizableObjectInstance->SkeletalMeshes[CustomizableObjectInstanceUsage->GetComponentIndex()]->GetPhysicsAsset() : nullptr);
 			}
 		}
 	}
@@ -2799,21 +2801,21 @@ namespace impl
 		bool bIsInEditorViewport = false;
 
 #if WITH_EDITOR
-		for (TObjectIterator<UCustomizableSkeletalComponent> CustomizableSkeletalComponent; CustomizableSkeletalComponent && !bIsInEditorViewport; ++CustomizableSkeletalComponent)
+		for (TObjectIterator<UCustomizableObjectInstanceUsage> CustomizableObjectInstanceUsage; CustomizableObjectInstanceUsage && !bIsInEditorViewport; ++CustomizableObjectInstanceUsage)
 		{
-			if (CustomizableSkeletalComponent && CustomizableSkeletalComponent->IsNetMode(NM_DedicatedServer))
+			if (CustomizableObjectInstanceUsage && CustomizableObjectInstanceUsage->IsNetMode(NM_DedicatedServer))
 			{
 				continue;
 			}
 
-			if (CustomizableSkeletalComponent &&
-				CustomizableSkeletalComponent->CustomizableObjectInstance == CandidateInstance)
+			if (CustomizableObjectInstanceUsage &&
+				CustomizableObjectInstanceUsage->GetCustomizableObjectInstance() == CandidateInstance)
 			{
 				EWorldType::Type WorldType = EWorldType::Type::None;
 
-				if (CustomizableSkeletalComponent->GetWorld())
+				if (CustomizableObjectInstanceUsage->GetWorld())
 				{
-					WorldType = CustomizableSkeletalComponent->GetWorld()->WorldType;
+					WorldType = CustomizableObjectInstanceUsage->GetWorld()->WorldType;
 				}
 
 				switch (WorldType)
diff --git a/Engine/Plugins/Experimental/Mutable/Source/CustomizableObject/Private/MuCO/CustomizableSkeletalComponent.cpp b/Engine/Plugins/Experimental/Mutable/Source/CustomizableObject/Private/MuCO/CustomizableSkeletalComponent.cpp
--- a/Engine/Plugins/Experimental/Mutable/Source/CustomizableObject/Private/MuCO/CustomizableSkeletalComponent.cpp
+++ b/Engine/Plugins/Experimental/Mutable/Source/CustomizableObject/Private/MuCO/CustomizableSkeletalComponent.cpp
@@ -1,35 +1,29 @@
 // Copyright Epic Games, Inc. All Rights Reserved.
 
 #include "MuCO/CustomizableSkeletalComponent.h"
+#include "MuCO/CustomizableObjectInstanceUsage.h"
 
-#include "MuCO/CustomizableObjectSystem.h"
-#include "MuCO/CustomizableInstancePrivateData.h"
-
-#include "AnimationRuntime.h"
+#include "UObject/UObjectGlobals.h"
 #include "Components/SkeletalMeshComponent.h"
-#include "Engine/SkeletalMesh.h"
-#include "Engine/SkinnedAssetCommon.h"
-#include "Engine/World.h"
-#include "GameFramework/Pawn.h"
-#include "Materials/MaterialInstanceDynamic.h"
-#include "Rendering/SkeletalMeshRenderData.h"
-#include "UObject/ObjectSaveContext.h"
-#include "MuCO/UnrealPortabilityHelpers.h"
 #include "MuCO/CustomizableObject.h"
 
-#include UE_INLINE_GENERATED_CPP_BY_NAME(CustomizableSkeletalComponent)
-
 
-UCustomizableSkeletalComponent::UCustomizableSkeletalComponent()
+void UCustomizableSkeletalComponent::CreateCustomizableObjectInstanceUsage()
 {
-	PrimaryComponentTick.bCanEverTick = true;
-	//bTickInEditor = true;
+	if (!CustomizableObjectInstanceUsage && !HasAnyFlags(RF_ClassDefaultObject))
+	{
+		CustomizableObjectInstanceUsage = NewObject<UCustomizableObjectInstanceUsage>(this, NAME_None, RF_Transient);
+		CustomizableObjectInstanceUsage->CustomizableSkeletalComponent = this;
+	}
 }
 
 
 void UCustomizableSkeletalComponent::Callbacks() const
 {
-	UpdatedDelegate.ExecuteIfBound();
+	if (CustomizableObjectInstanceUsage)
+	{
+		CustomizableObjectInstanceUsage->Callbacks();
+	}
 }
 
 
@@ -38,83 +32,30 @@ USkeletalMesh* UCustomizableSkeletalComponent::GetSkeletalMesh() const
 	return CustomizableObjectInstance ? CustomizableObjectInstance->GetSkeletalMesh(ComponentIndex) : nullptr;
 }
 
+
 void UCustomizableSkeletalComponent::SetSkeletalMesh(USkeletalMesh* SkeletalMesh, bool bReinitPose, bool bForceClothReset)
 {
-	USkeletalMeshComponent* Parent = Cast<USkeletalMeshComponent>(GetAttachParent());
-
-	// LINUX_PLATFORM needs a more aggressive workaround so morph and cloth glitches
-	// are not visible. 
-	// TODO: Try to find a better way of setting BP Morphs and clothing and remove
-	// the workaround.  
-	if (Parent)
+	if (CustomizableObjectInstanceUsage)
 	{
-#if PLATFORM_LINUX
-		const bool bDisableClothSimulation = Parent->bDisableClothSimulation;
-		if (bForceClothReset)
-		{
-			Parent->bDisableClothSimulation = true;
-		}
-#endif
-
-		if(UE_MUTABLE_GETSKINNEDASSET(Parent) == SkeletalMesh)
-		{
-			Parent->RecreateRenderState_Concurrent();
-		}
-		
-		TMap<FName, float> MorphTargetCurves = Parent->GetMorphTargetCurves();
-		PreUpdateDelegate.ExecuteIfBound(this, SkeletalMesh);
-		Parent->SetSkeletalMesh(SkeletalMesh, bReinitPose);
-		
-		// USkeletalMeshCompoent MorphTargetCurves are reset when SetSkeletalMesh is called.
-		// Re-enable them if not bReinitPose.
-		if (!bReinitPose && MorphTargetCurves.Num() > 0)
-		{
-			for (const TPair<FName, float>& MorphTarget : MorphTargetCurves)
-			{
-				Parent->SetMorphTarget(MorphTarget.Key, MorphTarget.Value);
-			}
-
-#if PLATFORM_LINUX
-			FRenderStateRecreator RenderStateRecreator(Parent);
-			FAnimationRuntime::AppendActiveMorphTargets(SkeletalMesh, Parent->GetMorphTargetCurves(), Parent->ActiveMorphTargets, Parent->MorphTargetWeights);		
-#endif
-		}
-
-		if (bForceClothReset)
-		{
-			Parent->ForceClothNextUpdateTeleportAndReset();
-
-#if PLATFORM_LINUX
-			Parent->bDisableClothSimulation = bDisableClothSimulation;
-#endif
-		}
-
-		if (Parent->GetNumOverrideMaterials() > 0)
-		{
-			// For some reason the reference skeletal mesh materials are added as override materials, clear them if necessary
-			Parent->EmptyOverrideMaterials();
-		}	
+		CustomizableObjectInstanceUsage->SetSkeletalMesh(SkeletalMesh, bReinitPose, bForceClothReset);
 	}
 }
 
+
 void UCustomizableSkeletalComponent::SetPhysicsAsset(UPhysicsAsset* PhysicsAsset)
 {
-	USkeletalMeshComponent* Parent = Cast<USkeletalMeshComponent>(GetAttachParent());
-
-	if (Parent && Parent->GetWorld())
+	if (CustomizableObjectInstanceUsage)
 	{
-		Parent->SetPhysicsAsset(PhysicsAsset, true);
+		CustomizableObjectInstanceUsage->SetPhysicsAsset(PhysicsAsset);
 	}
 }
 
 
 USkeletalMesh* UCustomizableSkeletalComponent::GetAttachedSkeletalMesh() const
 {
-	USkeletalMeshComponent* Parent = Cast<USkeletalMeshComponent>(GetAttachParent());
-
-	if (Parent)
+	if (CustomizableObjectInstanceUsage)
 	{
-		return UE_MUTABLE_GETSKELETALMESHASSET(Parent);
+		return CustomizableObjectInstanceUsage->GetAttachedSkeletalMesh();
 	}
 
 	return nullptr;
@@ -139,221 +80,17 @@ void UCustomizableSkeletalComponent::UpdateSkeletalMeshAsyncResult(FInstanceUpda
 }
 
 
-void UCustomizableSkeletalComponent::BeginDestroy()
-{
-	Super::BeginDestroy();
-}
-
-
 #if WITH_EDITOR
 
-void UCustomizableSkeletalComponent::UpdateDistFromComponentToLevelEditorCamera(const FVector& CameraPosition)
-{
-	// We want instances in the editor to be generated
-	if (!GetWorld() || GetWorld()->WorldType != EWorldType::Editor)
-		return;
-
-	if (CustomizableObjectInstance)
-	{
-		AActor* ParentActor = GetAttachmentRootActor();
-		if (ParentActor && ParentActor->IsValidLowLevel())
-		{
-			// update distance to camera and set the instance as being used by a component
-			CustomizableObjectInstance->GetPrivate()->SetCOInstanceFlags(UsedByComponent);
-
-			float SquareDist = FVector::DistSquared(CameraPosition, ParentActor->GetActorLocation());
-			CustomizableObjectInstance->GetPrivate()->MinSquareDistFromComponentToPlayer = 
-				FMath::Min(SquareDist, CustomizableObjectInstance->GetPrivate()->MinSquareDistFromComponentToPlayer);
-		}
-
-		USkeletalMesh* AttachedSkeletalMesh = GetAttachedSkeletalMesh();
-		USkeletalMesh* GeneratedSKeletalMesh = CustomizableObjectInstance->GetSkeletalMesh(ComponentIndex);
-		if (!AttachedSkeletalMesh && !CustomizableObjectInstance->GetPrivate()->HasCOInstanceFlags(CreatingSkeletalMesh))
-		{
-			UCustomizableObject* Object = CustomizableObjectInstance->GetCustomizableObject(); 
-			USkeletalMesh* RefMesh = Object ? Object->GetRefSkeletalMesh(ComponentIndex) : nullptr;
-			SetSkeletalMesh(GeneratedSKeletalMesh ? GeneratedSKeletalMesh : RefMesh);
-		}
-		else if (GeneratedSKeletalMesh && AttachedSkeletalMesh != GeneratedSKeletalMesh)
-		{
-			SetSkeletalMesh(GeneratedSKeletalMesh);
-		}
-	}
-}
-
-
 void UCustomizableSkeletalComponent::EditorUpdateComponent()
 {
-	if (CustomizableObjectInstance)
-	{
-		CustomizableObjectInstance->GetPrivate()->SetCOInstanceFlags(UsedByComponent);
-
-		AActor* ParentActor = GetAttachmentRootActor();
-		if (ParentActor)
-		{
-			USkeletalMesh* AttachedSkeletalMesh = GetAttachedSkeletalMesh();
-			USkeletalMesh* GeneratedSKeletalMesh = CustomizableObjectInstance->GetSkeletalMesh(ComponentIndex);
-			if (!AttachedSkeletalMesh && !CustomizableObjectInstance->GetPrivate()->HasCOInstanceFlags(CreatingSkeletalMesh))
-			{
-				SetSkeletalMesh(GeneratedSKeletalMesh ? GeneratedSKeletalMesh : CustomizableObjectInstance->GetCustomizableObject()->GetRefSkeletalMesh(ComponentIndex));
-			}
-			else if (GeneratedSKeletalMesh && AttachedSkeletalMesh != GeneratedSKeletalMesh)
-			{
-				SetSkeletalMesh(GeneratedSKeletalMesh);
-			}
-		}
-	}
-}
-#endif
-
-
-void UCustomizableSkeletalComponent::SetVisibilityOfSkeletalMeshSectionWithMaterialName(bool bInVisible, const FString& MaterialName, int32 LOD)
-{
-	USkeletalMesh* SkeletalMesh = CustomizableObjectInstance->GetSkeletalMesh(ComponentIndex);
-
-	if (!SkeletalMesh)
-	{
-		return;
-	}
-
-	for (int32 i = 0; i < SkeletalMesh->GetMaterials().Num(); ++i)
-	{
-		UMaterialInstanceDynamic* MatInst = Cast<UMaterialInstanceDynamic>(SkeletalMesh->GetMaterials()[i].MaterialInterface);
-
-		if (MatInst && MatInst->Parent->GetName() == MaterialName)
-		{
-			if (FSkeletalMeshRenderData* SkelResource = SkeletalMesh->GetResourceForRendering() )
-			{
-				if ( SkelResource->LODRenderData.IsValidIndex(LOD))
-				{
-					FSkeletalMeshLODRenderData& LodData = SkelResource->LODRenderData[LOD];
-
-					for (int32 j = 0; j < LodData.RenderSections.Num(); ++j)
-					{
-						FSkelMeshRenderSection& Section = LodData.RenderSections[j];
-
-						if (Section.MaterialIndex == i)
-						{
-							Section.bDisabled = !bInVisible;
-						}
-					}
-				}
-			}
-		}
-	}
-}
-
-
-void UCustomizableSkeletalComponent::UpdateDistFromComponentToPlayer(const AActor* ViewCenter, const bool bForceEvenIfNotBegunPlay)
-{
-	if (CustomizableObjectInstance)
+	if (CustomizableObjectInstanceUsage)
 	{
-		AActor* ParentActor = GetAttachmentRootActor();
-		CustomizableObjectInstance->SetIsPlayerOrNearIt(false);
-
-		if (ParentActor && ParentActor->IsValidLowLevel())
-		{
-			if (ParentActor->HasActorBegunPlay() || bForceEvenIfNotBegunPlay)
-			{
-				float SquareDist = FLT_MAX;
-
-				if (ViewCenter && ViewCenter->IsValidLowLevel())
-				{
-					APawn* Pawn = Cast<APawn>(ParentActor);
-					bool bIsPlayer = Pawn ? Pawn->IsPlayerControlled() : false;
-					CustomizableObjectInstance->SetIsPlayerOrNearIt(bIsPlayer);
-
-					if (bIsPlayer)
-					{
-						SquareDist = -0.01f; // Negative value to give the player character more priority than any other character
-					}
-					else
-					{
-						SquareDist = FVector::DistSquared(ViewCenter->GetActorLocation(), ParentActor->GetActorLocation());
-					}
-				}
-				else if (bForceEvenIfNotBegunPlay)
-				{
-					SquareDist = -0.01f; // This is a manual update before begin play and the creation of the pawn, so it should probably be high priority
-					CustomizableObjectInstance->GetPrivate()->LastMinSquareDistFromComponentToPlayer = FMath::Min(SquareDist, CustomizableObjectInstance->GetPrivate()->MinSquareDistFromComponentToPlayer);
-				}
-				else
-				{
-					SquareDist = 0.f; // This a mutable tick before begin play and the creation of the pawn, so it should have a definite and high priority but less than a manual update
-					CustomizableObjectInstance->GetPrivate()->LastMinSquareDistFromComponentToPlayer = FMath::Min(SquareDist, CustomizableObjectInstance->GetPrivate()->MinSquareDistFromComponentToPlayer);
-				}
-
-				CustomizableObjectInstance->GetPrivate()->MinSquareDistFromComponentToPlayer = FMath::Min(SquareDist, CustomizableObjectInstance->GetPrivate()->MinSquareDistFromComponentToPlayer);
-				CustomizableObjectInstance->SetIsBeingUsedByComponentInPlay(true);
-
-				if (CustomizableObjectInstance->GetPrivate()->MinSquareDistFromComponentToPlayer == SquareDist)
-				{
-					CustomizableObjectInstance->NearestToActor = this;
-					CustomizableObjectInstance->NearestToViewCenter = ViewCenter;
-				}
-			}
-		}
-
-		if (ParentActor && GetAttachedSkeletalMesh() == nullptr && CustomizableObjectInstance->GetSkeletalMesh(ComponentIndex) && !CustomizableObjectInstance->GetPrivate()->HasCOInstanceFlags(CreatingSkeletalMesh))
-		{
-			SetSkeletalMesh(CustomizableObjectInstance->GetSkeletalMesh(ComponentIndex));
-		}
+		CustomizableObjectInstanceUsage->EditorUpdateComponent();
 	}
 }
-
-
-void UCustomizableSkeletalComponent::TickComponent(float DeltaTime, ELevelTick TickType, FActorComponentTickFunction* ThisTickFunction)
-{
-	Super::TickComponent(DeltaTime, TickType, ThisTickFunction);
-
-	if (!bPendingSetSkeletalMesh || !CustomizableObjectInstance || !CustomizableObjectInstance->GetCustomizableObject())
-	{
-		return;
-	}
-
-	if (USkeletalMeshComponent* Parent = Cast<USkeletalMeshComponent>(GetAttachParent()))
-	{
-		UCustomizableObject* CustomizableObject = CustomizableObjectInstance->GetCustomizableObject();
-
-		// Hacky. Replace once we know if the instance has been generated
-		const bool bInstanceGenerated = CustomizableObjectInstance->HasAnySkeletalMesh();
-
-		// Generated SkeletalMesh to set, can be null if the component is empty
-		USkeletalMesh* SkeletalMesh = CustomizableObjectInstance->GetSkeletalMesh(ComponentIndex);
-
-		// If not generated yet, conditionally set the SkeletalMesh of reference
-		if (!bInstanceGenerated && !bSkipSetReferenceSkeletalMesh
-#if WITH_EDITORONLY_DATA
-			&& CustomizableObject->bEnableUseRefSkeletalMeshAsPlaceholder
-#endif
-			)
-		{
-			// Can be nullptr
-			SkeletalMesh = CustomizableObject->GetRefSkeletalMesh(ComponentIndex);
-		}
-
-		// Set SkeletalMesh
-		if (bInstanceGenerated || SkeletalMesh)
-		{
-			Parent->SetSkeletalMesh(SkeletalMesh);
-
-			if (Parent->OverrideMaterials.Num() > 0)
-			{
-#if !(UE_BUILD_SHIPPING || UE_BUILD_TEST)
-				if (Parent->GetClass()->GetFName() != FName(TEXT("SkeletalMeshComponentBudgeted"))) // Reduce unnecessary logging
-				{
-					UE_LOG(LogMutable, Warning, TEXT("Attaching Customizable Skeletal Component to Skeletal Mesh Component with overriden materials! Deleting overrides."));
-				}
 #endif
 
-				Parent->EmptyOverrideMaterials();
-			}
-
-			bPendingSetSkeletalMesh = false;
-		}
-	}
-}
-
 
 void UCustomizableSkeletalComponent::OnAttachmentChanged()
 {
@@ -361,12 +98,20 @@ void UCustomizableSkeletalComponent::OnAttachmentChanged()
 
 	USkeletalMeshComponent* Parent = Cast<USkeletalMeshComponent>(GetAttachParent());
 
-	if (Parent)
+	if (Parent && CustomizableObjectInstanceUsage)
 	{
-		bPendingSetSkeletalMesh = true;
+		CustomizableObjectInstanceUsage->SetPendingSetSkeletalMesh(true);
 	}
 	else if(!GetAttachParent())
 	{
 		DestroyComponent();
 	}
 }
+
+
+void UCustomizableSkeletalComponent::PostInitProperties()
+{
+	Super::PostInitProperties();
+
+	CreateCustomizableObjectInstanceUsage();
+}
diff --git a/Engine/Plugins/Experimental/Mutable/Source/CustomizableObject/Private/MuCO/LogBenchmarkUtil.cpp b/Engine/Plugins/Experimental/Mutable/Source/CustomizableObject/Private/MuCO/LogBenchmarkUtil.cpp
--- a/Engine/Plugins/Experimental/Mutable/Source/CustomizableObject/Private/MuCO/LogBenchmarkUtil.cpp
+++ b/Engine/Plugins/Experimental/Mutable/Source/CustomizableObject/Private/MuCO/LogBenchmarkUtil.cpp
@@ -8,6 +8,7 @@
 #include "Engine/SkeletalMesh.h"
 #include "TextureResource.h"
 #include "GameFramework/Pawn.h"
+#include "Components/SkeletalMeshComponent.h"
 
 #if WITH_EDITOR
 #include "Editor/EditorEngine.h"
@@ -17,7 +18,7 @@
 #include "MuCO/CustomizableObjectSystemPrivate.h"
 #include "MuCO/CustomizableObjectSystem.h"
 #include "MuCO/CustomizableInstancePrivateData.h"
-#include "MuCO/CustomizableSkeletalComponent.h"
+#include "MuCO/CustomizableObjectInstanceUsage.h"
 
 
 extern ENGINE_API UEngine* GEngine;
@@ -294,18 +295,20 @@ void LogBenchmarkUtil::UpdateStats(FMutableStats& StatsToUpdate, const TArray< T
 				bool bFoundPlayer = false;
 				int32 MsgIndex = 15820; // Arbitrary big value to prevent collisions with other on-screen messages
 
-				for (TObjectIterator<UCustomizableSkeletalComponent> CustomizableSkeletalComponent; CustomizableSkeletalComponent; ++CustomizableSkeletalComponent)
+				for (TObjectIterator<UCustomizableObjectInstanceUsage> CustomizableObjectInstanceUsage; CustomizableObjectInstanceUsage; ++CustomizableObjectInstanceUsage)
 				{
-					if (CustomizableSkeletalComponent && CustomizableSkeletalComponent->IsNetMode(NM_DedicatedServer))
+					if (CustomizableObjectInstanceUsage && CustomizableObjectInstanceUsage->IsNetMode(NM_DedicatedServer))
 					{
 						continue;
 					}
 
-					AActor* ParentActor = CustomizableSkeletalComponent->GetAttachmentRootActor();
-					UCustomizableObjectInstance* Instance = CustomizableSkeletalComponent->CustomizableObjectInstance;
+					AActor* ParentActor = CustomizableObjectInstanceUsage->GetAttachParent() ? 
+						CustomizableObjectInstanceUsage->GetAttachParent()->GetAttachmentRootActor() 
+						:  nullptr;
+					UCustomizableObjectInstance* Instance = CustomizableObjectInstanceUsage->GetCustomizableObjectInstance();
 
 					APawn* PlayerPawn = nullptr;
-					if (UWorld* World = CustomizableSkeletalComponent->GetWorld())
+					if (UWorld* World = CustomizableObjectInstanceUsage->GetWorld())
 					{
 						PlayerPawn = UGameplayStatics::GetPlayerPawn(World, 0);
 					}
@@ -326,7 +329,7 @@ void LogBenchmarkUtil::UpdateStats(FMutableStats& StatsToUpdate, const TArray< T
 						GEngine->AddOnScreenDebugMessage(MsgIndex++, .0f, FColor::Green, TEXT("Animation tags: ") + TagString);
 
 						check(Instance->GetPrivate() != nullptr);
-						FCustomizableInstanceComponentData* ComponentData = Instance->GetPrivate()->GetComponentData(CustomizableSkeletalComponent->ComponentIndex);
+						FCustomizableInstanceComponentData* ComponentData = Instance->GetPrivate()->GetComponentData(CustomizableObjectInstanceUsage->GetComponentIndex());
 
 						if (ComponentData)
 						{
@@ -359,7 +362,7 @@ void LogBenchmarkUtil::UpdateStats(FMutableStats& StatsToUpdate, const TArray< T
 
 						GEngine->AddOnScreenDebugMessage(MsgIndex++, .0f, FColor::Cyan,
 							TEXT("Player Pawn Mutable Mesh/Animation info for component ") + FString::Printf(TEXT("%d"),
-								CustomizableSkeletalComponent->ComponentIndex));
+								CustomizableObjectInstanceUsage->GetComponentIndex()));
 					}
 				}
 
diff --git a/Engine/Plugins/Experimental/Mutable/Source/CustomizableObject/Public/MuCO/CustomizableObjectDGGUI.h b/Engine/Plugins/Experimental/Mutable/Source/CustomizableObject/Public/MuCO/CustomizableObjectDGGUI.h
--- a/Engine/Plugins/Experimental/Mutable/Source/CustomizableObject/Public/MuCO/CustomizableObjectDGGUI.h
+++ b/Engine/Plugins/Experimental/Mutable/Source/CustomizableObject/Public/MuCO/CustomizableObjectDGGUI.h
@@ -7,7 +7,7 @@
 #include "CustomizableObjectDGGUI.generated.h"
 
 class UAnimInstance;
-class UCustomizableSkeletalComponent;
+class UCustomizableObjectInstanceUsage;
 class UObject;
 class UWorld;
 
@@ -18,11 +18,11 @@ class CUSTOMIZABLEOBJECT_API UDGGUI : public UUserWidget
 	GENERATED_BODY()
 public:
 	UFUNCTION(BlueprintImplementableEvent, Category = "DGGUI")
-	class UCustomizableSkeletalComponent* GetCustomizableSkeletalComponent();
+	class UCustomizableObjectInstanceUsage* GetCustomizableObjectInstanceUsage();
 
 	UFUNCTION(BlueprintImplementableEvent, Category = "DGGUI")
-	void SetCustomizableSkeletalComponent(class UCustomizableSkeletalComponent* CustomizableSkeletalComponent);
+	void SetCustomizableObjectInstanceUsage(class UCustomizableObjectInstanceUsage* CustomizableObjectInstanceUsage);
 
-	static void OpenDGGUI(const int32 SlotID, UCustomizableSkeletalComponent* SelectedCustomizableSkeletalComponent, const UWorld* CurrentWorld, const int32 PlayerIndex = 0);
+	static void OpenDGGUI(const int32 SlotID, UCustomizableObjectInstanceUsage* SelectedCustomizableObjectInstanceUsage, const UWorld* CurrentWorld, const int32 PlayerIndex = 0);
 	static bool CloseExistingDGGUI(const UWorld* CurrentWorld);
 };
diff --git a/Engine/Plugins/Experimental/Mutable/Source/CustomizableObject/Public/MuCO/ICustomizableObjectModule.h b/Engine/Plugins/Experimental/Mutable/Source/CustomizableObject/Public/MuCO/ICustomizableObjectModule.h
--- a/Engine/Plugins/Experimental/Mutable/Source/CustomizableObject/Public/MuCO/ICustomizableObjectModule.h
+++ b/Engine/Plugins/Experimental/Mutable/Source/CustomizableObject/Public/MuCO/ICustomizableObjectModule.h
@@ -7,7 +7,7 @@
 #include "MuCO/CustomizableObjectExtension.h"
 #include "MuCO/CustomizableObject.h"
 
-class UCustomizableSkeletalComponent;
+class UCustomizableObjectInstanceUsage;
 
 struct FRegisteredCustomizableObjectPinType
 {
@@ -82,4 +82,4 @@ public:
 
 /* Utility function for commands*/
 CUSTOMIZABLEOBJECT_API
-UCustomizableSkeletalComponent* GetPlayerCustomizableSkeletalComponent(const int32 SlotID, const UWorld* CurrentWorld, const int32 PlayerIndex = 0);
+UCustomizableObjectInstanceUsage* GetPlayerCustomizableObjectInstanceUsage(const int32 SlotID, const UWorld* CurrentWorld, const int32 PlayerIndex = 0);
diff --git a/Engine/Plugins/Experimental/Mutable/Source/CustomizableObjectEditor/Private/MuCOE/CustomizableInstanceDetails.h b/Engine/Plugins/Experimental/Mutable/Source/CustomizableObjectEditor/Private/MuCOE/CustomizableInstanceDetails.h
--- a/Engine/Plugins/Experimental/Mutable/Source/CustomizableObjectEditor/Private/MuCOE/CustomizableInstanceDetails.h
+++ b/Engine/Plugins/Experimental/Mutable/Source/CustomizableObjectEditor/Private/MuCOE/CustomizableInstanceDetails.h
@@ -8,7 +8,6 @@
 class IDetailLayoutBuilder;
 class SCustomizableInstanceProperties;
 class UCustomizableObjectInstance;
-class UCustomizableSkeletalComponent;
 
 
 class FCustomizableInstanceDetails : public IDetailCustomization
diff --git a/Engine/Plugins/Experimental/Mutable/Source/CustomizableObjectEditor/Private/MuCOE/CustomizableObjectEditorModule.cpp b/Engine/Plugins/Experimental/Mutable/Source/CustomizableObjectEditor/Private/MuCOE/CustomizableObjectEditorModule.cpp
--- a/Engine/Plugins/Experimental/Mutable/Source/CustomizableObjectEditor/Private/MuCOE/CustomizableObjectEditorModule.cpp
+++ b/Engine/Plugins/Experimental/Mutable/Source/CustomizableObjectEditor/Private/MuCOE/CustomizableObjectEditorModule.cpp
@@ -8,7 +8,7 @@
 #include "ISettingsSection.h"
 #include "MessageLogModule.h"
 #include "MuCO/CustomizableObjectSystem.h"		// For defines related to memory function replacements.
-#include "MuCO/CustomizableSkeletalComponent.h"
+#include "MuCO/CustomizableObjectInstanceUsage.h"
 #include "MuCO/CustomizableSkeletalMeshActor.h"
 #include "MuCO/ICustomizableObjectModule.h"		// For instance editor command utility function
 #include "MuCOE/CustomizableInstanceDetails.h"
@@ -363,9 +363,9 @@ void FCustomizableObjectEditorModule::OpenCOIE(const TArray<FString>& Arguments)
 	const int32 PlayerIndex = 0;
 
 	// Open the Customizable Object Instance Editor
-	if (UCustomizableSkeletalComponent* SelectedCustomizableSkeletalComponent = GetPlayerCustomizableSkeletalComponent(SlotID, CurrentWorld, PlayerIndex))
+	if (UCustomizableObjectInstanceUsage* SelectedCustomizableObjectInstanceUsage = GetPlayerCustomizableObjectInstanceUsage(SlotID, CurrentWorld, PlayerIndex))
 	{
-		if (UCustomizableObjectInstance* COInstance = SelectedCustomizableSkeletalComponent->CustomizableObjectInstance)
+		if (UCustomizableObjectInstance* COInstance = SelectedCustomizableObjectInstanceUsage->GetCustomizableObjectInstance())
 		{
 			FAssetToolsModule& AssetToolsModule = FModuleManager::LoadModuleChecked<FAssetToolsModule>(TEXT("AssetTools"));
 			TWeakPtr<IAssetTypeActions> WeakAssetTypeActions = AssetToolsModule.Get().GetAssetTypeActionsForClass(UCustomizableObjectInstance::StaticClass());
