[Mutable] Fix OpMeshMerge missing scopes due to invalid naming.

#rnx

diff --git a/Engine/Plugins/Experimental/Mutable/Source/MutableRuntime/Private/MuR/OpMeshMerge.h b/Engine/Plugins/Experimental/Mutable/Source/MutableRuntime/Private/MuR/OpMeshMerge.h
--- a/Engine/Plugins/Experimental/Mutable/Source/MutableRuntime/Private/MuR/OpMeshMerge.h
+++ b/Engine/Plugins/Experimental/Mutable/Source/MutableRuntime/Private/MuR/OpMeshMerge.h
@@ -439,7 +439,7 @@ namespace mu
 
 		if (bMergeSkeletons)
 		{
-			MUTABLE_CPUPROFILER_SCOPE(Skeleton);
+			MUTABLE_CPUPROFILER_SCOPE(MergeSkeleton);
 			Ptr<Skeleton> pResultSkeleton;
 
 			mu::SkeletonPtrConst pFirstSkeleton = pFirst->GetSkeleton();
@@ -750,7 +750,7 @@ namespace mu
 			const int32 SecondCount = pSecond->GetVertexBuffers().GetElementCount();
 
 			// TODO: when formats match, which at runtime should be always.
-			bool FastPath = pFirst->GetVertexBuffers().HasSameFormat( pSecond->GetVertexBuffers() );
+			bool bFastPath = pFirst->GetVertexBuffers().HasSameFormat( pSecond->GetVertexBuffers() );
 
 			// Check if the format of the BoneIndex buffer has to change
 			bool bChangeBoneIndicesFormat = false;
@@ -802,13 +802,13 @@ namespace mu
 					}
 				}
 
-				FastPath = FastPath && !bChangeBoneIndicesFormat;
+				bFastPath = bFastPath && !bChangeBoneIndicesFormat;
 			}
 
 			Ptr<const Mesh> pVFirst;
 			Ptr<const Mesh> pVSecond;
 			
-			if (!FastPath)
+			if (!bFastPath)
 			{
                 MUTABLE_CPUPROFILER_SCOPE(SlowPath);
 
@@ -1015,6 +1015,7 @@ namespace mu
 
             // first copy all the vertex data
 			{
+				MUTABLE_CPUPROFILER_SCOPE(CopyVertexData);
 				for (int32 vb = 0; vb < Result->GetVertexBuffers().m_buffers.Num(); ++vb)
 				{
 					MESH_BUFFER& result = Result->GetVertexBuffers().m_buffers[vb];
