[mutable] Fix compilation issue with stray pragma in some platforms.
#rnx

diff --git a/Engine/Plugins/Experimental/Mutable/Source/MutableRuntime/Private/MuR/OpMeshClipWithMesh.cpp b/Engine/Plugins/Experimental/Mutable/Source/MutableRuntime/Private/MuR/OpMeshClipWithMesh.cpp
--- a/Engine/Plugins/Experimental/Mutable/Source/MutableRuntime/Private/MuR/OpMeshClipWithMesh.cpp
+++ b/Engine/Plugins/Experimental/Mutable/Source/MutableRuntime/Private/MuR/OpMeshClipWithMesh.cpp
@@ -18,7 +18,8 @@
 
 #include "Spatial/PointHashGrid3.h"
 
-namespace mu { namespace  {
+namespace mu { namespace 
+{
 
     //---------------------------------------------------------------------------------------------
     //! Create a map from vertices into vertices, collapsing vertices that have the same position, 
@@ -592,15 +593,13 @@ namespace mu { namespace  {
 
 		MUTABLE_CPUPROFILER_SCOPE(MeshUVMaskClassifyVertices);
 
-		int32 OrigVertCount = Base->GetVertexBuffers().GetElementCount();
+		uint32 VertexCount = Base->GetVertexCount();
 
 		// Stores whether each vertex in the original mesh in the clip mesh volume
-		VertexClipped.SetNum(0,false);
-		VertexClipped.SetNumZeroed(OrigVertCount);
+		VertexClipped.SetNumUninitialized(VertexCount,false);
+		FMemory::Memzero(VertexClipped.GetData(),VertexClipped.GetAllocatedSize());
 
 		// Now go through all vertices in the mesh and record whether they are inside or outside of the ClipMesh
-		uint32 DestVertexCount = Base->GetVertexCount();
-
 		const FMeshBufferSet& MBSPriv = Base->GetVertexBuffers();
 		for (int32 b = 0; b < MBSPriv.m_buffers.Num(); ++b)
 		{
@@ -621,7 +620,7 @@ namespace mu { namespace  {
 				}
 
 				UntypedMeshBufferIteratorConst It(Base->GetVertexBuffers(), Sem, SemIndex);
-				for (uint32 V = 0; V < DestVertexCount; ++V)
+				for (uint32 V = 0; V < VertexCount; ++V)
 				{
 					// \TODO: This could be optimized.
 					FVector2f UV = It.GetAsVec2f();
@@ -702,9 +701,7 @@ namespace mu { namespace  {
   //      //	Result->CopyFrom(pBase); // If all faces have been discarded, return a copy of the unmodified mesh because unreal doesn't like empty meshes
   //      //}
 
-  //      // [jordi] Remove unused vertices. This is necessary to avoid returning a mesh with vertices and no faces, which screws some
-  //      // engines like Unreal Engine 4.
-  //      MeshRemoveUnusedVertices(Result);
+    //      MeshRemoveUnusedVertices(Result);
   //  }
 
 
@@ -825,8 +822,7 @@ namespace mu
         //	Result->CopyFrom(pBase); // If all faces have been discarded, return a copy of the unmodified mesh because unreal doesn't like empty meshes
         //}
 
-        // [jordi] Remove unused vertices. This is necessary to avoid returning a mesh with vertices and no faces, which screws some
-        // engines like Unreal Engine 4.
+
         MeshRemoveUnusedVertices(Result);
     }
 
