[Mutable] Fix wrongly setting the parent Skeletal Mesh when the Instance was not generated.

- Moved bEnableUseRefSkeletalMeshAsPlaceholder to non-EditorOnly.

#rnx

diff --git a/Engine/Plugins/Experimental/Mutable/Source/CustomizableObject/Classes/MuCO/CustomizableObject.h b/Engine/Plugins/Experimental/Mutable/Source/CustomizableObject/Classes/MuCO/CustomizableObject.h
--- a/Engine/Plugins/Experimental/Mutable/Source/CustomizableObject/Classes/MuCO/CustomizableObject.h
+++ b/Engine/Plugins/Experimental/Mutable/Source/CustomizableObject/Classes/MuCO/CustomizableObject.h
@@ -1207,15 +1207,15 @@ public:
 	UPROPERTY()
 	TArray<FCustomizableObjectStreamedExtensionData> StreamedExtensionData;
 
-#if WITH_EDITORONLY_DATA
-
+private:
 	/** Use the SkeletalMesh of reference as a placeholder until the custom mesh is ready to use.
 	  * 
-	  * Note: If disabled, a null mesh will be used to replace the discarded mesh due to 'ReplaceDiscardedWithReferenceMesh' being enabled.
-	  */
+	  * Note: If disabled, a null mesh will be used to replace the discarded mesh due to 'ReplaceDiscardedWithReferenceMesh' being enabled. */
 	UPROPERTY(EditAnywhere, Category = CustomizableObject)
 	bool bEnableUseRefSkeletalMeshAsPlaceholder = true;
 
+public:
+#if WITH_EDITORONLY_DATA
 	// Hide this property because it is not used yet.
 	//UPROPERTY(EditAnywhere, Category = CustomizableObject)
 	UPROPERTY()
@@ -1711,6 +1711,9 @@ public:
 	/** Return the MinLOD index to generate based on the active LODSettings (PerPlatformMinLOD or PerQualityLevelMinLOD) */
 	int32 GetMinLODIndex() const;
 
+	/** See bEnableUseRefSkeletalMeshAsPlaceholder. */
+	bool IsEnableUseRefSkeletalMeshAsPlaceholder() const;
+
 private:
 	
 	/** Stores the bone names of all the bones that can possibly use the generated meshes */
diff --git a/Engine/Plugins/Experimental/Mutable/Source/CustomizableObject/Private/MuCO/CustomizableObject.cpp b/Engine/Plugins/Experimental/Mutable/Source/CustomizableObject/Private/MuCO/CustomizableObject.cpp
--- a/Engine/Plugins/Experimental/Mutable/Source/CustomizableObject/Private/MuCO/CustomizableObject.cpp
+++ b/Engine/Plugins/Experimental/Mutable/Source/CustomizableObject/Private/MuCO/CustomizableObject.cpp
@@ -1846,6 +1846,12 @@ int32 UCustomizableObject::GetMinLODIndex() const
 }
 
 
+bool UCustomizableObject::IsEnableUseRefSkeletalMeshAsPlaceholder() const
+{
+	return bEnableUseRefSkeletalMeshAsPlaceholder;
+}
+
+
 FGuid UCustomizableObject::GetCompilationGuid() const
 {
 	return CompilationGuid;
@@ -2148,7 +2154,7 @@ FArchive& operator<<(FArchive& Ar, FMutableRefSkeletalMeshData& Data)
 void FMutableRefSkeletalMeshData::InitResources(UCustomizableObject* InOuter)
 {
 	check(InOuter);
-	if (InOuter->bEnableUseRefSkeletalMeshAsPlaceholder)
+	if (InOuter->IsEnableUseRefSkeletalMeshAsPlaceholder())
 	{
 		SkeletalMesh = TSoftObjectPtr<USkeletalMesh>(SkeletalMeshAssetPath).LoadSynchronous();
 	}
diff --git a/Engine/Plugins/Experimental/Mutable/Source/CustomizableObject/Private/MuCO/CustomizableObjectInstanceUsage.cpp b/Engine/Plugins/Experimental/Mutable/Source/CustomizableObject/Private/MuCO/CustomizableObjectInstanceUsage.cpp
--- a/Engine/Plugins/Experimental/Mutable/Source/CustomizableObject/Private/MuCO/CustomizableObjectInstanceUsage.cpp
+++ b/Engine/Plugins/Experimental/Mutable/Source/CustomizableObject/Private/MuCO/CustomizableObjectInstanceUsage.cpp
@@ -465,32 +465,38 @@ void UCustomizableObjectInstanceUsage::Tick(float DeltaTime)
 {
 	UCustomizableObjectInstance* CustomizableObjectInstance = GetCustomizableObjectInstance();
 
-	if (!GetPendingSetSkeletalMesh() || !CustomizableObjectInstance || !CustomizableObjectInstance->GetCustomizableObject())
+	if (!GetPendingSetSkeletalMesh() || !CustomizableObjectInstance)
 	{
 		return;
 	}
+	
+	UCustomizableObject* CustomizableObject = CustomizableObjectInstance->GetCustomizableObject();
+	if (!CustomizableObject)
+	{
+		return;
+	}	
 
 	if (USkeletalMeshComponent* Parent = Cast<USkeletalMeshComponent>(GetAttachParent()))
 	{
-		UCustomizableObject* CustomizableObject = CustomizableObjectInstance->GetCustomizableObject();
-
-		// Hacky. Replace once we know if the instance has been generated
-		const bool bInstanceGenerated = CustomizableObjectInstance->HasAnySkeletalMesh();
-
-		int32 ComponentIndex = GetComponentIndex();
+		const int32 ComponentIndex = GetComponentIndex();
 
-		// Generated SkeletalMesh to set, can be null if the component is empty
-		USkeletalMesh* SkeletalMesh = CustomizableObjectInstance->GetSkeletalMesh(ComponentIndex);
+		USkeletalMesh* SkeletalMesh = nullptr;
 
-		// If not generated yet, conditionally set the SkeletalMesh of reference
-		if (!bInstanceGenerated && !GetSkipSetReferenceSkeletalMesh()
-#if WITH_EDITORONLY_DATA
-			&& CustomizableObject->bEnableUseRefSkeletalMeshAsPlaceholder
-#endif
-			)
+		const bool bInstanceGenerated = CustomizableObjectInstance->GetPrivate()->GetSkeletalMeshStatus() == ESkeletalMeshStatus::Success;		
+		if (bInstanceGenerated)
+		{
+			// Generated SkeletalMesh to set, can be null if the component is empty
+			SkeletalMesh = CustomizableObjectInstance->GetSkeletalMesh(ComponentIndex);
+		}
+		else
 		{
-			// Can be nullptr
-			SkeletalMesh = CustomizableObject->GetRefSkeletalMesh(ComponentIndex);
+			// If not generated yet, conditionally set the SkeletalMesh of reference
+			if (!bInstanceGenerated && !GetSkipSetReferenceSkeletalMesh() && 
+				CustomizableObject->IsEnableUseRefSkeletalMeshAsPlaceholder())
+			{
+				// Can be nullptr
+				SkeletalMesh = CustomizableObject->GetRefSkeletalMesh(ComponentIndex);
+			}
 		}
 
 		// Set SkeletalMesh
