[Mutable] Fix double call to the PostCompileDelegate.

#rnx

diff --git a/Engine/Plugins/Experimental/Mutable/Source/CustomizableObject/Private/MuCO/CustomizableObjectInstance.cpp b/Engine/Plugins/Experimental/Mutable/Source/CustomizableObject/Private/MuCO/CustomizableObjectInstance.cpp
--- a/Engine/Plugins/Experimental/Mutable/Source/CustomizableObject/Private/MuCO/CustomizableObjectInstance.cpp
+++ b/Engine/Plugins/Experimental/Mutable/Source/CustomizableObject/Private/MuCO/CustomizableObjectInstance.cpp
@@ -246,7 +246,8 @@ void UCustomizableObjectInstance::BindPostCompileDelegate(UCustomizableObject* I
 	}
 
 	// Bind callback to the new CO
-	if (InCustomizableObject)
+	if (InCustomizableObject &&
+		!InCustomizableObject->PostCompileDelegate.IsBoundToObject(GetPrivate()))
 	{
 		InCustomizableObject->PostCompileDelegate.AddUObject(GetPrivate(), &UCustomizableInstancePrivateData::OnPostCompile);
 	}
