[Mutable] Added Update stats.

- Updated FLogBenchmarkUitl class to output only used stats.
- Corrected some wrongly reported stats.
- FLogBenchmarkUtil now also reports stats on csv files.
- Moved FinishUpdateGlobal from outside UpdateSkeletalMesh.
- Fixed GetNumSkeletalMeshes.
- Fixed a case where UE_MUTABLE_UPDATE_REGION region would never end.

#rnx

diff --git a/Engine/Plugins/Experimental/Mutable/Source/CustomizableObject/Private/MuCO/CustomizableObjectInstance.cpp b/Engine/Plugins/Experimental/Mutable/Source/CustomizableObject/Private/MuCO/CustomizableObjectInstance.cpp
--- a/Engine/Plugins/Experimental/Mutable/Source/CustomizableObject/Private/MuCO/CustomizableObjectInstance.cpp
+++ b/Engine/Plugins/Experimental/Mutable/Source/CustomizableObject/Private/MuCO/CustomizableObjectInstance.cpp
@@ -54,7 +54,7 @@ UTexture2D* UCustomizableInstancePrivateData::CreateTexture()
 		NAME_None,
 		RF_Transient
 		);
-	UCustomizableObjectSystem::GetInstance()->GetPrivate()->CreatedTexture(NewTexture);
+	UCustomizableObjectSystem::GetInstance()->GetPrivate()->LogBenchmarkUtil.AddTexture(*NewTexture);
 	NewTexture->SetPlatformData( nullptr );
 
 	return NewTexture;
@@ -1833,8 +1833,6 @@ bool UCustomizableInstancePrivateData::UpdateSkeletalMesh_PostBeginUpdate0(UCust
 		USkeletalMesh* SkeletalMesh = Public->SkeletalMeshes[Component.Id];
 		check(SkeletalMesh);
 
-		INC_DWORD_STAT(STAT_MutableNumSkeletalMeshes);
-
 		// Set up the default information any mesh from this component will have (LODArrayInfos, RenderData, Mesh settings, etc). 
 		InitSkeletalMeshData(OperationData, SkeletalMesh, RefSkeletalMeshData, *CustomizableObject, Component.Id);
 
diff --git a/Engine/Plugins/Experimental/Mutable/Source/CustomizableObject/Private/MuCO/CustomizableObjectSystem.cpp b/Engine/Plugins/Experimental/Mutable/Source/CustomizableObject/Private/MuCO/CustomizableObjectSystem.cpp
--- a/Engine/Plugins/Experimental/Mutable/Source/CustomizableObject/Private/MuCO/CustomizableObjectSystem.cpp
+++ b/Engine/Plugins/Experimental/Mutable/Source/CustomizableObject/Private/MuCO/CustomizableObjectSystem.cpp
@@ -449,18 +449,6 @@ void UCustomizableObjectSystem::InitSystem()
 	Private->TickDelegateHandle = FTSTicker::GetCoreTicker().AddTicker(Private->TickDelegate, 0.f);
 #endif // !UE_SERVER
 
-	Private->MutableStats.TotalBuildMs = 0;
-	Private->MutableStats.TotalBuiltInstances = 0;
-	Private->MutableStats.NumInstances = 0;
-	Private->MutableStats.TextureMemoryUsed = 0;
-
-#if !(UE_BUILD_SHIPPING || UE_BUILD_TEST)
-	SET_DWORD_STAT(STAT_MutableNumSkeletalMeshes, 0);
-	SET_DWORD_STAT(STAT_MutableNumTextures, 0);
-	SET_DWORD_STAT(STAT_MutableInstanceBuildTime, 0);
-	SET_DWORD_STAT(STAT_MutableInstanceBuildTimeAvrg, 0);
-#endif
-
 	Private->LastWorkingMemoryBytes = CVarWorkingMemory.GetValueOnGameThread() * 1024;
 	Private->LastGeneratedResourceCacheSize = CVarGeneratedResourcesCacheSize.GetValueOnGameThread();
 
@@ -580,30 +568,6 @@ void UCustomizableObjectSystem::SetNewCompilerFunc(FCustomizableObjectCompilerBa
 }
 
 
-void FCustomizableObjectSystemPrivate::CreatedTexture(UTexture2D* Texture) const
-{
-#if !(UE_BUILD_SHIPPING || UE_BUILD_TEST)
-	const bool bLogEnabled = true;
-#else
-	const bool bLogEnabled = LogBenchmarkUtil::IsLoggingActive();
-#endif
-	if (bLogEnabled)
-	{
-		MutableStats.TextureTrackerArray.Add(Texture);
-
-		for (auto Iterator = MutableStats.TextureTrackerArray.CreateIterator(); Iterator; ++Iterator)
-		{
-			if (Iterator->IsStale())
-			{
-				Iterator.RemoveCurrent();
-			}
-		}
-	}
-
-	INC_DWORD_STAT(STAT_MutableNumTextures);
-}
-
-
 int32 FCustomizableObjectSystemPrivate::EnableMutableAnimInfoDebugging = 0;
 
 static FAutoConsoleVariableRef CVarEnableMutableAnimInfoDebugging(
@@ -777,6 +741,9 @@ void FinishUpdateGlobal(const TSharedRef<FUpdateContextPrivate>& Context)
 
 	UCustomizableObjectInstance* Instance = Context->Instance.Get();
 
+	UCustomizableObjectSystem* System = UCustomizableObjectSystem::GetInstance();
+	FCustomizableObjectSystemPrivate* SystemPrivate = System ? System->GetPrivate() : nullptr;
+
 	if (Instance)
 	{
 		UCustomizableInstancePrivateData* PrivateInstance = Instance->GetPrivate();
@@ -838,10 +805,24 @@ void FinishUpdateGlobal(const TSharedRef<FUpdateContextPrivate>& Context)
 		
 	Context->UpdateCallback.ExecuteIfBound(ContextPublic);
 
-	UCustomizableObjectSystem::GetInstance()->GetPrivate()->MutableTaskGraph.AllowLaunchingMutableTaskLowPriority(true, false);
-
+	if (SystemPrivate)
+	{
+		SystemPrivate->MutableTaskGraph.AllowLaunchingMutableTaskLowPriority(true, false);
+	}
+	
 	const uint32 InstanceId = Instance ? Instance->GetUniqueID() : 0;
-	UE_LOG(LogMutable, Log, TEXT("Finished UpdateSkeletalMesh Async. of Instance %d , frame=%d"), InstanceId, GFrameNumber);
+	Context->UpdateTime = FPlatformTime::Seconds() - Context->StartUpdateTime;
+	UE_LOG(LogMutable, Log, TEXT("Finished UpdateSkeletalMesh Async. of Instance %d, Frame=%d, UpdateTime=%f"), InstanceId, GFrameNumber, Context->UpdateTime);
+
+	if (SystemPrivate)
+	{
+		SystemPrivate->LogBenchmarkUtil.FinishUpdate(Context);		
+	}
+	
+	if (Context->UpdateStarted)
+	{
+		TRACE_END_REGION(UE_MUTABLE_UPDATE_REGION);		
+	}
 }
 
 
@@ -903,8 +884,6 @@ void UpdateSkeletalMesh(const TSharedRef<FUpdateContextPrivate>& Context)
 	CustomizableObjectInstancePrivateData->ClearCOInstanceFlags(CreatingSkeletalMesh);
 
 	CustomizableObjectInstance->bEditorPropertyChanged = false;
-
-	FinishUpdateGlobal(Context);
 }
 
 
@@ -946,9 +925,9 @@ void FCustomizableObjectSystemPrivate::SetReplaceDiscardedWithReferenceMeshEnabl
 }
 
 
-int32 FCustomizableObjectSystemPrivate::GetCountAllocatedSkeletalMesh() const
+int32 FCustomizableObjectSystemPrivate::GetNumSkeletalMeshes() const
 {
-	return MutableStats.CountAllocatedSkeletalMesh;
+	return NumSkeletalMeshes;
 }
 
 
@@ -1016,7 +995,7 @@ EUpdateRequired FCustomizableObjectSystemPrivate::IsUpdateRequired(const UCustom
 
 	if (!bIsGenerated && // Prevent generating more instances than the limit, but let updates to existing instances run normally
 		NumGeneratedInstancesLimit > 0 &&
-		System->GetPrivate()->GetCountAllocatedSkeletalMesh() > NumGeneratedInstancesLimit + NumGeneratedInstancesLimitLOD1 + NumGeneratedInstancesLimitLOD2)
+		System->GetPrivate()->GetNumSkeletalMeshes() > NumGeneratedInstancesLimit + NumGeneratedInstancesLimitLOD1 + NumGeneratedInstancesLimitLOD2)
 	{
 		return EUpdateRequired::NoUpdate;
 	}
@@ -1159,6 +1138,7 @@ void FCustomizableObjectSystemPrivate::EnqueueUpdateSkeletalMesh(const TSharedRe
 		{
 			Context->UpdateResult = EUpdateResult::Success;
 			UpdateSkeletalMesh(Context);
+			FinishUpdateGlobal(Context);
 		}
 		else
 		{
@@ -1907,7 +1887,6 @@ namespace impl
 			if (bCached)
 			{
 				UE_LOG(LogMutable, VeryVerbose, TEXT("Texture resource with id [%llu] is cached."), Image.ImageID);
-				INC_DWORD_STAT(STAT_MutableNumCachedTextures);
 			}
 			else
 			{
@@ -2155,6 +2134,7 @@ namespace impl
 	void Task_Mutable_Update_GetMesh(const TSharedRef<FUpdateContextPrivate>& OperationData, const TSharedPtr<mu::Model>& Model)
 	{
 		MUTABLE_CPUPROFILER_SCOPE(Task_Mutable_Update_GetMesh)
+		FScopeTimer Timer(OperationData->TaskGetMeshTime);
 
 #if WITH_EDITOR
 		const uint32 StartCycles = FPlatformTime::Cycles();
@@ -2185,6 +2165,7 @@ namespace impl
 	void Task_Mutable_Update_GetImages(const TSharedRef<FUpdateContextPrivate>& OperationData)
 	{
 		MUTABLE_CPUPROFILER_SCOPE(Task_Mutable_GetImages)
+		FScopeTimer Timer(OperationData->TaskGetImagesTime);
 
 #if WITH_EDITOR
 		uint32 StartCycles = FPlatformTime::Cycles();
@@ -2268,6 +2249,7 @@ namespace impl
 	void Task_Game_Callbacks(const TSharedRef<FUpdateContextPrivate>& OperationData)
 	{
 		MUTABLE_CPUPROFILER_SCOPE(Task_Game_Callbacks)
+		FScopeTimer Timer(OperationData->TaskCallbacksTime);
 
 		check(IsInGameThread());
 
@@ -2310,24 +2292,17 @@ namespace impl
 			CustomizableObjectInstancePrivateData->TexturesToRelease.Empty();
 		}
 
-		{
-			double delta = FPlatformTime::Seconds() - CustomizableObjectSystemPrivateData->CurrentMutableOperation->StartUpdateTime;
-			UE_LOG(LogMutable, Log, TEXT("Finished update in %.3f ms."), delta*1000.0);
-		}
-
-		LogBenchmarkUtil::UpdateBuildTimeStats(CustomizableObjectSystemPrivateData->MutableStats, 
-												CustomizableObjectSystemPrivateData->CurrentMutableOperation->StartUpdateTime);
-
 		// End Update
 		System->ClearCurrentMutableOperation();
 
-		TRACE_END_REGION(UE_MUTABLE_UPDATE_REGION);
+		FinishUpdateGlobal(OperationData);
 	}
 
 
 	void Task_Game_ConvertResources(const TSharedRef<FUpdateContextPrivate>& OperationData)
 	{
 		MUTABLE_CPUPROFILER_SCOPE(Task_Game_ConvertResources)
+		FScopeTimer Timer(OperationData->TaskConvertResourcesTime);
 
 		check(IsInGameThread());
 
@@ -2475,6 +2450,7 @@ namespace impl
 	void Task_Game_LockCache(const TSharedRef<FUpdateContextPrivate>& OperationData)
 	{
 		MUTABLE_CPUPROFILER_SCOPE(Task_Game_LockCache)
+		FScopeTimer Timer(OperationData->TaskLockCacheTime);
 
 		check(IsInGameThread());
 
@@ -2636,6 +2612,8 @@ namespace impl
 	{
 		MUTABLE_CPUPROFILER_SCOPE(Task_Game_StartUpdate)
 
+		Operation->StartUpdateTime = FPlatformTime::Seconds();
+
 		UCustomizableObjectSystem::GetInstance()->GetPrivate()->MutableTaskGraph.AllowLaunchingMutableTaskLowPriority(false, false);
 		
 		UCustomizableObjectSystem* System = UCustomizableObjectSystem::GetInstance();
@@ -2676,6 +2654,7 @@ namespace impl
 
 			Operation->UpdateResult = EUpdateResult::ErrorOptimized;
 			UpdateSkeletalMesh(Operation);
+			FinishUpdateGlobal(Operation);
 			return;
 		}
 
@@ -2721,11 +2700,6 @@ namespace impl
 			return;
 		}
 
-		if (LogBenchmarkUtil::IsLoggingActive())
-		{
-			Operation->StartUpdateTime = FPlatformTime::Seconds();
-		}
-
 		FCustomizableObjectSystemPrivate* SystemPrivateData = System->GetPrivateChecked();
 
 		SystemPrivateData->CurrentInstanceBeingUpdated = CandidateInstance;
@@ -2957,6 +2931,8 @@ bool UCustomizableObjectSystem::Tick(float DeltaTime)
 	
 	TickPendingReleaseMaterials();
 
+	Private->UpdateStats();
+	
 	// Get a new operation if we aren't working on one
 	if (!Private->CurrentMutableOperation)
 	{
@@ -3085,7 +3061,6 @@ bool UCustomizableObjectSystem::Tick(float DeltaTime)
 					LODUpdateWithSameInstance->ApplyLODUpdateParamsToInstance(&PendingInstanceUpdateFound->Context.Get());
 				}
 
-				TRACE_BEGIN_REGION(UE_MUTABLE_UPDATE_REGION);
 				Private->StartUpdateSkeletalMesh(PendingInstanceUpdateFound->Context);				
 				Private->MutablePendingInstanceWork.RemoveUpdate(PendingInstanceUpdateFound->Context->Instance);
 			}
@@ -3094,7 +3069,6 @@ bool UCustomizableObjectSystem::Tick(float DeltaTime)
 				// Commit the LOD changes
 				LODUpdateCandidateFound->ApplyLODUpdateParamsToInstance();
 
-				TRACE_BEGIN_REGION(UE_MUTABLE_UPDATE_REGION);
 				const TSharedRef<FUpdateContextPrivate> Context = MakeShared<FUpdateContextPrivate>(*LODUpdateCandidateFound->CustomizableObjectInstance);
 				Private->StartUpdateSkeletalMesh(Context);
 			}
@@ -3125,15 +3099,14 @@ bool UCustomizableObjectSystem::Tick(float DeltaTime)
 		AdvanceCurrentOperation();
 	}
 
-	Private->MutableStats.MutablePendingInstanceWorkCount = Private->MutablePendingInstanceWork.Num();
-	LogBenchmarkUtil::UpdateStats(Private->MutableStats, ProtectedCachedTextures);
-
 #if WITH_EDITOR
 	TickRecompileCustomizableObjects();
 #endif
 
 	Private->MutableTaskGraph.Tick();
 
+	Private->LogBenchmarkUtil.UpdateStats(); // Must to be the last thing to perform
+	
 	return true;
 }
 
@@ -3296,27 +3269,47 @@ void UnCacheTexturesParameters(const TArray<FName>& TextureParameters)
 
 int32 UCustomizableObjectSystem::GetNumInstances() const
 {
-	return GetPrivateChecked()->MutableStats.NumInstances;
+	int32 NumInstances;
+	int32 NumBuiltInstances;
+	int32 NumInstancesLOD0;
+	int32 NumInstancesLOD1;
+	int32 NumInstancesLOD2;
+	int32 NumAllocatedSkeletalMeshes;
+	GetPrivateChecked()->LogBenchmarkUtil.GetInstancesStats(NumInstances, NumBuiltInstances, NumInstancesLOD0, NumInstancesLOD1, NumInstancesLOD2, NumAllocatedSkeletalMeshes);
+
+	return NumBuiltInstances;
 }
 
 int32 UCustomizableObjectSystem::GetNumPendingInstances() const
 {
-	return GetPrivateChecked()->MutableStats.NumPendingInstances;
+	return GetPrivateChecked()->MutablePendingInstanceWork.Num();
 }
 
 int32 UCustomizableObjectSystem::GetTotalInstances() const
 {
-	return GetPrivateChecked()->MutableStats.TotalInstances;
+	int32 NumInstances = 0;
+	
+	for (TObjectIterator<UCustomizableObjectInstance> Instance; Instance; ++Instance)
+	{
+		if (!Instance ||
+			Instance->HasAnyFlags(RF_ClassDefaultObject))
+		{
+			continue;
+		}
+				
+		++NumInstances;
+	}
+	return NumInstances;
 }
 
 int32 UCustomizableObjectSystem::GetTextureMemoryUsed() const
 {
-	return static_cast<int32>(GetPrivateChecked()->MutableStats.TextureMemoryUsed);
+	return 0; // Currently disabled since the reported value was incorrect due to MIP streaming.
 }
 
 int32 UCustomizableObjectSystem::GetAverageBuildTime() const
 {
-	return GetPrivateChecked()->MutableStats.TotalBuiltInstances == 0 ? 0 : Private->MutableStats.TotalBuildMs / Private->MutableStats.TotalBuiltInstances;
+	return GetPrivateChecked()->LogBenchmarkUtil.InstanceBuildTimeAvrg.GetValue() * 1000;
 }
 
 
@@ -3414,13 +3407,13 @@ void UCustomizableObjectSystem::AddUncompiledCOWarning(const UCustomizableObject
 
 void UCustomizableObjectSystem::EnableBenchmark()
 {
-	LogBenchmarkUtil::StartLogging();
+	GetPrivate()->LogBenchmarkUtil.SetEnable(true);
 }
 
 
 void UCustomizableObjectSystem::EndBenchmark()
 {
-	LogBenchmarkUtil::ShutdownAndSaveResults();
+	GetPrivate()->LogBenchmarkUtil.SetEnable(false);
 }
 
 
@@ -3671,6 +3664,9 @@ FUnrealMutableImageProvider* FCustomizableObjectSystemPrivate::GetImageProviderC
 
 void FCustomizableObjectSystemPrivate::StartUpdateSkeletalMesh(const TSharedRef<FUpdateContextPrivate>& Context)
 {
+	Context->UpdateStarted = true;
+	TRACE_BEGIN_REGION(UE_MUTABLE_UPDATE_REGION);
+
 	check(!CurrentMutableOperation); // Can not start an update if there is already another in progress
 	check(Context->Instance.IsValid()) // The instance has to be alive to start the update
 		
@@ -3694,6 +3690,22 @@ bool FCustomizableObjectSystemPrivate::IsUpdating(const UCustomizableObjectInsta
 }
 
 
+void FCustomizableObjectSystemPrivate::UpdateStats()
+{
+	NumSkeletalMeshes = 0;
+	
+	for (TObjectIterator<UCustomizableObjectInstance> Instance; Instance; ++Instance)
+	{
+		if (!IsValidChecked(*Instance))
+		{
+			continue;
+		}
+
+		NumSkeletalMeshes += Instance->SkeletalMeshes.Num();
+	}
+}
+
+
 bool UCustomizableObjectSystem::IsMutableAnimInfoDebuggingEnabled() const
 {
 #if WITH_EDITOR
diff --git a/Engine/Plugins/Experimental/Mutable/Source/CustomizableObject/Private/MuCO/CustomizableObjectSystemPrivate.h b/Engine/Plugins/Experimental/Mutable/Source/CustomizableObject/Private/MuCO/CustomizableObjectSystemPrivate.h
--- a/Engine/Plugins/Experimental/Mutable/Source/CustomizableObject/Private/MuCO/CustomizableObjectSystemPrivate.h
+++ b/Engine/Plugins/Experimental/Mutable/Source/CustomizableObject/Private/MuCO/CustomizableObjectSystemPrivate.h
@@ -2,6 +2,7 @@
 
 #pragma once
 
+#include "LogBenchmarkUtil.h"
 #include "Containers/Queue.h"
 #include "MuCO/CustomizableObject.h"
 #include "MuCO/CustomizableObjectInstance.h"
@@ -464,13 +465,22 @@ public:
 
 	mu::FImageOperator::FImagePixelFormatFunc PixelFormatOverride;
 
+	bool UpdateStarted = false;
+
 	// Update stats
 	double StartUpdateTime = 0.0;
+	double UpdateTime = 0.0;
+
+	double TaskGetMeshTime = 0.0;
+	double TaskLockCacheTime = 0.0;
+	double TaskGetImagesTime = 0.0;
+	double TaskConvertResourcesTime = 0.0f;
+	double TaskCallbacksTime = 0.0;
 	
 #if WITH_EDITOR
 	/** Used for profiling in the editor. */
 	uint32 MutableRuntimeCycles = 0;
-#endif 
+#endif
 };
 
 
@@ -481,36 +491,6 @@ struct FMutableReleasePlatformOperationData
 };
 
 
-struct FMutableStats
-{
-	// Stores the number of pending instance updates, LOD Updates, discards and releases last tick
-	int32 MutablePendingInstanceWorkCount = 0;
-
-	// Total time spent building instances
-	int32 TotalBuildMs = 0;
-
-	// Total number of instances built or updated.
-	int32 TotalBuiltInstances = 0;
-
-	// Number of instances alive (built)
-	int32 NumInstances = 0;
-
-	// Number of CustomizableObjectInstances waiting to be updated.
-	int32 NumPendingInstances = 0;
-
-	// Total number of CustomizableObjectInstances, including not built.
-	int32 TotalInstances = 0;
-
-	// Total memory in bytes used for generated textures
-	int64_t TextureMemoryUsed = 0;
-
-	uint32 CountAllocatedSkeletalMesh = 0;
-
-	// \TODO: Remove this array if we are not gathering stats!
-	TArray<TWeakObjectPtr<UTexture2D>> TextureTrackerArray;
-};
-
-
 class FCustomizableObjectSystemPrivate : public FGCObject
 {
 public:
@@ -599,9 +579,8 @@ public:
 	bool IsReplaceDiscardedWithReferenceMeshEnabled() const;
 	void SetReplaceDiscardedWithReferenceMeshEnabled(bool bIsEnabled);
 
-	int32 GetCountAllocatedSkeletalMesh() const;
-
-	mutable FMutableStats MutableStats;
+	/** Updated at the beginning of each tick. */
+	int32 GetNumSkeletalMeshes() const;
 
 	bool bReplaceDiscardedWithReferenceMesh = false;
 	bool bReleaseTexturesImmediately = false;
@@ -610,8 +589,6 @@ public:
 
 	static FCustomizableObjectCompilerBase* (*NewCompilerFunc)();
 
-	void CreatedTexture(UTexture2D* Texture) const;
-
 	TMap<FMutableImageCacheKey, uint32> TextureReferenceCount; // Keeps a count of texture usage to decide if they have to be blocked from GC during an update
 
 	// This is protected from GC by AddReferencedObjects
@@ -635,6 +612,10 @@ public:
 
 	/** See UCustomizableObjectInstance::IsUpdating. */
 	bool IsUpdating(const UCustomizableObjectInstance& Instance) const;
+
+	/** Update stats at each tick.
+	 * Used for stats that are costly to update. */
+	void UpdateStats();
 	
 	/** Mutable TaskGraph system (Mutable Thread). */
 	FMutableTaskGraph MutableTaskGraph;
@@ -643,5 +624,9 @@ public:
 	/** Mutable default image provider. Used by the COIEditor and Instance/Descriptor APIs. */
 	TObjectPtr<UEditorImageProvider> EditorImageProvider = nullptr;
 #endif
+
+	FLogBenchmarkUtil LogBenchmarkUtil;
+
+	int32 NumSkeletalMeshes = 0;
 };
 
diff --git a/Engine/Plugins/Experimental/Mutable/Source/CustomizableObject/Private/MuCO/LogBenchmarkUtil.cpp b/Engine/Plugins/Experimental/Mutable/Source/CustomizableObject/Private/MuCO/LogBenchmarkUtil.cpp
--- a/Engine/Plugins/Experimental/Mutable/Source/CustomizableObject/Private/MuCO/LogBenchmarkUtil.cpp
+++ b/Engine/Plugins/Experimental/Mutable/Source/CustomizableObject/Private/MuCO/LogBenchmarkUtil.cpp
@@ -23,371 +23,284 @@
 
 extern ENGINE_API UEngine* GEngine;
 
-#if WITH_EDITOR
-extern UNREALED_API UEditorEngine* GEditor;
-#endif
 
-DEFINE_STAT(STAT_MutableNumSkeletalMeshes);
-DEFINE_STAT(STAT_MutableNumCachedSkeletalMeshes);
-DEFINE_STAT(STAT_MutableNumAllocatedSkeletalMeshes);
-DEFINE_STAT(STAT_MutableNumInstancesLOD0);
-DEFINE_STAT(STAT_MutableNumInstancesLOD1);
-DEFINE_STAT(STAT_MutableNumInstancesLOD2);
-DEFINE_STAT(STAT_MutableNumTextures);
-DEFINE_STAT(STAT_MutableNumCachedTextures);
-DEFINE_STAT(STAT_MutableNumAllocatedTextures);
-DEFINE_STAT(STAT_MutableTextureResourceMemory);
-DEFINE_STAT(STAT_MutableTextureGeneratedMemory);
-DEFINE_STAT(STAT_MutableTextureCacheMemory);
-DEFINE_STAT(STAT_MutablePendingInstanceUpdates);
-DEFINE_STAT(STAT_MutableAbandonedInstanceUpdates);
-DEFINE_STAT(STAT_MutableInstanceBuildTime);
-DEFINE_STAT(STAT_MutableInstanceBuildTimeAvrg);
-
-bool LogBenchmarkUtil::bLoggingActive = false;
-
-
-void LogBenchmarkUtil::StartLogging()
+static TAutoConsoleVariable<bool> CVarEnableBenchmark(
+	TEXT("mutable.EnableBenchmark"),
+	false,
+	TEXT("Enable or disable the benchmarking."));
+
+
+void Write(FArchive& Archive, FStringView Text)
 {
-	bLoggingActive = true;
+	Archive.Serialize(const_cast<ANSICHAR*>(StringCast<ANSICHAR>(Text.GetData(), Text.Len()).Get()), Text.Len() * sizeof(ANSICHAR));
 }
 
-void LogBenchmarkUtil::ShutdownAndSaveResults()
+
+TSharedPtr<FArchive> CreateFile()
 {
-	if (bLoggingActive)
-	{
-		FString SaveDirectory = FPaths::ConvertRelativePathToFull(FPaths::ProjectSavedDir()) + "/Logs";
-		FString FileName = FString("BenchmarkResult.txt");
-		FString AbsoluteFilePath = SaveDirectory + "/" + FileName;
-		IPlatformFile& PlatformFile = FPlatformFileManager::Get().GetPlatformFile();
+	const FString Directory = FPaths::ProfilingDir() + TEXT("Mutable/Benchmark/");
+	IFileManager::Get().MakeDirectory(*Directory, true);
 
-		if (PlatformFile.CreateDirectoryTree(*SaveDirectory))
-		{
-			FString CurrentTime = FDateTime::Now().ToString();
-			FString DateStamp = LINE_TERMINATOR + FString("[" + CurrentTime + "]LogBenchmark: (string) ") + FString("benchmark_results_location : ") + SaveDirectory;
-			FFileHelper::SaveStringToFile(DateStamp, *AbsoluteFilePath, FFileHelper::EEncodingOptions::ForceAnsi);
+	const FDateTime FileDate = FDateTime::Now();
+	const FString Filename = FString::Printf(TEXT("%s/%s.csv"), *Directory,  *FileDate.ToString());
 
-			FString KeyName;
-			int32 MaxValue;
-			float AverageValue;
+	TSharedPtr<FArchive> Archive = MakeShareable(IFileManager::Get().CreateFileWriter(*Filename, FILEWRITE_AllowRead | FILEWRITE_NoFail));
+	check(Archive);
 
-			// Iterate over the integer stats and log them to a file
-			for (TPair<FName, IntStat> Ist : GetInstance().IntStats)
-			{
-				CurrentTime = FDateTime::Now().ToString();
-
-				KeyName = FString(Ist.Key.ToString());
-				MaxValue = Ist.Value.Max;
-				AverageValue = (long double)Ist.Value.Total / (long double)Ist.Value.TotalCalls;
-
-				FString TextToSave = LINE_TERMINATOR;
-				TextToSave += FString("[" + CurrentTime + "]");
-				TextToSave += FString("LogBenchmark: (int) peak_");
-				TextToSave += FString::Printf(TEXT("%s : %lld"), *KeyName, MaxValue);
-				TextToSave += LINE_TERMINATOR;
-				TextToSave += FString("[" + CurrentTime + "]");
-				TextToSave += FString("LogBenchmark: (float) average_");
-				TextToSave += FString::Printf(TEXT("%s : %f"), *KeyName, AverageValue);
-
-				FFileHelper::SaveStringToFile(TextToSave, *AbsoluteFilePath, FFileHelper::EEncodingOptions::ForceAnsi, &IFileManager::Get(), FILEWRITE_Append);
-			}
+	Write(*Archive, TEXT("ID_CO,ID_COI,ID_UpdateType,Time_Update,Time_TaskGetMesh,Time_TaskLockCache,Time_TaskGetImages,Time_TaskConvertResources,Time_TaskCallbacks\n"));		
 
-			// Iterate over the Float stats and log them to a file
-			for (TPair<FName, FloatStat> dst : GetInstance().FloatStats)
-			{
-				CurrentTime = FDateTime::Now().ToString();
-
-				KeyName = FString(dst.Key.ToString());
-				MaxValue = dst.Value.Max;
-				AverageValue = dst.Value.Total / dst.Value.TotalCalls;
-
-				FString TextToSave = LINE_TERMINATOR;
-				TextToSave += FString("[" + CurrentTime + "]");
-				TextToSave += FString("LogBenchmark: (int) peak_");
-				TextToSave += FString::Printf(TEXT("%s : %lld"), *KeyName, MaxValue);
-				TextToSave += LINE_TERMINATOR;
-				TextToSave += FString("[" + CurrentTime + "]");
-				TextToSave += FString("LogBenchmark: (float) average_");
-				TextToSave += FString::Printf(TEXT("%s : %f"), *KeyName, AverageValue);
-
-				FFileHelper::SaveStringToFile(TextToSave, *AbsoluteFilePath, FFileHelper::EEncodingOptions::ForceAnsi, &IFileManager::Get(), FILEWRITE_Append);
-			}
-		}
-
-		bLoggingActive = false;
-	}
+	return Archive;
 }
 
-void LogBenchmarkUtil::UpdateStat(const FName& Stat, int32 Value)
+
+FLogBenchmarkUtil::~FLogBenchmarkUtil()
 {
-	if (bLoggingActive)
+	if (Archive)
 	{
-		if (GetInstance().IntStats.Contains(Stat))
-		{
-			GetInstance().IntStats[Stat] += Value;
-		}
-		else
-		{
-			GetInstance().IntStats.Add(Stat) += Value;
-		}
+		Archive->Close();
 	}
 }
 
-void LogBenchmarkUtil::UpdateStat(const FName & Stat, double Value)
+
+void FLogBenchmarkUtil::SetEnable(bool bEnabled)
 {
-	UpdateStat(Stat, (long double)Value);
+	CVarEnableBenchmark->Set(bEnabled);
 }
 
-void LogBenchmarkUtil::UpdateStat(const FName& Stat, long double Value)
+
+void FLogBenchmarkUtil::GetInstancesStats(int32& OutNumInstances, int32& OutNumBuiltInstances, int32& OutNumInstancesLOD0, int32& OutNumInstancesLOD1, int32& OutNumInstancesLOD2, int32& OutNumAllocatedSkeletalMeshes) const
 {
-	if (bLoggingActive)
+	OutNumInstances = 0;
+	OutNumBuiltInstances = 0;
+	OutNumInstancesLOD0 = 0;
+	OutNumInstancesLOD1 = 0;
+	OutNumInstancesLOD2 = 0;
+	OutNumAllocatedSkeletalMeshes = 0;
+	
+	for (TObjectIterator<UCustomizableObjectInstance> Instance; Instance; ++Instance)
 	{
-		if (GetInstance().FloatStats.Contains(Stat))
+		if (!IsValidChecked(*Instance) ||
+			Instance->HasAnyFlags(RF_ClassDefaultObject))
 		{
-			GetInstance().FloatStats[Stat] += Value;
+			continue;
 		}
-		else
+		
+		++OutNumInstances;
+
+		OutNumInstancesLOD0 += Instance->GetCurrentMinLOD() <= 0 && Instance->GetCurrentMaxLOD() >= 0;
+		OutNumInstancesLOD1 += Instance->GetCurrentMinLOD() <= 1 && Instance->GetCurrentMaxLOD() >= 1;
+		OutNumInstancesLOD2 += Instance->GetCurrentMinLOD() >= 2;
+
+		if (Instance->GetPrivate()->GetSkeletalMeshStatus() == ESkeletalMeshStatus::Success)
 		{
-			GetInstance().FloatStats.Add(Stat) += Value;
+			++OutNumBuiltInstances;
+		}
+		
+		for (int32 ComponentIndex = 0; ComponentIndex < Instance->SkeletalMeshes.Num(); ++ComponentIndex)
+		{
+			if (Instance->SkeletalMeshes[ComponentIndex] && Instance->SkeletalMeshes[ComponentIndex]->GetResourceForRendering())
+			{
+				++OutNumAllocatedSkeletalMeshes;
+			}
 		}
 	}
 }
 
-bool LogBenchmarkUtil::IsLoggingActive()
+
+void FLogBenchmarkUtil::AddTexture(UTexture2D& Texture)
 {
-	return bLoggingActive;
+	check(IsInGameThread());
+
+	if (!CVarEnableBenchmark.GetValueOnGameThread())
+	{
+		return;
+	}
+
+	TextureTrackerArray.Add(&Texture);
 }
 
 
-void LogBenchmarkUtil::UpdateStats(FMutableStats& StatsToUpdate, const TArray< TObjectPtr<UTexture2D> >& ProtectedCachedTextures)
+void FLogBenchmarkUtil::UpdateStats()
 {
-	if (LogBenchmarkUtil::IsLoggingActive())
+	check(IsInGameThread());
+	
+	if (!CVarEnableBenchmark.GetValueOnGameThread())
 	{
-		uint64 SizeCache = 0;
+		return;
+	}
 
-		for (const UTexture2D* CachedTextures : ProtectedCachedTextures)
+	const UCustomizableObjectSystem* System = UCustomizableObjectSystem::GetInstanceChecked();
+	
+	// Instances and Skeletal Mesh
+	int32 LocalNumInstances = 0;
+	int32 LocalNumBuiltInstances = 0;
+	int32 LocalNumInstancesLOD0 = 0;
+	int32 LocalNumInstancesLOD1 = 0;
+	int32 LocalNumInstancesLOD2 = 0;
+	int32 LocalNumAllocatedSkeletalMeshes = 0;
+
+	GetInstancesStats(LocalNumInstances, LocalNumBuiltInstances, LocalNumInstancesLOD0, LocalNumInstancesLOD1, LocalNumInstancesLOD2, LocalNumAllocatedSkeletalMeshes);
+
+	NumInstances = LocalNumInstances;
+	NumBuiltInstances = LocalNumBuiltInstances;
+	NumInstancesLOD0 = LocalNumInstancesLOD0;
+	NumInstancesLOD1 = LocalNumInstancesLOD1;
+	NumInstancesLOD2 = LocalNumInstancesLOD2;
+	NumAllocatedSkeletalMeshes = LocalNumAllocatedSkeletalMeshes;
+
+	NumPendingInstanceUpdates = System->GetPrivate()->MutablePendingInstanceWork.Num();
+
+	// Textures
+	uint32 LocalNumAllocatedTextures = 0;
+	for (auto Iterator = TextureTrackerArray.CreateIterator(); Iterator; ++Iterator)
+	{
+		if (Iterator->IsStale())
 		{
-			if (CachedTextures)
-			{
-				SizeCache += CachedTextures->CalcTextureMemorySizeEnum(TMC_AllMips);
-			}
+			Iterator.RemoveCurrent();
+		}
+		else if (Iterator->IsValid() && (*Iterator)->GetResource())
+		{
+			++LocalNumAllocatedTextures;
 		}
+	}
+
+	NumAllocatedTextures = LocalNumAllocatedTextures;
 
-		SET_DWORD_STAT(STAT_MutableTextureCacheMemory, SizeCache / 1024.f);
+#if WITH_EDITORONLY_DATA
+	if (UCustomizableObjectSystem::GetInstance()->IsMutableAnimInfoDebuggingEnabled())
+	{
+		return;
 	}
 
-#if !(UE_BUILD_SHIPPING || UE_BUILD_TEST)
-	bool bLogEnabled = true;
-#else
-	bool bLogEnabled = LogBenchmarkUtil::IsLoggingActive();
-#endif
-	if (bLogEnabled)
+	if (GEngine)
 	{
-		const int32& MutablePendingInstanceWorkCount = StatsToUpdate.MutablePendingInstanceWorkCount;
-		StatsToUpdate.CountAllocatedSkeletalMesh = 0;
+		return;
+	}
 
-		int32 CountLOD0 = 0;
-		int32 CountLOD1 = 0;
-		int32 CountLOD2 = 0;
-		int32 CountTotal = 0;
+	bool bFoundPlayer = false;
+	int32 MsgIndex = 15820; // Arbitrary big value to prevent collisions with other on-screen messages
 
-		for (TObjectIterator<UCustomizableObjectInstance> CustomizableObjectInstance; CustomizableObjectInstance; ++CustomizableObjectInstance)
+	for (TObjectIterator<UCustomizableObjectInstanceUsage> CustomizableObjectInstanceUsage; CustomizableObjectInstanceUsage; ++CustomizableObjectInstanceUsage)
+	{
+		if (CustomizableObjectInstanceUsage && CustomizableObjectInstanceUsage->IsNetMode(NM_DedicatedServer))
 		{
-			if (IsValidChecked(*CustomizableObjectInstance) && CustomizableObjectInstance->GetPrivate())
-			{
-				++CountTotal;
-
-				for (int32 ComponentIndex = 0; ComponentIndex < CustomizableObjectInstance->SkeletalMeshes.Num(); ++ComponentIndex)
-				{
-					if (CustomizableObjectInstance->SkeletalMeshes[ComponentIndex] && CustomizableObjectInstance->SkeletalMeshes[ComponentIndex]->GetResourceForRendering())
-					{
-						StatsToUpdate.CountAllocatedSkeletalMesh++;
-
-						if (CustomizableObjectInstance->GetCurrentMinLOD() < 1)
-						{
-							++CountLOD0;
-						}
-						else if (CustomizableObjectInstance->GetCurrentMaxLOD() < 2)
-						{
-							++CountLOD1;
-						}
-						else
-						{
-							++CountLOD2;
-						}
-					}
-				}
-			}
+			continue;
 		}
 
-		StatsToUpdate.NumInstances = CountLOD0 + CountLOD1 + CountLOD2;
-		StatsToUpdate.TotalInstances = CountTotal;
-		StatsToUpdate.NumPendingInstances = MutablePendingInstanceWorkCount;
-		SET_DWORD_STAT(STAT_MutablePendingInstanceUpdates, MutablePendingInstanceWorkCount);
+		AActor* ParentActor = CustomizableObjectInstanceUsage->GetAttachParent() ? 
+			CustomizableObjectInstanceUsage->GetAttachParent()->GetAttachmentRootActor()
+			:  nullptr;
+		UCustomizableObjectInstance* Instance = CustomizableObjectInstanceUsage->GetCustomizableObjectInstance();
 
-		uint64 Size = 0;
-		uint32 CountAllocated = 0;
-		for (TWeakObjectPtr<UTexture2D>& Tracker : StatsToUpdate.TextureTrackerArray)
+		APawn* PlayerPawn = nullptr;
+		if (UWorld* World = CustomizableObjectInstanceUsage->GetWorld())
 		{
-			if (Tracker.IsValid() && Tracker->GetResource())
-			{
-				CountAllocated++;
+			PlayerPawn = UGameplayStatics::GetPlayerPawn(World, 0);
+		}
 
-				if (Tracker->GetResource()->TextureRHI)
-				{
-					Size += Tracker->CalcTextureMemorySizeEnum(TMC_AllMips);
-				}
+		if (ParentActor && (ParentActor == PlayerPawn) && Instance)
+		{
+			bFoundPlayer = true;
+
+			FString TagString;
+			const FGameplayTagContainer& Tags = Instance->GetAnimationGameplayTags();
+
+			for (const FGameplayTag& Tag : Tags)
+			{
+				TagString += !TagString.IsEmpty() ? FString(TEXT(", ")) : FString();
+				TagString += Tag.ToString();
 			}
-		}
 
-		StatsToUpdate.TextureMemoryUsed = int64_t(Size / 1024);
+			GEngine->AddOnScreenDebugMessage(MsgIndex++, .0f, FColor::Green, TEXT("Animation tags: ") + TagString);
 
-		uint64 SizeGenerated = 0;
-		for (TObjectIterator<UCustomizableObjectInstance> CustomizableObjectInstance; CustomizableObjectInstance; ++CustomizableObjectInstance)
-		{
-			if (IsValidChecked(*CustomizableObjectInstance) && CustomizableObjectInstance->GetPrivate() && CustomizableObjectInstance->HasAnySkeletalMesh())
+			check(Instance->GetPrivate() != nullptr);
+			FCustomizableInstanceComponentData* ComponentData = Instance->GetPrivate()->GetComponentData(CustomizableObjectInstanceUsage->GetComponentIndex());
+
+			if (ComponentData)
 			{
-				bool bHasResourceForRendering = false;
-				for (int32 MeshIndex = 0; !bHasResourceForRendering && MeshIndex < CustomizableObjectInstance->SkeletalMeshes.Num(); ++MeshIndex)
+				for (TPair<FName, TSoftClassPtr<UAnimInstance>>& Entry : ComponentData->AnimSlotToBP)
 				{
-					bHasResourceForRendering = CustomizableObjectInstance->GetSkeletalMesh(MeshIndex) && CustomizableObjectInstance->GetSkeletalMesh(MeshIndex)->GetResourceForRendering();
-				}
+					FString AnimBPSlot;
 
-				if (bHasResourceForRendering)
-				{
-					for (const FGeneratedTexture& GeneratedTextures : CustomizableObjectInstance->GetPrivate()->GeneratedTextures)
-					{
-						if (GeneratedTextures.Texture)
-						{
-							check(GeneratedTextures.Texture != nullptr);
-							SizeGenerated += GeneratedTextures.Texture->CalcTextureMemorySizeEnum(TMC_AllMips);
-						}
-					}
+					AnimBPSlot += Entry.Key.ToString() + FString("-") + Entry.Value.GetAssetName();
+					GEngine->AddOnScreenDebugMessage(MsgIndex++, .0f, FColor::Green, AnimBPSlot);
 				}
 			}
-		}
 
-		if (LogBenchmarkUtil::IsLoggingActive())
-		{
-			LogBenchmarkUtil::UpdateStat("customizable_objects", (int32)StatsToUpdate.CountAllocatedSkeletalMesh);
-			LogBenchmarkUtil::UpdateStat("pending_instance_updates", MutablePendingInstanceWorkCount);
-			LogBenchmarkUtil::UpdateStat("allocated_textures", (int32)CountAllocated);
-			LogBenchmarkUtil::UpdateStat("texture_resource_memory", (long double)Size / 1048576.0L);
-			LogBenchmarkUtil::UpdateStat("texture_generated_memory", (long double)SizeGenerated / 1048576.0L);
-		}
-
-#if !(UE_BUILD_SHIPPING || UE_BUILD_TEST)
-		SET_DWORD_STAT(STAT_MutableNumInstancesLOD0, CountLOD0);
-		SET_DWORD_STAT(STAT_MutableNumInstancesLOD1, CountLOD1);
-		SET_DWORD_STAT(STAT_MutableNumInstancesLOD2, CountLOD2);
-		SET_DWORD_STAT(STAT_MutableNumAllocatedSkeletalMeshes, StatsToUpdate.CountAllocatedSkeletalMesh);
-		SET_DWORD_STAT(STAT_MutablePendingInstanceUpdates, MutablePendingInstanceWorkCount);
-		SET_DWORD_STAT(STAT_MutableNumAllocatedTextures, CountAllocated);
-		SET_DWORD_STAT(STAT_MutableTextureResourceMemory, Size / 1024.f);
-		SET_DWORD_STAT(STAT_MutableTextureGeneratedMemory, SizeGenerated / 1024.f);
-#endif
+			GEngine->AddOnScreenDebugMessage(MsgIndex++, .0f, FColor::Green, TEXT("Slots-AnimBP: "));
 
-#if WITH_EDITORONLY_DATA
-		if (UCustomizableObjectSystem::GetInstance()->IsMutableAnimInfoDebuggingEnabled())
-		{
-			if (GEngine)
+			if (ComponentData)
 			{
-				bool bFoundPlayer = false;
-				int32 MsgIndex = 15820; // Arbitrary big value to prevent collisions with other on-screen messages
-
-				for (TObjectIterator<UCustomizableObjectInstanceUsage> CustomizableObjectInstanceUsage; CustomizableObjectInstanceUsage; ++CustomizableObjectInstanceUsage)
+				if (ComponentData->MeshPartPaths.IsEmpty())
 				{
-					if (CustomizableObjectInstanceUsage && CustomizableObjectInstanceUsage->IsNetMode(NM_DedicatedServer))
-					{
-						continue;
-					}
-
-					AActor* ParentActor = CustomizableObjectInstanceUsage->GetAttachParent() ? 
-						CustomizableObjectInstanceUsage->GetAttachParent()->GetAttachmentRootActor() 
-						:  nullptr;
-					UCustomizableObjectInstance* Instance = CustomizableObjectInstanceUsage->GetCustomizableObjectInstance();
-
-					APawn* PlayerPawn = nullptr;
-					if (UWorld* World = CustomizableObjectInstanceUsage->GetWorld())
-					{
-						PlayerPawn = UGameplayStatics::GetPlayerPawn(World, 0);
-					}
-
-					if (ParentActor && (ParentActor == PlayerPawn) && Instance)
-					{
-						bFoundPlayer = true;
-
-						FString TagString;
-						const FGameplayTagContainer& Tags = Instance->GetAnimationGameplayTags();
-
-						for (const FGameplayTag& Tag : Tags)
-						{
-							TagString += !TagString.IsEmpty() ? FString(TEXT(", ")) : FString();
-							TagString += Tag.ToString();
-						}
-
-						GEngine->AddOnScreenDebugMessage(MsgIndex++, .0f, FColor::Green, TEXT("Animation tags: ") + TagString);
-
-						check(Instance->GetPrivate() != nullptr);
-						FCustomizableInstanceComponentData* ComponentData = Instance->GetPrivate()->GetComponentData(CustomizableObjectInstanceUsage->GetComponentIndex());
-
-						if (ComponentData)
-						{
-							for (TPair<FName, TSoftClassPtr<UAnimInstance>>& Entry : ComponentData->AnimSlotToBP)
-							{
-								FString AnimBPSlot;
-
-								AnimBPSlot += Entry.Key.ToString() + FString("-") + Entry.Value.GetAssetName();
-								GEngine->AddOnScreenDebugMessage(MsgIndex++, .0f, FColor::Green, AnimBPSlot);
-							}
-						}
-
-						GEngine->AddOnScreenDebugMessage(MsgIndex++, .0f, FColor::Green, TEXT("Slots-AnimBP: "));
-
-						if (ComponentData)
-						{
-							if (ComponentData->MeshPartPaths.IsEmpty())
-							{
-								GEngine->AddOnScreenDebugMessage(MsgIndex++, .0f, FColor::Magenta,
-									TEXT("No meshes found. In order to see the meshes compile the pawn's root CustomizableObject after the 'mutable.EnableMutableAnimInfoDebugging 1' command has been run."));
-							}
-
-							for (const FString& MeshPath : ComponentData->MeshPartPaths)
-							{
-								GEngine->AddOnScreenDebugMessage(MsgIndex++, .0f, FColor::Magenta, MeshPath);
-							}
-						}
-
-						GEngine->AddOnScreenDebugMessage(MsgIndex++, .0f, FColor::Magenta, TEXT("Meshes: "));
-
-						GEngine->AddOnScreenDebugMessage(MsgIndex++, .0f, FColor::Cyan,
-							TEXT("Player Pawn Mutable Mesh/Animation info for component ") + FString::Printf(TEXT("%d"),
-								CustomizableObjectInstanceUsage->GetComponentIndex()));
-					}
+					GEngine->AddOnScreenDebugMessage(MsgIndex++, .0f, FColor::Magenta,
+						TEXT("No meshes found. In order to see the meshes compile the pawn's root CustomizableObject after the 'mutable.EnableMutableAnimInfoDebugging 1' command has been run."));
 				}
 
-				if (!bFoundPlayer)
+				for (const FString& MeshPath : ComponentData->MeshPartPaths)
 				{
-					GEngine->AddOnScreenDebugMessage(MsgIndex, .0f, FColor::Yellow, TEXT("Mutable Animation info: N/A"));
+					GEngine->AddOnScreenDebugMessage(MsgIndex++, .0f, FColor::Magenta, MeshPath);
 				}
 			}
+
+			GEngine->AddOnScreenDebugMessage(MsgIndex++, .0f, FColor::Magenta, TEXT("Meshes: "));
+
+			GEngine->AddOnScreenDebugMessage(MsgIndex++, .0f, FColor::Cyan,
+				TEXT("Player Pawn Mutable Mesh/Animation info for component ") + FString::Printf(TEXT("%d"),
+					CustomizableObjectInstanceUsage->GetComponentIndex()));
 		}
-#endif
 	}
+
+	if (!bFoundPlayer)
+	{
+		GEngine->AddOnScreenDebugMessage(MsgIndex, .0f, FColor::Yellow, TEXT("Mutable Animation info: N/A"));
+	}
+#endif
 }
 
 
-void LogBenchmarkUtil::UpdateBuildTimeStats(FMutableStats& StatsToUpdate, double StartUpdateTime)
+void FLogBenchmarkUtil::FinishUpdate(const TSharedRef<FUpdateContextPrivate>& Context)
 {
-	if (LogBenchmarkUtil::IsLoggingActive())
+	check(IsInGameThread());
+
+	if (!CVarEnableBenchmark.GetValueOnGameThread())
 	{
-		double DeltaSeconds = FPlatformTime::Seconds() - StartUpdateTime;
-		int32 DeltaMs = int32(DeltaSeconds * 1000);
-
-		LogBenchmarkUtil::UpdateStat("customizable_instance_build_time", DeltaMs);
-		StatsToUpdate.TotalBuildMs += DeltaMs;
-		StatsToUpdate.TotalBuiltInstances++;
-		SET_DWORD_STAT(STAT_MutableInstanceBuildTime, DeltaMs);
-		SET_DWORD_STAT(STAT_MutableInstanceBuildTimeAvrg, StatsToUpdate.TotalBuildMs / StatsToUpdate.TotalBuiltInstances);
+		return;
 	}
+	
+	TotalUpdateTime += Context->UpdateTime;
+	++NumUpdates;
+
+	InstanceBuildTimeAvrg = TotalUpdateTime / NumUpdates;
+
+	const UCustomizableObjectInstance* Instance = Context->Instance.Get();
+	if (!Instance)
+	{
+		return;
+	}
+	
+	const UCustomizableObject* Object = Instance->GetCustomizableObject();
+	if (!Object)
+	{
+		return;
+	}
+
+	if (!Archive) // Create file if it does not exist.
+	{
+		Archive = CreateFile();
+	}
+
+	const FString ID_CO = Context->Instance->GetCustomizableObject()->GetPathName();
+	const FString ID_COI = Instance->GetPathName();
+	const FString ID_UpdateType = TEXT("Mesh");
+	const double Time_Update = Context->UpdateTime * 1000;
+	const double Time_TaskGetMesh = Context->TaskGetMeshTime * 1000;
+	const double Time_TaskLockCache = Context->TaskLockCacheTime * 1000;
+	const double Time_TaskGetImages = Context->TaskGetImagesTime * 1000;
+	const double Time_TaskConvertResources = Context->TaskConvertResourcesTime * 1000;
+	const double Time_TaskCallbacks = Context->TaskCallbacksTime * 1000;
+	
+	Write(*Archive, FString::Printf(TEXT("%s,%s,%s,%f,%f,%f,%f,%f,%f\n"), *ID_CO, *ID_COI, *ID_UpdateType, Time_Update, Time_TaskGetMesh, Time_TaskLockCache, Time_TaskGetImages, Time_TaskConvertResources, Time_TaskCallbacks));
+	Archive->Flush();
 }
+
diff --git a/Engine/Plugins/Experimental/Mutable/Source/CustomizableObject/Private/MuCO/LogBenchmarkUtil.h b/Engine/Plugins/Experimental/Mutable/Source/CustomizableObject/Private/MuCO/LogBenchmarkUtil.h
--- a/Engine/Plugins/Experimental/Mutable/Source/CustomizableObject/Private/MuCO/LogBenchmarkUtil.h
+++ b/Engine/Plugins/Experimental/Mutable/Source/CustomizableObject/Private/MuCO/LogBenchmarkUtil.h
@@ -6,84 +6,100 @@
 #include "UObject/NameTypes.h"
 #include "Stats/Stats.h"
 
+class UTexture2D;
+
+/** Stat which automatically gets updated to Insights when modified. */
+#define DECLARE_BENCHMARK_STAT(Name, Type) \
+	struct F##Name \
+	{ \
+	private: \
+		Type Value = 0; \
+		\
+		mutable FCriticalSection Lock##Name; \
+		\
+	public: \
+		Type GetValue() const \
+		{ \
+			FScopeLock ScopedLock(&Lock##Name); \
+			return Value; \
+		} \
+		\
+		F##Name& operator+=(const Type& Rhs) \
+		{ \
+			FScopeLock ScopedLock(&Lock##Name); \
+			Value += Rhs; \
+			return *this; \
+		} \
+		\
+		F##Name& operator=(const Type& Rhs) \
+		{ \
+			FScopeLock ScopedLock(&Lock##Name); \
+			Value = Rhs; \
+		\
+		SET_DWORD_STAT(STAT_Mutable##Name, Value); \
+		\
+		return *this; \
+		} \
+	}; \
+	F##Name Name;
+
+#define DECLARE_BENCHMARK_INSIGHTS(Name, Description) \
+	DECLARE_DWORD_ACCUMULATOR_STAT(Description, STAT_Mutable##Name, STATGROUP_Mutable);
+
+
+class FUpdateContextPrivate;
 // Mutable stats
 DECLARE_STATS_GROUP(TEXT("Mutable"), STATGROUP_Mutable, STATCAT_Advanced);
-DECLARE_DWORD_ACCUMULATOR_STAT_EXTERN(TEXT("Num Created Mutable Skeletal Meshes"), STAT_MutableNumSkeletalMeshes, STATGROUP_Mutable, );
-DECLARE_DWORD_ACCUMULATOR_STAT_EXTERN(TEXT("Num Cached Mutable Skeletal Meshes"), STAT_MutableNumCachedSkeletalMeshes, STATGROUP_Mutable, );
-DECLARE_DWORD_ACCUMULATOR_STAT_EXTERN(TEXT("Num Allocated Mutable Skeletal Meshes"), STAT_MutableNumAllocatedSkeletalMeshes, STATGROUP_Mutable, );
-DECLARE_DWORD_ACCUMULATOR_STAT_EXTERN(TEXT("Num Instances at LOD 0"), STAT_MutableNumInstancesLOD0, STATGROUP_Mutable, );
-DECLARE_DWORD_ACCUMULATOR_STAT_EXTERN(TEXT("Num Instances at LOD 1"), STAT_MutableNumInstancesLOD1, STATGROUP_Mutable, );
-DECLARE_DWORD_ACCUMULATOR_STAT_EXTERN(TEXT("Num Instances at LOD 2 or more"), STAT_MutableNumInstancesLOD2, STATGROUP_Mutable, );
-DECLARE_DWORD_ACCUMULATOR_STAT_EXTERN(TEXT("Num Created Mutable Textures"), STAT_MutableNumTextures, STATGROUP_Mutable, );
-DECLARE_DWORD_ACCUMULATOR_STAT_EXTERN(TEXT("Num Cached Mutable Textures"), STAT_MutableNumCachedTextures, STATGROUP_Mutable, );
-DECLARE_DWORD_ACCUMULATOR_STAT_EXTERN(TEXT("Num Allocated Mutable Textures"), STAT_MutableNumAllocatedTextures, STATGROUP_Mutable, );
-DECLARE_DWORD_ACCUMULATOR_STAT_EXTERN(TEXT("Texture Resource Memory"), STAT_MutableTextureResourceMemory, STATGROUP_Mutable, );
-DECLARE_DWORD_ACCUMULATOR_STAT_EXTERN(TEXT("Texture Generated Memory"), STAT_MutableTextureGeneratedMemory, STATGROUP_Mutable, );
-DECLARE_DWORD_ACCUMULATOR_STAT_EXTERN(TEXT("Texture Locked Memory"), STAT_MutableTextureCacheMemory, STATGROUP_Mutable, );
-DECLARE_DWORD_ACCUMULATOR_STAT_EXTERN(TEXT("Pending Instance Updates"), STAT_MutablePendingInstanceUpdates, STATGROUP_Mutable, );
-DECLARE_DWORD_ACCUMULATOR_STAT_EXTERN(TEXT("Abandoned Instance Updates"), STAT_MutableAbandonedInstanceUpdates, STATGROUP_Mutable, );
-DECLARE_DWORD_ACCUMULATOR_STAT_EXTERN(TEXT("Last Instance Build Time"), STAT_MutableInstanceBuildTime, STATGROUP_Mutable, );
-DECLARE_DWORD_ACCUMULATOR_STAT_EXTERN(TEXT("Avrg Instance Build Time"), STAT_MutableInstanceBuildTimeAvrg, STATGROUP_Mutable, );
-
-struct FMutableStats;
-
-class LogBenchmarkUtil
+
+DECLARE_BENCHMARK_INSIGHTS(NumAllocatedSkeletalMeshes, TEXT("Num Allocated Mutable Skeletal Meshes"));
+DECLARE_BENCHMARK_INSIGHTS(NumAllocatedTextures, TEXT("Num Allocated Mutable Textures"));
+DECLARE_BENCHMARK_INSIGHTS(NumInstances, TEXT("Num Instances"));
+DECLARE_BENCHMARK_INSIGHTS(NumInstancesLOD0, TEXT("Num Instances at LOD 0"));
+DECLARE_BENCHMARK_INSIGHTS(NumInstancesLOD1, TEXT("Num Instances at LOD 1"));
+DECLARE_BENCHMARK_INSIGHTS(NumInstancesLOD2, TEXT("Num Instances at LOD 2 or more"));
+DECLARE_BENCHMARK_INSIGHTS(NumPendingInstanceUpdates, TEXT("Num Pending Instance Updates"));
+DECLARE_BENCHMARK_INSIGHTS(NumBuiltInstances, TEXT("Num Built Instances"))
+DECLARE_BENCHMARK_INSIGHTS(InstanceBuildTimeAvrg, TEXT("Avrg Instance Build Time"));
+
+
+/** Benchmarking system. Gathers stats and send it to Insights an Benchmarking Files. */
+class FLogBenchmarkUtil
 {
 public:
-	static void StartLogging();
-	static void ShutdownAndSaveResults();
+	~FLogBenchmarkUtil();
+
+	/** Enable or disable the system. */
+	void SetEnable(bool bEnable);
 
-	//! Update the stats logged in unreal's stats system.
-	static void UpdateStats(FMutableStats& StatsToUpdate, const TArray< TObjectPtr<class UTexture2D> >& ProtectedCachedTextures);
+	/** Get stats. */
+	void GetInstancesStats(int32& OutNumInstances, int32& OutNumBuiltInstances, int32& OutNumInstancesLOD0, int32& OutNumInstancesLOD1, int32& OutNumInstancesLOD2, int32& OutNumAllocatedSkeletalMeshes) const;
 
-	static void UpdateBuildTimeStats(FMutableStats& StatsToUpdate, double StartUpdateTime);
+	/** Add Mutable created Texture to track. */
+	void AddTexture(UTexture2D& Texture);
 
-	static void UpdateStat(const FName& Stat, int32 Value);
-	static void UpdateStat(const FName& Stat, double Value);
-	static void UpdateStat(const FName& Stat, long double Value);
+	/** Update stats which can only be updated on the tick. */
+	void UpdateStats();
 
-	static bool IsLoggingActive();
+	/** Gathers stats update stats when it has finished. */
+	void FinishUpdate(const TSharedRef<FUpdateContextPrivate>& Context);
+	
+	DECLARE_BENCHMARK_STAT(NumAllocatedTextures, uint32);
+	DECLARE_BENCHMARK_STAT(NumAllocatedSkeletalMeshes, int32);
+
+	DECLARE_BENCHMARK_STAT(NumInstances, int32);
+	DECLARE_BENCHMARK_STAT(NumInstancesLOD0, int32);
+	DECLARE_BENCHMARK_STAT(NumInstancesLOD1, int32);
+	DECLARE_BENCHMARK_STAT(NumInstancesLOD2, int32);
+
+	DECLARE_BENCHMARK_STAT(NumPendingInstanceUpdates, int32)
+	DECLARE_BENCHMARK_STAT(NumBuiltInstances, int32)
+	DECLARE_BENCHMARK_STAT(InstanceBuildTimeAvrg, double);
+
+	TArray<TWeakObjectPtr<UTexture2D>> TextureTrackerArray;
 
 private:
-	struct IntStat
-	{
-		int32 Total;
-		int32 Max;
-		int32 TotalCalls;
-
-		IntStat& operator+=(const int32& Rhs)
-		{
-			Total += Rhs;
-			Max = FGenericPlatformMath::Max(Max, Rhs);
-			TotalCalls++;
-			return *this;
-		}
-	};
-
-	struct FloatStat
-	{
-		long double Total;
-		long double Max;
-		int32 TotalCalls;
-
-		FloatStat& operator+=(const long double& Rhs)
-		{
-			Total += Rhs;
-			Max = FGenericPlatformMath::Max(Max, Rhs);
-			TotalCalls++;
-			return *this;
-		}
-	};
-
-	static bool bLoggingActive;
-	TMap<FName, IntStat> IntStats;
-	TMap<FName, FloatStat> FloatStats;
-
-	static LogBenchmarkUtil &GetInstance()
-	{
-		static LogBenchmarkUtil Instance;
-		return Instance;
-	}
-
-	LogBenchmarkUtil() = default;
+	double TotalUpdateTime = 0;
+	uint32 NumUpdates = 0;
+
+	TSharedPtr<FArchive> Archive;
 };
diff --git a/Engine/Plugins/Experimental/Mutable/Source/MutableRuntime/Public/MuR/MutableTrace.h b/Engine/Plugins/Experimental/Mutable/Source/MutableRuntime/Public/MuR/MutableTrace.h
--- a/Engine/Plugins/Experimental/Mutable/Source/MutableRuntime/Public/MuR/MutableTrace.h
+++ b/Engine/Plugins/Experimental/Mutable/Source/MutableRuntime/Public/MuR/MutableTrace.h
@@ -3,6 +3,7 @@
 #pragma once
 
 #include "Trace/Trace.h"
+#include "HAL/PlatformTime.h"
 
 namespace UE::Trace { class FChannel; }
 
@@ -11,3 +12,23 @@ UE_TRACE_CHANNEL_EXTERN(MutableChannel, MUTABLERUNTIME_API)
 /** Custom Mutable profiler scope. */
 #define MUTABLE_CPUPROFILER_SCOPE(Name) \
 	TRACE_CPUPROFILER_EVENT_SCOPE_ON_CHANNEL(Mutable_##Name, MutableChannel)
+
+
+/** Simple class that saves the time the scope is alive in the given location. */
+class FScopeTimer
+{
+public:
+	FScopeTimer(double& InResult) : Result(InResult)
+	{
+		StartTime = FPlatformTime::Seconds();
+	}
+
+	~FScopeTimer()
+	{
+		Result =  FPlatformTime::Seconds() - StartTime;
+	}
+	
+private:
+	double StartTime = 0.0;
+	double& Result;
+};
