[Mutable] Fix FScopeTimer class naming collision.

#tests non-unity build
#rnx

diff --git a/Engine/Plugins/Experimental/Mutable/Source/CustomizableObject/Private/MuCO/CustomizableObjectSystem.cpp b/Engine/Plugins/Experimental/Mutable/Source/CustomizableObject/Private/MuCO/CustomizableObjectSystem.cpp
--- a/Engine/Plugins/Experimental/Mutable/Source/CustomizableObject/Private/MuCO/CustomizableObjectSystem.cpp
+++ b/Engine/Plugins/Experimental/Mutable/Source/CustomizableObject/Private/MuCO/CustomizableObjectSystem.cpp
@@ -2134,7 +2134,7 @@ namespace impl
 	void Task_Mutable_Update_GetMesh(const TSharedRef<FUpdateContextPrivate>& OperationData, const TSharedPtr<mu::Model>& Model)
 	{
 		MUTABLE_CPUPROFILER_SCOPE(Task_Mutable_Update_GetMesh)
-		FScopeTimer Timer(OperationData->TaskGetMeshTime);
+		FMutableScopeTimer Timer(OperationData->TaskGetMeshTime);
 
 #if WITH_EDITOR
 		const uint32 StartCycles = FPlatformTime::Cycles();
@@ -2165,7 +2165,7 @@ namespace impl
 	void Task_Mutable_Update_GetImages(const TSharedRef<FUpdateContextPrivate>& OperationData)
 	{
 		MUTABLE_CPUPROFILER_SCOPE(Task_Mutable_GetImages)
-		FScopeTimer Timer(OperationData->TaskGetImagesTime);
+		FMutableScopeTimer Timer(OperationData->TaskGetImagesTime);
 
 #if WITH_EDITOR
 		uint32 StartCycles = FPlatformTime::Cycles();
@@ -2249,7 +2249,7 @@ namespace impl
 	void Task_Game_Callbacks(const TSharedRef<FUpdateContextPrivate>& OperationData)
 	{
 		MUTABLE_CPUPROFILER_SCOPE(Task_Game_Callbacks)
-		FScopeTimer Timer(OperationData->TaskCallbacksTime);
+		FMutableScopeTimer Timer(OperationData->TaskCallbacksTime);
 
 		check(IsInGameThread());
 
@@ -2302,7 +2302,7 @@ namespace impl
 	void Task_Game_ConvertResources(const TSharedRef<FUpdateContextPrivate>& OperationData)
 	{
 		MUTABLE_CPUPROFILER_SCOPE(Task_Game_ConvertResources)
-		FScopeTimer Timer(OperationData->TaskConvertResourcesTime);
+		FMutableScopeTimer Timer(OperationData->TaskConvertResourcesTime);
 
 		check(IsInGameThread());
 
@@ -2450,7 +2450,7 @@ namespace impl
 	void Task_Game_LockCache(const TSharedRef<FUpdateContextPrivate>& OperationData)
 	{
 		MUTABLE_CPUPROFILER_SCOPE(Task_Game_LockCache)
-		FScopeTimer Timer(OperationData->TaskLockCacheTime);
+		FMutableScopeTimer Timer(OperationData->TaskLockCacheTime);
 
 		check(IsInGameThread());
 
diff --git a/Engine/Plugins/Experimental/Mutable/Source/MutableRuntime/Public/MuR/MutableTrace.h b/Engine/Plugins/Experimental/Mutable/Source/MutableRuntime/Public/MuR/MutableTrace.h
--- a/Engine/Plugins/Experimental/Mutable/Source/MutableRuntime/Public/MuR/MutableTrace.h
+++ b/Engine/Plugins/Experimental/Mutable/Source/MutableRuntime/Public/MuR/MutableTrace.h
@@ -15,15 +15,15 @@ UE_TRACE_CHANNEL_EXTERN(MutableChannel, MUTABLERUNTIME_API)
 
 
 /** Simple class that saves the time the scope is alive in the given location. */
-class FScopeTimer
+class FMutableScopeTimer
 {
 public:
-	FScopeTimer(double& InResult) : Result(InResult)
+	FMutableScopeTimer(double& InResult) : Result(InResult)
 	{
 		StartTime = FPlatformTime::Seconds();
 	}
 
-	~FScopeTimer()
+	~FMutableScopeTimer()
 	{
 		Result =  FPlatformTime::Seconds() - StartTime;
 	}
