#rnx

diff --git a/Engine/Plugins/Experimental/Mutable/Source/CustomizableObject/Classes/MuCO/CustomizableObjectInstance.h b/Engine/Plugins/Experimental/Mutable/Source/CustomizableObject/Classes/MuCO/CustomizableObjectInstance.h
--- a/Engine/Plugins/Experimental/Mutable/Source/CustomizableObject/Classes/MuCO/CustomizableObjectInstance.h
+++ b/Engine/Plugins/Experimental/Mutable/Source/CustomizableObject/Classes/MuCO/CustomizableObjectInstance.h
@@ -270,9 +270,6 @@ public:
 
 	/** Immediately destroy the Mutable Core Live Update Instance attached to this (if exists). */
 	void DestroyLiveUpdateInstance();
-	
-	// Releases all the mutable resources this instance holds, should only be called when it is not going to be used any more.
-	void ReleaseMutableResources(bool bCalledFromBeginDestroy);
 
 	// Returns de description texture (ex: color bar) for this parameter and DescIndex
 	// This will only be valid if bBuildParameterDecorations was set to true before the last update.
diff --git a/Engine/Plugins/Experimental/Mutable/Source/CustomizableObject/Private/MuCO/CustomizableInstancePrivateData.h b/Engine/Plugins/Experimental/Mutable/Source/CustomizableObject/Private/MuCO/CustomizableInstancePrivateData.h
--- a/Engine/Plugins/Experimental/Mutable/Source/CustomizableObject/Private/MuCO/CustomizableInstancePrivateData.h
+++ b/Engine/Plugins/Experimental/Mutable/Source/CustomizableObject/Private/MuCO/CustomizableInstancePrivateData.h
@@ -265,7 +265,10 @@ public:
 	// The following method is basically copied from PostEditChangeProperty and/or SkeletalMesh.cpp to be able to replicate PostEditChangeProperty without the editor
 	void PostEditChangePropertyWithoutEditor(USkeletalMesh* SkeletalMesh);
 	
-	void DiscardResourcesAndSetReferenceSkeletalMesh(UCustomizableObjectInstance* Public);
+	void DiscardResourcesAndSetReferenceSkeletalMesh(UCustomizableObjectInstance* Instance);
+
+	// Releases all the mutable resources this instance holds, should only be called when it is not going to be used any more.
+	void ReleaseMutableResources(bool bCalledFromBeginDestroy, const UCustomizableObjectInstance& Instance);
 
 	const TArray<FAnimInstanceOverridePhysicsAsset>* GetGeneratedPhysicsAssetsForAnimInstance(TSubclassOf<UAnimInstance> AnimInstance) const;
 
diff --git a/Engine/Plugins/Experimental/Mutable/Source/CustomizableObject/Private/MuCO/CustomizableObjectInstance.cpp b/Engine/Plugins/Experimental/Mutable/Source/CustomizableObject/Private/MuCO/CustomizableObjectInstance.cpp
--- a/Engine/Plugins/Experimental/Mutable/Source/CustomizableObject/Private/MuCO/CustomizableObjectInstance.cpp
+++ b/Engine/Plugins/Experimental/Mutable/Source/CustomizableObject/Private/MuCO/CustomizableObjectInstance.cpp
@@ -320,11 +320,12 @@ void UCustomizableObjectInstance::BeginDestroy()
 		{
 			PrivateData->StreamingHandle->CancelHandle();
 		}
+
 		PrivateData->StreamingHandle = nullptr;
+
+		PrivateData.Get()->ReleaseMutableResources(true, *this);
 	}
 	
-	ReleaseMutableResources(true);
-	
 	Super::BeginDestroy();
 }
 
@@ -345,17 +346,17 @@ void UCustomizableObjectInstance::DestroyLiveUpdateInstance()
 }
 
 
-void UCustomizableObjectInstance::ReleaseMutableResources(bool bCalledFromBeginDestroy)
+void UCustomizableInstancePrivateData::ReleaseMutableResources(bool bCalledFromBeginDestroy, const UCustomizableObjectInstance& Instance)
 {
-	PrivateData->GeneratedMaterials.Empty();
+	GeneratedMaterials.Empty();
 
 	if (UCustomizableObjectSystem::IsCreated()) // Need to check this because the object might be destroyed after the CustomizableObjectSystem at shutdown
 	{
 		FCustomizableObjectSystemPrivate* CustomizableObjectSystem = UCustomizableObjectSystem::GetInstance()->GetPrivate();
 		// Get the cache of resources of all live instances of this object
-		FMutableResourceCache& Cache = CustomizableObjectSystem->GetObjectCache(GetCustomizableObject());
+		FMutableResourceCache& Cache = CustomizableObjectSystem->GetObjectCache(Instance.GetCustomizableObject());
 
-		for (FGeneratedTexture& Texture : PrivateData->GeneratedTextures)
+		for (FGeneratedTexture& Texture : GeneratedTextures)
 		{
 			if (CustomizableObjectSystem->RemoveTextureReference(Texture.Key))
 			{
@@ -363,7 +364,7 @@ void UCustomizableObjectInstance::ReleaseMutableResources(bool bCalledFromBeginD
 				// instance's remaining sk meshes and GC is being performed anyway so it will free the textures if needed
 				if (!bCalledFromBeginDestroy && CustomizableObjectSystem->bReleaseTexturesImmediately)
 				{
-					UCustomizableInstancePrivateData::ReleaseMutableTexture(Texture.Key, Cast<UTexture2D>(Texture.Texture), Cache);
+					ReleaseMutableTexture(Texture.Key, Cast<UTexture2D>(Texture.Texture), Cache);
 				}
 			}
 		}
@@ -373,14 +374,14 @@ void UCustomizableObjectInstance::ReleaseMutableResources(bool bCalledFromBeginD
 		{
 			FUnrealMutableImageProvider* ImageProvider = UCustomizableObjectSystem::GetInstance()->GetPrivateChecked()->GetImageProviderChecked();
 
-			for (const FName& TextureParameter : PrivateData->UpdateTextureParameters)
+			for (const FName& TextureParameter : UpdateTextureParameters)
 			{
 				ImageProvider->UnCacheImage(TextureParameter, false);
 			}
 		}
 	}
 
-	PrivateData->GeneratedTextures.Empty();
+	GeneratedTextures.Empty();
 }
 
 
@@ -2073,6 +2074,7 @@ UCustomizableObjectInstance* UCustomizableObjectInstance::Clone()
 
 	// Default Outer is the transient package.
 	UCustomizableObjectInstance* NewInstance = NewObject<UCustomizableObjectInstance>();
+	check(NewInstance->PrivateData);
 	NewInstance->CopyParametersFromInstance(this);
 
 	return NewInstance;
@@ -2685,20 +2687,20 @@ bool UCustomizableObjectInstance::IsSelectedParameterProfileDirty() const
 //}
 
 
-void UCustomizableInstancePrivateData::DiscardResourcesAndSetReferenceSkeletalMesh(UCustomizableObjectInstance* Public )
+void UCustomizableInstancePrivateData::DiscardResourcesAndSetReferenceSkeletalMesh(UCustomizableObjectInstance* Instance)
 {
 	if (HasCOInstanceFlags(Generated))
 	{
-		for (int32 Component = 0; Component < Public->SkeletalMeshes.Num(); ++Component)
+		for (int32 Component = 0; Component < Instance->SkeletalMeshes.Num(); ++Component)
 		{
-			if (Public->SkeletalMeshes[Component] && Public->SkeletalMeshes[Component]->IsValidLowLevel())
+			if (Instance->SkeletalMeshes[Component] && Instance->SkeletalMeshes[Component]->IsValidLowLevel())
 			{
-				Public->SkeletalMeshes[Component]->ReleaseResources();
-				Public->SkeletalMeshes[Component] = nullptr;
+				Instance->SkeletalMeshes[Component]->ReleaseResources();
+				Instance->SkeletalMeshes[Component] = nullptr;
 			}
 		}
 
-		Public->ReleaseMutableResources(false);
+		ReleaseMutableResources(false, *Instance);
 	}
 
 	ClearCOInstanceFlags(Generated);
@@ -2706,7 +2708,7 @@ void UCustomizableInstancePrivateData::DiscardResourcesAndSetReferenceSkeletalMe
 	
 	InvalidateGeneratedData();
 	
-	Public->SkeletalMeshes.Reset();
+	Instance->SkeletalMeshes.Reset();
 	DescriptorRuntimeHash = FDescriptorRuntimeHash();
 
 	for (TObjectIterator<UCustomizableObjectInstanceUsage> It; It; ++It)
@@ -2720,9 +2722,9 @@ void UCustomizableInstancePrivateData::DiscardResourcesAndSetReferenceSkeletalMe
 		}
 #endif
 
-		if (CustomizableObjectInstanceUsage && CustomizableObjectInstanceUsage->GetCustomizableObjectInstance() == Public)
+		if (CustomizableObjectInstanceUsage && CustomizableObjectInstanceUsage->GetCustomizableObjectInstance() == Instance)
 		{
-			UCustomizableObject* CustomizableObject = Public->GetCustomizableObject();
+			UCustomizableObject* CustomizableObject = Instance->GetCustomizableObject();
 			bool bReplaceDiscardedWithReferenceMesh = UCustomizableObjectSystem::GetInstance()->GetPrivate()->IsReplaceDiscardedWithReferenceMeshEnabled();
 			CustomizableObjectInstanceUsage->SetSkeletalMesh(CustomizableObject && bReplaceDiscardedWithReferenceMesh ? CustomizableObject->GetRefSkeletalMesh(CustomizableObjectInstanceUsage->GetComponentIndex()) : nullptr);
 		}
