Fix error C2872: 'Write': ambiguous symbol in MSVC static analysis
#rnx

diff --git a/Engine/Plugins/Experimental/Mutable/Source/CustomizableObject/Private/MuCO/LogBenchmarkUtil.cpp b/Engine/Plugins/Experimental/Mutable/Source/CustomizableObject/Private/MuCO/LogBenchmarkUtil.cpp
--- a/Engine/Plugins/Experimental/Mutable/Source/CustomizableObject/Private/MuCO/LogBenchmarkUtil.cpp
+++ b/Engine/Plugins/Experimental/Mutable/Source/CustomizableObject/Private/MuCO/LogBenchmarkUtil.cpp
@@ -30,12 +30,14 @@ static TAutoConsoleVariable<bool> CVarEnableBenchmark(
 	TEXT("Enable or disable the benchmarking."));
 
 
-void Write(FArchive& Archive, FStringView Text)
+namespace LogBenchmarkUtil
 {
-	Archive.Serialize(const_cast<ANSICHAR*>(StringCast<ANSICHAR>(Text.GetData(), Text.Len()).Get()), Text.Len() * sizeof(ANSICHAR));
+	void Write(FArchive& Archive, FStringView Text)
+	{
+		Archive.Serialize(const_cast<ANSICHAR*>(StringCast<ANSICHAR>(Text.GetData(), Text.Len()).Get()), Text.Len() * sizeof(ANSICHAR));
+	}
 }
 
-
 TSharedPtr<FArchive> CreateFile()
 {
 	const FString Directory = FPaths::ProfilingDir() + TEXT("Mutable/Benchmark/");
@@ -47,7 +49,7 @@ TSharedPtr<FArchive> CreateFile()
 	TSharedPtr<FArchive> Archive = MakeShareable(IFileManager::Get().CreateFileWriter(*Filename, FILEWRITE_AllowRead | FILEWRITE_NoFail));
 	check(Archive);
 
-	Write(*Archive, TEXT("ID_CO,ID_COI,ID_UpdateType,Time_Update,Time_TaskGetMesh,Time_TaskLockCache,Time_TaskGetImages,Time_TaskConvertResources,Time_TaskCallbacks\n"));		
+	LogBenchmarkUtil::Write(*Archive, TEXT("ID_CO,ID_COI,ID_UpdateType,Time_Update,Time_TaskGetMesh,Time_TaskLockCache,Time_TaskGetImages,Time_TaskConvertResources,Time_TaskCallbacks\n"));		
 
 	return Archive;
 }
@@ -300,7 +302,7 @@ void FLogBenchmarkUtil::FinishUpdate(const TSharedRef<FUpdateContextPrivate>& Co
 	const double Time_TaskConvertResources = Context->TaskConvertResourcesTime * 1000;
 	const double Time_TaskCallbacks = Context->TaskCallbacksTime * 1000;
 	
-	Write(*Archive, FString::Printf(TEXT("%s,%s,%s,%f,%f,%f,%f,%f,%f\n"), *ID_CO, *ID_COI, *ID_UpdateType, Time_Update, Time_TaskGetMesh, Time_TaskLockCache, Time_TaskGetImages, Time_TaskConvertResources, Time_TaskCallbacks));
+	LogBenchmarkUtil::Write(*Archive, FString::Printf(TEXT("%s,%s,%s,%f,%f,%f,%f,%f,%f\n"), *ID_CO, *ID_COI, *ID_UpdateType, Time_Update, Time_TaskGetMesh, Time_TaskLockCache, Time_TaskGetImages, Time_TaskConvertResources, Time_TaskCallbacks));
 	Archive->Flush();
 }
 
