[Mutable] Added option to set the output path for the mutable benchmark files.

diff --git a/Engine/Plugins/Experimental/Mutable/Source/CustomizableObject/Private/MuCO/LogBenchmarkUtil.cpp b/Engine/Plugins/Experimental/Mutable/Source/CustomizableObject/Private/MuCO/LogBenchmarkUtil.cpp
--- a/Engine/Plugins/Experimental/Mutable/Source/CustomizableObject/Private/MuCO/LogBenchmarkUtil.cpp
+++ b/Engine/Plugins/Experimental/Mutable/Source/CustomizableObject/Private/MuCO/LogBenchmarkUtil.cpp
@@ -29,6 +29,10 @@ static TAutoConsoleVariable<bool> CVarEnableBenchmark(
 	false,
 	TEXT("Enable or disable the benchmarking."));
 
+static TAutoConsoleVariable<FString> CVarBenchmarkFilePath(
+	 TEXT("mutable.BenchmarkFilePath"),
+	FPaths::ProfilingDir() + TEXT("Mutable/Benchmark"),
+	TEXT("Sets the path where to store the generated mutable benchmark report file."));
 
 namespace LogBenchmarkUtil
 {
@@ -40,7 +44,7 @@ namespace LogBenchmarkUtil
 
 TSharedPtr<FArchive> CreateFile()
 {
-	const FString Directory = FPaths::ProfilingDir() + TEXT("Mutable/Benchmark/");
+	const FString Directory = CVarBenchmarkFilePath.GetValueOnGameThread();
 	IFileManager::Get().MakeDirectory(*Directory, true);
 
 	const FDateTime FileDate = FDateTime::Now();
