[Mutable] Move some of the image related types and enums to ImageTypes.h

- This is needed to avoid cyclical dependencies if a member of Image defined in another file needs to access some of the types.

#rnx

diff --git a/Engine/Plugins/Experimental/Mutable/Source/MutableRuntime/Private/MuR/Image.cpp b/Engine/Plugins/Experimental/Mutable/Source/MutableRuntime/Private/MuR/Image.cpp
--- a/Engine/Plugins/Experimental/Mutable/Source/MutableRuntime/Private/MuR/Image.cpp
+++ b/Engine/Plugins/Experimental/Mutable/Source/MutableRuntime/Private/MuR/Image.cpp
@@ -10,105 +10,18 @@
 #include "MuR/MutableMath.h"
 #include "MuR/MutableTrace.h"
 
-#include <initializer_list>
-
 namespace
 {
 	static bool bDisableCompressedImageBlackBlockInit = false;
 	static FAutoConsoleVariableRef CVarDisableCompressedImageBlackBlockInit(
 		TEXT("mutable.DisableCompressedImageBlackBlockInit"),
 		bDisableCompressedImageBlackBlockInit,
-		TEXT("A value of 1 disables mutable compressed black block initialization"),
+		TEXT("A value of 1 disables mutable compressed black block initialization."),
 		ECVF_Default);
 }
 
 namespace mu
 {
-	MUTABLE_IMPLEMENT_ENUM_SERIALISABLE( EBlendType );
-	MUTABLE_IMPLEMENT_ENUM_SERIALISABLE( EMipmapFilterType );
-	MUTABLE_IMPLEMENT_ENUM_SERIALISABLE( ECompositeImageMode );
-	MUTABLE_IMPLEMENT_ENUM_SERIALISABLE( ESamplingMethod );
-	MUTABLE_IMPLEMENT_ENUM_SERIALISABLE( EMinFilterMethod );
-	MUTABLE_IMPLEMENT_ENUM_SERIALISABLE( EImageFormat );
-
-    //---------------------------------------------------------------------------------------------
-    static FImageFormatData s_imageFormatData[uint32(EImageFormat::IF_COUNT)] =
-    {
-		FImageFormatData( 0, 0, 0, 0 ),	// IF_NONE
-		FImageFormatData( 1, 1, 3, 3 ),	// IF_RGB_UBYTE
-		FImageFormatData( 1, 1, 4, 4 ),	// IF_RGBA_UBYTE
-		FImageFormatData( 1, 1, 1, 1 ),	// IF_U_UBYTE
-
-		FImageFormatData( 0, 0, 0, 0 ),	// IF_PVRTC2 (deprecated)
-        FImageFormatData( 0, 0, 0, 0 ),	// IF_PVRTC4 (deprecated)
-        FImageFormatData( 0, 0, 0, 0 ),	// IF_ETC1 (deprecated)
-        FImageFormatData( 0, 0, 0, 0 ),	// IF_ETC2 (deprecated)
-
-        FImageFormatData( 0, 0, 0, 1 ),	// IF_L_UBYTE_RLE
-        FImageFormatData( 0, 0, 0, 3 ),	// IF_RGB_UBYTE_RLE
-        FImageFormatData( 0, 0, 0, 4 ),	// IF_RGBA_UBYTE_RLE
-        FImageFormatData( 0, 0, 0, 1 ),	// IF_L_UBIT_RLE
-
-        FImageFormatData( 4, 4, 8,  4 ),	// IF_BC1
-        FImageFormatData( 4, 4, 16, 4 ),	// IF_BC2
-        FImageFormatData( 4, 4, 16, 4 ),	// IF_BC3
-        FImageFormatData( 4, 4, 8,  1 ),	// IF_BC4
-        FImageFormatData( 4, 4, 16, 2 ),	// IF_BC5
-        FImageFormatData( 4, 4, 16, 3 ),	// IF_BC6
-        FImageFormatData( 4, 4, 16, 4 ),	// IF_BC7
-
-        FImageFormatData( 1, 1, 4, 4 ),		// IF_BGRA_UBYTE
-
-		FImageFormatData(4, 4, 16, 3, {252, 253, 255, 255, 255, 255, 255, 255, 0, 0, 0, 0, 0, 0, 255, 255}), // IF_ASTC_4x4_RGB_LDR
-		FImageFormatData(4, 4, 16, 4, {252, 253, 255, 255, 255, 255, 255, 255, 0, 0, 0, 0, 0, 0, 0, 0}),     // IF_ASTC_4x4_RGBA_LDR
-        FImageFormatData( 4, 4, 16, 2 ),	// IF_ASTC_4x4_RG_LDR // TODO: check black block for RG.
-		
-		FImageFormatData(8, 8, 16, 3, {252, 253, 255, 255, 255, 255, 255, 255, 0, 0, 0, 0, 0, 0, 255, 255}),	// IF_ASTC_8x8_RGB_LDR,
-		FImageFormatData(8, 8, 16, 4, {252, 253, 255, 255, 255, 255, 255, 255, 0, 0, 0, 0, 0, 0, 0, 0}),		// IF_ASTC_8x8_RGBA_LDR,
-		FImageFormatData(8, 8, 16, 2),		// IF_ASTC_8x8_RG_LDR,
-		FImageFormatData(12, 12, 16, 3, {252, 253, 255, 255, 255, 255, 255, 255, 0, 0, 0, 0, 0, 0, 255, 255}),	// IF_ASTC_12x12_RGB_LDR
-		FImageFormatData(12, 12, 16, 4, {252, 253, 255, 255, 255, 255, 255, 255, 0, 0, 0, 0, 0, 0, 0, 0}),		// IF_ASTC_12x12_RGBA_LDR
-		FImageFormatData(12, 12, 16, 2),	// IF_ASTC_12x12_RG_LDR
-		FImageFormatData(6, 6, 16, 3, {252, 253, 255, 255, 255, 255, 255, 255, 0, 0, 0, 0, 0, 0, 255, 255}),	// IF_ASTC_6x6_RGB_LDR,
-		FImageFormatData(6, 6, 16, 4, {252, 253, 255, 255, 255, 255, 255, 255, 0, 0, 0, 0, 0, 0, 0, 0}),		// IF_ASTC_6x6_RGBA_LDR,
-		FImageFormatData(6, 6, 16, 2),		// IF_ASTC_6x6_RG_LDR,
-		FImageFormatData(10, 10, 16, 3, {252, 253, 255, 255, 255, 255, 255, 255, 0, 0, 0, 0, 0, 0, 255, 255}),	// IF_ASTC_6x6_RGB_LDR,
-		FImageFormatData(10, 10, 16, 4, {252, 253, 255, 255, 255, 255, 255, 255, 0, 0, 0, 0, 0, 0, 0, 0}),		// IF_ASTC_6x6_RGBA_LDR,
-		FImageFormatData(10, 10, 16, 2),	// IF_ASTC_6x6_RG_LDR,
-
-    };
-
-
-    //---------------------------------------------------------------------------------------------
-    const FImageFormatData& GetImageFormatData( EImageFormat format )
-    {
-        check( format < EImageFormat::IF_COUNT );
-        return s_imageFormatData[ uint8(format) ];
-    }
-
-
-	//---------------------------------------------------------------------------------------------
-	void FMipmapGenerationSettings::Serialise(OutputArchive& arch) const
-	{
-		uint32 ver = 0;
-		arch << ver;
-
-		arch << m_sharpenFactor;
-		arch << m_filterType;
-		arch << m_ditherMipmapAlpha;
-	}
-
-	void FMipmapGenerationSettings::Unserialise(InputArchive& arch)
-	{
-		uint32 ver = 0;
-		arch >> ver;
-		check(ver == 0);
-
-		arch >> m_sharpenFactor;
-		arch >> m_filterType;
-		arch >> m_ditherMipmapAlpha;
-	}
-
 
     //---------------------------------------------------------------------------------------------
     Image::Image()
diff --git a/Engine/Plugins/Experimental/Mutable/Source/MutableRuntime/Private/MuR/ImagePrivate.h b/Engine/Plugins/Experimental/Mutable/Source/MutableRuntime/Private/MuR/ImagePrivate.h
--- a/Engine/Plugins/Experimental/Mutable/Source/MutableRuntime/Private/MuR/ImagePrivate.h
+++ b/Engine/Plugins/Experimental/Mutable/Source/MutableRuntime/Private/MuR/ImagePrivate.h
@@ -3,6 +3,7 @@
 #pragma once
 
 #include "MuR/Image.h"
+#include "MuR/ImageTypes.h"
 
 #include "MuR/SerialisationPrivate.h"
 #include "MuR/MutableMath.h"
@@ -17,7 +18,6 @@ namespace mu
 	MUTABLE_DEFINE_ENUM_SERIALISABLE( EImageFormat );
 	MUTABLE_DEFINE_ENUM_SERIALISABLE( EAddressMode );
 
-
 	//---------------------------------------------------------------------------------------------
 	//!
 	//---------------------------------------------------------------------------------------------
diff --git a/Engine/Plugins/Experimental/Mutable/Source/MutableRuntime/Private/MuR/ImageTypes.cpp b/Engine/Plugins/Experimental/Mutable/Source/MutableRuntime/Private/MuR/ImageTypes.cpp
new file mode 100644
--- /dev/null
+++ b/Engine/Plugins/Experimental/Mutable/Source/MutableRuntime/Private/MuR/ImageTypes.cpp
@@ -0,0 +1,87 @@
+// Copyright Epic Games, Inc. All Rights Reserved.
+
+#include "MuR/ImageTypes.h"
+#include "MuR/SerialisationPrivate.h"
+
+namespace mu
+{
+	MUTABLE_IMPLEMENT_ENUM_SERIALISABLE(EBlendType);
+	MUTABLE_IMPLEMENT_ENUM_SERIALISABLE(EMipmapFilterType);
+	MUTABLE_IMPLEMENT_ENUM_SERIALISABLE(ECompositeImageMode);
+	MUTABLE_IMPLEMENT_ENUM_SERIALISABLE(ESamplingMethod);
+	MUTABLE_IMPLEMENT_ENUM_SERIALISABLE(EMinFilterMethod);
+	MUTABLE_IMPLEMENT_ENUM_SERIALISABLE(EImageFormat);
+
+    static FImageFormatData s_imageFormatData[uint32(EImageFormat::IF_COUNT)] =
+    {
+		FImageFormatData(0, 0, 0, 0),	// IF_NONE
+		FImageFormatData(1, 1, 3, 3),	// IF_RGB_UBYTE
+		FImageFormatData(1, 1, 4, 4),	// IF_RGBA_UBYTE
+		FImageFormatData(1, 1, 1, 1),	// IF_U_UBYTE
+
+		FImageFormatData(0, 0, 0, 0),	// IF_PVRTC2 (deprecated)
+        FImageFormatData(0, 0, 0, 0),	// IF_PVRTC4 (deprecated)
+        FImageFormatData(0, 0, 0, 0),	// IF_ETC1 (deprecated)
+        FImageFormatData(0, 0, 0, 0),	// IF_ETC2 (deprecated)
+
+        FImageFormatData(0, 0, 0, 1),	// IF_L_UBYTE_RLE
+        FImageFormatData(0, 0, 0, 3),	// IF_RGB_UBYTE_RLE
+        FImageFormatData(0, 0, 0, 4),	// IF_RGBA_UBYTE_RLE
+        FImageFormatData(0, 0, 0, 1),	// IF_L_UBIT_RLE
+
+        FImageFormatData(4, 4, 8,  4),	// IF_BC1
+        FImageFormatData(4, 4, 16, 4),	// IF_BC2
+        FImageFormatData(4, 4, 16, 4),	// IF_BC3
+        FImageFormatData(4, 4, 8,  1),	// IF_BC4
+        FImageFormatData(4, 4, 16, 2),	// IF_BC5
+        FImageFormatData(4, 4, 16, 3),	// IF_BC6
+        FImageFormatData(4, 4, 16, 4),	// IF_BC7
+
+        FImageFormatData(1, 1, 4, 4 ),		// IF_BGRA_UBYTE
+
+		FImageFormatData(4, 4, 16, 3, {252, 253, 255, 255, 255, 255, 255, 255, 0, 0, 0, 0, 0, 0, 255, 255}), // IF_ASTC_4x4_RGB_LDR
+		FImageFormatData(4, 4, 16, 4, {252, 253, 255, 255, 255, 255, 255, 255, 0, 0, 0, 0, 0, 0, 0, 0}),     // IF_ASTC_4x4_RGBA_LDR
+        FImageFormatData(4, 4, 16, 2),	// IF_ASTC_4x4_RG_LDR // TODO: check black block for RG.
+		
+		FImageFormatData(8, 8, 16, 3, {252, 253, 255, 255, 255, 255, 255, 255, 0, 0, 0, 0, 0, 0, 255, 255}),	// IF_ASTC_8x8_RGB_LDR,
+		FImageFormatData(8, 8, 16, 4, {252, 253, 255, 255, 255, 255, 255, 255, 0, 0, 0, 0, 0, 0, 0, 0}),		// IF_ASTC_8x8_RGBA_LDR,
+		FImageFormatData(8, 8, 16, 2),		// IF_ASTC_8x8_RG_LDR,
+		FImageFormatData(12, 12, 16, 3, {252, 253, 255, 255, 255, 255, 255, 255, 0, 0, 0, 0, 0, 0, 255, 255}),	// IF_ASTC_12x12_RGB_LDR
+		FImageFormatData(12, 12, 16, 4, {252, 253, 255, 255, 255, 255, 255, 255, 0, 0, 0, 0, 0, 0, 0, 0}),		// IF_ASTC_12x12_RGBA_LDR
+		FImageFormatData(12, 12, 16, 2),	// IF_ASTC_12x12_RG_LDR
+		FImageFormatData(6, 6, 16, 3, {252, 253, 255, 255, 255, 255, 255, 255, 0, 0, 0, 0, 0, 0, 255, 255}),	// IF_ASTC_6x6_RGB_LDR,
+		FImageFormatData(6, 6, 16, 4, {252, 253, 255, 255, 255, 255, 255, 255, 0, 0, 0, 0, 0, 0, 0, 0}),		// IF_ASTC_6x6_RGBA_LDR,
+		FImageFormatData(6, 6, 16, 2),		// IF_ASTC_6x6_RG_LDR,
+		FImageFormatData(10, 10, 16, 3, {252, 253, 255, 255, 255, 255, 255, 255, 0, 0, 0, 0, 0, 0, 255, 255}),	// IF_ASTC_6x6_RGB_LDR,
+		FImageFormatData(10, 10, 16, 4, {252, 253, 255, 255, 255, 255, 255, 255, 0, 0, 0, 0, 0, 0, 0, 0}),		// IF_ASTC_6x6_RGBA_LDR,
+		FImageFormatData(10, 10, 16, 2),	// IF_ASTC_6x6_RG_LDR,
+    };
+
+    const FImageFormatData& GetImageFormatData( EImageFormat format )
+    {
+        check( format < EImageFormat::IF_COUNT );
+        return s_imageFormatData[ uint8(format) ];
+    }
+
+	void FMipmapGenerationSettings::Serialise(OutputArchive& arch) const
+	{
+		uint32 ver = 0;
+		arch << ver;
+
+		arch << m_sharpenFactor;
+		arch << m_filterType;
+		arch << m_ditherMipmapAlpha;
+	}
+
+	void FMipmapGenerationSettings::Unserialise(InputArchive& arch)
+	{
+		uint32 ver = 0;
+		arch >> ver;
+		check(ver == 0);
+
+		arch >> m_sharpenFactor;
+		arch >> m_filterType;
+		arch >> m_ditherMipmapAlpha;
+	}
+}
+
diff --git a/Engine/Plugins/Experimental/Mutable/Source/MutableRuntime/Public/MuR/Image.h b/Engine/Plugins/Experimental/Mutable/Source/MutableRuntime/Public/MuR/Image.h
--- a/Engine/Plugins/Experimental/Mutable/Source/MutableRuntime/Public/MuR/Image.h
+++ b/Engine/Plugins/Experimental/Mutable/Source/MutableRuntime/Public/MuR/Image.h
@@ -2,6 +2,8 @@
 
 #pragma once
 
+#include "MuR/ImageTypes.h"
+
 #include "MuR/MutableMath.h"
 #include "MuR/MutableTrace.h"
 #include "MuR/Ptr.h"
@@ -29,166 +31,6 @@ namespace mu
 	typedef Ptr<Image> ImagePtr;
 	typedef Ptr<const Image> ImagePtrConst;
 
-	//!
-	using FImageSize = UE::Math::TIntVector2<uint16>;
-	using FImageRect = box<vec2<uint16>>;
-
-	//! Pixel formats supported by the images.
-	//! \ingroup runtime
-	enum class EImageFormat : uint8
-	{
-		IF_NONE,
-		IF_RGB_UBYTE,
-		IF_RGBA_UBYTE,
-		IF_L_UBYTE,
-
-        //! Deprecated formats
-        IF_PVRTC2_DEPRECATED,
-        IF_PVRTC4_DEPRECATED,
-        IF_ETC1_DEPRECATED,
-        IF_ETC2_DEPRECATED,
-
-		IF_L_UBYTE_RLE,
-		IF_RGB_UBYTE_RLE,
-		IF_RGBA_UBYTE_RLE,
-		IF_L_UBIT_RLE,
-
-        //! Common S3TC formats
-        IF_BC1,
-        IF_BC2,
-        IF_BC3,
-        IF_BC4,
-        IF_BC5,
-
-        //! Not really supported yet
-        IF_BC6,
-        IF_BC7,
-
-        //! Swizzled versions, engineers be damned.
-        IF_BGRA_UBYTE,
-
-        //! The new standard
-        IF_ASTC_4x4_RGB_LDR,
-        IF_ASTC_4x4_RGBA_LDR,
-        IF_ASTC_4x4_RG_LDR,
-
-		IF_ASTC_8x8_RGB_LDR,
-		IF_ASTC_8x8_RGBA_LDR,
-		IF_ASTC_8x8_RG_LDR,
-		IF_ASTC_12x12_RGB_LDR,
-		IF_ASTC_12x12_RGBA_LDR,
-		IF_ASTC_12x12_RG_LDR,
-		IF_ASTC_6x6_RGB_LDR,
-		IF_ASTC_6x6_RGBA_LDR,
-		IF_ASTC_6x6_RG_LDR,
-		IF_ASTC_10x10_RGB_LDR,
-		IF_ASTC_10x10_RGBA_LDR,
-		IF_ASTC_10x10_RG_LDR,
-
-        IF_COUNT
-	};
-
-	struct MUTABLERUNTIME_API FImageDesc
-	{
-		FImageDesc()
-		{
-			m_format = EImageFormat::IF_NONE;
-			m_size = FImageSize(0, 0);
-			m_lods = 0;
-		}
-
-		FImageDesc(const FImageSize& size, EImageFormat format, uint8 lods)
-		{
-			m_format = format;
-			m_size = size;
-			m_lods = lods;
-		}
-
-		//!
-		FImageSize m_size = FImageSize(0, 0);
-
-		//!
-		EImageFormat m_format;
-
-		//! Levels of detail (mipmaps)
-		uint8 m_lods;
-
-		//!
-		inline bool operator==(const FImageDesc& o) const
-		{
-			return  (m_format == o.m_format) &&
-				(m_lods == o.m_lods) &&
-				(m_size == o.m_size);
-		}
-	};
-
-
-	//! List of supported modes in generic image layering operations.
-	enum class EBlendType
-	{
-		BT_NONE = 0,
-		BT_SOFTLIGHT,
-		BT_HARDLIGHT,
-		BT_BURN,
-		BT_DODGE,
-		BT_SCREEN,
-		BT_OVERLAY,
-		BT_BLEND,
-		BT_MULTIPLY,
-		BT_LIGHTEN,				// Increase the channel value by a given proportion of what is missing from white 
-		BT_NORMAL_COMBINE,
-		_BT_COUNT
-	};
-
-	enum class EAddressMode
-	{
-		None,
-		Wrap,
-		ClampToEdge,
-		ClampToBlack,
-	};
-
-	enum class EMipmapFilterType
-	{
-		MFT_Unfiltered,
-		MFT_SimpleAverage,
-		MFT_Sharpen,
-		_MFT_COUNT
-	};
-
-	enum class ECompositeImageMode
-	{
-		CIM_Disabled,
-		CIM_NormalRoughnessToRed,
-		CIM_NormalRoughnessToGreen,
-		CIM_NormalRoughnessToBlue,
-		CIM_NormalRoughnessToAlpha,
-		_CIM_COUNT
-	};
-
-	enum class EInitializationType
-	{
-		NotInitialized,
-		Black
-	};
-
-	enum class ESamplingMethod : uint8
-	{
-		Point = 0,
-		BiLinear,
-		MaxValue
-	};
-	static_assert(uint32(ESamplingMethod::MaxValue) <= (1 << 3), "ESampligMethod enum cannot hold more than 8 values");
-
-	enum class EMinFilterMethod : uint8
-	{
-		None = 0,
-		TotalAreaHeuristic,
-		MaxValue
-	};
-
-	static_assert(uint32(EMinFilterMethod::MaxValue) <= (1 << 3), "EMinFilterMethod enum cannot hold more than 8 values");
-
     //! \brief 2D image resource with mipmaps.
 	//! \ingroup runtime
     class MUTABLERUNTIME_API Image : public Resource
@@ -466,69 +308,6 @@ namespace mu
 		void GetNonBlackRect_Reference(FImageRect& rect) const;
 	};
 
-	/** */
-	struct FImageFormatData
-	{
-		static constexpr SIZE_T MAX_BYTES_PER_BLOCK = 16;
-
-		FImageFormatData
-		(
-			uint8 InPixelsPerBlockX = 0,
-			uint8 InPixelsPerBlockY = 0,
-			uint16 InBytesPerBlock = 0,
-			uint16 InChannels = 0
-		)
-		{
-			PixelsPerBlockX = InPixelsPerBlockX;
-			PixelsPerBlockY = InPixelsPerBlockY;
-			BytesPerBlock = InBytesPerBlock;
-			Channels = InChannels;
-		}
-
-		FImageFormatData
-		(
-			uint8 InPixelsPerBlockX,
-			uint8 InPixelsPerBlockY,
-			uint16 InBytesPerBlock,
-			uint16 InChannels,
-			std::initializer_list<uint8> BlackBlockInit
-		)
-			: FImageFormatData(InPixelsPerBlockX, InPixelsPerBlockY, InBytesPerBlock, InChannels)
-		{
-			check(MAX_BYTES_PER_BLOCK >= BlackBlockInit.size());
-
-			const SIZE_T SanitizedBlockSize = FMath::Min<SIZE_T>(MAX_BYTES_PER_BLOCK, BlackBlockInit.size());
-			FMemory::Memcpy(BlackBlock, BlackBlockInit.begin(), SanitizedBlockSize);
-		}
-
-		/** For block based formats, size of the block size.For uncompressed formats it will always be 1,1. For non-block-based compressed formats, it will be 0,0. */
-		uint8 PixelsPerBlockX, PixelsPerBlockY;
-
-		/** Number of bytes used by every pixel block, if uncompressed or block-compressed format.
-		 * For non-block-compressed formats, it returns 0.
-		 */
-		uint16 BytesPerBlock;
-
-		/** Channels in every pixel of the image. */
-		uint16 Channels;
-
-		/** Representation of a black block of the image. */
-		uint8 BlackBlock[MAX_BYTES_PER_BLOCK] = { 0 };
-	};
-
-	MUTABLERUNTIME_API const FImageFormatData& GetImageFormatData(EImageFormat format);
-
-	struct MUTABLERUNTIME_API FMipmapGenerationSettings
-	{
-		float m_sharpenFactor = 0.0f;
-		EMipmapFilterType m_filterType = EMipmapFilterType::MFT_SimpleAverage;
-		EAddressMode m_addressMode = EAddressMode::None;
-		bool m_ditherMipmapAlpha = false;
-
-		void Serialise(OutputArchive& arch) const;
-		void Unserialise(InputArchive& arch);
-	};
-
 	/** This struct contains image operations that may need to allocate and free temporary images to work.
 	* It's purpose is to override the creation, release and clone functions depending on the context.
 	*/
diff --git a/Engine/Plugins/Experimental/Mutable/Source/MutableRuntime/Public/MuR/ImageTypes.h b/Engine/Plugins/Experimental/Mutable/Source/MutableRuntime/Public/MuR/ImageTypes.h
new file mode 100644
--- /dev/null
+++ b/Engine/Plugins/Experimental/Mutable/Source/MutableRuntime/Public/MuR/ImageTypes.h
@@ -0,0 +1,229 @@
+// Copyright Epic Games, Inc. All Rights Reserved.
+
+#pragma once
+
+#include "MuR/MutableMath.h"
+#include "MuR/Serialisation.h"
+
+#include <initializer_list>
+
+namespace mu
+{
+	//!
+	using FImageSize = UE::Math::TIntVector2<uint16>;
+	using FImageRect = box<vec2<uint16>>;
+
+	//! Pixel formats supported by the images.
+	//! \ingroup runtime
+	enum class EImageFormat : uint8
+	{
+		IF_NONE,
+		IF_RGB_UBYTE,
+		IF_RGBA_UBYTE,
+		IF_L_UBYTE,
+
+        //! Deprecated formats
+        IF_PVRTC2_DEPRECATED,
+        IF_PVRTC4_DEPRECATED,
+        IF_ETC1_DEPRECATED,
+        IF_ETC2_DEPRECATED,
+
+		IF_L_UBYTE_RLE,
+		IF_RGB_UBYTE_RLE,
+		IF_RGBA_UBYTE_RLE,
+		IF_L_UBIT_RLE,
+
+        //! Common S3TC formats
+        IF_BC1,
+        IF_BC2,
+        IF_BC3,
+        IF_BC4,
+        IF_BC5,
+
+        //! Not really supported yet
+        IF_BC6,
+        IF_BC7,
+
+        //! Swizzled versions, engineers be damned.
+        IF_BGRA_UBYTE,
+
+        //! The new standard
+        IF_ASTC_4x4_RGB_LDR,
+        IF_ASTC_4x4_RGBA_LDR,
+        IF_ASTC_4x4_RG_LDR,
+
+		IF_ASTC_8x8_RGB_LDR,
+		IF_ASTC_8x8_RGBA_LDR,
+		IF_ASTC_8x8_RG_LDR,
+		IF_ASTC_12x12_RGB_LDR,
+		IF_ASTC_12x12_RGBA_LDR,
+		IF_ASTC_12x12_RG_LDR,
+		IF_ASTC_6x6_RGB_LDR,
+		IF_ASTC_6x6_RGBA_LDR,
+		IF_ASTC_6x6_RG_LDR,
+		IF_ASTC_10x10_RGB_LDR,
+		IF_ASTC_10x10_RGBA_LDR,
+		IF_ASTC_10x10_RG_LDR,
+
+        IF_COUNT
+	};
+	
+	struct MUTABLERUNTIME_API FImageDesc
+	{
+		FImageDesc()
+		{
+		}
+
+		FImageDesc(const FImageSize& InSize, EImageFormat InFormat, uint8 InLods)
+			: m_size(InSize), m_format(InFormat), m_lods(InLods) 
+		{
+		}
+
+		//!
+		FImageSize m_size = FImageSize(0, 0);
+
+		//!
+		EImageFormat m_format = EImageFormat::IF_NONE;
+
+		//! Levels of detail (mipmaps)
+		uint8 m_lods = 0;
+
+		//!
+		FORCEINLINE bool operator==(const FImageDesc& Other) const
+		{
+			return m_format == Other.m_format &&
+				   m_lods == Other.m_lods     &&
+				   m_size == Other.m_size;
+		}
+	};
+
+
+	//! List of supported modes in generic image layering operations.
+	enum class EBlendType
+	{
+		BT_NONE = 0,
+		BT_SOFTLIGHT,
+		BT_HARDLIGHT,
+		BT_BURN,
+		BT_DODGE,
+		BT_SCREEN,
+		BT_OVERLAY,
+		BT_BLEND,
+		BT_MULTIPLY,
+		BT_LIGHTEN,				// Increase the channel value by a given proportion of what is missing from white 
+		BT_NORMAL_COMBINE,
+		_BT_COUNT
+	};
+	
+	enum class EAddressMode
+	{
+		None,
+		Wrap,
+		ClampToEdge,
+		ClampToBlack,
+	};
+
+	enum class EMipmapFilterType
+	{
+		MFT_Unfiltered,
+		MFT_SimpleAverage,
+		MFT_Sharpen,
+		_MFT_COUNT
+	};
+	
+	enum class ECompositeImageMode
+	{
+		CIM_Disabled,
+		CIM_NormalRoughnessToRed,
+		CIM_NormalRoughnessToGreen,
+		CIM_NormalRoughnessToBlue,
+		CIM_NormalRoughnessToAlpha,
+		_CIM_COUNT
+	};
+	
+	enum class ESamplingMethod : uint8
+	{
+		Point = 0,
+		BiLinear,
+		MaxValue
+	};
+	static_assert(uint32(ESamplingMethod::MaxValue) <= (1 << 3), "ESampligMethod enum cannot hold more than 8 values");
+
+	enum class EInitializationType
+	{
+		NotInitialized,
+		Black
+	};
+
+	enum class EMinFilterMethod : uint8
+	{
+		None = 0,
+		TotalAreaHeuristic,
+		MaxValue
+	};
+	static_assert(uint32(EMinFilterMethod::MaxValue) <= (1 << 3), "EMinFilterMethod enum cannot hold more than 8 values");
+
+	/** */
+	struct FImageFormatData
+	{
+		static constexpr SIZE_T MAX_BYTES_PER_BLOCK = 16;
+
+		FImageFormatData
+		(
+			uint8 InPixelsPerBlockX = 0,
+			uint8 InPixelsPerBlockY = 0,
+			uint16 InBytesPerBlock = 0,
+			uint16 InChannels = 0
+		)
+		{
+			PixelsPerBlockX = InPixelsPerBlockX;
+			PixelsPerBlockY = InPixelsPerBlockY;
+			BytesPerBlock = InBytesPerBlock;
+			Channels = InChannels;
+		}
+
+		FImageFormatData
+		(
+			uint8 InPixelsPerBlockX,
+			uint8 InPixelsPerBlockY,
+			uint16 InBytesPerBlock,
+			uint16 InChannels,
+			std::initializer_list<uint8> BlackBlockInit
+		)
+			: FImageFormatData(InPixelsPerBlockX, InPixelsPerBlockY, InBytesPerBlock, InChannels)
+		{
+			check(MAX_BYTES_PER_BLOCK >= BlackBlockInit.size());
+
+			const SIZE_T SanitizedBlockSize = FMath::Min<SIZE_T>(MAX_BYTES_PER_BLOCK, BlackBlockInit.size());
+			FMemory::Memcpy(BlackBlock, BlackBlockInit.begin(), SanitizedBlockSize);
+		}
+
+		/** For block based formats, size of the block size.For uncompressed formats it will always be 1,1. For non-block-based compressed formats, it will be 0,0. */
+		uint8 PixelsPerBlockX, PixelsPerBlockY;
+
+		/** Number of bytes used by every pixel block, if uncompressed or block-compressed format.
+		 * For non-block-compressed formats, it returns 0.
+		 */
+		uint16 BytesPerBlock;
+
+		/** Channels in every pixel of the image. */
+		uint16 Channels;
+
+		/** Representation of a black block of the image. */
+		uint8 BlackBlock[MAX_BYTES_PER_BLOCK] = { 0 };
+	};
+
+	MUTABLERUNTIME_API const FImageFormatData& GetImageFormatData(EImageFormat format);
+
+	struct MUTABLERUNTIME_API FMipmapGenerationSettings
+	{
+		float m_sharpenFactor = 0.0f;
+		EMipmapFilterType m_filterType = EMipmapFilterType::MFT_SimpleAverage;
+		EAddressMode m_addressMode = EAddressMode::None;
+		bool m_ditherMipmapAlpha = false;
+
+		void Serialise(OutputArchive& arch) const;
+		void Unserialise(InputArchive& arch);
+	};
+}
+
