[Mutable] Add console commands to dump mutable global memory stats current and peak. Enable mutable memory peak tracking preprocessor guard by default on debug and development builds.

- Added mutable.EnableGlobalMemoryTracking cvar.
- Added mutable.EnableUpdateLLMStats cvar.
- Added mutable.DumpGlobaleMemoryCounter command.
- Added mutable.ZeroGlobalMemoryCounter command.
- Added mutable.RestoreGlobalMemoryCounter command.

- Programs will now update LLM stats at every allocation. This behavior is conditional on the mutable.EnableGlobalMemoryTracking and mutable.EnableUpdateLLMStats cvars.

#rnx

diff --git a/Engine/Plugins/Experimental/Mutable/Source/MutableRuntime/Private/MuR/MemoryTrackingUtils.cpp b/Engine/Plugins/Experimental/Mutable/Source/MutableRuntime/Private/MuR/MemoryTrackingUtils.cpp
new file mode 100644
--- /dev/null
+++ b/Engine/Plugins/Experimental/Mutable/Source/MutableRuntime/Private/MuR/MemoryTrackingUtils.cpp
@@ -0,0 +1,217 @@
+// Copyright Epic Games, Inc. All Rights Reserved.
+
+#include "MuR/MemoryTrackingUtils.h"
+
+
+#include "MuR/Platform.h"
+#include "Logging/LogMacros.h"
+#include "HAL/LowLevelMemTracker.h"
+#include "HAL/IConsoleManager.h"
+
+namespace mu::Private
+{
+	// Call LLM system tick (we do this to simulate a frame since the LLM system is not entirelly designed to run over a program)
+	FORCEINLINE void UpdateLLMStats()
+	{
+		// This code will only be compiled if the global definition to enable LLM tracking is set to 1 for the host program
+		// E.g.: GlobalDefinitions.Add("LLM_ENABLED_IN_CONFIG=1");
+#if ENABLE_LOW_LEVEL_MEM_TRACKER && IS_PROGRAM
+		FLowLevelMemTracker& MemTracker = FLowLevelMemTracker::Get();
+		if (MemTracker.IsEnabled())
+		{
+			MemTracker.UpdateStatsPerFrame();
+		}
+#endif
+	}
+}
+
+namespace mu
+{
+
+	struct FMemoryTrackingConsoleFlags
+	{
+		static inline bool bEnableGlobalMemoryTracking = true;
+		static inline FAutoConsoleVariableRef CVarEnableGlobalMemoryTracking = FAutoConsoleVariableRef(
+			TEXT("mutable.EnableGlobalMemoryTracking"),
+			bEnableGlobalMemoryTracking,
+			TEXT("Enables Mutable plugin global memory tracking, peak and total."),
+			ECVF_Default);
+
+#if ENABLE_LOW_LEVEL_MEM_TRACKER && IS_PROGRAM
+		static inline bool bEnableUpdateLLMStats = false;
+		static inline FAutoConsoleVariableRef CVarEnableUpdateLLMStats = FAutoConsoleVariableRef(
+			TEXT("mutable.EnableUpdateLLMStats"),
+			bEnableUpdateLLMStats,
+			TEXT("Enables LLM updates at every tracked allocation, only enabled on program builds."),
+			ECVF_Default);
+#endif
+	};
+
+#if UE_MUTABLE_TRACK_ALLOCATOR_MEMORY_PEAK
+	void FGlobalMemoryCounter::Update(SSIZE_T Differential)
+	{
+		if (LIKELY(!FMemoryTrackingConsoleFlags::bEnableGlobalMemoryTracking))
+		{
+			return;
+		}
+
+		FScopeLock Lock(&Mutex);
+
+		Counter += Differential;
+		PeakValue = FMath::Max(PeakValue, Counter);
+
+		AbsoluteCounter += Differential;
+		AbsolutePeakValue = FMath::Max(AbsolutePeakValue, AbsoluteCounter);
+
+#if ENABLE_LOW_LEVEL_MEM_TRACKER && IS_PROGRAM
+		if (UNLIKELY(FMemoryTrackingConsoleFlags::bEnableUpdateLLMStats))
+		{
+			Private::UpdateLLMStats();
+		}
+#endif
+	}
+
+	void FGlobalMemoryCounter::Zero()
+	{
+		if (LIKELY(!FMemoryTrackingConsoleFlags::bEnableGlobalMemoryTracking))
+		{
+			return;
+		}
+
+		FScopeLock Lock(&Mutex);
+
+		Counter = 0;	
+		PeakValue = 0;
+	}
+
+	void FGlobalMemoryCounter::Restore()
+	{
+		if (LIKELY(!FMemoryTrackingConsoleFlags::bEnableGlobalMemoryTracking))
+		{
+			return;
+		}
+
+		FScopeLock Lock(&Mutex);
+
+		Counter = AbsoluteCounter;
+		PeakValue = AbsolutePeakValue;
+	}
+
+	SSIZE_T FGlobalMemoryCounter::GetPeak()
+	{
+		if (LIKELY(!FMemoryTrackingConsoleFlags::bEnableGlobalMemoryTracking))
+		{
+			return 0;
+		}
+
+		FScopeLock Lock(&Mutex);
+
+		const volatile SSIZE_T Result = PeakValue;
+		return Result;
+	}
+
+	SSIZE_T FGlobalMemoryCounter::GetCounter()
+	{
+		if (LIKELY(!FMemoryTrackingConsoleFlags::bEnableGlobalMemoryTracking))
+		{
+			return 0;
+		}
+
+		FScopeLock Lock(&Mutex);
+
+		const volatile SSIZE_T Result = Counter;
+		return Result;
+	}
+#else
+	void FGlobalMemoryCounter::Zero()
+	{
+	}
+
+	void FGlobalMemoryCounter::Restore()
+	{
+	}
+
+	SSIZE_T FGlobalMemoryCounter::GetPeak()
+	{
+		return 0;
+	}
+
+	SSIZE_T FGlobalMemoryCounter::GetCounter()
+	{
+		return 0;
+	}
+#endif
+
+
+	struct FMemoryTrackingCommands
+	{
+		static inline FAutoConsoleCommand CmdDumpMutableGlobalMemoryCounter = FAutoConsoleCommand(
+			TEXT("mutable.DumpMutableGlobalMemoryCounter"),
+			TEXT("Dump Mutable plugin global tracked memory counter."),
+			FConsoleCommandDelegate::CreateStatic([]()
+			{
+				if (!UE_MUTABLE_TRACK_ALLOCATOR_MEMORY_PEAK)
+				{
+					UE_LOG(LogMutableCore, Display, TEXT("The build does not have UE_MUTABLE_TRACK_ALLOCATOR_MEMORY_PEAK enabled, no counting performed."));
+					return;
+				}
+
+				if (!FMemoryTrackingConsoleFlags::bEnableGlobalMemoryTracking)
+				{
+					UE_LOG(LogMutableCore, Display, TEXT("mutable.EnableGlobalMemoryTracking is set to false, no counting performed."));
+					return;
+				}
+
+				UE_LOG(LogMutableCore, 
+					   Display, 
+					   TEXT("Mutable Memory: Current %2.f MiB, Peak %2.f MiB"), 
+					   double(mu::FGlobalMemoryCounter::GetCounter()) / double(1024*1024),
+					   double(mu::FGlobalMemoryCounter::GetPeak())    / double(1024*1024)
+				);
+			})
+		);
+
+		static inline FAutoConsoleCommand CmdZeroMutableGlobalMemoryCounter = FAutoConsoleCommand(
+			TEXT("mutable.ZeroMutableGlobalMemoryCounter"),
+			TEXT("Zero Mutable plugin global tracked memory counter."),
+			FConsoleCommandDelegate::CreateStatic([]()
+			{
+				if (!UE_MUTABLE_TRACK_ALLOCATOR_MEMORY_PEAK)
+				{
+					UE_LOG(LogMutableCore, Display, TEXT("The build does not have UE_MUTABLE_TRACK_ALLOCATOR_MEMORY_PEAK enabled, no counting performed."));
+					return;
+				}
+
+				if (!FMemoryTrackingConsoleFlags::bEnableGlobalMemoryTracking)
+				{
+					UE_LOG(LogMutableCore, Display, TEXT("mutable.EnableGlobalMemoryTracking is set to false, no counting performed."));
+					return;
+				}
+
+				FGlobalMemoryCounter::Zero();	
+			})
+		);
+
+		static inline FAutoConsoleCommand CmdRestoreMutableGlobalMemoryCounter = FAutoConsoleCommand(
+			TEXT("mutable.RestoreMutableGlobalMemoryCounter"),
+			TEXT("Restore Mutable plugin global tracked memory counter."),
+			FConsoleCommandDelegate::CreateStatic([]()
+			{
+				if (!UE_MUTABLE_TRACK_ALLOCATOR_MEMORY_PEAK)
+				{
+					UE_LOG(LogMutableCore, Display, TEXT("The build does not have UE_MUTABLE_TRACK_ALLOCATOR_MEMORY_PEAK enabled, no counting performed."));
+					return;
+				}
+
+				if (!FMemoryTrackingConsoleFlags::bEnableGlobalMemoryTracking)
+				{
+					UE_LOG(LogMutableCore, Display, TEXT("mutable.EnableGlobalMemoryTracking is set to false, no counting performed."));
+					return;
+				}
+
+				FGlobalMemoryCounter::Restore();	
+			})
+		);
+	};
+
+}
diff --git a/Engine/Plugins/Experimental/Mutable/Source/MutableRuntime/Public/MuR/MemoryTrackingAllocationPolicy.h b/Engine/Plugins/Experimental/Mutable/Source/MutableRuntime/Public/MuR/MemoryTrackingAllocationPolicy.h
--- a/Engine/Plugins/Experimental/Mutable/Source/MutableRuntime/Public/MuR/MemoryTrackingAllocationPolicy.h
+++ b/Engine/Plugins/Experimental/Mutable/Source/MutableRuntime/Public/MuR/MemoryTrackingAllocationPolicy.h
@@ -4,98 +4,13 @@
 
 #include "Containers/ContainerAllocationPolicies.h"
 #include "Containers/ContainersFwd.h"
-#include "HAL/CriticalSection.h"
+#include "MuR/MemoryTrackingUtils.h"
 
 #include <atomic>
 #include <type_traits>
 
-
-#ifndef UE_MUTABLE_TRACK_ALLOCATOR_MEMORY_PEAK
-#define UE_MUTABLE_TRACK_ALLOCATOR_MEMORY_PEAK 0
-#endif
-
 namespace mu
 {
-
-#if UE_MUTABLE_TRACK_ALLOCATOR_MEMORY_PEAK
-	struct FGlobalMemoryCounter
-	{
-	private:
-		 static inline SSIZE_T AbsoluteCounter { 0 };
-		 static inline SSIZE_T AbsolutePeakValue { 0 };
-		 static inline SSIZE_T Counter { 0 };
-		 static inline SSIZE_T PeakValue { 0 };
-		
-		 static inline FCriticalSection Mutex {};
-
-	public:
-		static void Update(SSIZE_T Differential)
-		{
-			FScopeLock Lock(&Mutex);
-
-			Counter += Differential;
-			PeakValue = FMath::Max(PeakValue, Counter);
-
-			AbsoluteCounter += Differential;
-			AbsolutePeakValue = FMath::Max(AbsolutePeakValue, AbsoluteCounter);
-		}
-
-		static void Zero()
-		{
-			FScopeLock Lock(&Mutex);
-
-			Counter = 0;	
-			PeakValue = 0;
-		}
-
-		static void Restore()
-		{
-			FScopeLock Lock(&Mutex);
-
-			Counter = AbsoluteCounter;
-			PeakValue = AbsolutePeakValue;
-		}
-
-		static SSIZE_T GetPeak()
-		{
-			FScopeLock Lock(&Mutex);
-
-			const volatile SSIZE_T Result = PeakValue;
-			return Result;
-		}
-
-
-		static SSIZE_T GetCounter()
-		{
-			FScopeLock Lock(&Mutex);
-
-			const volatile SSIZE_T Result = Counter;
-			return Result;
-		}
-	};
-#else
-	struct FGlobalMemoryCounter
-	{
-		static void Zero()
-		{
-		}
-
-		static void Restore()
-		{
-		}
-
-		static SSIZE_T GetPeak()
-		{
-			return 0;
-		}
-
-		static SSIZE_T GetCounter()
-		{
-			return 0;
-		}
-	};
-#endif
-
 	template<typename TagType>
 	struct TMemoryCounter
 	{
diff --git a/Engine/Plugins/Experimental/Mutable/Source/MutableRuntime/Public/MuR/MemoryTrackingUtils.h b/Engine/Plugins/Experimental/Mutable/Source/MutableRuntime/Public/MuR/MemoryTrackingUtils.h
new file mode 100644
--- /dev/null
+++ b/Engine/Plugins/Experimental/Mutable/Source/MutableRuntime/Public/MuR/MemoryTrackingUtils.h
@@ -0,0 +1,47 @@
+// Copyright Epic Games, Inc. All Rights Reserved.
+
+#pragma once
+
+#include "Misc/Build.h"
+#include "HAL/Platform.h"
+#include "HAL/CriticalSection.h"
+
+#ifndef UE_MUTABLE_TRACK_ALLOCATOR_MEMORY_PEAK
+	#if UE_BUILD_DEBUG || UE_BUILD_DEVELOPMENT
+		#define UE_MUTABLE_TRACK_ALLOCATOR_MEMORY_PEAK 1
+	#else
+		#define UE_MUTABLE_TRACK_ALLOCATOR_MEMORY_PEAK 0
+	#endif
+#endif
+
+namespace mu
+{
+#if UE_MUTABLE_TRACK_ALLOCATOR_MEMORY_PEAK 
+	struct MUTABLERUNTIME_API FGlobalMemoryCounter
+	{
+	private:
+		 static inline SSIZE_T AbsoluteCounter   { 0 };
+		 static inline SSIZE_T AbsolutePeakValue { 0 };
+		 static inline SSIZE_T Counter           { 0 };
+		 static inline SSIZE_T PeakValue         { 0 };
+		
+		 static inline FCriticalSection Mutex {};
+
+	public:
+		 static void Update(SSIZE_T Differential);
+		 static void Zero();
+		 static void Restore();
+		 static SSIZE_T GetPeak();
+		 static SSIZE_T GetCounter();
+	};
+#else
+	struct MUTABLERUNTIME_API FGlobalMemoryCounter
+	{
+		 //static void Update(SSIZE_T Differential);
+		 static void Zero();
+		 static void Restore();
+		 static SSIZE_T GetPeak();
+		 static SSIZE_T GetCounter();
+	};
+#endif
+}
