Undo CL 29126522 LogBenchmarkUtil.cpp (revision 5)

Causing a crash on boot:
"Attempting to get the command line but it hasn't been initialized yet."

diff --git a/Engine/Plugins/Experimental/Mutable/Source/CustomizableObject/Private/MuCO/LogBenchmarkUtil.cpp b/Engine/Plugins/Experimental/Mutable/Source/CustomizableObject/Private/MuCO/LogBenchmarkUtil.cpp
--- a/Engine/Plugins/Experimental/Mutable/Source/CustomizableObject/Private/MuCO/LogBenchmarkUtil.cpp
+++ b/Engine/Plugins/Experimental/Mutable/Source/CustomizableObject/Private/MuCO/LogBenchmarkUtil.cpp
@@ -29,10 +29,6 @@ static TAutoConsoleVariable<bool> CVarEnableBenchmark(
 	false,
 	TEXT("Enable or disable the benchmarking."));
 
-static TAutoConsoleVariable<FString> CVarBenchmarkFilePath(
-	 TEXT("mutable.BenchmarkFilePath"),
-	FPaths::ProfilingDir() + TEXT("Mutable/Benchmark"),
-	TEXT("Sets the path where to store the generated mutable benchmark report file."));
 
 namespace LogBenchmarkUtil
 {
@@ -44,7 +40,7 @@ namespace LogBenchmarkUtil
 
 TSharedPtr<FArchive> CreateFile()
 {
-	const FString Directory = CVarBenchmarkFilePath.GetValueOnGameThread();
+	const FString Directory = FPaths::ProfilingDir() + TEXT("Mutable/Benchmark/");
 	IFileManager::Get().MakeDirectory(*Directory, true);
 
 	const FDateTime FileDate = FDateTime::Now();
