Resolve CL: 29126522 to not use the FCommandLine::Get() before it could be initialized

"Attempting to get the command line but it hasn't been initialized yet."

diff --git a/Engine/Plugins/Experimental/Mutable/Source/CustomizableObject/Private/MuCO/LogBenchmarkUtil.cpp b/Engine/Plugins/Experimental/Mutable/Source/CustomizableObject/Private/MuCO/LogBenchmarkUtil.cpp
--- a/Engine/Plugins/Experimental/Mutable/Source/CustomizableObject/Private/MuCO/LogBenchmarkUtil.cpp
+++ b/Engine/Plugins/Experimental/Mutable/Source/CustomizableObject/Private/MuCO/LogBenchmarkUtil.cpp
@@ -29,18 +29,27 @@ static TAutoConsoleVariable<bool> CVarEnableBenchmark(
 	false,
 	TEXT("Enable or disable the benchmarking."));
 
-
 namespace LogBenchmarkUtil
 {
 	void Write(FArchive& Archive, FStringView Text)
 	{
 		Archive.Serialize(const_cast<ANSICHAR*>(StringCast<ANSICHAR>(Text.GetData(), Text.Len()).Get()), Text.Len() * sizeof(ANSICHAR));
 	}
+
+	static FString LocalBenchmarkFilePath = TEXT("");
+	static FAutoConsoleVariableRef CVarBenchmarkFilePath(
+		TEXT("mutable.BenchmarkFilePath"),
+		LocalBenchmarkFilePath,
+		TEXT("Sets the path where to store the generated mutable benchmark report file"));
 }
 
 TSharedPtr<FArchive> CreateFile()
 {
-	const FString Directory = FPaths::ProfilingDir() + TEXT("Mutable/Benchmark/");
+	FString Directory = LogBenchmarkUtil::LocalBenchmarkFilePath;
+	if (Directory.IsEmpty())
+	{
+		Directory = FPaths::ProfilingDir() + TEXT("Mutable/Benchmark");
+	}
 	IFileManager::Get().MakeDirectory(*Directory, true);
 
 	const FDateTime FileDate = FDateTime::Now();
@@ -54,6 +63,13 @@ TSharedPtr<FArchive> CreateFile()
 	return Archive;
 }
 
+FLogBenchmarkUtil::FLogBenchmarkUtil()
+{
+	if (LogBenchmarkUtil::LocalBenchmarkFilePath.IsEmpty())
+	{
+		LogBenchmarkUtil::LocalBenchmarkFilePath = FPaths::ProfilingDir() + TEXT("Mutable/Benchmark");
+	}
+}
 
 FLogBenchmarkUtil::~FLogBenchmarkUtil()
 {
diff --git a/Engine/Plugins/Experimental/Mutable/Source/CustomizableObject/Private/MuCO/LogBenchmarkUtil.h b/Engine/Plugins/Experimental/Mutable/Source/CustomizableObject/Private/MuCO/LogBenchmarkUtil.h
--- a/Engine/Plugins/Experimental/Mutable/Source/CustomizableObject/Private/MuCO/LogBenchmarkUtil.h
+++ b/Engine/Plugins/Experimental/Mutable/Source/CustomizableObject/Private/MuCO/LogBenchmarkUtil.h
@@ -66,6 +66,7 @@ DECLARE_BENCHMARK_INSIGHTS(InstanceBuildTimeAvrg, TEXT("Avrg Instance Build Time
 class FLogBenchmarkUtil
 {
 public:
+	FLogBenchmarkUtil();
 	~FLogBenchmarkUtil();
 
 	/** Enable or disable the system. */
