[mutable] Fixes in mesh clipping for meshes whose buffers have not been cache optimized.
#rnx

diff --git a/Engine/Plugins/Experimental/Mutable/Source/MutableRuntime/Private/MuR/OpMeshApplyLayout.h b/Engine/Plugins/Experimental/Mutable/Source/MutableRuntime/Private/MuR/OpMeshApplyLayout.h
--- a/Engine/Plugins/Experimental/Mutable/Source/MutableRuntime/Private/MuR/OpMeshApplyLayout.h
+++ b/Engine/Plugins/Experimental/Mutable/Source/MutableRuntime/Private/MuR/OpMeshApplyLayout.h
@@ -195,7 +195,7 @@ namespace mu
 			}
 	
 			// Remove
-			MeshRemoveVerticesWithMap( pApplied, VertexMask.GetData(), VertexMask.Num());
+			MeshRemoveVerticesWithMap( pApplied, VertexMask.GetData(), (uint32)VertexMask.Num());
 		}
 
 		//
diff --git a/Engine/Plugins/Experimental/Mutable/Source/MutableRuntime/Private/MuR/OpMeshRemove.cpp b/Engine/Plugins/Experimental/Mutable/Source/MutableRuntime/Private/MuR/OpMeshRemove.cpp
--- a/Engine/Plugins/Experimental/Mutable/Source/MutableRuntime/Private/MuR/OpMeshRemove.cpp
+++ b/Engine/Plugins/Experimental/Mutable/Source/MutableRuntime/Private/MuR/OpMeshRemove.cpp
@@ -74,51 +74,57 @@ namespace mu
     }
 
 
-	void MeshRemoveVerticesWithMap( Mesh* Result, uint8* removedVertices, int32 resultVertexCount)
+	void MeshRemoveVerticesWithMap( Mesh* Result, const uint8* RemovedVertices, uint32 RemovedVertexCount)
 	{
 		int32 firstFreeVertex = 0;
 
         // Rebuild index buffers
 
         // Map from source vertex index, to new vertex index for used vertices.
-        // These are indices as in the index buffer, not the absoulte vertex index as in the
+        // These are indices as in the index buffer, not the absolute vertex index as in the
         // vertexbuffer MBS_VERTEXINDEX buffers.
-		TArray<int32> usedVertices;
-		usedVertices.Init( -1, Result->GetVertexCount() );
+		TArray<int32> UsedVertexMap;
+		UsedVertexMap.Init( -1, Result->GetVertexCount() );
         {
-            SIZE_T removedIndices = 0;
+            int32 removedIndices = 0;
 
             if ( Result->GetIndexBuffers().GetElementSize(0)==4 )
             {
                 MeshBufferIteratorConst<MBF_UINT32,uint32,1> itSource( Result->GetIndexBuffers(), MBS_VERTEXINDEX );
                 MeshBufferIterator<MBF_UINT32,uint32,1> itDest( Result->GetIndexBuffers(), MBS_VERTEXINDEX );
 
-                int32 indexCount = Result->GetIndexCount();
-                for ( int32 f=0; f<indexCount/3; ++f )
+                int32 IndexCount = Result->GetIndexCount();
+				check(IndexCount%3==0);
+                for ( int32 f=0; f<IndexCount/3; ++f )
                 {
                     uint32 sourceIndices[3];
                     sourceIndices[0] = itSource[0][0];
                     sourceIndices[1] = itSource[1][0];
                     sourceIndices[2] = itSource[2][0];
 
-                    bool faceRemoved =  ( removedVertices[ sourceIndices[0] ]
-                            + removedVertices[ sourceIndices[1] ]
-                            + removedVertices[ sourceIndices[2] ]
-                            ) > 0;
+					check(sourceIndices[0] < RemovedVertexCount);
+					check(sourceIndices[1] < RemovedVertexCount);
+					check(sourceIndices[2] < RemovedVertexCount);
+
+					int32 RemoveCount =  (RemovedVertices[ sourceIndices[0] ]
+                            + RemovedVertices[ sourceIndices[1] ]
+                            + RemovedVertices[ sourceIndices[2] ]
+                            );
+					bool bFaceRemoved = RemoveCount > 0;
 
-                    if ( !faceRemoved )
+                    if ( !bFaceRemoved)
                     {
                         for (int32 i=0;i<3;++i)
                         {
                             uint32 sourceIndex = sourceIndices[i];
 
-                            if ( usedVertices[ sourceIndex ] < 0 )
+                            if (UsedVertexMap[ sourceIndex ] < 0 )
                             {
-                                usedVertices[ sourceIndex ] = firstFreeVertex;
+								UsedVertexMap[ sourceIndex ] = firstFreeVertex;
                                 firstFreeVertex++;
                             }
 
-                            uint32 destIndex = usedVertices[ sourceIndex ];
+                            uint32 destIndex = UsedVertexMap[ sourceIndex ];
                             *(uint32*)itDest.ptr() = destIndex;
 
                             itDest++;
@@ -144,24 +150,29 @@ namespace mu
                     sourceIndices[1] = itSource[1][0];
                     sourceIndices[2] = itSource[2][0];
 
-                    bool faceRemoved =  ( removedVertices[ sourceIndices[0] ]
-                            + removedVertices[ sourceIndices[1] ]
-                            + removedVertices[ sourceIndices[2] ]
-                            ) > 0;
+					check(sourceIndices[0] < RemovedVertexCount);
+					check(sourceIndices[1] < RemovedVertexCount);
+					check(sourceIndices[2] < RemovedVertexCount);
 
-                    if ( !faceRemoved )
+					int32 RemoveCount =  (RemovedVertices[ sourceIndices[0] ]
+                            + RemovedVertices[ sourceIndices[1] ]
+                            + RemovedVertices[ sourceIndices[2] ]
+                            );
+					bool bFaceRemoved = RemoveCount > 0;
+
+                    if ( !bFaceRemoved)
                     {
                         for (int i=0;i<3;++i)
                         {
                             uint16 sourceIndex = sourceIndices[i];
 
-                            if ( usedVertices[ sourceIndex ] < 0 )
+                            if (UsedVertexMap[ sourceIndex ] < 0 )
                             {
-                                usedVertices[ sourceIndex ] = firstFreeVertex;
+								UsedVertexMap[ sourceIndex ] = firstFreeVertex;
                                 firstFreeVertex++;
                             }
 
-                            uint16 destIndex = (uint16)usedVertices[ sourceIndex ];
+                            uint16 destIndex = (uint16)UsedVertexMap[ sourceIndex ];
                             *(uint16*)itDest.ptr() = destIndex;
 
                             itDest++;
@@ -176,7 +187,7 @@ namespace mu
 
             else
             {
-                // Case not implemented
+                // Index buffer format case not implemented
                 check( false );
             }
 
@@ -187,18 +198,26 @@ namespace mu
             Result->GetIndexBuffers().SetElementCount( faceCount*3-(int)removedIndices );
         }
 
+
         // Rebuild the vertex buffers
-        for ( int32 b=0; b<Result->GetVertexBuffers().GetBufferCount(); ++b )
+
+		// The temp array is necessary because if the vertex buffer is not sorted according to the index buffer we cannot do it in-place
+		// This happens with some mesh import options.
+		TArray<uint8> Temp;
+		for ( int32 b=0; b<Result->GetVertexBuffers().GetBufferCount(); ++b )
         {
             int32 elemSize = Result->GetVertexBuffers().GetElementSize( b );
-            const uint8* pSourceData = Result->GetVertexBuffers().GetBufferData( b );
-            uint8* pData = Result->GetVertexBuffers().GetBufferData( b );
+            const uint8* SourceData = Result->GetVertexBuffers().GetBufferData( b );
+
+			Temp.SetNumUninitialized(firstFreeVertex*elemSize,false);
+            uint8* DestData = Temp.GetData();
+
             for ( int32 v=0; v<Result->GetVertexCount(); ++v )
             {
                 int32 span = 0;
                 for ( int32 s=0; v+s<Result->GetVertexCount(); ++s )
                 {
-                    if ( usedVertices[v+s]>=0 )
+                    if (UsedVertexMap[v+s]>=0 )
                     {
                         if (span==0)
                         {
@@ -206,7 +225,7 @@ namespace mu
                         }
                         else
                         {
-                            if ( usedVertices[v+s] == usedVertices[v+s-1]+1 )
+                            if (UsedVertexMap[v+s] == UsedVertexMap[v+s-1]+1 )
                             {
                                 ++span;
                             }
@@ -225,11 +244,14 @@ namespace mu
 
                 if (span>0)
                 {
-                    FMemory::Memmove( pData+elemSize*usedVertices[v], pSourceData+elemSize*v, elemSize*span );
+                    FMemory::Memcpy( DestData+elemSize*UsedVertexMap[v], SourceData+elemSize*v, elemSize*span );
 
                     v += span-1;
                 }
             }
+
+			// Copy from temp buffer to final vertex buffer
+			FMemory::Memcpy( Result->GetVertexBuffers().GetBufferData(b), DestData, firstFreeVertex*elemSize);
         }
         Result->GetVertexBuffers().SetElementCount( firstFreeVertex );
 
@@ -240,47 +262,44 @@ namespace mu
     }
 
 
-	void MeshRemoveMask(Mesh* Result, const Mesh* pSource, const Mesh* pMask, bool& bOutSuccess)
+	void MeshRemoveMask(Mesh* Result, const Mesh* Source, const Mesh* Mask, bool& bOutSuccess)
     {
         MUTABLE_CPUPROFILER_SCOPE(MeshRemoveMask);
 		bOutSuccess = true;
 
-        if (!pMask->GetVertexCount() || !pSource->GetVertexCount() || !pSource->GetIndexCount())
+        if (!Mask->GetVertexCount() || !Source->GetVertexCount() || !Source->GetIndexCount())
         {
 			bOutSuccess = false;
             return;
         }
 
-		Result->CopyFrom(*pSource);
+		Result->CopyFrom(*Source);
 
-        MeshBufferIteratorConst<MBF_UINT32,uint32,1> itMaskVI( pMask->GetVertexBuffers(), MBS_VERTEXINDEX );
-
-        int32 firstFreeVertex = 0;
+		int32 MaskElementCount = Mask->GetVertexBuffers().GetElementCount();
 
         // For each source vertex, true if it is removed.
-        int32 resultVertexCount = Result->GetVertexCount();
-		uint8* removedVertices = reinterpret_cast<uint8*>(FMemory::Malloc(resultVertexCount, 16));
-		FMemory::Memzero(removedVertices,resultVertexCount);
+        int32 ResultVertexCount = Result->GetVertexCount();
+		TArray<uint8> RemovedVertices;
+		RemovedVertices.SetNumZeroed(ResultVertexCount);
         {
-			TArray<ID_INTERVAL> intervals;
-            ExtractVertexIndexIntervals(intervals, Result);
+			TArray<ID_INTERVAL> Intervals;
+            ExtractVertexIndexIntervals(Intervals, Result);
 
-            for ( int32 mv=0; mv<pMask->GetVertexBuffers().GetElementCount(); ++mv )
+			MeshBufferIteratorConst<MBF_UINT32, uint32, 1> itMaskVI(Mask->GetVertexBuffers(), MBS_VERTEXINDEX);
+			for ( int32 mv=0; mv<MaskElementCount; ++mv )
             {
-                uint32 maskVertexId = (*itMaskVI)[0];
+                uint32 MaskVertexId = (*itMaskVI)[0];
                 ++itMaskVI;
 
-                int32 indexInSource = FindPositionInIntervals(intervals, maskVertexId);
-                if (indexInSource>=0)
+                int32 IndexInSource = FindPositionInIntervals(Intervals, MaskVertexId);
+                if (IndexInSource >=0)
                 {
-                    removedVertices[indexInSource] = 1;
+					RemovedVertices[IndexInSource] = 1;
                 }
             }
         }
 
-		MeshRemoveVerticesWithMap(Result, removedVertices, resultVertexCount);
-
-		FMemory::Free(removedVertices);
+		MeshRemoveVerticesWithMap(Result, RemovedVertices.GetData(), (uint32)RemovedVertices.Num());
 	}
 
 }
diff --git a/Engine/Plugins/Experimental/Mutable/Source/MutableRuntime/Private/MuR/OpMeshRemove.h b/Engine/Plugins/Experimental/Mutable/Source/MutableRuntime/Private/MuR/OpMeshRemove.h
--- a/Engine/Plugins/Experimental/Mutable/Source/MutableRuntime/Private/MuR/OpMeshRemove.h
+++ b/Engine/Plugins/Experimental/Mutable/Source/MutableRuntime/Private/MuR/OpMeshRemove.h
@@ -14,6 +14,6 @@ namespace mu
 	extern void MeshRemoveMask(Mesh* Result, const Mesh* Source, const Mesh* Mask, bool& bOutSuccess);
 
 	/** Remove a list of vertices and related faces from a mesh. The list is stored as a bool map for every vertex in the mesh. */
-	extern void MeshRemoveVerticesWithMap(Mesh* Result, uint8* RemovedVertices, int32 RemovedVerticesCount);
+	extern void MeshRemoveVerticesWithMap(Mesh* Result, const uint8* RemovedVertices, uint32 RemovedVerticesCount);
 
 }
