[Mutable] Server must always have the reference skeletal mesh loaded.

#rnx

diff --git a/Engine/Plugins/Experimental/Mutable/Source/CustomizableObject/Classes/MuCO/CustomizableObject.h b/Engine/Plugins/Experimental/Mutable/Source/CustomizableObject/Classes/MuCO/CustomizableObject.h
--- a/Engine/Plugins/Experimental/Mutable/Source/CustomizableObject/Classes/MuCO/CustomizableObject.h
+++ b/Engine/Plugins/Experimental/Mutable/Source/CustomizableObject/Classes/MuCO/CustomizableObject.h
@@ -1010,7 +1010,7 @@ struct FMutableRefSkeletalMeshData
 #if WITH_EDITORONLY_DATA
 	friend FArchive& operator<<(FArchive& Ar, FMutableRefSkeletalMeshData& Data);
 
-	void InitResources(UCustomizableObject* InOuter);
+	void InitResources(UCustomizableObject* InOuter, const ITargetPlatform* InTargetPlatform);
 #endif
 
 };
@@ -1535,7 +1535,7 @@ public:
 	
 	/** Generic Save/Load methods to write/read compiled data */
 	void SaveCompiledData(FArchive& Ar, bool bSkipEditorOnlyData = false);
-	void LoadCompiledData(FArchive& Ar, bool bSkipEditorOnlyData = false);
+	void LoadCompiledData(FArchive& Ar, const ITargetPlatform* InTargetPlatform, bool bSkipEditorOnlyData = false);
 
 	/** Load compiled data from disk, this is used to load Editor Compilations and Cook Compilations (when using OnCookStart) */
 	void LoadCompiledDataFromDisk(bool bIsEditorData = true, const ITargetPlatform* InTargetPlatform = nullptr);
diff --git a/Engine/Plugins/Experimental/Mutable/Source/CustomizableObject/Private/MuCO/CustomizableObject.cpp b/Engine/Plugins/Experimental/Mutable/Source/CustomizableObject/Private/MuCO/CustomizableObject.cpp
--- a/Engine/Plugins/Experimental/Mutable/Source/CustomizableObject/Private/MuCO/CustomizableObject.cpp
+++ b/Engine/Plugins/Experimental/Mutable/Source/CustomizableObject/Private/MuCO/CustomizableObject.cpp
@@ -114,11 +114,12 @@ void UCustomizableObject::PreSave(FObjectPreSaveContext ObjectSaveContext)
 #if WITH_EDITORONLY_DATA
 	if (ObjectSaveContext.IsCooking() && !bIsChildObject)
 	{
-		if (const FMutableCachedPlatformData* PlatformData = CachedPlatformsData.Find(ObjectSaveContext.GetTargetPlatform()->PlatformName()))
+		const ITargetPlatform* TargetPlatform = ObjectSaveContext.GetTargetPlatform();
+		if (const FMutableCachedPlatformData* PlatformData = CachedPlatformsData.Find(TargetPlatform->PlatformName()))
 		{
 			// Load cached data before saving
 			FMemoryReaderView MemoryReader(PlatformData->ModelData);
-			LoadCompiledData(MemoryReader, true);
+			LoadCompiledData(MemoryReader, TargetPlatform, true);
 
 			// Create an export object to manage the streamable data
 			BulkData = NewObject<UCustomizableObjectBulk>(this);
@@ -253,7 +254,9 @@ void UCustomizableObject::Serialize(FArchive& Ar_Asset)
 		
 		if (Ar_Asset.IsLoading())
 		{
-			LoadCompiledDataFromDisk();
+			ITargetPlatformManagerModule& TargetPlatformManager = GetTargetPlatformManagerRef();
+			const ITargetPlatform* RunningPlatform = TargetPlatformManager.GetRunningTargetPlatform();
+			LoadCompiledDataFromDisk(true, RunningPlatform);
 		}
 	}
 #else
@@ -497,7 +500,7 @@ void UCustomizableObject::SaveCompiledData(FArchive& MemoryWriter, bool bSkipEdi
 	MemoryWriter << LODSettings.bLODStreamingEnabled;
 }
 
-void UCustomizableObject::LoadCompiledData(FArchive& MemoryReader, bool bSkipEditorOnlyData)
+void UCustomizableObject::LoadCompiledData(FArchive& MemoryReader, const ITargetPlatform* InTargetPlatform, bool bSkipEditorOnlyData)
 {
 	PrivateData->SetModel(nullptr, FGuid());
 	ClearCompiledData();
@@ -515,7 +518,7 @@ void UCustomizableObject::LoadCompiledData(FArchive& MemoryReader, bool bSkipEdi
 		// Initialize resources. 
 		for(FMutableRefSkeletalMeshData& ReferenceSkeletalMeshData : ReferenceSkeletalMeshesData)
 		{
-			ReferenceSkeletalMeshData.InitResources(this);
+			ReferenceSkeletalMeshData.InitResources(this, InTargetPlatform);
 		}
 
 		// We can load
@@ -681,7 +684,7 @@ void UCustomizableObject::LoadCompiledDataFromDisk(bool bIsEditorData, const ITa
 				CompiledDataFileHandle->Read(CompiledDataBytes.GetData(), CompiledDataSize);
 
 				FMemoryReaderView MemoryReader(CompiledDataBytes);
-				LoadCompiledData(MemoryReader);
+				LoadCompiledData(MemoryReader, InTargetPlatform);
 			}
 			else if (!bIsEditorData)// Caching Cooked Data
 			{
@@ -2151,10 +2154,12 @@ FArchive& operator<<(FArchive& Ar, FMutableRefSkeletalMeshData& Data)
 }
 
 
-void FMutableRefSkeletalMeshData::InitResources(UCustomizableObject* InOuter)
+void FMutableRefSkeletalMeshData::InitResources(UCustomizableObject* InOuter, const ITargetPlatform* InTargetPlatform)
 {
 	check(InOuter);
-	if (InOuter->IsEnableUseRefSkeletalMeshAsPlaceholder())
+
+	const bool bHasServer = InTargetPlatform ? !InTargetPlatform->IsClientOnly() : false;	
+	if (InOuter->IsEnableUseRefSkeletalMeshAsPlaceholder() || bHasServer)
 	{
 		SkeletalMesh = TSoftObjectPtr<USkeletalMesh>(SkeletalMeshAssetPath).LoadSynchronous();
 	}
