[Mutable] Fix LOD section remapping getting wrong sections when LODMaterialMap is valid but the section has not changed.

#rnx

diff --git a/Engine/Plugins/Experimental/Mutable/Source/CustomizableObjectEditor/Private/MuCOE/GenerateMutableSource/GenerateMutableSourceMesh.cpp b/Engine/Plugins/Experimental/Mutable/Source/CustomizableObjectEditor/Private/MuCOE/GenerateMutableSource/GenerateMutableSourceMesh.cpp
--- a/Engine/Plugins/Experimental/Mutable/Source/CustomizableObjectEditor/Private/MuCOE/GenerateMutableSource/GenerateMutableSourceMesh.cpp
+++ b/Engine/Plugins/Experimental/Mutable/Source/CustomizableObjectEditor/Private/MuCOE/GenerateMutableSource/GenerateMutableSourceMesh.cpp
@@ -83,7 +83,9 @@ void GetLODAndSectionForAutomaticLODs(const FMutableGraphGenerationContext& Cont
 	const TArray<int32>& FromMaterialMap = SkeletalMesh.GetLODInfoArray()[LODIndexConnected].LODMaterialMap;
 	
 	// Material Index of the connected pin
-	const int32 SearchLODMaterialIndex = FromMaterialMap.IsValidIndex(SectionIndexConnected) ? FromMaterialMap[SectionIndexConnected] : FromSection.MaterialIndex;
+	const int32 SearchLODMaterialIndex = FromMaterialMap.IsValidIndex(SectionIndexConnected) && SkeletalMesh.GetMaterials().IsValidIndex(FromMaterialMap[SectionIndexConnected]) ?
+		FromMaterialMap[SectionIndexConnected] : 
+		FromSection.MaterialIndex;
 
 	const int32 CompilingLODIndex = LODIndexConnected + (Context.CurrentLOD - Context.FromLOD);
 	if (!ImportedModel->LODModels.IsValidIndex(CompilingLODIndex))
@@ -99,7 +101,9 @@ void GetLODAndSectionForAutomaticLODs(const FMutableGraphGenerationContext& Cont
 	bool bFound = false;
 	for (int32 SectionIndex = 0; SectionIndex < LODModel.Sections.Num(); ++SectionIndex)
 	{
-		const int32 MaterialIndex = MaterialMap.IsValidIndex(SectionIndex) ? MaterialMap[SectionIndex] : LODModel.Sections[SectionIndex].MaterialIndex; // MaterialMap overrides the MaterialIndex in the section
+		const int32 MaterialIndex =  MaterialMap.IsValidIndex(SectionIndex) && SkeletalMesh.GetMaterials().IsValidIndex(MaterialMap[SectionIndex]) ?
+			MaterialMap[SectionIndex] :
+			LODModel.Sections[SectionIndex].MaterialIndex; // MaterialMap overrides the MaterialIndex in the section
 			
 		if (MaterialIndex == SearchLODMaterialIndex &&
 			!LODModel.Sections[SectionIndex].bDisabled)
