[mutable] Ensure Mesh Switching nodes try to find layouts ignoring the unconnected options.
#rnx

diff --git a/Engine/Plugins/Experimental/Mutable/Source/MutableTools/Private/MuT/CodeGenerator_Mesh.cpp b/Engine/Plugins/Experimental/Mutable/Source/MutableTools/Private/MuT/CodeGenerator_Mesh.cpp
--- a/Engine/Plugins/Experimental/Mutable/Source/MutableTools/Private/MuT/CodeGenerator_Mesh.cpp
+++ b/Engine/Plugins/Experimental/Mutable/Source/MutableTools/Private/MuT/CodeGenerator_Mesh.cpp
@@ -1082,30 +1082,39 @@ namespace mu
         }
 
         // Options
+		bool bFirstValidConnectionFound = false;
         for ( int32 t=0; t< node.m_options.Num(); ++t )
         {
 			FMeshGenerationOptions TargetOptions = InOptions;
 
-            if (t!=0)
-            {
-				TargetOptions.OverrideLayouts = OutResult.GeneratedLayouts;
-            }
-
             if ( node.m_options[t] )
             {
-                FMeshGenerationResult BranchResults;
+				// Take the layouts from the first non-null connection.
+				// \TODO: Take them from the first connection that actually returns layouts?
+				if (!bFirstValidConnectionFound)
+				{
+					TargetOptions.OverrideLayouts = OutResult.GeneratedLayouts;
+				}
+				
+				FMeshGenerationResult BranchResults;
                 GenerateMesh(TargetOptions, BranchResults, node.m_options[t] );
 
-                auto branch = BranchResults.meshOp;
+                Ptr<ASTOp> branch = BranchResults.meshOp;
                 op->cases.Emplace((int16)t,op,branch);
 
-                if (t==0)
-                {
+				if (!bFirstValidConnectionFound)
+				{
+					bFirstValidConnectionFound = true;
                     OutResult = BranchResults;
                 }
             }
         }
 
+		if (!bFirstValidConnectionFound)
+		{
+			ensure(false);
+		}
+
         OutResult.meshOp = op;
     }
 
@@ -1116,9 +1125,10 @@ namespace mu
 		//
 		FMeshGenerationResult NewResult = OutResult;
 		int t = 0;
+		bool bFirstRowGenerated = false;
 
 		Ptr<ASTOp> Op = GenerateTableSwitch<NodeMeshTable::Private, ETableColumnType::Mesh, OP_TYPE::ME_SWITCH>(*TableNode->GetPrivate(),
-			[this, &NewResult, &t, &InOptions] (const NodeMeshTable::Private& node, int colIndex, int row, ErrorLog* pErrorLog)
+			[this, &NewResult, &bFirstRowGenerated, &InOptions] (const NodeMeshTable::Private& node, int colIndex, int row, ErrorLog* pErrorLog)
 			{
 				mu::Ptr<mu::Mesh> pMesh = node.m_pTable->GetPrivate()->Rows[row].Values[colIndex].Mesh;
 				FMeshGenerationResult BranchResults;
@@ -1138,7 +1148,7 @@ namespace mu
 
 					FMeshGenerationOptions TargetOptions = InOptions;
 
-					if (t != 0)
+					if (bFirstRowGenerated)
 					{
 						TargetOptions.OverrideLayouts = NewResult.GeneratedLayouts;
 					}
@@ -1147,12 +1157,11 @@ namespace mu
 
 					GenerateMesh(TargetOptions, BranchResults, pCell);
 
-					if (t == 0)
+					if (!bFirstRowGenerated)
 					{
 						NewResult = BranchResults;
+						bFirstRowGenerated = true;
 					}
-
-					++t;
 				}
 
 				return BranchResults.meshOp;
diff --git a/Engine/Plugins/Experimental/Mutable/Source/MutableTools/Private/MuT/NodeMeshSwitch.cpp b/Engine/Plugins/Experimental/Mutable/Source/MutableTools/Private/MuT/NodeMeshSwitch.cpp
--- a/Engine/Plugins/Experimental/Mutable/Source/MutableTools/Private/MuT/NodeMeshSwitch.cpp
+++ b/Engine/Plugins/Experimental/Mutable/Source/MutableTools/Private/MuT/NodeMeshSwitch.cpp
@@ -119,17 +119,17 @@ namespace mu
 	//---------------------------------------------------------------------------------------------
 	NodeLayoutPtr NodeMeshSwitch::Private::GetLayout( int index ) const
 	{
-		NodeLayoutPtr pResult;
-
-		if (m_options.Num()>0 && m_options[0] )
+		for (int32 i=0; i<m_options.Num(); ++i)
 		{
-			NodeMesh::Private* pPrivate =
-					dynamic_cast<NodeMesh::Private*>( m_options[0]->GetBasePrivate() );
-
-			pResult = pPrivate->GetLayout( index );
+			if (m_options[i])
+			{
+				NodeMesh::Private* pPrivate = dynamic_cast<NodeMesh::Private*>(m_options[i]->GetBasePrivate());
+				NodeLayoutPtr pResult = pPrivate->GetLayout(index);
+				return pResult;
+			}
 		}
 
-		return pResult;
+		return nullptr;
 	}
 
 
