[mutable] Some fixes to modifiers: assert fix with multiple components and modifiers with variations.
Add more preview protection for missing materials and UV render.
#rnx

diff --git a/Engine/Plugins/Experimental/Mutable/Source/CustomizableObjectEditor/Private/MuCOE/GenerateMutableSource/GenerateMutableSource.cpp b/Engine/Plugins/Experimental/Mutable/Source/CustomizableObjectEditor/Private/MuCOE/GenerateMutableSource/GenerateMutableSource.cpp
--- a/Engine/Plugins/Experimental/Mutable/Source/CustomizableObjectEditor/Private/MuCOE/GenerateMutableSource/GenerateMutableSource.cpp
+++ b/Engine/Plugins/Experimental/Mutable/Source/CustomizableObjectEditor/Private/MuCOE/GenerateMutableSource/GenerateMutableSource.cpp
@@ -1331,9 +1331,20 @@ bool AffectsCurrentComponent(const UEdGraphPin* Pin, FMutableGraphGenerationCont
 	}
 	else if (const UCustomizableObjectNodeModifierBase* TypedNodeModifier = Cast<UCustomizableObjectNodeModifierBase>(Node))
 	{
-		// Modifiers affect all compoenents at lod level. This branch should never be reached.
-		check(false);
-		return false;
+		// Because of the current implementation, modifiers affect all componeents at lod level. If there is only one component it is ok, but otherwise rise an error.
+		if (GenerationContext.NumMeshComponentsInRoot == 1)
+		{
+			ComponentIndex = 0;
+			return true;
+		}
+		else
+		{
+			// This case is not supported yet
+			FString Msg = FString::Printf(TEXT("Error! Node has modifiers when using multiple components in root object. This is currently not supported."));
+			GenerationContext.Compiler->CompilerLog(FText::FromString(Msg), Node, EMessageSeverity::Error);
+			ComponentIndex = 0;
+			return false;
+		}
 	}
 	else
 	{
diff --git a/Engine/Plugins/Experimental/Mutable/Source/CustomizableObjectEditor/Private/MuCOE/GenerateMutableSource/GenerateMutableSourceSurface.cpp b/Engine/Plugins/Experimental/Mutable/Source/CustomizableObjectEditor/Private/MuCOE/GenerateMutableSource/GenerateMutableSourceSurface.cpp
--- a/Engine/Plugins/Experimental/Mutable/Source/CustomizableObjectEditor/Private/MuCOE/GenerateMutableSource/GenerateMutableSourceSurface.cpp
+++ b/Engine/Plugins/Experimental/Mutable/Source/CustomizableObjectEditor/Private/MuCOE/GenerateMutableSource/GenerateMutableSourceSurface.cpp
@@ -386,13 +386,19 @@ mu::NodeSurfacePtr GenerateMutableSourceSurface(const UEdGraphPin * Pin, FMutabl
 		FString TableColumnName;
 
 		// Checking if we should not use the material of the table node even if it is linked to the material node
-		if (const UEdGraphPin* ConnectedPin = FollowInputPin(*TypedNodeMat->GetMaterialAssetPin()))
+		const UEdGraphPin* TempConnectedPin = nullptr;
+		if (TypedNodeMat->GetMaterialAssetPin())
 		{
-			if (const UCustomizableObjectNodeTable* TypedNodeTable = Cast< UCustomizableObjectNodeTable >(ConnectedPin->GetOwningNode()))
+			TempConnectedPin = FollowInputPin(*TypedNodeMat->GetMaterialAssetPin());
+		}
+
+		if (TempConnectedPin)
+		{
+			if (const UCustomizableObjectNodeTable* TypedNodeTable = Cast< UCustomizableObjectNodeTable >(TempConnectedPin->GetOwningNode()))
 			{
-				TableColumnName = ConnectedPin->PinFriendlyName.ToString();
+				TableColumnName = TempConnectedPin->PinFriendlyName.ToString();
 
-				if (UMaterialInstance * TableMaterial = TypedNodeTable->GetColumnDefaultAssetByType<UMaterialInstance>(ConnectedPin))
+				if (UMaterialInstance * TableMaterial = TypedNodeTable->GetColumnDefaultAssetByType<UMaterialInstance>(TempConnectedPin))
 				{
 					// Checking if the reference material of the Table Node has the same parent as the material of the Material Node 
 					if (!TypedNodeMat->Material || TableMaterial->GetMaterial() != TypedNodeMat->Material->GetMaterial())
diff --git a/Engine/Plugins/Experimental/Mutable/Source/CustomizableObjectEditor/Private/MuCOE/SCustomizableObjectEditorViewport.cpp b/Engine/Plugins/Experimental/Mutable/Source/CustomizableObjectEditor/Private/MuCOE/SCustomizableObjectEditorViewport.cpp
--- a/Engine/Plugins/Experimental/Mutable/Source/CustomizableObjectEditor/Private/MuCOE/SCustomizableObjectEditorViewport.cpp
+++ b/Engine/Plugins/Experimental/Mutable/Source/CustomizableObjectEditor/Private/MuCOE/SCustomizableObjectEditorViewport.cpp
@@ -806,8 +806,18 @@ void SCustomizableObjectEditorViewportTabBody::GenerateUVMaterialOptions()
 						continue;
 					}
 
+					if (!Materials[Section.MaterialIndex])
+					{
+						continue;
+					}
+
+					FString BaseMaterialName;
+
 					const UMaterial* BaseMaterial = Materials[Section.MaterialIndex]->GetBaseMaterial();
-					FString BaseMaterialName = BaseMaterial->GetName();
+					if (BaseMaterial)
+					{
+						BaseMaterialName = BaseMaterial->GetName();
+					}
 
 					BaseMaterialName += FString::Printf(TEXT(" LOD_%d_Component_%d"), LODIndex, ComponentIndex);
 
