#rnx

diff --git a/Engine/Plugins/Experimental/Mutable/Source/CustomizableObject/Private/MuCO/CustomizableObjectInstanceUsage.cpp b/Engine/Plugins/Experimental/Mutable/Source/CustomizableObject/Private/MuCO/CustomizableObjectInstanceUsage.cpp
--- a/Engine/Plugins/Experimental/Mutable/Source/CustomizableObject/Private/MuCO/CustomizableObjectInstanceUsage.cpp
+++ b/Engine/Plugins/Experimental/Mutable/Source/CustomizableObject/Private/MuCO/CustomizableObjectInstanceUsage.cpp
@@ -149,6 +149,9 @@ void UCustomizableObjectInstanceUsage::AttachTo(USkeletalMeshComponent* Skeletal
 	else
 	{
 		UsedSkeletalMeshComponent = SkeletalMeshComponent;
+
+		// To mimic the behavior of UCustomizableSkeletalComponent::OnAttachmentChanged()
+		SetPendingSetSkeletalMesh(true);
 	}
 }
 
