#rnx

diff --git a/Engine/Plugins/Experimental/Mutable/Source/CustomizableObject/Private/MuCO/CustomizableSkeletalComponent.cpp b/Engine/Plugins/Experimental/Mutable/Source/CustomizableObject/Private/MuCO/CustomizableSkeletalComponent.cpp
--- a/Engine/Plugins/Experimental/Mutable/Source/CustomizableObject/Private/MuCO/CustomizableSkeletalComponent.cpp
+++ b/Engine/Plugins/Experimental/Mutable/Source/CustomizableObject/Private/MuCO/CustomizableSkeletalComponent.cpp
@@ -10,9 +10,18 @@
 
 void UCustomizableSkeletalComponent::CreateCustomizableObjectInstanceUsage()
 {
+	if (CustomizableObjectInstanceUsage)
+	{
+		// CustomizableObjectInstanceUsage may already exist if duplicated from an existing Customizable Skeletal Component
+		if (CustomizableObjectInstanceUsage->CustomizableSkeletalComponent != this)
+		{
+			CustomizableObjectInstanceUsage = nullptr;
+		}
+	}
+
 	if (!CustomizableObjectInstanceUsage && !HasAnyFlags(RF_ClassDefaultObject))
 	{
-		CustomizableObjectInstanceUsage = NewObject<UCustomizableObjectInstanceUsage>(this, NAME_None, RF_Transient);
+		CustomizableObjectInstanceUsage = NewObject<UCustomizableObjectInstanceUsage>(this, TEXT("InstanceUsage"), RF_Transient);
 		CustomizableObjectInstanceUsage->CustomizableSkeletalComponent = this;
 	}
 }
@@ -29,7 +38,12 @@ void UCustomizableSkeletalComponent::Callbacks() const
 
 USkeletalMesh* UCustomizableSkeletalComponent::GetSkeletalMesh() const
 {
-	return CustomizableObjectInstance ? CustomizableObjectInstance->GetSkeletalMesh(ComponentIndex) : nullptr;
+	if (CustomizableObjectInstanceUsage)
+	{
+		return CustomizableObjectInstanceUsage->GetSkeletalMesh();
+	}
+
+	return nullptr;
 }
 
 
@@ -64,18 +78,18 @@ USkeletalMesh* UCustomizableSkeletalComponent::GetAttachedSkeletalMesh() const
 
 void UCustomizableSkeletalComponent::UpdateSkeletalMeshAsync(bool bNeverSkipUpdate)
 {
-	if (CustomizableObjectInstance)
+	if (CustomizableObjectInstanceUsage)
 	{
-		CustomizableObjectInstance->UpdateSkeletalMeshAsync(false, false);
+		CustomizableObjectInstanceUsage->UpdateSkeletalMeshAsync(bNeverSkipUpdate);
 	}
 }
 
 
 void UCustomizableSkeletalComponent::UpdateSkeletalMeshAsyncResult(FInstanceUpdateDelegate Callback, bool bIgnoreCloseDist, bool bForceHighPriority)
 {
-	if (CustomizableObjectInstance)
+	if (CustomizableObjectInstanceUsage)
 	{
-		CustomizableObjectInstance->UpdateSkeletalMeshAsyncResult(Callback, false, false);
+		CustomizableObjectInstanceUsage->UpdateSkeletalMeshAsyncResult(Callback, bIgnoreCloseDist, bForceHighPriority);
 	}
 }
 
