#rnx

diff --git a/Engine/Plugins/Experimental/Mutable/Source/CustomizableObject/Private/MuCO/CustomizableObjectInstanceUsage.cpp b/Engine/Plugins/Experimental/Mutable/Source/CustomizableObject/Private/MuCO/CustomizableObjectInstanceUsage.cpp
--- a/Engine/Plugins/Experimental/Mutable/Source/CustomizableObject/Private/MuCO/CustomizableObjectInstanceUsage.cpp
+++ b/Engine/Plugins/Experimental/Mutable/Source/CustomizableObject/Private/MuCO/CustomizableObjectInstanceUsage.cpp
@@ -466,15 +466,20 @@ void UCustomizableObjectInstanceUsage::UpdateDistFromComponentToPlayer(const AAc
 
 void UCustomizableObjectInstanceUsage::Tick(float DeltaTime)
 {
+	if (!IsValid(this))
+	{
+		return;
+	}
+
 	UCustomizableObjectInstance* CustomizableObjectInstance = GetCustomizableObjectInstance();
 
-	if (!GetPendingSetSkeletalMesh() || !CustomizableObjectInstance)
+	if (!GetPendingSetSkeletalMesh() || !CustomizableObjectInstance || !IsValid(CustomizableObjectInstance))
 	{
 		return;
 	}
 	
 	UCustomizableObject* CustomizableObject = CustomizableObjectInstance->GetCustomizableObject();
-	if (!CustomizableObject)
+	if (!CustomizableObject || !IsValid(CustomizableObject))
 	{
 		return;
 	}	
