[mutable] Add new policy property that allows to interprete the multiple tags in modifier as "and" or as "or".
Cleanup old string boilerplate.
#rnx

diff --git a/Engine/Plugins/Experimental/Mutable/Source/CustomizableObject/Private/MuCO/CustomizableObject.cpp b/Engine/Plugins/Experimental/Mutable/Source/CustomizableObject/Private/MuCO/CustomizableObject.cpp
--- a/Engine/Plugins/Experimental/Mutable/Source/CustomizableObject/Private/MuCO/CustomizableObject.cpp
+++ b/Engine/Plugins/Experimental/Mutable/Source/CustomizableObject/Private/MuCO/CustomizableObject.cpp
@@ -1158,7 +1158,7 @@ int32 UCustomizableObject::FindState( const FString& Name ) const
 	int32 Result = -1;
 	if (PrivateData->GetModel())
 	{
-		Result = PrivateData->GetModel()->FindState( StringCast<ANSICHAR>(*Name).Get() );
+		Result = PrivateData->GetModel()->FindState(Name);
 	}
 
 	return Result;
diff --git a/Engine/Plugins/Experimental/Mutable/Source/CustomizableObjectEditor/Classes/MuCOE/Nodes/CustomizableObjectNodeMeshClipDeform.h b/Engine/Plugins/Experimental/Mutable/Source/CustomizableObjectEditor/Classes/MuCOE/Nodes/CustomizableObjectNodeMeshClipDeform.h
--- a/Engine/Plugins/Experimental/Mutable/Source/CustomizableObjectEditor/Classes/MuCOE/Nodes/CustomizableObjectNodeMeshClipDeform.h
+++ b/Engine/Plugins/Experimental/Mutable/Source/CustomizableObjectEditor/Classes/MuCOE/Nodes/CustomizableObjectNodeMeshClipDeform.h
@@ -3,6 +3,7 @@
 #pragma once
 
 #include "MuCOE/Nodes/CustomizableObjectNodeModifierBase.h"
+#include "MuT/NodeModifier.h"
 
 #include "CustomizableObjectNodeMeshClipDeform.generated.h"
 
@@ -27,6 +28,10 @@ public:
 	UPROPERTY(EditAnywhere, BlueprintReadWrite, Category = MeshClipDeform)
 	TArray<FString> Tags;
 	
+	/** Policy to use tags in case more than one is added. */
+	UPROPERTY(EditAnywhere, Category = MeshClipDeform)
+	EMutableMultipleTagPolicy MultipleTagPolicy = EMutableMultipleTagPolicy::OnlyOneRequired;
+
 	UPROPERTY(EditAnywhere, Category = MeshClipDeform)
 	EShapeBindingMethod BindingMethod;
 	
diff --git a/Engine/Plugins/Experimental/Mutable/Source/CustomizableObjectEditor/Classes/MuCOE/Nodes/CustomizableObjectNodeMeshClipMorph.h b/Engine/Plugins/Experimental/Mutable/Source/CustomizableObjectEditor/Classes/MuCOE/Nodes/CustomizableObjectNodeMeshClipMorph.h
--- a/Engine/Plugins/Experimental/Mutable/Source/CustomizableObjectEditor/Classes/MuCOE/Nodes/CustomizableObjectNodeMeshClipMorph.h
+++ b/Engine/Plugins/Experimental/Mutable/Source/CustomizableObjectEditor/Classes/MuCOE/Nodes/CustomizableObjectNodeMeshClipMorph.h
@@ -3,6 +3,7 @@
 #pragma once
 
 #include "MuCOE/Nodes/CustomizableObjectNodeModifierBase.h"
+#include "MuT/NodeModifier.h"
 
 #include "CustomizableObjectNodeMeshClipMorph.generated.h"
 
@@ -26,6 +27,10 @@ public:
 	UPROPERTY(EditAnywhere, BlueprintReadWrite, Category = MeshToClipAndMorph)
 	TArray<FString> Tags;
 
+	/** Policy to use tags in case more than one is added. */
+	UPROPERTY(EditAnywhere, Category = MeshToClipAndMorph)
+	EMutableMultipleTagPolicy MultipleTagPolicy = EMutableMultipleTagPolicy::OnlyOneRequired;
+
 	UPROPERTY(EditAnywhere, Category = MeshToClipAndMorph)
 	uint32 ReferenceSkeletonIndex = 0;
 
diff --git a/Engine/Plugins/Experimental/Mutable/Source/CustomizableObjectEditor/Classes/MuCOE/Nodes/CustomizableObjectNodeMeshClipWithMesh.h b/Engine/Plugins/Experimental/Mutable/Source/CustomizableObjectEditor/Classes/MuCOE/Nodes/CustomizableObjectNodeMeshClipWithMesh.h
--- a/Engine/Plugins/Experimental/Mutable/Source/CustomizableObjectEditor/Classes/MuCOE/Nodes/CustomizableObjectNodeMeshClipWithMesh.h
+++ b/Engine/Plugins/Experimental/Mutable/Source/CustomizableObjectEditor/Classes/MuCOE/Nodes/CustomizableObjectNodeMeshClipWithMesh.h
@@ -3,6 +3,7 @@
 #pragma once
 
 #include "MuCOE/Nodes/CustomizableObjectNodeModifierBase.h"
+#include "MuT/NodeModifier.h"
 
 #include "CustomizableObjectNodeMeshClipWithMesh.generated.h"
 
@@ -25,6 +26,9 @@ public:
 	UPROPERTY(EditAnywhere, BlueprintReadWrite, Category = MeshToClip)
 	TArray<FString> Tags;
 
+	UPROPERTY(EditAnywhere, Category = CustomizableObjectToClip)
+	EMutableMultipleTagPolicy MultipleTagPolicy = EMutableMultipleTagPolicy::OnlyOneRequired;
+
 	//!< If assigned, then a material inside this CO will be clipped by this node.
     //!< If several materials with the same name, all are considered (to cover all LOD levels)
     UPROPERTY(EditAnywhere, Category = CustomizableObjectToClip)
diff --git a/Engine/Plugins/Experimental/Mutable/Source/CustomizableObjectEditor/Classes/MuCOE/Nodes/CustomizableObjectNodeModifierClipWithUVMask.h b/Engine/Plugins/Experimental/Mutable/Source/CustomizableObjectEditor/Classes/MuCOE/Nodes/CustomizableObjectNodeModifierClipWithUVMask.h
--- a/Engine/Plugins/Experimental/Mutable/Source/CustomizableObjectEditor/Classes/MuCOE/Nodes/CustomizableObjectNodeModifierClipWithUVMask.h
+++ b/Engine/Plugins/Experimental/Mutable/Source/CustomizableObjectEditor/Classes/MuCOE/Nodes/CustomizableObjectNodeModifierClipWithUVMask.h
@@ -3,6 +3,7 @@
 #pragma once
 
 #include "MuCOE/Nodes/CustomizableObjectNodeModifierBase.h"
+#include "MuT/NodeModifier.h"
 
 #include "CustomizableObjectNodeModifierClipWithUVMask.generated.h"
 
@@ -22,6 +23,10 @@ public:
 	UPROPERTY(EditAnywhere, BlueprintReadWrite, Category = MeshToClip)
 	TArray<FString> Tags;
 
+	/** Policy to use tags in case more than one is added. */
+	UPROPERTY(EditAnywhere, Category = MeshToClip)
+	EMutableMultipleTagPolicy MultipleTagPolicy = EMutableMultipleTagPolicy::OnlyOneRequired;
+
 	/** UV channel index that will be used to get the UVs to apply the clipping mask to. */
 	UPROPERTY(EditAnywhere, BlueprintReadWrite, Category = MeshToClip)
 	int32 UVChannelForMask = 0;
diff --git a/Engine/Plugins/Experimental/Mutable/Source/CustomizableObjectEditor/Private/MuCOE/GenerateMutableSource/GenerateMutableSource.cpp b/Engine/Plugins/Experimental/Mutable/Source/CustomizableObjectEditor/Private/MuCOE/GenerateMutableSource/GenerateMutableSource.cpp
--- a/Engine/Plugins/Experimental/Mutable/Source/CustomizableObjectEditor/Private/MuCOE/GenerateMutableSource/GenerateMutableSource.cpp
+++ b/Engine/Plugins/Experimental/Mutable/Source/CustomizableObjectEditor/Private/MuCOE/GenerateMutableSource/GenerateMutableSource.cpp
@@ -204,7 +204,7 @@ void FMutableGraphGenerationContext::GenerateClippingCOInternalTags()
 
 		for (j = 0; j < It->Value.Num(); ++j)
 		{
-			It->Value[j]->AddTag(StringCast<ANSICHAR>(*TagName).Get());
+			It->Value[j]->AddTag(TagName);
 		}
 
 		UCustomizableObjectNodeMeshClipWithMesh* CustomizableObjectNodeMeshClipWithMesh = It->Key;
@@ -231,7 +231,7 @@ void FMutableGraphGenerationContext::GenerateClippingCOInternalTags()
 
 					for (int32 k = 0; k < It2->Value.Num(); ++k)
 					{
-						It2->Value[k]->AddTag(StringCast<ANSICHAR>(*TagName).Get());
+						It2->Value[k]->AddTag(TagName);
 					}
 				}
 			}
@@ -698,8 +698,8 @@ mu::NodeObjectPtr GenerateMutableSource(const UEdGraphPin * Pin, FMutableGraphGe
 		mu::Ptr<mu::NodeObjectNew> ObjectNode = new mu::NodeObjectNew();
 		Result = ObjectNode;
 
-		ObjectNode->SetName(StringCast<ANSICHAR>(*TypedNodeObj->ObjectName).Get());
-		ObjectNode->SetUid(StringCast<ANSICHAR>(*GenerationContext.GetNodeIdUnique(TypedNodeObj).ToString()).Get());
+		ObjectNode->SetName(TypedNodeObj->ObjectName);
+		ObjectNode->SetUid(GenerationContext.GetNodeIdUnique(TypedNodeObj).ToString());
 
 		// LOD
 		const int32 NumLODs = TypedNodeObj->GetNumLODPins();
@@ -795,10 +795,10 @@ mu::NodeObjectPtr GenerateMutableSource(const UEdGraphPin * Pin, FMutableGraphGe
 		for (int StateIndex = 0; StateIndex < NumStates && bFilterStates; ++StateIndex)
 		{
 			const FCustomizableObjectState& State = TypedNodeObj->States[StateIndex];
-			ObjectNode->SetStateName(StateIndex, StringCast<ANSICHAR>(*State.Name).Get());
+			ObjectNode->SetStateName(StateIndex, State.Name);
 			for (int ParamIndex = 0; ParamIndex < State.RuntimeParameters.Num(); ++ParamIndex)
 			{
-				ObjectNode->AddStateParam(StateIndex, StringCast<ANSICHAR>(*State.RuntimeParameters[ParamIndex]).Get());
+				ObjectNode->AddStateParam(StateIndex, State.RuntimeParameters[ParamIndex]);
 			}
 
 			const ITargetPlatform* TargetPlatform = GenerationContext.Options.TargetPlatform;
@@ -1035,8 +1035,8 @@ mu::NodeObjectPtr GenerateMutableSource(const UEdGraphPin * Pin, FMutableGraphGe
 		GenerationContext.SocketPriorityStack.Push(TypedNodeGroup->SocketPriority);
 
 		GenerationContext.AddParameterNameUnique(TypedNodeGroup, TypedNodeGroup->GroupName);
-		GroupNode->SetName(StringCast<ANSICHAR>(*TypedNodeGroup->GroupName).Get());
-		GroupNode->SetUid(StringCast<ANSICHAR>(*TypedNodeGroup->NodeGuid.ToString()).Get());
+		GroupNode->SetName(TypedNodeGroup->GroupName);
+		GroupNode->SetUid(TypedNodeGroup->NodeGuid.ToString());
 
 		// Get all group projectors and put them in the generation context so that they are available to the child material nodes of this group node
 		uint32 NumProjectorCountBeforeNode = GenerationContext.ProjectorGroupMap.Num();
@@ -1122,7 +1122,7 @@ mu::NodeObjectPtr GenerateMutableSource(const UEdGraphPin * Pin, FMutableGraphGe
 			else
 			{
 				mu::NodeObjectPtr ChildNode = new mu::NodeObjectNew;
-				ChildNode->SetName(StringCast<ANSICHAR>(*CustomizableObjectNodeObject->ObjectName).Get());
+				ChildNode->SetName(CustomizableObjectNodeObject->ObjectName);
 				GroupNode->SetChild(ChildIndex, ChildNode.get());
 			}
 		}
@@ -1210,7 +1210,7 @@ mu::NodeObjectPtr GenerateMutableSource(const UEdGraphPin * Pin, FMutableGraphGe
 			else
 			{
 				mu::NodeObjectPtr ChildNode = new mu::NodeObjectNew;
-				ChildNode->SetName(StringCast<ANSICHAR>(*CustomizableObjectNodeObject->ObjectName).Get());
+				ChildNode->SetName(CustomizableObjectNodeObject->ObjectName);
 				GroupNode->SetChild(ChildIndex, ChildNode.get());
 			}
 
@@ -1368,7 +1368,7 @@ int32 AddTagToMutableMeshUnique(mu::Mesh& MutableMesh, const FString& Tag)
 	}
 
 	MutableMesh.SetTagCount(TagCount + 1);
-	MutableMesh.SetTag(TagCount, StringCast<ANSICHAR>(*Tag).Get());
+	MutableMesh.SetTag(TagCount, Tag);
 
 	return TagCount;
 }
diff --git a/Engine/Plugins/Experimental/Mutable/Source/CustomizableObjectEditor/Private/MuCOE/GenerateMutableSource/GenerateMutableSourceColor.cpp b/Engine/Plugins/Experimental/Mutable/Source/CustomizableObjectEditor/Private/MuCOE/GenerateMutableSource/GenerateMutableSourceColor.cpp
--- a/Engine/Plugins/Experimental/Mutable/Source/CustomizableObjectEditor/Private/MuCOE/GenerateMutableSource/GenerateMutableSourceColor.cpp
+++ b/Engine/Plugins/Experimental/Mutable/Source/CustomizableObjectEditor/Private/MuCOE/GenerateMutableSource/GenerateMutableSourceColor.cpp
@@ -64,8 +64,8 @@ mu::NodeColourPtr GenerateMutableSourceColor(const UEdGraphPin* Pin, FMutableGra
 
 		GenerationContext.AddParameterNameUnique(Node, TypedNodeColorParam->ParameterName);
 
-		ColorNode->SetName(StringCast<ANSICHAR>(*TypedNodeColorParam->ParameterName).Get());
-		ColorNode->SetUid(StringCast<ANSICHAR>(*GenerationContext.GetNodeIdUnique(Node).ToString()).Get());
+		ColorNode->SetName(TypedNodeColorParam->ParameterName);
+		ColorNode->SetUid(GenerationContext.GetNodeIdUnique(Node).ToString());
 		ColorNode->SetDefaultValue(TypedNodeColorParam->DefaultValue.R, TypedNodeColorParam->DefaultValue.G, TypedNodeColorParam->DefaultValue.B);
 
 		GenerationContext.ParameterUIDataMap.Add(TypedNodeColorParam->ParameterName, FParameterUIData(
@@ -259,7 +259,7 @@ mu::NodeColourPtr GenerateMutableSourceColor(const UEdGraphPin* Pin, FMutableGra
 			const UEdGraphPin* VariationPin = TypedNodeColorVar->VariationPin(VariationIndex);
 			if (!VariationPin) continue;
 
-			ColorNode->SetVariationTag(VariationIndex, StringCast<ANSICHAR>(*TypedNodeColorVar->GetVariation(VariationIndex).Tag).Get());
+			ColorNode->SetVariationTag(VariationIndex, TypedNodeColorVar->GetVariation(VariationIndex).Tag);
 			if (const UEdGraphPin* ConnectedPin = FollowInputPin(*VariationPin))
 			{
 				mu::NodeColourPtr ChildNode = GenerateMutableSourceColor(ConnectedPin, GenerationContext);
@@ -302,7 +302,7 @@ mu::NodeColourPtr GenerateMutableSourceColor(const UEdGraphPin* Pin, FMutableGra
 					mu::NodeColourTablePtr ColorTableNode = new mu::NodeColourTable();
 
 					// Generating a new Color column if not exists
-					if (Table->FindColumn(StringCast<ANSICHAR>(*ColumnName).Get()) == INDEX_NONE)
+					if (Table->FindColumn(ColumnName) == INDEX_NONE)
 					{
 						int32 Dummy = -1; // TODO MTBL-1512
 						bool Dummy2 = false;
@@ -320,8 +320,8 @@ mu::NodeColourPtr GenerateMutableSourceColor(const UEdGraphPin* Pin, FMutableGra
 						Result = ColorTableNode;
 
 						ColorTableNode->SetTable(Table);
-						ColorTableNode->SetColumn(StringCast<ANSICHAR>(*ColumnName).Get());
-						ColorTableNode->SetParameterName(StringCast<ANSICHAR>(*TypedNodeTable->ParameterName).Get());
+						ColorTableNode->SetColumn(ColumnName);
+						ColorTableNode->SetParameterName(TypedNodeTable->ParameterName);
 
 						GenerationContext.AddParameterNameUnique(Node, TypedNodeTable->ParameterName);
 					}
diff --git a/Engine/Plugins/Experimental/Mutable/Source/CustomizableObjectEditor/Private/MuCOE/GenerateMutableSource/GenerateMutableSourceFloat.cpp b/Engine/Plugins/Experimental/Mutable/Source/CustomizableObjectEditor/Private/MuCOE/GenerateMutableSource/GenerateMutableSourceFloat.cpp
--- a/Engine/Plugins/Experimental/Mutable/Source/CustomizableObjectEditor/Private/MuCOE/GenerateMutableSource/GenerateMutableSourceFloat.cpp
+++ b/Engine/Plugins/Experimental/Mutable/Source/CustomizableObjectEditor/Private/MuCOE/GenerateMutableSource/GenerateMutableSourceFloat.cpp
@@ -93,8 +93,8 @@ mu::NodeScalarPtr GenerateMutableSourceFloat(const UEdGraphPin* Pin, FMutableGra
 
 		GenerationContext.AddParameterNameUnique(Node, FloatParameterNode->ParameterName);
 
-		ScalarNode->SetName(StringCast<ANSICHAR>(*FloatParameterNode->ParameterName).Get());
-		ScalarNode->SetUid(StringCast<ANSICHAR>(*GenerationContext.GetNodeIdUnique(Node).ToString()).Get());
+		ScalarNode->SetName(FloatParameterNode->ParameterName);
+		ScalarNode->SetUid(GenerationContext.GetNodeIdUnique(Node).ToString());
 		ScalarNode->SetDefaultValue(FloatParameterNode->DefaultValue);
 
 		GenerationContext.ParameterUIDataMap.Add(FloatParameterNode->ParameterName, FParameterUIData(
@@ -113,8 +113,8 @@ mu::NodeScalarPtr GenerateMutableSourceFloat(const UEdGraphPin* Pin, FMutableGra
 
 		GenerationContext.AddParameterNameUnique(Node, EnumParamNode->ParameterName);
 
-		EnumParameterNode->SetName(StringCast<ANSICHAR>(*EnumParamNode->ParameterName).Get());
-		EnumParameterNode->SetUid(StringCast<ANSICHAR>(*GenerationContext.GetNodeIdUnique(Node).ToString()).Get());
+		EnumParameterNode->SetName(EnumParamNode->ParameterName);
+		EnumParameterNode->SetUid(GenerationContext.GetNodeIdUnique(Node).ToString());
 		EnumParameterNode->SetValueCount(NumSelectors);
 		EnumParameterNode->SetDefaultValueIndex(DefaultValue);
 
@@ -123,7 +123,7 @@ mu::NodeScalarPtr GenerateMutableSourceFloat(const UEdGraphPin* Pin, FMutableGra
 
 		for (int SelectorIndex = 0; SelectorIndex < NumSelectors; ++SelectorIndex)
 		{
-			EnumParameterNode->SetValue(SelectorIndex, (float)SelectorIndex, StringCast<ANSICHAR>(*EnumParamNode->Values[SelectorIndex].Name).Get());
+			EnumParameterNode->SetValue(SelectorIndex, (float)SelectorIndex, EnumParamNode->Values[SelectorIndex].Name);
 
 			ParameterUIData.ArrayIntegerParameterOption.Add(FIntegerParameterUIData(
 				EnumParamNode->Values[SelectorIndex].Name,
@@ -287,7 +287,7 @@ mu::NodeScalarPtr GenerateMutableSourceFloat(const UEdGraphPin* Pin, FMutableGra
 			UEdGraphPin* VariationPin = TypedNodeFloatVar->VariationPin(VariationIndex);
 			if (!VariationPin) continue;
 
-			FloatNode->SetVariationTag(VariationIndex, StringCast<ANSICHAR>(*TypedNodeFloatVar->GetVariation(VariationIndex).Tag).Get());
+			FloatNode->SetVariationTag(VariationIndex, TypedNodeFloatVar->GetVariation(VariationIndex).Tag);
 			if (const UEdGraphPin* ConnectedPin = FollowInputPin(*VariationPin))
 			{
 				mu::NodeScalarPtr ChildNode = GenerateMutableSourceFloat(ConnectedPin, GenerationContext);
@@ -342,7 +342,7 @@ mu::NodeScalarPtr GenerateMutableSourceFloat(const UEdGraphPin* Pin, FMutableGra
 					}
 
 					// Generating a new Float column if not exists
-					if (Table->FindColumn(StringCast<ANSICHAR>(*ColumnName).Get()) == INDEX_NONE)
+					if (Table->FindColumn(ColumnName) == INDEX_NONE)
 					{
 						int32 Dummy = -1; // TODO MTBL-1512
 						bool Dummy2 = false;
@@ -360,8 +360,8 @@ mu::NodeScalarPtr GenerateMutableSourceFloat(const UEdGraphPin* Pin, FMutableGra
 						Result = ScalarTableNode;
 
 						ScalarTableNode->SetTable(Table);
-						ScalarTableNode->SetColumn(StringCast<ANSICHAR>(*ColumnName).Get());
-						ScalarTableNode->SetParameterName(StringCast<ANSICHAR>(*TypedNodeTable->ParameterName).Get());
+						ScalarTableNode->SetColumn(ColumnName);
+						ScalarTableNode->SetParameterName(TypedNodeTable->ParameterName);
 
 						GenerationContext.AddParameterNameUnique(Node, TypedNodeTable->ParameterName);
 					}
diff --git a/Engine/Plugins/Experimental/Mutable/Source/CustomizableObjectEditor/Private/MuCOE/GenerateMutableSource/GenerateMutableSourceGroupProjector.cpp b/Engine/Plugins/Experimental/Mutable/Source/CustomizableObjectEditor/Private/MuCOE/GenerateMutableSource/GenerateMutableSourceGroupProjector.cpp
--- a/Engine/Plugins/Experimental/Mutable/Source/CustomizableObjectEditor/Private/MuCOE/GenerateMutableSource/GenerateMutableSourceGroupProjector.cpp
+++ b/Engine/Plugins/Experimental/Mutable/Source/CustomizableObjectEditor/Private/MuCOE/GenerateMutableSource/GenerateMutableSourceGroupProjector.cpp
@@ -437,8 +437,8 @@ bool GenerateMutableSourceGroupProjector(const UEdGraphPin* Pin, FMutableGraphGe
 
 			mu::NodeScalarParameterPtr NodeScalarParam = new mu::NodeScalarParameter;
 			FString NodeScalarParamName = ProjParamNode->ParameterName + FMultilayerProjector::NUM_LAYERS_PARAMETER_POSTFIX;
-			NodeScalarParam->SetName(StringCast<ANSICHAR>(*(NodeScalarParamName)).Get());
-			NodeScalarParam->SetUid(StringCast<ANSICHAR>(*(ProjectorParamUid + FString("_NL"))).Get());
+			NodeScalarParam->SetName(NodeScalarParamName);
+			NodeScalarParam->SetUid(ProjectorParamUid + FString("_NL"));
 			GenerationContext.AddParameterNameUnique(originalGroup, NodeScalarParamName);
 			NodeScalarParam->SetDefaultValue(0.f);
 
@@ -456,8 +456,8 @@ bool GenerateMutableSourceGroupProjector(const UEdGraphPin* Pin, FMutableGraphGe
 			
 			mu::NodeScalarEnumParameterPtr PoseEnumParameterNode = new mu::NodeScalarEnumParameter;
 			FString PoseNodeEnumParamName = ProjParamNode->ParameterName + FMultilayerProjector::POSE_PARAMETER_POSTFIX;
-			PoseEnumParameterNode->SetName(StringCast<ANSICHAR>(*(PoseNodeEnumParamName)).Get());
-			PoseEnumParameterNode->SetUid(StringCast<ANSICHAR>(*(ProjectorParamUid + FString("_SP"))).Get());
+			PoseEnumParameterNode->SetName(PoseNodeEnumParamName);
+			PoseEnumParameterNode->SetUid(ProjectorParamUid + FString("_SP"));
 			GenerationContext.AddParameterNameUnique(originalGroup, PoseNodeEnumParamName);
 			PoseEnumParameterNode->SetValueCount(ProjParamNode->OptionPoses.Num() + 1);
 			PoseEnumParameterNode->SetDefaultValueIndex(0);
@@ -470,8 +470,8 @@ bool GenerateMutableSourceGroupProjector(const UEdGraphPin* Pin, FMutableGraphGe
 
 			mu::NodeScalarParameterPtr OpacityParameterNode = new mu::NodeScalarParameter;
 			FString OpacityParameterNodeName = ProjParamNode->ParameterName + FMultilayerProjector::OPACITY_PARAMETER_POSTFIX;
-			OpacityParameterNode->SetName(StringCast<ANSICHAR>(*(OpacityParameterNodeName)).Get());
-			OpacityParameterNode->SetUid(StringCast<ANSICHAR>(*(ProjectorParamUid + FString("_O"))).Get());
+			OpacityParameterNode->SetName(OpacityParameterNodeName);
+			OpacityParameterNode->SetUid(ProjectorParamUid + FString("_O"));
 			GenerationContext.AddParameterNameUnique(originalGroup, OpacityParameterNodeName);
 			OpacityParameterNode->SetDefaultValue(0.75f);
 			OpacityParameterNode->SetRangeCount(1);
@@ -496,7 +496,7 @@ bool GenerateMutableSourceGroupProjector(const UEdGraphPin* Pin, FMutableGraphGe
 
 			for (int PoseIndex = 0; PoseIndex < ProjParamNode->OptionPoses.Num(); ++PoseIndex)
 			{
-				PoseEnumParameterNode->SetValue(PoseIndex + 1, (float)PoseIndex + 1.f, StringCast<ANSICHAR>(*ProjParamNode->OptionPoses[PoseIndex].PoseName).Get());
+				PoseEnumParameterNode->SetValue(PoseIndex + 1, (float)PoseIndex + 1.f, ProjParamNode->OptionPoses[PoseIndex].PoseName);
 
 				TArray<FString> ArrayBoneName;
 				TArray<FTransform> ArrayTransform;
@@ -517,8 +517,8 @@ bool GenerateMutableSourceGroupProjector(const UEdGraphPin* Pin, FMutableGraphGe
 		
 			mu::NodeScalarEnumParameterPtr EnumParameterNode = new mu::NodeScalarEnumParameter;
 			FString NodeEnumParamName = ProjParamNode->ParameterName + FMultilayerProjector::IMAGE_PARAMETER_POSTFIX;
-			EnumParameterNode->SetName(StringCast<ANSICHAR>(*NodeEnumParamName).Get());
-			EnumParameterNode->SetUid(StringCast<ANSICHAR>(*(ProjectorParamUid + FString("_SI"))).Get());
+			EnumParameterNode->SetName(NodeEnumParamName);
+			EnumParameterNode->SetUid(ProjectorParamUid + FString("_SI"));
 			GenerationContext.AddParameterNameUnique(originalGroup, NodeEnumParamName);
 			EnumParameterNode->SetValueCount(ArrayOptionImage.Num());
 			EnumParameterNode->SetDefaultValueIndex(0);
@@ -531,7 +531,7 @@ bool GenerateMutableSourceGroupProjector(const UEdGraphPin* Pin, FMutableGraphGe
 
 			for (int ImageIndex = 0; ImageIndex < ArrayOptionImage.Num(); ++ImageIndex)
 			{
-				EnumParameterNode->SetValue(ImageIndex, (float)ImageIndex, StringCast<ANSICHAR>(*ArrayOptionImage[ImageIndex].OptionName).Get());
+				EnumParameterNode->SetValue(ImageIndex, (float)ImageIndex, ArrayOptionImage[ImageIndex].OptionName);
 
 				FMutableParamUIMetadata optionMetadata = ParameterUIData.ParamUIMetadata;
 				optionMetadata.UIThumbnail = ArrayOptionImage[ImageIndex].OptionImage;
diff --git a/Engine/Plugins/Experimental/Mutable/Source/CustomizableObjectEditor/Private/MuCOE/GenerateMutableSource/GenerateMutableSourceImage.cpp b/Engine/Plugins/Experimental/Mutable/Source/CustomizableObjectEditor/Private/MuCOE/GenerateMutableSource/GenerateMutableSourceImage.cpp
--- a/Engine/Plugins/Experimental/Mutable/Source/CustomizableObjectEditor/Private/MuCOE/GenerateMutableSource/GenerateMutableSourceImage.cpp
+++ b/Engine/Plugins/Experimental/Mutable/Source/CustomizableObjectEditor/Private/MuCOE/GenerateMutableSource/GenerateMutableSourceImage.cpp
@@ -474,7 +474,7 @@ mu::NodeImagePtr GenerateMutableSourceImage(const UEdGraphPin* Pin, FMutableGrap
 				continue;
 			}
 
-			TextureNode->SetVariationTag(VariationIndex, StringCast<ANSICHAR>(*TypedNodeImageVar->Variations[VariationIndex].Tag).Get());
+			TextureNode->SetVariationTag(VariationIndex, TypedNodeImageVar->Variations[VariationIndex].Tag);
 			if (const UEdGraphPin* ConnectedPin = FollowInputPin(*VariationPin))
 			{
 				mu::NodeImagePtr ChildNode = GenerateMutableSourceImage(ConnectedPin, GenerationContext, MaxTextureSize);
@@ -1093,7 +1093,7 @@ mu::NodeImagePtr GenerateMutableSourceImage(const UEdGraphPin* Pin, FMutableGrap
 						}
 
 						// Generating a new Texture column if not exists
-						if (Table->FindColumn(StringCast<ANSICHAR>(*ColumnName).Get()) == INDEX_NONE)
+						if (Table->FindColumn(ColumnName) == INDEX_NONE)
 						{
 							int32 Dummy = -1; // TODO MTBL-1512
 							bool Dummy2 = false;
@@ -1111,8 +1111,8 @@ mu::NodeImagePtr GenerateMutableSourceImage(const UEdGraphPin* Pin, FMutableGrap
 							Result = ImageTableNode;
 
 							ImageTableNode->SetTable(Table);
-							ImageTableNode->SetColumn(StringCast<ANSICHAR>(*ColumnName).Get());
-							ImageTableNode->SetParameterName(StringCast<ANSICHAR>(*TypedNodeTable->ParameterName).Get());
+							ImageTableNode->SetColumn(ColumnName);
+							ImageTableNode->SetParameterName(TypedNodeTable->ParameterName);
 							ImageTableNode->SetMaxTextureSize(MaxTextureSize);
 
 							GenerationContext.AddParameterNameUnique(Node, TypedNodeTable->ParameterName);
diff --git a/Engine/Plugins/Experimental/Mutable/Source/CustomizableObjectEditor/Private/MuCOE/GenerateMutableSource/GenerateMutableSourceMesh.cpp b/Engine/Plugins/Experimental/Mutable/Source/CustomizableObjectEditor/Private/MuCOE/GenerateMutableSource/GenerateMutableSourceMesh.cpp
--- a/Engine/Plugins/Experimental/Mutable/Source/CustomizableObjectEditor/Private/MuCOE/GenerateMutableSource/GenerateMutableSourceMesh.cpp
+++ b/Engine/Plugins/Experimental/Mutable/Source/CustomizableObjectEditor/Private/MuCOE/GenerateMutableSource/GenerateMutableSourceMesh.cpp
@@ -2654,11 +2654,11 @@ mu::NodeMeshPtr GenerateMorphMesh(const UEdGraphPin* Pin,
 		{
 			const FName RowName = TypedNodeTable->GetRowNames()[RowIndex];
 
-			ColumnIndex = Table->FindColumn(StringCast<ANSICHAR>(*ColumnName).Get());
+			ColumnIndex = Table->FindColumn(ColumnName);
 
 			if (ColumnIndex == INDEX_NONE)
 			{
-				ColumnIndex = Table->AddColumn(StringCast<ANSICHAR>(*ColumnName).Get(), mu::ETableColumnType::Mesh);
+				ColumnIndex = Table->AddColumn(ColumnName, mu::ETableColumnType::Mesh);
 			}
 
 			mu::MeshPtr MorphedSourceTableMesh = BuildMorphedMutableMesh(Pin, TypedNodeMorphs[MorphIndex].MorphTargetName, GenerationContext, bOnlyConnectedLOD, RowName);
@@ -2671,8 +2671,8 @@ mu::NodeMeshPtr GenerateMorphMesh(const UEdGraphPin* Pin,
 
 			mu::NodeMeshTablePtr MorphedSourceMeshNodeTable = new mu::NodeMeshTable;
 			MorphedSourceMeshNodeTable->SetTable(Table);
-			MorphedSourceMeshNodeTable->SetColumn(StringCast<ANSICHAR>(*ColumnName).Get());
-			MorphedSourceMeshNodeTable->SetParameterName(StringCast<ANSICHAR>(*TypedNodeTable->ParameterName).Get());
+			MorphedSourceMeshNodeTable->SetColumn(ColumnName);
+			MorphedSourceMeshNodeTable->SetParameterName(TypedNodeTable->ParameterName);
 			MorphedSourceMeshNodeTable->SetMessageContext(MorphNode);
 
 			mu::NodeMeshMakeMorphPtr Morph = new mu::NodeMeshMakeMorph;
@@ -3648,7 +3648,7 @@ mu::NodeMeshPtr GenerateMutableSourceMesh(const UEdGraphPin* Pin,
 					FString MutableColumnName = TypedNodeTable->GetMutableColumnName(Pin, LODIndex);
 
 					// Generating a new Mesh column if not exists
-					if (Table->FindColumn(StringCast<ANSICHAR>(*MutableColumnName).Get()) == INDEX_NONE)
+					if (Table->FindColumn(MutableColumnName) == INDEX_NONE)
 					{
 						bSuccess = GenerateTableColumn(TypedNodeTable, Pin, Table, DataTableColumnName, Property, LODIndexConnected, SectionIndexConnected, LODIndex, SectionIndex, bOnlyConnectedLOD, GenerationContext);
 
@@ -3664,8 +3664,8 @@ mu::NodeMeshPtr GenerateMutableSourceMesh(const UEdGraphPin* Pin,
 						Result = MeshTableNode;
 
 						MeshTableNode->SetTable(Table);
-						MeshTableNode->SetColumn(StringCast<ANSICHAR>(*MutableColumnName).Get());
-						MeshTableNode->SetParameterName(StringCast<ANSICHAR>(*TypedNodeTable->ParameterName).Get());
+						MeshTableNode->SetColumn(MutableColumnName);
+						MeshTableNode->SetParameterName(TypedNodeTable->ParameterName);
 
 						GenerationContext.AddParameterNameUnique(Node, TypedNodeTable->ParameterName);
 
diff --git a/Engine/Plugins/Experimental/Mutable/Source/CustomizableObjectEditor/Private/MuCOE/GenerateMutableSource/GenerateMutableSourceModifier.cpp b/Engine/Plugins/Experimental/Mutable/Source/CustomizableObjectEditor/Private/MuCOE/GenerateMutableSource/GenerateMutableSourceModifier.cpp
--- a/Engine/Plugins/Experimental/Mutable/Source/CustomizableObjectEditor/Private/MuCOE/GenerateMutableSource/GenerateMutableSourceModifier.cpp
+++ b/Engine/Plugins/Experimental/Mutable/Source/CustomizableObjectEditor/Private/MuCOE/GenerateMutableSource/GenerateMutableSourceModifier.cpp
@@ -63,9 +63,10 @@ mu::NodeModifierPtr GenerateMutableSourceModifier(const UEdGraphPin * Pin, FMuta
 
 		ClipNode->SetVertexSelectionBone(GenerationContext.BoneNames.AddUnique(TypedNodeClip->BoneName), TypedNodeClip->MaxEffectRadius);
 
+		ClipNode->SetMultipleTagPolicy( TypedNodeClip->MultipleTagPolicy );
 		for (const FString& Tag : TypedNodeClip->Tags)
 		{
-			 ClipNode->AddTag(StringCast<ANSICHAR>(*Tag).Get());
+			 ClipNode->AddTag(Tag);
 		}
 	}
 
@@ -101,9 +102,10 @@ mu::NodeModifierPtr GenerateMutableSourceModifier(const UEdGraphPin * Pin, FMuta
 			ClipNode->SetBindingMethod(BindingMethod);
 		}
 	
+		ClipNode->SetMultipleTagPolicy(TypedNodeClipDeform->MultipleTagPolicy);
 		for (const FString& Tag : TypedNodeClipDeform->Tags)
 		{
-			ClipNode->AddTag(StringCast<ANSICHAR>(*Tag).Get());
+			ClipNode->AddTag(Tag);
 		}		
 	}
 
@@ -168,6 +170,7 @@ mu::NodeModifierPtr GenerateMutableSourceModifier(const UEdGraphPin * Pin, FMuta
 			ClipNode->SetClipMesh(ClipMesh.get());
 		}
 
+		ClipNode->SetMultipleTagPolicy(TypedNodeClipMesh->MultipleTagPolicy);
 		for (const FString& Tag : TypedNodeClipMesh->Tags)
 		{
 			 ClipNode->AddTag(Tag);
@@ -211,6 +214,7 @@ mu::NodeModifierPtr GenerateMutableSourceModifier(const UEdGraphPin * Pin, FMuta
 
 		ClipNode->SetLayoutIndex(TypedNodeClipUVMask->UVChannelForMask);
 
+		ClipNode->SetMultipleTagPolicy(TypedNodeClipUVMask->MultipleTagPolicy);
 		for (const FString& Tag : TypedNodeClipUVMask->Tags)
 		{
 			ClipNode->AddTag(Tag);
diff --git a/Engine/Plugins/Experimental/Mutable/Source/CustomizableObjectEditor/Private/MuCOE/GenerateMutableSource/GenerateMutableSourceProjector.cpp b/Engine/Plugins/Experimental/Mutable/Source/CustomizableObjectEditor/Private/MuCOE/GenerateMutableSource/GenerateMutableSourceProjector.cpp
--- a/Engine/Plugins/Experimental/Mutable/Source/CustomizableObjectEditor/Private/MuCOE/GenerateMutableSource/GenerateMutableSourceProjector.cpp
+++ b/Engine/Plugins/Experimental/Mutable/Source/CustomizableObjectEditor/Private/MuCOE/GenerateMutableSource/GenerateMutableSourceProjector.cpp
@@ -85,8 +85,8 @@ mu::NodeProjectorPtr GenerateMutableSourceProjector(const UEdGraphPin* Pin, FMut
 
 		GenerationContext.AddParameterNameUnique(Node, TypedNodeParam->ParameterName);
 
-		ProjectorNode->SetName(StringCast<ANSICHAR>(*TypedNodeParam->ParameterName).Get());
-		ProjectorNode->SetUid(StringCast<ANSICHAR>(*GenerationContext.GetNodeIdUnique(Node).ToString()).Get());
+		ProjectorNode->SetName(TypedNodeParam->ParameterName);
+		ProjectorNode->SetUid(GenerationContext.GetNodeIdUnique(Node).ToString());
 		switch ((int)TypedNodeParam->DefaultValue.ProjectionType)
 		{
 		case 0:
diff --git a/Engine/Plugins/Experimental/Mutable/Source/CustomizableObjectEditor/Private/MuCOE/GenerateMutableSource/GenerateMutableSourceSurface.cpp b/Engine/Plugins/Experimental/Mutable/Source/CustomizableObjectEditor/Private/MuCOE/GenerateMutableSource/GenerateMutableSourceSurface.cpp
--- a/Engine/Plugins/Experimental/Mutable/Source/CustomizableObjectEditor/Private/MuCOE/GenerateMutableSource/GenerateMutableSourceSurface.cpp
+++ b/Engine/Plugins/Experimental/Mutable/Source/CustomizableObjectEditor/Private/MuCOE/GenerateMutableSource/GenerateMutableSourceSurface.cpp
@@ -443,10 +443,10 @@ mu::NodeSurfacePtr GenerateMutableSourceSurface(const UEdGraphPin * Pin, FMutabl
 
 					const FString ImageName = TypedNodeMat->GetParameterName(EMaterialParameterType::Texture, ImageIndex).ToString();
 					FString SurfNodeImageName = FString::Printf(TEXT("%d"), GenerationContext.ImageProperties.Num());
-					SurfNode->SetImageName(ImageIndex, StringCast<ANSICHAR>(*SurfNodeImageName).Get());
+					SurfNode->SetImageName(ImageIndex, SurfNodeImageName);
 
 					SurfNode->SetImageLayoutIndex(ImageIndex, -1);
-					SurfNode->SetImageAdditionalNames(ImageIndex, StringCast<ANSICHAR>(*TypedNodeMat->Material->GetName()).Get(), StringCast<ANSICHAR>(*ImageName).Get());
+					SurfNode->SetImageAdditionalNames(ImageIndex, TypedNodeMat->Material->GetName(), ImageName);
 
 					// We don't need a reference texture or props here, but we do need the parameter name.
 					FGeneratedImageProperties Props;
@@ -752,10 +752,10 @@ mu::NodeSurfacePtr GenerateMutableSourceSurface(const UEdGraphPin * Pin, FMutabl
 							SurfNodeImageName += "-MutableLayerParam:" + FString::FromInt(LayerIndex);
 						}
 
-						SurfNode->SetImageName(ImageIndex, StringCast<ANSICHAR>(*SurfNodeImageName).Get());
+						SurfNode->SetImageName(ImageIndex, SurfNodeImageName);
 						int32 UVLayout = TypedNodeMat->GetImageUVLayout(ImageIndex);
 						SurfNode->SetImageLayoutIndex(ImageIndex, UVLayout);
-						SurfNode->SetImageAdditionalNames(ImageIndex, StringCast<ANSICHAR>(*TypedNodeMat->Material->GetName()).Get(), StringCast<ANSICHAR>(*ImageName).Get());
+						SurfNode->SetImageAdditionalNames(ImageIndex, TypedNodeMat->Material->GetName(), ImageName);
 
 						if (bShareProjectionTexturesBetweenLODs && bIsGroupProjectorImage)
 						{
@@ -773,7 +773,7 @@ mu::NodeSurfacePtr GenerateMutableSourceSurface(const UEdGraphPin * Pin, FMutabl
 					ensure(LOD > GenerationContext.FirstLODAvailable);
 					check(ProjectorInfo->SurfNode->GetImage(ImageIndex) == ProjectorInfo->ImageNode);
 					SurfNode->SetImage(ImageIndex, ProjectorInfo->ImageNode);
-					SurfNode->SetImageName(ImageIndex, StringCast<ANSICHAR>(*ProjectorInfo->TextureName).Get());
+					SurfNode->SetImageName(ImageIndex, ProjectorInfo->TextureName);
 					SurfNode->SetImageLayoutIndex(ImageIndex, ProjectorInfo->UVLayout);
 
 					TextureNameToProjectionResFactor.Add(ProjectorInfo->RealTextureName, ProjectorInfo->AlternateProjectionResolutionFactor);
@@ -805,7 +805,7 @@ mu::NodeSurfacePtr GenerateMutableSourceSurface(const UEdGraphPin * Pin, FMutabl
 					}
 
 					SurfNode->SetVector(VectorIndex, ColorNode);
-					SurfNode->SetVectorName(VectorIndex, StringCast<ANSICHAR>(*VectorName).Get());
+					SurfNode->SetVectorName(VectorIndex, VectorName);
 				}
 			}
 		}
@@ -833,7 +833,7 @@ mu::NodeSurfacePtr GenerateMutableSourceSurface(const UEdGraphPin * Pin, FMutabl
 					}
 
 					SurfNode->SetScalar(ScalarIndex, ScalarNode);
-					SurfNode->SetScalarName(ScalarIndex, StringCast<ANSICHAR>(*ScalarName).Get());
+					SurfNode->SetScalarName(ScalarIndex, ScalarName);
 				}
 			}
 		}
@@ -856,7 +856,7 @@ mu::NodeSurfacePtr GenerateMutableSourceSurface(const UEdGraphPin * Pin, FMutabl
 					mu::NodeScalarPtr ScalarNode = GenerateMutableSourceFloat(ConnectedPin, GenerationContext);
 
 					SurfNode->SetScalar(MaterialIndex, ScalarNode);
-					SurfNode->SetScalarName(MaterialIndex, StringCast<ANSICHAR>(*MaterialName).Get());
+					SurfNode->SetScalarName(MaterialIndex, MaterialName);
 				}
 			}
 			else
@@ -865,14 +865,14 @@ mu::NodeSurfacePtr GenerateMutableSourceSurface(const UEdGraphPin * Pin, FMutabl
 				ScalarNode->SetValue(ReferencedMaterialsIndex);
 
 				SurfNode->SetScalar(MaterialIndex, ScalarNode);
-				SurfNode->SetScalarName(MaterialIndex, StringCast<ANSICHAR>(*MaterialName).Get());
+				SurfNode->SetScalarName(MaterialIndex, MaterialName);
 			}
 		}
 		
 
 		for (const FString& Tag : TypedNodeMat->Tags)
 		{
-			SurfNode->AddTag(StringCast<ANSICHAR>(*Tag).Get());
+			SurfNode->AddTag(Tag);
 		}
 
 		TArray<mu::NodeSurfaceNewPtr>* ArraySurfaceNodePtr = GenerationContext.MapMaterialNodeToMutableSurfaceNodeArray.Find(TypedNodeMat);
@@ -915,7 +915,7 @@ mu::NodeSurfacePtr GenerateMutableSourceSurface(const UEdGraphPin * Pin, FMutabl
 
 			for (const FString& Tag : TypedNodeMat->Tags)
 			{
-				SurfNode2->AddTag(StringCast<ANSICHAR>(*Tag).Get());
+				SurfNode2->AddTag(Tag);
 			}
 
 			SurfNode2->SetImageCount(SurfNode->GetImageCount());
@@ -965,7 +965,7 @@ mu::NodeSurfacePtr GenerateMutableSourceSurface(const UEdGraphPin * Pin, FMutabl
 			mu::NodeSurfaceVariationPtr SurfaceVariation = new mu::NodeSurfaceVariation;
 			SurfaceVariation->SetVariationType(mu::NodeSurfaceVariation::VariationType::State);
 			SurfaceVariation->SetVariationCount(1);
-			SurfaceVariation->SetVariationTag(0, StringCast<ANSICHAR>(*AlternateResStateName).Get());
+			SurfaceVariation->SetVariationTag(0, AlternateResStateName);
 
 			SurfaceVariation->AddDefaultSurface(&*SurfNode);
 			SurfaceVariation->AddVariationSurface(0, &*SurfNode2);
@@ -1148,7 +1148,7 @@ mu::NodeSurfacePtr GenerateMutableSourceSurface(const UEdGraphPin * Pin, FMutabl
 		
 			for (const FString& Tag : TypedNodeExt->Tags)
 			{
-				SurfNode->AddTag(StringCast<ANSICHAR>(*Tag).Get());
+				SurfNode->AddTag(Tag);
 			}
 		}();
 	}
@@ -1510,7 +1510,7 @@ mu::NodeSurfacePtr GenerateMutableSourceSurface(const UEdGraphPin * Pin, FMutabl
 
 			if (UEdGraphPin* VariationPin = TypedNodeVar->VariationPin(VariationIndex))
 			{
-				SurfNode->SetVariationTag(VariationIndex, StringCast<ANSICHAR>(*TypedNodeVar->GetVariation(VariationIndex).Tag).Get());
+				SurfNode->SetVariationTag(VariationIndex, TypedNodeVar->GetVariation(VariationIndex).Tag);
 				for (const UEdGraphPin* ConnectedPin : FollowInputPinArray(*VariationPin))
 				{
 					// Is it a modifier?
diff --git a/Engine/Plugins/Experimental/Mutable/Source/CustomizableObjectEditor/Private/MuCOE/GenerateMutableSource/GenerateMutableSourceTable.cpp b/Engine/Plugins/Experimental/Mutable/Source/CustomizableObjectEditor/Private/MuCOE/GenerateMutableSource/GenerateMutableSourceTable.cpp
--- a/Engine/Plugins/Experimental/Mutable/Source/CustomizableObjectEditor/Private/MuCOE/GenerateMutableSource/GenerateMutableSourceTable.cpp
+++ b/Engine/Plugins/Experimental/Mutable/Source/CustomizableObjectEditor/Private/MuCOE/GenerateMutableSource/GenerateMutableSourceTable.cpp
@@ -40,11 +40,11 @@ bool FillTableColumn(const UCustomizableObjectNodeTable* TableNode,	mu::TablePtr
 				// Generating an Empty cell
 				FString MutableColumnName = TableNode->GenerateSkeletalMeshMutableColumName(ColumnName, LODIndex, SectionIndex);
 
-				CurrentColumn = MutableTable.get()->FindColumn(StringCast<ANSICHAR>(*MutableColumnName).Get());
+				CurrentColumn = MutableTable.get()->FindColumn(MutableColumnName);
 
 				if (CurrentColumn == -1)
 				{
-					CurrentColumn = MutableTable->AddColumn(StringCast<ANSICHAR>(*MutableColumnName).Get(), mu::ETableColumnType::Mesh);
+					CurrentColumn = MutableTable->AddColumn(MutableColumnName, mu::ETableColumnType::Mesh);
 				}
 
 				mu::MeshPtr EmptySkeletalMesh = nullptr;
@@ -216,11 +216,11 @@ bool FillTableColumn(const UCustomizableObjectNodeTable* TableNode,	mu::TablePtr
 
 			FString MutableColumnName = TableNode->GenerateSkeletalMeshMutableColumName(ColumnName, LODIndex, SectionIndex);
 
-			CurrentColumn = MutableTable.get()->FindColumn(StringCast<ANSICHAR>(*MutableColumnName).Get());
+			CurrentColumn = MutableTable.get()->FindColumn(MutableColumnName);
 
 			if (CurrentColumn == -1)
 			{
-				CurrentColumn = MutableTable->AddColumn(StringCast<ANSICHAR>(*MutableColumnName).Get(), mu::ETableColumnType::Mesh);
+				CurrentColumn = MutableTable->AddColumn(MutableColumnName, mu::ETableColumnType::Mesh);
 			}
 
 			// First process the mesh tags that are going to make the mesh unique and affect whether it's repeated in 
@@ -351,11 +351,11 @@ bool FillTableColumn(const UCustomizableObjectNodeTable* TableNode,	mu::TablePtr
 
 			FString MutableColumnName = TableNode->GenerateStaticMeshMutableColumName(ColumnName, SectionIndex);
 
-			CurrentColumn = MutableTable.get()->FindColumn(StringCast<ANSICHAR>(*MutableColumnName).Get());
+			CurrentColumn = MutableTable.get()->FindColumn(MutableColumnName);
 
 			if (CurrentColumn == -1)
 			{
-				CurrentColumn = MutableTable->AddColumn(StringCast<ANSICHAR>(*MutableColumnName).Get(), mu::ETableColumnType::Mesh);
+				CurrentColumn = MutableTable->AddColumn(MutableColumnName, mu::ETableColumnType::Mesh);
 			}
 
 			mu::MeshPtr MutableMesh = GenerateMutableMesh(StaticMesh, TSoftClassPtr<UAnimInstance>(), CurrentLOD, SectionIndex, CurrentLOD, SectionIndex, FString(), GenerationContext, TableNode);
@@ -394,11 +394,11 @@ bool FillTableColumn(const UCustomizableObjectNodeTable* TableNode,	mu::TablePtr
 			UTexture* Texture = Texture2D ? Cast<UTexture>(Texture2D) : Cast<UTexture>(TextureArray);
 
 			// Getting column index from column name
-			CurrentColumn = MutableTable->FindColumn(StringCast<ANSICHAR>(*ColumnName).Get());
+			CurrentColumn = MutableTable->FindColumn(ColumnName);
 
 			if (CurrentColumn == INDEX_NONE)
 			{
-				CurrentColumn = MutableTable->AddColumn(StringCast<ANSICHAR>(*ColumnName).Get(), mu::ETableColumnType::Image);
+				CurrentColumn = MutableTable->AddColumn(ColumnName, mu::ETableColumnType::Image);
 			}
 
 			// Removing encoding part
@@ -502,11 +502,11 @@ bool FillTableColumn(const UCustomizableObjectNodeTable* TableNode,	mu::TablePtr
 			FString EncodedSwitchParameterName = "__MutableMaterialId";
 			if (ColumnName.Contains(EncodedSwitchParameterName))
 			{
-				CurrentColumn = MutableTable.get()->FindColumn(StringCast<ANSICHAR>(*ColumnName).Get());
+				CurrentColumn = MutableTable.get()->FindColumn(ColumnName);
 
 				if (CurrentColumn == -1)
 				{
-					CurrentColumn = MutableTable->AddColumn(StringCast<ANSICHAR>(*ColumnName).Get(), mu::ETableColumnType::Scalar);
+					CurrentColumn = MutableTable->AddColumn(ColumnName, mu::ETableColumnType::Scalar);
 				}
 
 				const int32 lastMaterialAmount = GenerationContext.ReferencedMaterials.Num();
@@ -540,12 +540,12 @@ bool FillTableColumn(const UCustomizableObjectNodeTable* TableNode,	mu::TablePtr
 			if (ParameterIndex != INDEX_NONE && ParameterInfos[ParameterIndex].Name == GenerationContext.CurrentMaterialTableParameter)
 			{
 				// Getting column index from parameter name
-				ColumnIndex = MutableTable->FindColumn(StringCast<ANSICHAR>(*ColumnName).Get());
+				ColumnIndex = MutableTable->FindColumn(ColumnName);
 
 				if (ColumnIndex == INDEX_NONE)
 				{
 					// If there is no column with the parameters name, we generate a new one
-					ColumnIndex = MutableTable->AddColumn(StringCast<ANSICHAR>(*ColumnName).Get(), mu::ETableColumnType::Image);
+					ColumnIndex = MutableTable->AddColumn(ColumnName, mu::ETableColumnType::Image);
 				}
 
 				UTexture* ParentTextureValue = nullptr;
@@ -596,11 +596,11 @@ bool FillTableColumn(const UCustomizableObjectNodeTable* TableNode,	mu::TablePtr
 	{
 		if (StructProperty->Struct == TBaseStructure<FLinearColor>::Get())
 		{
-			CurrentColumn = MutableTable->FindColumn(StringCast<ANSICHAR>(*ColumnName).Get());
+			CurrentColumn = MutableTable->FindColumn(ColumnName);
 
 			if (CurrentColumn == INDEX_NONE)
 			{
-				CurrentColumn = MutableTable->AddColumn(StringCast<ANSICHAR>(*ColumnName).Get(), mu::ETableColumnType::Color);
+				CurrentColumn = MutableTable->AddColumn(ColumnName, mu::ETableColumnType::Color);
 			}
 
 			// Setting cell value
@@ -617,11 +617,11 @@ bool FillTableColumn(const UCustomizableObjectNodeTable* TableNode,	mu::TablePtr
 
 	else if (const FNumericProperty* FloatNumProperty = CastField<FFloatProperty>(ColumnProperty))
 	{
-		CurrentColumn = MutableTable->FindColumn(StringCast<ANSICHAR>(*ColumnName).Get());
+		CurrentColumn = MutableTable->FindColumn(ColumnName);
 
 		if (CurrentColumn == INDEX_NONE)
 		{
-			CurrentColumn = MutableTable->AddColumn(StringCast<ANSICHAR>(*ColumnName).Get(), mu::ETableColumnType::Scalar);
+			CurrentColumn = MutableTable->AddColumn(ColumnName, mu::ETableColumnType::Scalar);
 		}
 
 		// Setting cell value
@@ -631,11 +631,11 @@ bool FillTableColumn(const UCustomizableObjectNodeTable* TableNode,	mu::TablePtr
 
 	else if (const FNumericProperty* DoubleNumProperty = CastField<FDoubleProperty>(ColumnProperty))
 	{
-		CurrentColumn = MutableTable->FindColumn(StringCast<ANSICHAR>(*ColumnName).Get());
+		CurrentColumn = MutableTable->FindColumn(ColumnName);
 	
 		if (CurrentColumn == INDEX_NONE)
 		{
-			CurrentColumn = MutableTable->AddColumn(StringCast<ANSICHAR>(*ColumnName).Get(), mu::ETableColumnType::Scalar);
+			CurrentColumn = MutableTable->AddColumn(ColumnName, mu::ETableColumnType::Scalar);
 		}
 	
 		// Setting cell value
@@ -747,7 +747,7 @@ mu::TablePtr GenerateMutableSourceTable(const FString& TableName, const UEdGraph
 			{
 				MutableTable->AddRow(i);
 				FString RowName= RowNames[i].ToString();
-				MutableTable->SetCell(0, i, StringCast<ANSICHAR>(*RowName).Get());
+				MutableTable->SetCell(0, i, RowName);
 				ParameterUIData.ArrayIntegerParameterOption.Add(FIntegerParameterUIData(
 					RowName,
 					FMutableParamUIMetadata()));
diff --git a/Engine/Plugins/Experimental/Mutable/Source/CustomizableObjectEditor/Private/MuCOE/Nodes/CustomizableObjectNodeExtensionDataVariation.cpp b/Engine/Plugins/Experimental/Mutable/Source/CustomizableObjectEditor/Private/MuCOE/Nodes/CustomizableObjectNodeExtensionDataVariation.cpp
--- a/Engine/Plugins/Experimental/Mutable/Source/CustomizableObjectEditor/Private/MuCOE/Nodes/CustomizableObjectNodeExtensionDataVariation.cpp
+++ b/Engine/Plugins/Experimental/Mutable/Source/CustomizableObjectEditor/Private/MuCOE/Nodes/CustomizableObjectNodeExtensionDataVariation.cpp
@@ -85,7 +85,7 @@ mu::NodeExtensionDataPtr UCustomizableObjectNodeExtensionDataVariation::Generate
 	{
 		if (const UEdGraphPin* VariationPin = GetVariationPin(VariationIndex))
 		{
-			VariationNode->SetVariationTag(VariationIndex, StringCast<ANSICHAR>(*Variations[VariationIndex].Tag).Get());
+			VariationNode->SetVariationTag(VariationIndex, Variations[VariationIndex].Tag);
 			if (const UEdGraphPin* ConnectedPin = FollowInputPin(*VariationPin))
 			{
 				if (const ICustomizableObjectExtensionNode* ExtensionNode = Cast<ICustomizableObjectExtensionNode>(ConnectedPin->GetOwningNode()))
diff --git a/Engine/Plugins/Experimental/Mutable/Source/CustomizableObjectEditor/Private/MuCOE/Nodes/CustomizableObjectNodeLayoutBlocksDetails.cpp b/Engine/Plugins/Experimental/Mutable/Source/CustomizableObjectEditor/Private/MuCOE/Nodes/CustomizableObjectNodeLayoutBlocksDetails.cpp
--- a/Engine/Plugins/Experimental/Mutable/Source/CustomizableObjectEditor/Private/MuCOE/Nodes/CustomizableObjectNodeLayoutBlocksDetails.cpp
+++ b/Engine/Plugins/Experimental/Mutable/Source/CustomizableObjectEditor/Private/MuCOE/Nodes/CustomizableObjectNodeLayoutBlocksDetails.cpp
@@ -193,7 +193,7 @@ void FCustomizableObjectNodeLayoutBlocksDetails::OnGridComboBoxSelectionChanged(
 		mu::Model* Model = CustomInstance->CustomizableObject->GetModel();
 
 		CustomInstance->PreEditChange(NULL);
-		CustomInstance->State = Model->FindState( StringCast<ANSICHAR>(**Selection).Get() );
+		CustomInstance->State = Model->FindState( *Selection );
 		CustomInstance->PostEditChange();
 
 		ResetParamBox();
diff --git a/Engine/Plugins/Experimental/Mutable/Source/MutableTools/MutableTools.Build.cs b/Engine/Plugins/Experimental/Mutable/Source/MutableTools/MutableTools.Build.cs
--- a/Engine/Plugins/Experimental/Mutable/Source/MutableTools/MutableTools.Build.cs
+++ b/Engine/Plugins/Experimental/Mutable/Source/MutableTools/MutableTools.Build.cs
@@ -24,8 +24,9 @@ namespace UnrealBuildTool.Rules
 
 			PublicDependencyModuleNames.AddRange(
                 new string[] {
-					"MutableRuntime", 
+					"MutableRuntime",
 					"Core",
+					"CoreUObject",
 					"GeometryCore",
 					"ImageCore",
 					"TextureCompressor",
diff --git a/Engine/Plugins/Experimental/Mutable/Source/MutableTools/Private/MuT/CodeGenerator.cpp b/Engine/Plugins/Experimental/Mutable/Source/MutableTools/Private/MuT/CodeGenerator.cpp
--- a/Engine/Plugins/Experimental/Mutable/Source/MutableTools/Private/MuT/CodeGenerator.cpp
+++ b/Engine/Plugins/Experimental/Mutable/Source/MutableTools/Private/MuT/CodeGenerator.cpp
@@ -2076,58 +2076,85 @@ namespace mu
 
 	//---------------------------------------------------------------------------------------------
 	void CodeGenerator::GetModifiersFor(
-		const TArray<FString>& tags,
-		int LOD, bool bModifiersForBeforeOperations,
-		TArray<FirstPassGenerator::FModifier>& modifiers)
+		const TArray<FString>& SurfaceTags,
+		int32 LOD, bool bModifiersForBeforeOperations,
+		TArray<FirstPassGenerator::FModifier>& OutModifiers)
 	{
         MUTABLE_CPUPROFILER_SCOPE(GetModifiersFor);
 
-		if (tags.Num())
+		if (SurfaceTags.IsEmpty())
 		{
-			for (const FirstPassGenerator::FModifier& m: m_firstPass.modifiers)
+			return;
+		}
+
+		for (const FirstPassGenerator::FModifier& m: m_firstPass.modifiers)
+		{
+			// Correct LOD?
+			if (m.lod != LOD)
 			{
-				// Correct LOD?
-				if (m.lod != LOD)
-				{
-					continue;
-				}
+				continue;
+			}
 
-				// Correct stage
-				if (m.node->m_applyBeforeNormalOperations != bModifiersForBeforeOperations)
-				{
-					continue;
-				}
+			// Correct stage?
+			if (m.node->bApplyBeforeNormalOperations != bModifiersForBeforeOperations)
+			{
+				continue;
+			}
 
-				// Already there?
-				bool alreadyAdded = 
-					modifiers.FindByPredicate( [&m](const FirstPassGenerator::FModifier& c) {return c.node == m.node; })
-					!= 
-					nullptr;
+			// Already there?
+			bool bAlreadyAdded = 
+				OutModifiers.FindByPredicate( [&m](const FirstPassGenerator::FModifier& c) {return c.node == m.node; })
+				!= 
+				nullptr;
 
-				if (alreadyAdded)
-				{
-					continue;
-				}
+			if (bAlreadyAdded)
+			{
+				continue;
+			}
 
-				// Matching tags?
-				bool found = false;
-				for (auto it = m.node->m_tags.begin(); !found && it != m.node->m_tags.end(); ++it)
+			// Matching tags?
+			bool bApply = false;
+
+			switch (m.node->MultipleTagsPolicy)
+			{
+			case EMutableMultipleTagPolicy::OnlyOneRequired:
+			{
+				for (const FString& Tag: m.node->RequiredTags)
 				{
-					found = tags.Contains(*it);
+					if (SurfaceTags.Contains(Tag))
+					{
+						bApply = true;
+						break;
+					}
 				}
+				break;
+			}
 
-				if (found)
+			case EMutableMultipleTagPolicy::AllRequired:
+			{
+				bApply = true;
+				for (const FString& Tag : m.node->RequiredTags)
 				{
-					modifiers.Add(m);
+					if (!SurfaceTags.Contains(Tag))
+					{
+						bApply = false;
+						break;
+					}
 				}
 			}
+			}
+
+			if (bApply)
+			{
+				OutModifiers.Add(m);
+			}
 		}
 	}
 
 	//---------------------------------------------------------------------------------------------
 	Ptr<ASTOp> CodeGenerator::ApplyMeshModifiers(
 		const Ptr<ASTOp>& sourceOp,
-		const TArray<FString>& tags,
+		const TArray<FString>& SurfaceTags,
 		bool bModifiersForBeforeOperations,
 		const void* errorContext )
 	{
@@ -2137,7 +2164,7 @@ namespace mu
 		TArray<FirstPassGenerator::FModifier> modifiers;
 
 		int currentLOD = m_currentParents.Last().m_lod;
-		GetModifiersFor(tags, currentLOD, bModifiersForBeforeOperations, modifiers);
+		GetModifiersFor(SurfaceTags, currentLOD, bModifiersForBeforeOperations, modifiers);
 
 		Ptr<ASTOp> preModifiersMesh = lastMeshOp;
 
diff --git a/Engine/Plugins/Experimental/Mutable/Source/MutableTools/Private/MuT/CodeGenerator.h b/Engine/Plugins/Experimental/Mutable/Source/MutableTools/Private/MuT/CodeGenerator.h
--- a/Engine/Plugins/Experimental/Mutable/Source/MutableTools/Private/MuT/CodeGenerator.h
+++ b/Engine/Plugins/Experimental/Mutable/Source/MutableTools/Private/MuT/CodeGenerator.h
@@ -292,11 +292,11 @@ namespace mu
 		//-----------------------------------------------------------------------------------------
 
 		// Get the modifiers that have to be applied to elements with a specific tag.
-		void GetModifiersFor(const TArray<FString>& tags, int LOD,
-			bool bModifiersForBeforeOperations, TArray<FirstPassGenerator::FModifier>& modifiers);
+		void GetModifiersFor(const TArray<FString>& SurfaceTags, int32 LOD,
+			bool bModifiersForBeforeOperations, TArray<FirstPassGenerator::FModifier>& OutModifiers);
 
 		// Apply the required mesh modifiers to the given operation.
-		Ptr<ASTOp> ApplyMeshModifiers( const Ptr<ASTOp>& sourceOp, const TArray<FString>& tags,
+		Ptr<ASTOp> ApplyMeshModifiers( const Ptr<ASTOp>& sourceOp, const TArray<FString>& SurfaceTags,
 			bool bModifiersForBeforeOperations, const void* errorContext);
 
 		// Get the modifiers that have to be applied to elements with a specific tag.
diff --git a/Engine/Plugins/Experimental/Mutable/Source/MutableTools/Private/MuT/CodeGenerator_FirstPass.h b/Engine/Plugins/Experimental/Mutable/Source/MutableTools/Private/MuT/CodeGenerator_FirstPass.h
--- a/Engine/Plugins/Experimental/Mutable/Source/MutableTools/Private/MuT/CodeGenerator_FirstPass.h
+++ b/Engine/Plugins/Experimental/Mutable/Source/MutableTools/Private/MuT/CodeGenerator_FirstPass.h
@@ -137,18 +137,17 @@ namespace mu
 		{
             const NodeModifier::Private* node = nullptr;
 
-            // List of tags that are required for the presence of this surface
+            // List of tags that are required to apply this modifier
 			TArray<FString> positiveTags;
 
-            // List of tags that block the presence of this surface
+            // List of tags that block the activation of this modifier
 			TArray<FString> negativeTags;
 
             // This conditions is the condition of the object defining this modifier which may not
             // be the parent object where this surface will be added.
             Ptr<ASTOp> objectCondition;
 
-            // This conditions is the condition for this modifier to be enabled when all the object
-            // conditions are met.
+            // This conditions is the condition for this modifier to be enabled when all the object conditions are met.
             // This is filled in CodeGenerator_SecondPass.
             Ptr<ASTOp> surfaceCondition;
 
diff --git a/Engine/Plugins/Experimental/Mutable/Source/MutableTools/Private/MuT/CodeGenerator_Mesh.cpp b/Engine/Plugins/Experimental/Mutable/Source/MutableTools/Private/MuT/CodeGenerator_Mesh.cpp
--- a/Engine/Plugins/Experimental/Mutable/Source/MutableTools/Private/MuT/CodeGenerator_Mesh.cpp
+++ b/Engine/Plugins/Experimental/Mutable/Source/MutableTools/Private/MuT/CodeGenerator_Mesh.cpp
@@ -1110,11 +1110,6 @@ namespace mu
             }
         }
 
-		if (!bFirstValidConnectionFound)
-		{
-			ensure(false);
-		}
-
         OutResult.meshOp = op;
     }
 
diff --git a/Engine/Plugins/Experimental/Mutable/Source/MutableTools/Private/MuT/NodeColourVariation.cpp b/Engine/Plugins/Experimental/Mutable/Source/MutableTools/Private/MuT/NodeColourVariation.cpp
--- a/Engine/Plugins/Experimental/Mutable/Source/MutableTools/Private/MuT/NodeColourVariation.cpp
+++ b/Engine/Plugins/Experimental/Mutable/Source/MutableTools/Private/MuT/NodeColourVariation.cpp
@@ -83,31 +83,22 @@ namespace mu
 
 
     //---------------------------------------------------------------------------------------------
-    void NodeColourVariation::SetVariationCount( int num )
+    void NodeColourVariation::SetVariationCount( int32 num )
     {
         check( num >= 0 );
         m_pD->m_variations.SetNum( num );
     }
 
     //---------------------------------------------------------------------------------------------
-    void NodeColourVariation::SetVariationTag( int index, const char* strTag )
+    void NodeColourVariation::SetVariationTag( int32 index, const FString& Tag )
     {
         check( index >= 0 && index < m_pD->m_variations.Num() );
-        check( strTag );
-
-        if ( strTag )
-        {
-            m_pD->m_variations[index].m_tag = strTag;
-        }
-        else
-        {
-            m_pD->m_variations[index].m_tag = "";
-        }
+        m_pD->m_variations[index].m_tag = Tag;
     }
 
 
     //---------------------------------------------------------------------------------------------
-    void NodeColourVariation::SetVariationColour( int index, NodeColour* pNode )
+    void NodeColourVariation::SetVariationColour( int32 index, NodeColour* pNode )
     {
         check( index >= 0 && index < m_pD->m_variations.Num() );
 
diff --git a/Engine/Plugins/Experimental/Mutable/Source/MutableTools/Private/MuT/NodeExtensionDataVariation.cpp b/Engine/Plugins/Experimental/Mutable/Source/MutableTools/Private/MuT/NodeExtensionDataVariation.cpp
--- a/Engine/Plugins/Experimental/Mutable/Source/MutableTools/Private/MuT/NodeExtensionDataVariation.cpp
+++ b/Engine/Plugins/Experimental/Mutable/Source/MutableTools/Private/MuT/NodeExtensionDataVariation.cpp
@@ -92,12 +92,11 @@ namespace mu
 	}
 
 	//---------------------------------------------------------------------------------------------
-	void NodeExtensionDataVariation::SetVariationTag(int InIndex, const char* InStrTag)
+	void NodeExtensionDataVariation::SetVariationTag(int InIndex, const FString& Tag)
 	{
 		check(m_pD->Variations.IsValidIndex(InIndex));
-		check(InStrTag);
 
-		m_pD->Variations[InIndex].Tag = InStrTag;
+		m_pD->Variations[InIndex].Tag = Tag;
 	}
 
 	//---------------------------------------------------------------------------------------------
diff --git a/Engine/Plugins/Experimental/Mutable/Source/MutableTools/Private/MuT/NodeImageVariation.cpp b/Engine/Plugins/Experimental/Mutable/Source/MutableTools/Private/MuT/NodeImageVariation.cpp
--- a/Engine/Plugins/Experimental/Mutable/Source/MutableTools/Private/MuT/NodeImageVariation.cpp
+++ b/Engine/Plugins/Experimental/Mutable/Source/MutableTools/Private/MuT/NodeImageVariation.cpp
@@ -90,19 +90,11 @@ namespace mu
     }
 
     //---------------------------------------------------------------------------------------------
-    void NodeImageVariation::SetVariationTag( int index, const char* strTag )
+    void NodeImageVariation::SetVariationTag( int index, const FString& Tag )
     {
         check( index >= 0 && index < (int)m_pD->m_variations.Num() );
-        check( strTag );
 
-        if ( strTag )
-        {
-            m_pD->m_variations[index].m_tag = strTag;
-        }
-        else
-        {
-            m_pD->m_variations[index].m_tag = "";
-        }
+		m_pD->m_variations[index].m_tag = Tag;
     }
 
 
diff --git a/Engine/Plugins/Experimental/Mutable/Source/MutableTools/Private/MuT/NodeModifier.cpp b/Engine/Plugins/Experimental/Mutable/Source/MutableTools/Private/MuT/NodeModifier.cpp
--- a/Engine/Plugins/Experimental/Mutable/Source/MutableTools/Private/MuT/NodeModifier.cpp
+++ b/Engine/Plugins/Experimental/Mutable/Source/MutableTools/Private/MuT/NodeModifier.cpp
@@ -16,32 +16,24 @@
 
 namespace mu
 {
-
-
-	//---------------------------------------------------------------------------------------------
 	// Static initialisation
-	//---------------------------------------------------------------------------------------------
 	static NODE_TYPE s_nodeModifierType =
 			NODE_TYPE( "NodeModifier", Node::GetStaticType() );
 
+	MUTABLE_IMPLEMENT_ENUM_SERIALISABLE(EMutableMultipleTagPolicy);
 
-	//---------------------------------------------------------------------------------------------
-	//---------------------------------------------------------------------------------------------
-	//---------------------------------------------------------------------------------------------
 	const NODE_TYPE* NodeModifier::GetType() const
 	{
 		return GetStaticType();
 	}
 
 
-	//---------------------------------------------------------------------------------------------
 	const NODE_TYPE* NodeModifier::GetStaticType()
 	{
           return &s_nodeModifierType;
         }
 
 
-	//---------------------------------------------------------------------------------------------
 	void NodeModifier::Serialise( const NodeModifier* p, OutputArchive& arch )
 	{
         uint32 ver = 0;
@@ -64,7 +56,6 @@ namespace mu
     }
 
         
-	//---------------------------------------------------------------------------------------------
 	NodeModifierPtr NodeModifier::StaticUnserialise( InputArchive& arch )
 	{
         uint32 ver;
@@ -86,23 +77,30 @@ namespace mu
 		return 0;
 	}
 
-	//---------------------------------------------------------------------------------------------
-	void NodeModifier::AddTag(const FString& tagName)
+	void NodeModifier::AddTag(const FString& Value)
+	{
+		NodeModifier::Private* pD = dynamic_cast<NodeModifier::Private*>(GetBasePrivate());
+		check(pD);
+
+		pD->RequiredTags.Add(Value);
+	}
+
+
+	void NodeModifier::SetMultipleTagPolicy(EMutableMultipleTagPolicy Value)
 	{
 		NodeModifier::Private* pD = dynamic_cast<NodeModifier::Private*>(GetBasePrivate());
 		check(pD);
 
-		pD->m_tags.Add(tagName);
+		pD->MultipleTagsPolicy = Value;
 	}
 
 
-	//---------------------------------------------------------------------------------------------
 	void NodeModifier::SetStage(bool bBeforeNormalOperation)
 	{
 		NodeModifier::Private* pD = dynamic_cast<NodeModifier::Private*>(GetBasePrivate());
 		check(pD);
 
-		pD->m_applyBeforeNormalOperations = bBeforeNormalOperation; 
+		pD->bApplyBeforeNormalOperations = bBeforeNormalOperation; 
 	}
 
 }
diff --git a/Engine/Plugins/Experimental/Mutable/Source/MutableTools/Private/MuT/NodeModifierPrivate.h b/Engine/Plugins/Experimental/Mutable/Source/MutableTools/Private/MuT/NodeModifierPrivate.h
--- a/Engine/Plugins/Experimental/Mutable/Source/MutableTools/Private/MuT/NodeModifierPrivate.h
+++ b/Engine/Plugins/Experimental/Mutable/Source/MutableTools/Private/MuT/NodeModifierPrivate.h
@@ -5,28 +5,36 @@
 #include "MuT/NodePrivate.h"
 #include "MuT/NodeModifier.h"
 
+#include <string>
+
 
 namespace mu
 {
 
+	MUTABLE_DEFINE_ENUM_SERIALISABLE(EMutableMultipleTagPolicy);
+
 	class NodeModifier::Private : public Node::Private
 	{
 	public:
 
-		// Tags that target surface need to have enabled to receive this modifier.
-		TArray<FString> m_tags;
+		/** Tags that target surface need to have enabled to receive this modifier. */
+		TArray<FString> RequiredTags;
+
+		/** In case of multiple tags in RequiredTags: are they all required, or one is enough? */
+		EMutableMultipleTagPolicy MultipleTagsPolicy = EMutableMultipleTagPolicy::OnlyOneRequired;
 
 		// Wether the modifier has to be applied after the normal node operations or before
-		bool m_applyBeforeNormalOperations = true;
+		bool bApplyBeforeNormalOperations = true;
 
 		//!
 		void Serialise(OutputArchive& arch) const
 		{
-            uint32_t ver = 3;
+            uint32_t ver = 4;
 			arch << ver;
 
-			arch << m_tags;
-			arch << m_applyBeforeNormalOperations;
+			arch << RequiredTags;
+			arch << MultipleTagsPolicy;
+			arch << bApplyBeforeNormalOperations;
 		}
 
 		//!
@@ -40,17 +48,23 @@ namespace mu
 			{
 				TArray<std::string> Temp;
 				arch >> Temp;
-				m_tags.SetNum(Temp.Num());
+				RequiredTags.SetNum(Temp.Num());
 				for( int32 i=0; i<Temp.Num(); ++i)
 				{
-					m_tags[i] = Temp[i].c_str();
+					RequiredTags[i] = Temp[i].c_str();
 				}
 			}
 			else
 			{
-				arch >> m_tags;
+				arch >> RequiredTags;
 			}
-			arch >> m_applyBeforeNormalOperations;
+
+			if (ver >= 4)
+			{
+				arch >> MultipleTagsPolicy;
+			}
+
+			arch >> bApplyBeforeNormalOperations;
 		}
 
 	};
diff --git a/Engine/Plugins/Experimental/Mutable/Source/MutableTools/Private/MuT/NodeScalarParameter.cpp b/Engine/Plugins/Experimental/Mutable/Source/MutableTools/Private/MuT/NodeScalarParameter.cpp
--- a/Engine/Plugins/Experimental/Mutable/Source/MutableTools/Private/MuT/NodeScalarParameter.cpp
+++ b/Engine/Plugins/Experimental/Mutable/Source/MutableTools/Private/MuT/NodeScalarParameter.cpp
@@ -61,30 +61,16 @@ namespace mu
 
 
 	//---------------------------------------------------------------------------------------------
-	void NodeScalarParameter::SetName( const char* strName )
+	void NodeScalarParameter::SetName( const FString& Name )
 	{
-		if ( strName )
-		{
-			m_pD->m_name = strName;
-		}
-		else
-		{
-			m_pD->m_name = "";
-		}
+		m_pD->m_name = Name;
 	}
 
 
 	//---------------------------------------------------------------------------------------------
-	void NodeScalarParameter::SetUid( const char* strUid )
+	void NodeScalarParameter::SetUid( const FString& Uid )
 	{
-		if ( strUid )
-		{
-			m_pD->m_uid = strUid;
-		}
-		else
-		{
-			m_pD->m_uid = "";
-		}
+		m_pD->m_uid = Uid;
 	}
 
 
diff --git a/Engine/Plugins/Experimental/Mutable/Source/MutableTools/Private/MuT/NodeScalarVariation.cpp b/Engine/Plugins/Experimental/Mutable/Source/MutableTools/Private/MuT/NodeScalarVariation.cpp
--- a/Engine/Plugins/Experimental/Mutable/Source/MutableTools/Private/MuT/NodeScalarVariation.cpp
+++ b/Engine/Plugins/Experimental/Mutable/Source/MutableTools/Private/MuT/NodeScalarVariation.cpp
@@ -75,41 +75,30 @@ namespace mu
     //---------------------------------------------------------------------------------------------
     // Own Interface
     //---------------------------------------------------------------------------------------------
-    void NodeScalarVariation::SetDefaultScalar( NodeScalar* p ) { m_pD->m_defaultScalar = p; }
+    void NodeScalarVariation::SetDefaultScalar( NodeScalar* p ) 
+	{ 
+		m_pD->m_defaultScalar = p; 
+	}
 
 
-    //---------------------------------------------------------------------------------------------
-    int NodeScalarVariation::GetVariationCount() const { return int( m_pD->m_variations.Num() ); }
-
-
-    //---------------------------------------------------------------------------------------------
-    void NodeScalarVariation::SetVariationCount( int num )
+	void NodeScalarVariation::SetVariationCount( int32 num )
     {
         check( num >= 0 );
         m_pD->m_variations.SetNum( num );
     }
 
-    //---------------------------------------------------------------------------------------------
-    void NodeScalarVariation::SetVariationTag( int index, const char* strTag )
+
+	void NodeScalarVariation::SetVariationTag( int32 index, const FString& Tag )
     {
-        check( index >= 0 && index < (int)m_pD->m_variations.Num() );
-        check( strTag );
+        check( index >= 0 && index < m_pD->m_variations.Num() );
 
-        if ( strTag )
-        {
-            m_pD->m_variations[index].m_tag = strTag;
-        }
-        else
-        {
-            m_pD->m_variations[index].m_tag = "";
-        }
+		m_pD->m_variations[index].m_tag = Tag;
     }
 
 
-    //---------------------------------------------------------------------------------------------
-    void NodeScalarVariation::SetVariationScalar( int index, NodeScalar* pNode )
+    void NodeScalarVariation::SetVariationScalar( int32 index, NodeScalar* pNode )
     {
-        check( index >= 0 && index < (int)m_pD->m_variations.Num() );
+        check( index >= 0 && index < m_pD->m_variations.Num() );
 
         m_pD->m_variations[index].m_scalar = pNode;
     }
diff --git a/Engine/Plugins/Experimental/Mutable/Source/MutableTools/Private/MuT/NodeStringParameter.cpp b/Engine/Plugins/Experimental/Mutable/Source/MutableTools/Private/MuT/NodeStringParameter.cpp
--- a/Engine/Plugins/Experimental/Mutable/Source/MutableTools/Private/MuT/NodeStringParameter.cpp
+++ b/Engine/Plugins/Experimental/Mutable/Source/MutableTools/Private/MuT/NodeStringParameter.cpp
@@ -75,37 +75,23 @@ namespace mu
 
 
 	//---------------------------------------------------------------------------------------------
-	void NodeStringParameter::SetName( const char* strName )
+	void NodeStringParameter::SetName( const FString& Name )
 	{
-		if ( strName )
-		{
-			m_pD->m_name = strName;
-		}
-		else
-		{
-			m_pD->m_name = "";
-		}
+		m_pD->m_name = Name;
 	}
 
 
 	//---------------------------------------------------------------------------------------------
-	void NodeStringParameter::SetUid( const char* strUid )
+	void NodeStringParameter::SetUid( const FString& Uid )
 	{
-		if ( strUid )
-		{
-			m_pD->m_uid = strUid;
-		}
-		else
-		{
-			m_pD->m_uid = "";
-		}
+		m_pD->m_uid = Uid;
 	}
 
 
 	//---------------------------------------------------------------------------------------------
-	void NodeStringParameter::SetDefaultValue( const char* v )
+	void NodeStringParameter::SetDefaultValue( const FString& v )
 	{
-		m_pD->m_defaultValue = v?v:"";
+		m_pD->m_defaultValue = v;
 	}
 
 
diff --git a/Engine/Plugins/Experimental/Mutable/Source/MutableTools/Private/MuT/NodeSurfaceVariation.cpp b/Engine/Plugins/Experimental/Mutable/Source/MutableTools/Private/MuT/NodeSurfaceVariation.cpp
--- a/Engine/Plugins/Experimental/Mutable/Source/MutableTools/Private/MuT/NodeSurfaceVariation.cpp
+++ b/Engine/Plugins/Experimental/Mutable/Source/MutableTools/Private/MuT/NodeSurfaceVariation.cpp
@@ -159,19 +159,11 @@ namespace mu
 
 
 	//---------------------------------------------------------------------------------------------
-	void NodeSurfaceVariation::SetVariationTag(int index, const char* strTag)
+	void NodeSurfaceVariation::SetVariationTag(int index, const FString& Tag)
 	{
 		check(index >= 0 && index < m_pD->m_variations.Num());
-		check(strTag);
 
-        if (strTag)
-        {
-            m_pD->m_variations[index].m_tag = strTag;
-        }
-        else
-        {
-            m_pD->m_variations[index].m_tag = "";
-        }
+		m_pD->m_variations[index].m_tag = Tag;
 	}
 
 
diff --git a/Engine/Plugins/Experimental/Mutable/Source/MutableTools/Public/MuT/NodeColourVariation.h b/Engine/Plugins/Experimental/Mutable/Source/MutableTools/Public/MuT/NodeColourVariation.h
--- a/Engine/Plugins/Experimental/Mutable/Source/MutableTools/Public/MuT/NodeColourVariation.h
+++ b/Engine/Plugins/Experimental/Mutable/Source/MutableTools/Public/MuT/NodeColourVariation.h
@@ -53,16 +53,16 @@ namespace mu
         void SetDefaultColour( NodeColour* Colour );
 
         //! Set the number of tags to consider in this variation
-        void SetVariationCount( int count );
+        void SetVariationCount( int32 count );
 
         //!
         int GetVariationCount() const;
 
         //! Set the tag or state name that will enable a specific vartiation
-        void SetVariationTag( int index, const char* strTag );
+        void SetVariationTag( int32 index, const FString& Tag );
 
         //!
-        void SetVariationColour( int index, NodeColour* Colour );
+        void SetVariationColour( int32 index, NodeColour* Colour );
 
         //!}
 
diff --git a/Engine/Plugins/Experimental/Mutable/Source/MutableTools/Public/MuT/NodeExtensionDataVariation.h b/Engine/Plugins/Experimental/Mutable/Source/MutableTools/Public/MuT/NodeExtensionDataVariation.h
--- a/Engine/Plugins/Experimental/Mutable/Source/MutableTools/Public/MuT/NodeExtensionDataVariation.h
+++ b/Engine/Plugins/Experimental/Mutable/Source/MutableTools/Public/MuT/NodeExtensionDataVariation.h
@@ -45,7 +45,7 @@ namespace mu
 		void SetVariationCount(int InCount);
 		int GetVariationCount() const;
 
-		void SetVariationTag(int InIndex, const char* InStrTag);
+		void SetVariationTag(int InIndex, const FString& InTag);
 
 		void SetVariationValue(int InIndex, NodeExtensionDataPtr InValue);
 
diff --git a/Engine/Plugins/Experimental/Mutable/Source/MutableTools/Public/MuT/NodeImageVariation.h b/Engine/Plugins/Experimental/Mutable/Source/MutableTools/Public/MuT/NodeImageVariation.h
--- a/Engine/Plugins/Experimental/Mutable/Source/MutableTools/Public/MuT/NodeImageVariation.h
+++ b/Engine/Plugins/Experimental/Mutable/Source/MutableTools/Public/MuT/NodeImageVariation.h
@@ -62,7 +62,7 @@ namespace mu
         int GetVariationCount() const;
 
         //! Set the tag or state name that will enable a specific vartiation
-        void SetVariationTag( int index, const char* strTag );
+        void SetVariationTag( int index, const FString& Tag );
 
         //!
         void SetVariationImage( int index, NodeImage* Image );
diff --git a/Engine/Plugins/Experimental/Mutable/Source/MutableTools/Public/MuT/NodeModifier.h b/Engine/Plugins/Experimental/Mutable/Source/MutableTools/Public/MuT/NodeModifier.h
--- a/Engine/Plugins/Experimental/Mutable/Source/MutableTools/Public/MuT/NodeModifier.h
+++ b/Engine/Plugins/Experimental/Mutable/Source/MutableTools/Public/MuT/NodeModifier.h
@@ -8,6 +8,19 @@
 #include "MuT/Node.h"
 #include "Containers/UnrealString.h"
 
+#include "NodeModifier.generated.h"
+
+
+/** Despite being an UEnum, this is not always version-serialized (in MutableTools).
+* Beware of changing the enum options or order.
+*/
+UENUM()
+enum class EMutableMultipleTagPolicy : uint8
+{
+	OnlyOneRequired,
+	AllRequired
+};
+
 
 namespace mu
 {
@@ -17,7 +30,6 @@ namespace mu
 	typedef Ptr<NodeModifier> NodeModifierPtr;
 	typedef Ptr<const NodeModifier> NodeModifierConst;
 
-
 	//! This class is the parent of all nodes that output a component.
 	//! \ingroup model
 	class MUTABLETOOLS_API NodeModifier : public Node
@@ -52,15 +64,13 @@ namespace mu
         // Own interface
         //-----------------------------------------------------------------------------------------
 
-        //! \name Tags
-        //! \{
-
-        //! Add a tag to the surface, which will be affected by modifier nodes with the same tag
-        void AddTag(const FString& tagName);
+        /** Add a tag to the surface, which will be affected by modifier nodes with the same tag. */
+        void AddTag(const FString& TagName);
 
-        //! \}
+		/** Set the policy to interprete the tags when there is more than one. */
+		void SetMultipleTagPolicy(EMutableMultipleTagPolicy);
 
-		//! Set the stage to apply this modifier in. Default is before normal operations.
+		/** Set the stage to apply this modifier in.Default is before normal operations. */
 		void SetStage( bool bBeforeNormalOperation );
 
 		//-----------------------------------------------------------------------------------------
diff --git a/Engine/Plugins/Experimental/Mutable/Source/MutableTools/Public/MuT/NodeScalarParameter.h b/Engine/Plugins/Experimental/Mutable/Source/MutableTools/Public/MuT/NodeScalarParameter.h
--- a/Engine/Plugins/Experimental/Mutable/Source/MutableTools/Public/MuT/NodeScalarParameter.h
+++ b/Engine/Plugins/Experimental/Mutable/Source/MutableTools/Public/MuT/NodeScalarParameter.h
@@ -59,10 +59,10 @@ namespace mu
 		//-----------------------------------------------------------------------------------------
 
 		//! Get the name of the parameter. It will be exposed in the final compiled data.
-		void SetName( const char* );
+		void SetName( const FString& );
 
 		//! Get the uid of the parameter. It will be exposed in the final compiled data.
-		void SetUid( const char* );
+		void SetUid( const FString&);
 
 		//! Get the default value of the parameter.
 		void SetDefaultValue( float v );
diff --git a/Engine/Plugins/Experimental/Mutable/Source/MutableTools/Public/MuT/NodeScalarVariation.h b/Engine/Plugins/Experimental/Mutable/Source/MutableTools/Public/MuT/NodeScalarVariation.h
--- a/Engine/Plugins/Experimental/Mutable/Source/MutableTools/Public/MuT/NodeScalarVariation.h
+++ b/Engine/Plugins/Experimental/Mutable/Source/MutableTools/Public/MuT/NodeScalarVariation.h
@@ -53,16 +53,13 @@ namespace mu
         void SetDefaultScalar( NodeScalar* Scalar );
 
         //! Set the number of tags to consider in this variation
-        void SetVariationCount( int count );
-
-        //!
-        int GetVariationCount() const;
+        void SetVariationCount( int32 count );
 
         //! Set the tag or state name that will enable a specific vartiation
-        void SetVariationTag( int index, const char* strTag );
+        void SetVariationTag( int32 index, const FString& strTag );
 
         //!
-        void SetVariationScalar( int index, NodeScalar* Scalar );
+        void SetVariationScalar( int32 index, NodeScalar* Scalar );
 
         //!}
 
diff --git a/Engine/Plugins/Experimental/Mutable/Source/MutableTools/Public/MuT/NodeStringParameter.h b/Engine/Plugins/Experimental/Mutable/Source/MutableTools/Public/MuT/NodeStringParameter.h
--- a/Engine/Plugins/Experimental/Mutable/Source/MutableTools/Public/MuT/NodeStringParameter.h
+++ b/Engine/Plugins/Experimental/Mutable/Source/MutableTools/Public/MuT/NodeStringParameter.h
@@ -61,13 +61,13 @@ namespace mu
 		//-----------------------------------------------------------------------------------------
 
 		//! Set the name of the parameter. It will be exposed in the final compiled data.
-		void SetName( const char* );
+		void SetName( const FString& );
 
 		//! Set the uid of the parameter. It will be exposed in the final compiled data.
-		void SetUid( const char* );
+		void SetUid( const FString& );
 
 		//! Set the default value of the parameter.
-		void SetDefaultValue( const char* v );
+		void SetDefaultValue( const FString& );
 
         //! Set the number of ranges (dimensions) for this parameter.
         //! By default a parameter has 0 ranges, meaning it only has one value.
diff --git a/Engine/Plugins/Experimental/Mutable/Source/MutableTools/Public/MuT/NodeSurfaceVariation.h b/Engine/Plugins/Experimental/Mutable/Source/MutableTools/Public/MuT/NodeSurfaceVariation.h
--- a/Engine/Plugins/Experimental/Mutable/Source/MutableTools/Public/MuT/NodeSurfaceVariation.h
+++ b/Engine/Plugins/Experimental/Mutable/Source/MutableTools/Public/MuT/NodeSurfaceVariation.h
@@ -78,7 +78,7 @@ namespace mu
         void SetVariationType(VariationType);
 
         //! Set the tag or state name that will enable a specific vartiation
-		void SetVariationTag(int index, const char* strTag);
+		void SetVariationTag(int index, const FString& Tag);
 
 		//! 
         void AddVariationSurface(int index, NodeSurface* surface);
