[mutable] Add initial version of Material Switching node.
#rnx

diff --git a/Engine/Plugins/Experimental/Mutable/Source/CustomizableObjectEditor/Classes/MuCOE/Nodes/CustomizableObjectNodeMaterialSwitch.h b/Engine/Plugins/Experimental/Mutable/Source/CustomizableObjectEditor/Classes/MuCOE/Nodes/CustomizableObjectNodeMaterialSwitch.h
new file mode 100644
--- /dev/null
+++ b/Engine/Plugins/Experimental/Mutable/Source/CustomizableObjectEditor/Classes/MuCOE/Nodes/CustomizableObjectNodeMaterialSwitch.h
@@ -0,0 +1,19 @@
+// Copyright Epic Games, Inc. All Rights Reserved.
+
+#pragma once
+
+#include "MuCOE/Nodes/CustomizableObjectNodeSwitchBase.h"
+
+#include "CustomizableObjectNodeMaterialSwitch.generated.h"
+
+
+UCLASS()
+class CUSTOMIZABLEOBJECTEDITOR_API UCustomizableObjectNodeMaterialSwitch : public UCustomizableObjectNodeSwitchBase
+{
+public:
+	GENERATED_BODY()
+	
+	// UCustomizableObjectNodeSwitchBase interface
+	virtual FName GetCategory() const override;
+};
+
diff --git a/Engine/Plugins/Experimental/Mutable/Source/CustomizableObjectEditor/Private/MuCOE/EdGraphSchema_CustomizableObject.cpp b/Engine/Plugins/Experimental/Mutable/Source/CustomizableObjectEditor/Private/MuCOE/EdGraphSchema_CustomizableObject.cpp
--- a/Engine/Plugins/Experimental/Mutable/Source/CustomizableObjectEditor/Private/MuCOE/EdGraphSchema_CustomizableObject.cpp
+++ b/Engine/Plugins/Experimental/Mutable/Source/CustomizableObjectEditor/Private/MuCOE/EdGraphSchema_CustomizableObject.cpp
@@ -33,6 +33,7 @@
 #include "MuCOE/Nodes/CustomizableObjectNodeGroupProjectorParameter.h"
 #include "MuCOE/Nodes/CustomizableObjectNodeLayoutBlocks.h"
 #include "MuCOE/Nodes/CustomizableObjectNodeMaterialVariation.h"
+#include "MuCOE/Nodes/CustomizableObjectNodeMaterialSwitch.h"
 #include "MuCOE/Nodes/CustomizableObjectNodeMeshClipDeform.h"
 #include "MuCOE/Nodes/CustomizableObjectNodeMeshClipMorph.h"
 #include "MuCOE/Nodes/CustomizableObjectNodeMeshClipWithMesh.h"
@@ -374,6 +375,7 @@ void UEdGraphSchema_CustomizableObject::GetGraphContextActions(FGraphContextMenu
 			ContextMenuBuilder.CreateTemplateNode<UCustomizableObjectNodeRemoveMeshBlocks>(),
 			ContextMenuBuilder.CreateTemplateNode<UCustomizableObjectNodeEditMaterial>(),
 			ContextMenuBuilder.CreateTemplateNode<UCustomizableObjectNodeMaterialVariation>(),
+			ContextMenuBuilder.CreateTemplateNode<UCustomizableObjectNodeMaterialSwitch>(),
 			ContextMenuBuilder.CreateTemplateNode<UCustomizableObjectNodeMorphMaterial>(),
 			ContextMenuBuilder.CreateTemplateNode<UCustomizableObjectNodeCopyMaterial>(),
 			ContextMenuBuilder.CreateTemplateNode<UCustomizableObjectNodeMeshClipMorph>(),
diff --git a/Engine/Plugins/Experimental/Mutable/Source/CustomizableObjectEditor/Private/MuCOE/GenerateMutableSource/GenerateMutableSource.cpp b/Engine/Plugins/Experimental/Mutable/Source/CustomizableObjectEditor/Private/MuCOE/GenerateMutableSource/GenerateMutableSource.cpp
--- a/Engine/Plugins/Experimental/Mutable/Source/CustomizableObjectEditor/Private/MuCOE/GenerateMutableSource/GenerateMutableSource.cpp
+++ b/Engine/Plugins/Experimental/Mutable/Source/CustomizableObjectEditor/Private/MuCOE/GenerateMutableSource/GenerateMutableSource.cpp
@@ -19,6 +19,7 @@
 #include "MuCOE/Nodes/CustomizableObjectNodeExtendMaterial.h"
 #include "MuCOE/Nodes/CustomizableObjectNodeMaterial.h"
 #include "MuCOE/Nodes/CustomizableObjectNodeMaterialVariation.h"
+#include "MuCOE/Nodes/CustomizableObjectNodeMaterialSwitch.h"
 #include "MuCOE/Nodes/CustomizableObjectNodeMesh.h"
 #include "MuCOE/Nodes/CustomizableObjectNodeMeshClipWithMesh.h"
 #include "MuCOE/Nodes/CustomizableObjectNodeObjectGroup.h"
@@ -1315,6 +1316,37 @@ bool AffectsCurrentComponent(const UEdGraphPin* Pin, FMutableGraphGenerationCont
 
 		return bAffectsCurrentComponent;
 	}
+
+	else if (const UCustomizableObjectNodeMaterialSwitch* TypedNodeSwitch = Cast<UCustomizableObjectNodeMaterialSwitch>(Node))
+	{
+		bool bAffectsCurrentComponent = false;
+
+		int32 OptionCount = TypedNodeSwitch->GetNumElements();
+		for (int32 OptionIndex=0; OptionIndex<OptionCount; ++OptionIndex)
+		{
+			if (UEdGraphPin* OptionPin = TypedNodeSwitch->GetElementPin(OptionIndex))
+			{
+				for (const UEdGraphPin* ConnectedPin : FollowInputPinArray(*OptionPin))
+				{
+					if (!AffectsCurrentComponent(ConnectedPin, GenerationContext))
+					{
+						if (bAffectsCurrentComponent)
+						{
+							FString Msg = FString::Printf(TEXT("Error! One or more materials nodes linked to a material switch node have different component index"));
+							GenerationContext.Compiler->CompilerLog(FText::FromString(Msg), Node, EMessageSeverity::Error);
+						}
+
+						return false;
+					}
+
+					bAffectsCurrentComponent = true;
+				}
+			}
+		}
+
+		return bAffectsCurrentComponent;
+	}
+
 	else if (const UCustomizableObjectNodeMaterial* TypedNodeMat = Cast<UCustomizableObjectNodeMaterial>(Node))
 	{
 		ComponentIndex = TypedNodeMat->MeshComponentIndex;
diff --git a/Engine/Plugins/Experimental/Mutable/Source/CustomizableObjectEditor/Private/MuCOE/GenerateMutableSource/GenerateMutableSourceSurface.cpp b/Engine/Plugins/Experimental/Mutable/Source/CustomizableObjectEditor/Private/MuCOE/GenerateMutableSource/GenerateMutableSourceSurface.cpp
--- a/Engine/Plugins/Experimental/Mutable/Source/CustomizableObjectEditor/Private/MuCOE/GenerateMutableSource/GenerateMutableSourceSurface.cpp
+++ b/Engine/Plugins/Experimental/Mutable/Source/CustomizableObjectEditor/Private/MuCOE/GenerateMutableSource/GenerateMutableSourceSurface.cpp
@@ -32,6 +32,7 @@
 #include "MuCOE/Nodes/CustomizableObjectNodeFloatConstant.h"
 #include "MuCOE/Nodes/CustomizableObjectNodeFloatParameter.h"
 #include "MuCOE/Nodes/CustomizableObjectNodeMaterialVariation.h"
+#include "MuCOE/Nodes/CustomizableObjectNodeMaterialSwitch.h"
 #include "MuCOE/Nodes/CustomizableObjectNodeMorphMaterial.h"
 #include "MuCOE/Nodes/CustomizableObjectNodeRemoveMesh.h"
 #include "MuCOE/Nodes/CustomizableObjectNodeRemoveMeshBlocks.h"
@@ -48,6 +49,7 @@
 #include "MuT/NodePatchMesh.h"
 #include "MuT/NodeScalarConstant.h"
 #include "MuT/NodeSurfaceEdit.h"
+#include "MuT/NodeSurfaceSwitch.h"
 #include "MuT/NodeSurfaceVariation.h"
 #include "MuT/UnrealPixelFormatOverride.h"
 
@@ -277,7 +279,7 @@ mu::NodeSurfacePtr GenerateMutableSourceSurface(const UEdGraphPin * Pin, FMutabl
 		}
 
 		const UEdGraphPin* ConnectedMaterialPin = FollowInputPin(*TypedNodeMat->GetMeshPin());
-		// Warn when texture connections are  improperly used by connecting them directly to material inputs when no layout is used
+		// Warn when texture connections are improperly used by connecting them directly to material inputs when no layout is used
 		// TODO: delete the if clause and the warning when static meshes are operational again
 		if (ConnectedMaterialPin)
 		{
@@ -1529,6 +1531,81 @@ mu::NodeSurfacePtr GenerateMutableSourceSurface(const UEdGraphPin * Pin, FMutabl
 		}
 	}
 
+	else if (const UCustomizableObjectNodeMaterialSwitch* TypedNodeSwitch = Cast<UCustomizableObjectNodeMaterialSwitch>(Node))
+	{
+		// Using a lambda so control flow is easier to manage.
+		Result = [&]()
+		{
+			const UEdGraphPin* SwitchParameter = TypedNodeSwitch->SwitchParameter();
+
+			// Check Switch Parameter arity preconditions.
+			if (const UEdGraphPin* EnumPin = FollowInputPin(*SwitchParameter))
+			{
+				mu::NodeScalarPtr SwitchParam = GenerateMutableSourceFloat(EnumPin, GenerationContext);
+
+				// Switch Param not generated
+				if (!SwitchParam)
+				{
+					// Warn about a failure.
+					if (EnumPin)
+					{
+						const FText Message = LOCTEXT("FailedToGenerateSwitchParam", "Could not generate switch enum parameter. Please refesh the switch node and connect an enum.");
+						GenerationContext.Compiler->CompilerLog(Message, Node);
+					}
+
+					return Result;
+				}
+
+				if (SwitchParam->GetType() != mu::NodeScalarEnumParameter::GetStaticType())
+				{
+					const FText Message = LOCTEXT("WrongSwitchParamType", "Switch parameter of incorrect type.");
+					GenerationContext.Compiler->CompilerLog(Message, Node);
+
+					return Result;
+				}
+
+				const int32 NumSwitchOptions = TypedNodeSwitch->GetNumElements();
+
+				mu::NodeScalarEnumParameter* EnumParameter = static_cast<mu::NodeScalarEnumParameter*>(SwitchParam.get());
+				if (NumSwitchOptions != EnumParameter->GetValueCount())
+				{
+					const FText Message = LOCTEXT("MismatchedSwitch", "Switch enum and switch node have different number of options. Please refresh the switch node to make sure the outcomes are labeled properly.");
+					GenerationContext.Compiler->CompilerLog(Message, Node);
+				}
+
+				mu::Ptr<mu::NodeSurfaceSwitch> SwitchNode = new mu::NodeSurfaceSwitch;
+				SwitchNode->SetParameter(SwitchParam);
+				SwitchNode->SetOptionCount(NumSwitchOptions);
+
+				for (int32 SelectorIndex = 0; SelectorIndex < NumSwitchOptions; ++SelectorIndex)
+				{
+					if (const UEdGraphPin* ConnectedPin = FollowInputPin(*TypedNodeSwitch->GetElementPin(SelectorIndex)))
+					{
+						FMutableGraphSurfaceGenerationData DummySurfaceData;
+						mu::NodeSurfacePtr ChildNode = GenerateMutableSourceSurface(ConnectedPin, GenerationContext, DummySurfaceData);
+						if (ChildNode)
+						{
+							SwitchNode->SetOption(SelectorIndex, ChildNode.get());
+						}
+						else
+						{
+							// Probably ok
+							//GenerationContext.Compiler->CompilerLog(LOCTEXT("SurfaceModifierFailed", "Surface generation failed."), Node);
+						}
+					}
+				}
+
+				Result = SwitchNode;
+				return Result;
+			}
+			else
+			{
+				GenerationContext.Compiler->CompilerLog(LOCTEXT("NoEnumParamInSwitch", "Switch nodes must have an enum switch parameter. Please connect an enum and refesh the switch node."), Node);
+				return Result;
+			}
+		}(); // invoke lambda.
+	}
+
 	else
 	{
 		GenerationContext.Compiler->CompilerLog(LOCTEXT("UnimplementedNode", "Node type not implemented yet."), Node);
diff --git a/Engine/Plugins/Experimental/Mutable/Source/CustomizableObjectEditor/Private/MuCOE/Nodes/CustomizableObjectNodeMaterialSwitch.cpp b/Engine/Plugins/Experimental/Mutable/Source/CustomizableObjectEditor/Private/MuCOE/Nodes/CustomizableObjectNodeMaterialSwitch.cpp
new file mode 100644
--- /dev/null
+++ b/Engine/Plugins/Experimental/Mutable/Source/CustomizableObjectEditor/Private/MuCOE/Nodes/CustomizableObjectNodeMaterialSwitch.cpp
@@ -0,0 +1,17 @@
+// Copyright Epic Games, Inc. All Rights Reserved.
+
+#include "MuCOE/Nodes/CustomizableObjectNodeMaterialSwitch.h"
+
+#include "MuCOE/EdGraphSchema_CustomizableObject.h"
+
+#define LOCTEXT_NAMESPACE "CustomizableObjectEditor"
+
+
+FName UCustomizableObjectNodeMaterialSwitch::GetCategory() const
+{
+	return UEdGraphSchema_CustomizableObject::PC_Material;
+}
+
+
+#undef LOCTEXT_NAMESPACE
+
diff --git a/Engine/Plugins/Experimental/Mutable/Source/MutableTools/Private/MuT/CodeGenerator.cpp b/Engine/Plugins/Experimental/Mutable/Source/MutableTools/Private/MuT/CodeGenerator.cpp
--- a/Engine/Plugins/Experimental/Mutable/Source/MutableTools/Private/MuT/CodeGenerator.cpp
+++ b/Engine/Plugins/Experimental/Mutable/Source/MutableTools/Private/MuT/CodeGenerator.cpp
@@ -82,11 +82,10 @@
 #include "MuT/NodeScalar.h"
 #include "MuT/NodeScalarConstant.h"
 #include "MuT/NodeSurface.h"
-#include "MuT/NodeSurfaceEdit.h"
 #include "MuT/NodeSurfaceEditPrivate.h"
 #include "MuT/NodeSurfaceNewPrivate.h"
-#include "MuT/NodeSurfaceVariation.h"
 #include "MuT/NodeSurfaceVariationPrivate.h"
+#include "MuT/NodeSurfaceSwitchPrivate.h"
 #include "MuT/TablePrivate.h"
 #include "Trace/Detail/Channel.h"
 
@@ -114,7 +113,7 @@ namespace mu
 		MUTABLE_CPUPROFILER_SCOPE(Generate);
 
 		// First pass
-		m_firstPass.Generate(m_pErrorLog, pNode->GetBasePrivate(), m_compilerOptions->bIgnoreStates);
+		m_firstPass.Generate(m_pErrorLog, pNode->GetBasePrivate(), m_compilerOptions->bIgnoreStates, this);
 
 		// Second pass
 		SecondPassGenerator SecondPass(&m_firstPass, m_compilerOptions);
@@ -217,6 +216,12 @@ namespace mu
 			return nullptr;
 		}
 
+		else if (dynamic_cast<const NodeSurfaceSwitch*>(pNode.get()))
+		{
+			// This happens only if we generate a node graph that has a NodeSurfaceVariation at the root.
+			return nullptr;
+		}
+
 		else if (dynamic_cast<const NodeSurfaceEdit*>(pNode.get()))
 		{
 			// This happens only if we generate a node graph that has a NodeSurfaceEdit at the root.
@@ -507,6 +512,20 @@ namespace mu
 				pResult = FindSourceMesh(pTypedSV->GetPrivate()->m_defaultSurfaces[0].get());
 			}
 		}
+		else if (const NodeSurfaceSwitch* pTypedSS = dynamic_cast<const NodeSurfaceSwitch*>(pNode))
+		{
+			for (int32 i=0; i<pTypedSS->GetPrivate()->Options.Num(); ++i)
+			{
+				if (pTypedSS->GetPrivate()->Options[i])
+				{
+					pResult = FindSourceMesh(pTypedSS->GetPrivate()->Options[i].get());
+					if (pResult)
+					{
+						break;
+					}
+				}
+			}
+		}
 		else if (const NodeMeshInterpolate* pTypedMI = dynamic_cast<const NodeMeshInterpolate*>(pNode))
         {
             if ( pTypedMI->GetTargetCount()>0 )
@@ -1875,7 +1894,7 @@ namespace mu
         MUTABLE_CPUPROFILER_SCOPE(NodeComponentEdit);
 
         // Nothing to do. Surface information will be already collected in the suitable
-        // parent components furing the first and second passes.
+        // parent components during the first and second passes.
 
         return nullptr;
     }
diff --git a/Engine/Plugins/Experimental/Mutable/Source/MutableTools/Private/MuT/CodeGenerator.h b/Engine/Plugins/Experimental/Mutable/Source/MutableTools/Private/MuT/CodeGenerator.h
--- a/Engine/Plugins/Experimental/Mutable/Source/MutableTools/Private/MuT/CodeGenerator.h
+++ b/Engine/Plugins/Experimental/Mutable/Source/MutableTools/Private/MuT/CodeGenerator.h
@@ -122,6 +122,9 @@ namespace mu
                           public Visitor<NodeObjectGroup::Private, Ptr<ASTOp>, true>,
                           public Visitor<NodePatchImage::Private, Ptr<ASTOp>, true>
     {
+		
+		friend class FirstPassGenerator;
+
     public:
 
         CodeGenerator( CompilerOptions::Private* options );
@@ -666,7 +669,7 @@ namespace mu
 		typedef TMap<FVisitedKeyMap, FScalarGenerationResult> GeneratedScalarsMap;
 		GeneratedScalarsMap m_generatedScalars;
 
-		void GenerateScalar(FScalarGenerationResult&, const NodeScalarPtrConst&);
+		void GenerateScalar(FScalarGenerationResult&, const Ptr<const NodeScalar>&);
 		void GenerateScalar_Constant(FScalarGenerationResult&, const Ptr<const NodeScalarConstant>&);
 		void GenerateScalar_Parameter(FScalarGenerationResult&, const Ptr<const NodeScalarParameter>&);
 		void GenerateScalar_Switch(FScalarGenerationResult&, const Ptr<const NodeScalarSwitch>&);
diff --git a/Engine/Plugins/Experimental/Mutable/Source/MutableTools/Private/MuT/CodeGenerator_FirstPass.cpp b/Engine/Plugins/Experimental/Mutable/Source/MutableTools/Private/MuT/CodeGenerator_FirstPass.cpp
--- a/Engine/Plugins/Experimental/Mutable/Source/MutableTools/Private/MuT/CodeGenerator_FirstPass.cpp
+++ b/Engine/Plugins/Experimental/Mutable/Source/MutableTools/Private/MuT/CodeGenerator_FirstPass.cpp
@@ -2,6 +2,7 @@
 
 #include "MuT/CodeGenerator_FirstPass.h"
 
+#include "MuT/CodeGenerator.h"
 #include "HAL/PlatformMath.h"
 #include "Misc/AssertionMacros.h"
 #include "MuR/MutableTrace.h"
@@ -31,6 +32,8 @@
 #include "MuT/NodeSurfaceEditPrivate.h"
 #include "MuT/NodeSurfaceNewPrivate.h"
 #include "MuT/NodeSurfaceVariationPrivate.h"
+#include "MuT/NodeSurfaceSwitchPrivate.h"
+#include "MuT/NodeScalarEnumParameterPrivate.h"
 
 namespace mu
 {
@@ -48,10 +51,12 @@ namespace mu
 	//---------------------------------------------------------------------------------------------
     void FirstPassGenerator::Generate( ErrorLogPtr pErrorLog,
                                        const Node::Private* root,
-                                       bool ignoreStates )
+                                       bool ignoreStates,
+									   CodeGenerator* InGenerator )
 	{
 		MUTABLE_CPUPROFILER_SCOPE(FirstPassGenerate);
 
+		Generator = InGenerator;
 		m_pErrorLog = pErrorLog;
         m_ignoreStates = ignoreStates;
 
@@ -424,6 +429,66 @@ namespace mu
 	}
 
 
+	//---------------------------------------------------------------------------------------------
+	Ptr<ASTOp> FirstPassGenerator::Visit(const NodeSurfaceSwitch::Private& node)
+	{
+		if (node.Options.Num() == 0)
+		{
+			// No options in the switch!
+			return nullptr;
+		}
+
+		// Prepare the enumeration parameter
+		CodeGenerator::FScalarGenerationResult ScalarResult;
+		if (node.Parameter)
+		{
+			Generator->GenerateScalar( ScalarResult, node.Parameter );
+
+			// Not really necessary but it currently should always be this type.
+			Ptr<ASTOpParameter> EnumOp = dynamic_cast<ASTOpParameter*>(ScalarResult.op.get());
+			check(EnumOp);
+		}
+		else
+		{
+			// This argument is required
+			ScalarResult.op = Generator->GenerateMissingScalarCode(TEXT("Switch variable"), 0.0f, node.m_errorContext);
+		}
+
+		// Parse the options
+		for (int32 t = 0; t < node.Options.Num(); ++t)
+		{
+			// Create a comparison operation as the boolean parameter for the child
+			Ptr<ASTOpFixed> ParamOp = new ASTOpFixed();
+			ParamOp->op.type = OP_TYPE::BO_EQUAL_INT_CONST;
+			ParamOp->SetChild(ParamOp->op.args.BoolEqualScalarConst.value, ScalarResult.op);
+			ParamOp->op.args.BoolEqualScalarConst.constant = (int16)t;
+
+			// Combine the new condition with previous conditions coming from parent objects
+			if (m_currentCondition.Last().objectCondition)
+			{
+				Ptr<ASTOpFixed> op = new ASTOpFixed();
+				op->op.type = OP_TYPE::BO_AND;
+				op->SetChild(op->op.args.BoolBinary.a, m_currentCondition.Last().objectCondition);
+				op->SetChild(op->op.args.BoolBinary.b, ParamOp);
+				ParamOp = op;
+			}
+
+			FConditionContext data;
+			data.objectCondition = ParamOp;
+			m_currentCondition.Push(data);
+
+			if (node.Options[t])
+			{
+				node.Options[t]->GetBasePrivate()->Accept(*this);
+			}
+
+			m_currentCondition.Pop();
+		}
+
+		return nullptr;
+	}
+
+
 	//---------------------------------------------------------------------------------------------
     Ptr<ASTOp> FirstPassGenerator::Visit(const NodeComponentNew::Private& node)
 	{
diff --git a/Engine/Plugins/Experimental/Mutable/Source/MutableTools/Private/MuT/CodeGenerator_FirstPass.h b/Engine/Plugins/Experimental/Mutable/Source/MutableTools/Private/MuT/CodeGenerator_FirstPass.h
--- a/Engine/Plugins/Experimental/Mutable/Source/MutableTools/Private/MuT/CodeGenerator_FirstPass.h
+++ b/Engine/Plugins/Experimental/Mutable/Source/MutableTools/Private/MuT/CodeGenerator_FirstPass.h
@@ -22,6 +22,7 @@
 #include "MuT/NodeSurfaceEdit.h"
 #include "MuT/NodeSurfaceNew.h"
 #include "MuT/NodeSurfaceVariation.h"
+#include "MuT/NodeSurfaceSwitch.h"
 #include "MuT/Visitor.h"
 
 #include <stdint.h>
@@ -44,8 +45,9 @@ namespace mu
 
         public Visitor<NodeSurfaceNew::Private, Ptr<ASTOp>, true>,
         public Visitor<NodeSurfaceEdit::Private, Ptr<ASTOp>, true>,
-        public Visitor<NodeSurfaceVariation::Private, Ptr<ASTOp>, true>,
-        public Visitor<NodeComponentNew::Private, Ptr<ASTOp>, true>,
+		public Visitor<NodeSurfaceVariation::Private, Ptr<ASTOp>, true>,
+		public Visitor<NodeSurfaceSwitch::Private, Ptr<ASTOp>, true>,
+		public Visitor<NodeComponentNew::Private, Ptr<ASTOp>, true>,
         public Visitor<NodeComponentEdit::Private, Ptr<ASTOp>, true>,
         public Visitor<NodeLOD::Private, Ptr<ASTOp>, true>,
         public Visitor<NodeObjectNew::Private, Ptr<ASTOp>, true>,
@@ -61,9 +63,7 @@ namespace mu
 
 		FirstPassGenerator();
 
-        void Generate(ErrorLogPtr pErrorLog,
-                      const Node::Private* root,
-                      bool ignoreStates);
+        void Generate(ErrorLogPtr pErrorLog, const Node::Private* root, bool ignoreStates, class CodeGenerator*);
 
 	public:
 
@@ -188,8 +188,9 @@ namespace mu
 
         Ptr<ASTOp> Visit(const NodeSurfaceNew::Private&) override;
         Ptr<ASTOp> Visit(const NodeSurfaceEdit::Private&) override;
-        Ptr<ASTOp> Visit(const NodeSurfaceVariation::Private&) override;
-        Ptr<ASTOp> Visit(const NodeComponentNew::Private&) override;
+		Ptr<ASTOp> Visit(const NodeSurfaceVariation::Private&) override;
+		Ptr<ASTOp> Visit(const NodeSurfaceSwitch::Private&) override;
+		Ptr<ASTOp> Visit(const NodeComponentNew::Private&) override;
         Ptr<ASTOp> Visit(const NodeComponentEdit::Private&) override;
         Ptr<ASTOp> Visit(const NodeLOD::Private&) override;
         Ptr<ASTOp> Visit(const NodeObjectNew::Private& node) override;
@@ -221,6 +222,9 @@ namespace mu
 		//! Index of the LOD we are processing
         int m_currentLOD = -1;
 
+		/** Non-owned reference to main code generator. */
+		CodeGenerator* Generator = nullptr;
+
         //!
         ErrorLogPtr m_pErrorLog;
 
diff --git a/Engine/Plugins/Experimental/Mutable/Source/MutableTools/Private/MuT/CodeGenerator_Scalar.cpp b/Engine/Plugins/Experimental/Mutable/Source/MutableTools/Private/MuT/CodeGenerator_Scalar.cpp
--- a/Engine/Plugins/Experimental/Mutable/Source/MutableTools/Private/MuT/CodeGenerator_Scalar.cpp
+++ b/Engine/Plugins/Experimental/Mutable/Source/MutableTools/Private/MuT/CodeGenerator_Scalar.cpp
@@ -54,7 +54,7 @@ namespace mu
 
 
 	//-------------------------------------------------------------------------------------------------
-	void CodeGenerator::GenerateScalar(FScalarGenerationResult& result, const NodeScalarPtrConst& untyped)
+	void CodeGenerator::GenerateScalar(FScalarGenerationResult& result, const Ptr<const NodeScalar>& untyped)
 	{
 		if (!untyped)
 		{
diff --git a/Engine/Plugins/Experimental/Mutable/Source/MutableTools/Private/MuT/NodeSurface.cpp b/Engine/Plugins/Experimental/Mutable/Source/MutableTools/Private/MuT/NodeSurface.cpp
--- a/Engine/Plugins/Experimental/Mutable/Source/MutableTools/Private/MuT/NodeSurface.cpp
+++ b/Engine/Plugins/Experimental/Mutable/Source/MutableTools/Private/MuT/NodeSurface.cpp
@@ -9,6 +9,7 @@
 #include "MuT/NodeSurfaceEdit.h"
 #include "MuT/NodeSurfaceNew.h"
 #include "MuT/NodeSurfaceVariation.h"
+#include "MuT/NodeSurfaceSwitch.h"
 
 #include <stdint.h>
 
@@ -67,6 +68,7 @@ namespace mu
         case 1: return NodeSurfaceEdit::StaticUnserialise( arch ); break;
 		//case 2: return NodeSelectSurface::StaticUnserialise(arch); break;
 		case 3: return NodeSurfaceVariation::StaticUnserialise(arch); break;
+		case 4: return NodeSurfaceSwitch::StaticUnserialise(arch); break;
 		default : check(false);
 		}
 
diff --git a/Engine/Plugins/Experimental/Mutable/Source/MutableTools/Private/MuT/NodeSurfaceSwitch.cpp b/Engine/Plugins/Experimental/Mutable/Source/MutableTools/Private/MuT/NodeSurfaceSwitch.cpp
new file mode 100644
--- /dev/null
+++ b/Engine/Plugins/Experimental/Mutable/Source/MutableTools/Private/MuT/NodeSurfaceSwitch.cpp
@@ -0,0 +1,118 @@
+// Copyright Epic Games, Inc. All Rights Reserved.
+
+#include "MuT/NodeSurfaceSwitch.h"
+
+#include "Misc/AssertionMacros.h"
+#include "MuT/NodeLayout.h"
+#include "MuT/NodeSurfaceSwitchPrivate.h"
+#include "MuT/NodePrivate.h"
+#include "MuT/NodeScalar.h"
+
+
+namespace mu
+{
+
+	//---------------------------------------------------------------------------------------------
+	// Static initialisation
+	//---------------------------------------------------------------------------------------------
+	NODE_TYPE NodeSurfaceSwitch::Private::s_type =
+			NODE_TYPE( "SurfaceSwitch", NodeSurface::GetStaticType() );
+
+
+	//---------------------------------------------------------------------------------------------
+	//!
+	//---------------------------------------------------------------------------------------------
+
+	MUTABLE_IMPLEMENT_NODE( NodeSurfaceSwitch, EType::Switch, Node, Node::EType::Surface)
+
+
+	//---------------------------------------------------------------------------------------------
+	// Node Interface
+	//---------------------------------------------------------------------------------------------
+	int NodeSurfaceSwitch::GetInputCount() const
+	{
+		return 1 + m_pD->Options.Num();
+	}
+
+
+	//---------------------------------------------------------------------------------------------
+	Node* NodeSurfaceSwitch::GetInputNode( int i ) const
+	{
+		check( i>=0 && i<GetInputCount() );
+
+		Node* pResult = 0;
+
+		switch (i)
+		{
+		case 0:
+			pResult = m_pD->Parameter.get();
+			break;
+
+		default:
+			pResult = m_pD->Options[i-1].get();
+			break;
+		}
+
+		return pResult;
+	}
+
+
+	//---------------------------------------------------------------------------------------------
+	void NodeSurfaceSwitch::SetInputNode( int i, NodePtr pNode )
+	{
+		check( i>=0 && i<GetInputCount() );
+
+		switch (i)
+		{
+		case 0:
+			m_pD->Parameter = dynamic_cast<NodeScalar*>(pNode.get());
+			break;
+
+		default:
+			m_pD->Options[i-1] = dynamic_cast<NodeSurface*>(pNode.get());
+			break;
+		}
+	}
+
+
+	//---------------------------------------------------------------------------------------------
+	// Own Interface
+	//---------------------------------------------------------------------------------------------
+	NodeScalarPtr NodeSurfaceSwitch::GetParameter() const
+	{
+		return m_pD->Parameter.get();
+	}
+
+
+	//---------------------------------------------------------------------------------------------
+	void NodeSurfaceSwitch::SetParameter( NodeScalarPtr pNode )
+	{
+		m_pD->Parameter = pNode;
+	}
+
+
+	//---------------------------------------------------------------------------------------------
+	void NodeSurfaceSwitch::SetOptionCount( int32 t )
+	{
+		m_pD->Options.SetNum(t);
+	}
+
+
+	//---------------------------------------------------------------------------------------------
+	NodeSurfacePtr NodeSurfaceSwitch::GetOption( int32 t ) const
+	{
+		check( t>=0 && t<m_pD->Options.Num() );
+		return m_pD->Options[t].get();
+	}
+
+
+	//---------------------------------------------------------------------------------------------
+	void NodeSurfaceSwitch::SetOption( int32 t, NodeSurfacePtr pNode )
+	{
+		check( t>=0 && t<m_pD->Options.Num() );
+		m_pD->Options[t] = pNode;
+	}
+
+}
+
+
diff --git a/Engine/Plugins/Experimental/Mutable/Source/MutableTools/Private/MuT/NodeSurfaceSwitchPrivate.h b/Engine/Plugins/Experimental/Mutable/Source/MutableTools/Private/MuT/NodeSurfaceSwitchPrivate.h
new file mode 100644
--- /dev/null
+++ b/Engine/Plugins/Experimental/Mutable/Source/MutableTools/Private/MuT/NodeSurfaceSwitchPrivate.h
@@ -0,0 +1,52 @@
+// Copyright Epic Games, Inc. All Rights Reserved.
+
+#pragma once
+
+#include "MuT/NodeSurfaceSwitch.h"
+#include "MuT/NodeScalar.h"
+#include "MuT/NodeMeshPrivate.h"
+#include "MuT/AST.h"
+
+
+namespace mu
+{
+
+
+	class NodeSurfaceSwitch::Private : public NodeSurface::Private
+	{
+	public:
+
+		MUTABLE_DEFINE_CONST_VISITABLE()
+
+	public:
+
+		static NODE_TYPE s_type;
+
+		Ptr<NodeScalar> Parameter;
+		TArray<Ptr<NodeSurface>> Options;
+
+		//!
+		void Serialise( OutputArchive& arch ) const
+		{
+            uint32_t ver = 0;
+			arch << ver;
+
+			arch << Parameter;
+			arch << Options;
+		}
+
+		//!
+		void Unserialise( InputArchive& arch )
+		{
+            uint32_t ver;
+			arch >> ver;
+			check(ver==0);
+
+			arch >> Parameter;
+			arch >> Options;
+		}
+
+	};
+
+
+}
diff --git a/Engine/Plugins/Experimental/Mutable/Source/MutableTools/Public/MuT/NodeSurface.h b/Engine/Plugins/Experimental/Mutable/Source/MutableTools/Public/MuT/NodeSurface.h
--- a/Engine/Plugins/Experimental/Mutable/Source/MutableTools/Public/MuT/NodeSurface.h
+++ b/Engine/Plugins/Experimental/Mutable/Source/MutableTools/Public/MuT/NodeSurface.h
@@ -30,6 +30,7 @@ namespace mu
 			New = 0,
 			Edit = 1,
 			Variation = 3,
+			Switch = 4,
 
 			None
 		};
diff --git a/Engine/Plugins/Experimental/Mutable/Source/MutableTools/Public/MuT/NodeSurfaceSwitch.h b/Engine/Plugins/Experimental/Mutable/Source/MutableTools/Public/MuT/NodeSurfaceSwitch.h
new file mode 100644
--- /dev/null
+++ b/Engine/Plugins/Experimental/Mutable/Source/MutableTools/Public/MuT/NodeSurfaceSwitch.h
@@ -0,0 +1,90 @@
+// Copyright Epic Games, Inc. All Rights Reserved.
+
+#pragma once
+
+#include "MuR/Ptr.h"
+#include "MuR/RefCounted.h"
+#include "MuT/Node.h"
+#include "MuT/NodeSurface.h"
+
+
+namespace mu
+{
+
+	// Forward definitions
+	class NodeScalar;
+	typedef Ptr<NodeScalar> NodeScalarPtr;
+	typedef Ptr<const NodeScalar> NodeScalarPtrConst;
+
+	class NodeSurfaceSwitch;
+	typedef Ptr<NodeSurfaceSwitch> NodeSurfaceSwitchPtr;
+	typedef Ptr<const NodeSurfaceSwitch> NodeSurfaceSwitchPtrConst;
+
+
+	//! This node selects an output Surface from a set of input Surfaces based on a parameter.
+	//! \ingroup model
+	class MUTABLETOOLS_API NodeSurfaceSwitch : public NodeSurface
+	{
+	public:
+
+		//-----------------------------------------------------------------------------------------
+		// Life cycle
+		//-----------------------------------------------------------------------------------------
+
+		NodeSurfaceSwitch();
+
+		void SerialiseWrapper(OutputArchive& arch) const override;
+		static void Serialise( const NodeSurfaceSwitch* pNode, OutputArchive& arch );
+		static NodeSurfaceSwitchPtr StaticUnserialise( InputArchive& arch );
+
+
+		//-----------------------------------------------------------------------------------------
+		// Node Interface
+		//-----------------------------------------------------------------------------------------
+
+        
+
+		const NODE_TYPE* GetType() const override;
+		static const NODE_TYPE* GetStaticType();
+
+		int GetInputCount() const override;
+		Node* GetInputNode( int i ) const override;
+		void SetInputNode( int i, NodePtr pNode ) override;
+
+		//-----------------------------------------------------------------------------------------
+		// Own Interface
+		//-----------------------------------------------------------------------------------------
+
+		//! Get the node generating the parameter used to select the option.
+		NodeScalarPtr GetParameter() const;
+		void SetParameter( NodeScalarPtr );
+
+		//! Set the number of option Surfaces. It will keep the currently set targets and initialise
+		//! the new ones as null.
+		void SetOptionCount( int32 );
+
+		//! Get the node generating the t-th option Surface.
+		NodeSurfacePtr GetOption( int32 t ) const;
+		void SetOption( int32 t, NodeSurfacePtr );
+
+
+		//-----------------------------------------------------------------------------------------
+		// Interface pattern
+		//-----------------------------------------------------------------------------------------
+		class Private;
+		Private* GetPrivate() const;
+        Node::Private* GetBasePrivate() const override;
+
+	protected:
+
+		//! Forbidden. Manage with the Ptr<> template.
+		~NodeSurfaceSwitch();
+
+	private:
+
+		Private* m_pD;
+
+	};
+
+
+}
