[mutable] Added deterministic generation of random instances using a random stream.

diff --git a/Engine/Plugins/Experimental/Mutable/Source/CustomizableObject/Classes/MuCO/CustomizableObjectInstance.h b/Engine/Plugins/Experimental/Mutable/Source/CustomizableObject/Classes/MuCO/CustomizableObjectInstance.h
--- a/Engine/Plugins/Experimental/Mutable/Source/CustomizableObject/Classes/MuCO/CustomizableObjectInstance.h
+++ b/Engine/Plugins/Experimental/Mutable/Source/CustomizableObject/Classes/MuCO/CustomizableObjectInstance.h
@@ -5,6 +5,7 @@
 #include "Async/TaskGraphFwd.h"
 #include "MuCO/CustomizableObjectInstanceDescriptor.h"
 #include "Templates/SubclassOf.h"
+#include "Math/RandomStream.h"
 
 #include "CustomizableObjectInstance.generated.h"
 
@@ -230,6 +231,10 @@ public:
 	UFUNCTION(BlueprintCallable, Category = CustomizableObjectInstance)
 	void SetRandomValues();
 
+	/**  Set random values to the parameters using a stream. */
+	UFUNCTION(BlueprintCallable, Category = CustomizableObjectInstance)
+	void SetRandomValuesFromStream(const FRandomStream& Stream);
+
 	/**  Set random values to the parameters using a seed. Useful for testing only. */
 	void SetRandomValues(const int32 InRandomizationSeed);
 
diff --git a/Engine/Plugins/Experimental/Mutable/Source/CustomizableObject/Private/MuCO/CustomizableObjectInstance.cpp b/Engine/Plugins/Experimental/Mutable/Source/CustomizableObject/Private/MuCO/CustomizableObjectInstance.cpp
--- a/Engine/Plugins/Experimental/Mutable/Source/CustomizableObject/Private/MuCO/CustomizableObjectInstance.cpp
+++ b/Engine/Plugins/Experimental/Mutable/Source/CustomizableObject/Private/MuCO/CustomizableObjectInstance.cpp
@@ -2446,6 +2446,11 @@ void UCustomizableObjectInstance::SetRandomValues()
 	Descriptor.SetRandomValues(FMath::Rand());
 }
 
+void UCustomizableObjectInstance::SetRandomValuesFromStream(const FRandomStream& Stream)
+{
+	Descriptor.SetRandomValuesFromStream(Stream);
+}
+
 void UCustomizableObjectInstance::SetRandomValues(const int32 InRandomizationSeed)
 {
 	Descriptor.SetRandomValues(InRandomizationSeed);
diff --git a/Engine/Plugins/Experimental/Mutable/Source/CustomizableObject/Private/MuCO/CustomizableObjectInstanceDescriptor.cpp b/Engine/Plugins/Experimental/Mutable/Source/CustomizableObject/Private/MuCO/CustomizableObjectInstanceDescriptor.cpp
--- a/Engine/Plugins/Experimental/Mutable/Source/CustomizableObject/Private/MuCO/CustomizableObjectInstanceDescriptor.cpp
+++ b/Engine/Plugins/Experimental/Mutable/Source/CustomizableObject/Private/MuCO/CustomizableObjectInstanceDescriptor.cpp
@@ -2302,6 +2302,41 @@ void FCustomizableObjectInstanceDescriptor::SetRandomValues(const int32& InRando
 }
 
 
+void FCustomizableObjectInstanceDescriptor::SetRandomValuesFromStream(const FRandomStream& Stream)
+{
+	check(CustomizableObject);
+	RETURN_ON_UNCOMPILED_CO(CustomizableObject, TEXT("Error: Cannot set random values"))
+		
+	if (!CustomizableObject) return;
+
+	for (int32 i = 0; i < FloatParameters.Num(); ++i)
+	{
+		FloatParameters[i].ParameterValue = Stream.FRand();
+	}
+
+	for (int32 i = 0; i < BoolParameters.Num(); ++i)
+	{
+		BoolParameters[i].ParameterValue = Stream.GetUnsignedInt() % 2 == 0;
+	}
+
+	for (int32 i = 0; i < IntParameters.Num(); ++i)
+	{
+		int32 ParameterIndexInCO = CustomizableObject->FindParameter(IntParameters[i].ParameterName);
+
+		if (ParameterIndexInCO >= 0)
+		{
+			int32 NumValues = CustomizableObject->GetIntParameterNumOptions(ParameterIndexInCO);
+			if (NumValues > 0)
+			{
+				int32 Index = Stream.GetUnsignedInt() % NumValues;
+				FString Option = CustomizableObject->GetIntParameterAvailableOption(ParameterIndexInCO, Index);
+				SetIntParameterSelectedOption(i, Option);
+			}
+		}
+	}
+}
+
+
 bool FCustomizableObjectInstanceDescriptor::CreateMultiLayerProjector(const FName& ProjectorParamName)
 {
 	if (!FMultilayerProjector::AreDescriptorParametersValid(*this, ProjectorParamName.ToString()))
diff --git a/Engine/Plugins/Experimental/Mutable/Source/CustomizableObject/Public/MuCO/CustomizableObjectInstanceDescriptor.h b/Engine/Plugins/Experimental/Mutable/Source/CustomizableObject/Public/MuCO/CustomizableObjectInstanceDescriptor.h
--- a/Engine/Plugins/Experimental/Mutable/Source/CustomizableObject/Public/MuCO/CustomizableObjectInstanceDescriptor.h
+++ b/Engine/Plugins/Experimental/Mutable/Source/CustomizableObject/Public/MuCO/CustomizableObjectInstanceDescriptor.h
@@ -4,6 +4,7 @@
 
 #include "MuCO/MultilayerProjector.h"
 #include "MuCO/CustomizableObjectParameterTypeDefinitions.h"
+#include "Math/RandomStream.h"
 
 #include "CustomizableObjectInstanceDescriptor.generated.h"
 
@@ -294,6 +295,8 @@ struct CUSTOMIZABLEOBJECT_API FCustomizableObjectInstanceDescriptor
 	
 	void SetRandomValues(const int32& InRandomizationSeed);
 
+	void SetRandomValuesFromStream(const FRandomStream& Stream);
+
 	// ------------------------------------------------------------
 	// Multilayer Projectors
 	// ------------------------------------------------------------
