#rnx

diff --git a/Engine/Plugins/Experimental/Mutable/Source/CustomizableObject/Classes/MuCO/CustomizableObjectInstanceUsage.h b/Engine/Plugins/Experimental/Mutable/Source/CustomizableObject/Classes/MuCO/CustomizableObjectInstanceUsage.h
--- a/Engine/Plugins/Experimental/Mutable/Source/CustomizableObject/Classes/MuCO/CustomizableObjectInstanceUsage.h
+++ b/Engine/Plugins/Experimental/Mutable/Source/CustomizableObject/Classes/MuCO/CustomizableObjectInstanceUsage.h
@@ -106,7 +106,7 @@ private:
 
 	// Begin FTickableGameObject
 	virtual void Tick(float DeltaTime) override;
-	virtual ETickableTickType GetTickableTickType() const override { return ETickableTickType::Always; }
+	virtual ETickableTickType GetTickableTickType() const override;
 	virtual TStatId GetStatId() const override;
 	virtual bool IsTickableWhenPaused() const override { return true; }
 	virtual bool IsTickableInEditor() const override { return true; }
diff --git a/Engine/Plugins/Experimental/Mutable/Source/CustomizableObject/Private/MuCO/CustomizableObjectInstanceUsage.cpp b/Engine/Plugins/Experimental/Mutable/Source/CustomizableObject/Private/MuCO/CustomizableObjectInstanceUsage.cpp
--- a/Engine/Plugins/Experimental/Mutable/Source/CustomizableObject/Private/MuCO/CustomizableObjectInstanceUsage.cpp
+++ b/Engine/Plugins/Experimental/Mutable/Source/CustomizableObject/Private/MuCO/CustomizableObjectInstanceUsage.cpp
@@ -18,6 +18,14 @@
 #include "Stats/Stats.h"
 
 
+ETickableTickType UCustomizableObjectInstanceUsage::GetTickableTickType() const
+{ 
+	return (HasAnyFlags(RF_ClassDefaultObject))
+		? ETickableTickType::Never
+		: ETickableTickType::Always;
+}
+
+
 void UCustomizableObjectInstanceUsage::Callbacks() const
 {
 	if (CustomizableSkeletalComponent)
diff --git a/Engine/Plugins/Experimental/Mutable/Source/CustomizableObject/Private/MuCO/CustomizableSkeletalComponent.cpp b/Engine/Plugins/Experimental/Mutable/Source/CustomizableObject/Private/MuCO/CustomizableSkeletalComponent.cpp
--- a/Engine/Plugins/Experimental/Mutable/Source/CustomizableObject/Private/MuCO/CustomizableSkeletalComponent.cpp
+++ b/Engine/Plugins/Experimental/Mutable/Source/CustomizableObject/Private/MuCO/CustomizableSkeletalComponent.cpp
@@ -19,7 +19,12 @@ void UCustomizableSkeletalComponent::CreateCustomizableObjectInstanceUsage()
 		}
 	}
 
-	if (!CustomizableObjectInstanceUsage && !HasAnyFlags(RF_ClassDefaultObject))
+	AActor* RootActor = GetAttachmentRootActor();
+	bool bIsDefaultActor = RootActor ?
+		RootActor->HasAnyFlags(RF_ClassDefaultObject) :
+		false;
+
+	if (!CustomizableObjectInstanceUsage && !HasAnyFlags(RF_ClassDefaultObject) && !bIsDefaultActor)
 	{
 		CustomizableObjectInstanceUsage = NewObject<UCustomizableObjectInstanceUsage>(this, TEXT("InstanceUsage"), RF_Transient);
 		CustomizableObjectInstanceUsage->CustomizableSkeletalComponent = this;
