Fix warning V501: There are identical sub-expressions

diff --git a/Engine/Plugins/Experimental/AnimNext/Source/AnimNext/Private/RigVMRuntimeDataRegistry.cpp b/Engine/Plugins/Experimental/AnimNext/Source/AnimNext/Private/RigVMRuntimeDataRegistry.cpp
--- a/Engine/Plugins/Experimental/AnimNext/Source/AnimNext/Private/RigVMRuntimeDataRegistry.cpp
+++ b/Engine/Plugins/Experimental/AnimNext/Source/AnimNext/Private/RigVMRuntimeDataRegistry.cpp
@@ -77,7 +77,7 @@ static thread_local TMap<FRigVMRuntimeDataID, FRigVMRuntimeData> RuntimeDataStor
 
 	if (RigVMRuntimeData = FindRuntimeData(RigVMRuntimeDataID); RigVMRuntimeData != nullptr)
 	{
-		if (RigVMRuntimeData->Context.VMHash != RigVMRuntimeData->Context.VMHash)
+		if (RigVMRuntimeData->Context.VMHash != ReferenceContext.VMHash)
 		{
 			RigVMRuntimeData->Context = ReferenceContext;
 		}
diff --git a/Engine/Plugins/Experimental/Mutable/Source/MutableRuntime/Private/MuR/OpImageMipmap.cpp b/Engine/Plugins/Experimental/Mutable/Source/MutableRuntime/Private/MuR/OpImageMipmap.cpp
--- a/Engine/Plugins/Experimental/Mutable/Source/MutableRuntime/Private/MuR/OpImageMipmap.cpp
+++ b/Engine/Plugins/Experimental/Mutable/Source/MutableRuntime/Private/MuR/OpImageMipmap.cpp
@@ -382,7 +382,7 @@ namespace mu
 
 							// Average of 2 unsigned integers without overflow extended to work on multiple bytes.
 							constexpr uint64 ShiftMask = 0xFEFEFEFEFEFEFEFE;
-							const uint64 ErrorCorrection = XorRow0Row1Bits & XorRow0Row1Bits & 0x0101010101010101;
+							const uint64 ErrorCorrection = XorRow0Row1Bits & 0x0101010101010101;
 							const uint64 AvgLowBits = (Row0Bits & Row1Bits) + ((XorRow0Row1Bits & ShiftMask) >> 1) + ErrorCorrection;
 							const uint64 AvgHighBits = AvgLowBits >> 32;
 							const uint32 Result = (AvgLowBits & AvgHighBits) + (((AvgLowBits ^ AvgHighBits) & ShiftMask) >> 1);
diff --git a/Engine/Plugins/Interchange/Runtime/Source/Import/Private/Texture/InterchangePSDTranslator.cpp b/Engine/Plugins/Interchange/Runtime/Source/Import/Private/Texture/InterchangePSDTranslator.cpp
--- a/Engine/Plugins/Interchange/Runtime/Source/Import/Private/Texture/InterchangePSDTranslator.cpp
+++ b/Engine/Plugins/Interchange/Runtime/Source/Import/Private/Texture/InterchangePSDTranslator.cpp
@@ -212,7 +212,7 @@ namespace UE
 				}
 
 				// Skip Clut, Image Resource Section, and Layer/Mask Section
-				if (SkipSection(Current) == false ||
+				if (SkipSection(Current) == false || //-V501
 					SkipSection(Current) == false ||
 					SkipSection(Current) == false)
 				{
diff --git a/Engine/Source/Runtime/MovieSceneTracks/Public/Tracks/MovieSceneMaterialTrack.h b/Engine/Source/Runtime/MovieSceneTracks/Public/Tracks/MovieSceneMaterialTrack.h
--- a/Engine/Source/Runtime/MovieSceneTracks/Public/Tracks/MovieSceneMaterialTrack.h
+++ b/Engine/Source/Runtime/MovieSceneTracks/Public/Tracks/MovieSceneMaterialTrack.h
@@ -49,7 +49,7 @@ struct FComponentMaterialInfo
 	}
 	friend bool operator==(const FComponentMaterialInfo& A, const FComponentMaterialInfo& B)
 	{
-		return A.MaterialSlotName == B.MaterialSlotName && A.MaterialSlotIndex == B.MaterialSlotIndex && A.MaterialSlotName == B.MaterialSlotName && A.MaterialType == B.MaterialType;
+		return A.MaterialSlotIndex == B.MaterialSlotIndex && A.MaterialSlotName == B.MaterialSlotName && A.MaterialType == B.MaterialType;
 	}
 
 	FString ToString() const 
diff --git a/Engine/Source/Runtime/VectorVM/Private/VectorVMExperimental.cpp b/Engine/Source/Runtime/VectorVM/Private/VectorVMExperimental.cpp
--- a/Engine/Source/Runtime/VectorVM/Private/VectorVMExperimental.cpp
+++ b/Engine/Source/Runtime/VectorVM/Private/VectorVMExperimental.cpp
@@ -777,7 +777,7 @@ static void VVMBuildMapTableCaches(FVectorVMExecContext *ExecCtx)
 
 static const uint8 *VVM_Output32_from_16(const bool CT_MultipleLoops, const uint8 *InsPtr, FVectorVMBatchState *BatchState, FVectorVMExecContext *ExecCtx)
 {
-	VVM_OUTPUT_FUNCTION_HEADER(EVectorVMOp::outputdata_float);
+	VVM_OUTPUT_FUNCTION_HEADER(EVectorVMOp::outputdata_float); //-V501
 	VVM_OUTPUT_FUNCTION_FOOTER;
 }
 
