Add function to compile CustomizableObjects synchronously, exposed to editor scripts

diff --git a/Engine/Plugins/Experimental/Mutable/Source/CustomizableObjectEditor/Private/MuCOE/CustomizableObjectEditorFunctionLibrary.cpp b/Engine/Plugins/Experimental/Mutable/Source/CustomizableObjectEditor/Private/MuCOE/CustomizableObjectEditorFunctionLibrary.cpp
new file mode 100644
--- /dev/null
+++ b/Engine/Plugins/Experimental/Mutable/Source/CustomizableObjectEditor/Private/MuCOE/CustomizableObjectEditorFunctionLibrary.cpp
@@ -0,0 +1,64 @@
+// Copyright Epic Games, Inc. All Rights Reserved.
+#pragma once
+
+#include "MuCOE/CustomizableObjectEditorFunctionLibrary.h"
+#include "MuCOE/CustomizableObjectCompiler.h"
+#include "MuCO/CustomizableObjectSystem.h"
+#include "MuCO/CustomizableObject.h"
+
+ECustomizableObjectCompilationState UCustomizableObjectEditorFunctionLibrary::CompileCustomizableObjectSynchronously(UCustomizableObject* CustomizableObject)
+{
+	if (!CustomizableObject)
+	{
+		UE_LOG(LogMutable, Warning, TEXT("Attempted to compile nullptr Customizable Object"));
+		return ECustomizableObjectCompilationState::Failed;
+	}
+
+	FString ObjectPath = CustomizableObject->GetPathName();
+	if (CustomizableObject->IsLocked())
+	{
+		// Take this if you need a hack:
+		// UCustomizableObjectSystem::GetInstance()->UnlockObject(CustomizableObject);
+		UE_LOG( LogMutable, Warning, 
+			TEXT("Attempted to compile %s Customizable Object when it is locked"),
+			*ObjectPath);
+		return ECustomizableObjectCompilationState::Failed;
+	}
+
+	const double StartTime = FPlatformTime::Seconds();
+	const double PrintFrequencySeconds = 15.0;
+	double PrintTime = StartTime + PrintFrequencySeconds;
+
+	FCustomizableObjectCompiler Compiler;
+	FCompilationOptions Options = CustomizableObject->CompileOptions;
+	Options.bSilentCompilation = false;
+	Compiler.Compile(*CustomizableObject, Options, true);
+
+	Compiler.Tick();
+	
+	while (Compiler.GetCompilationState() == ECustomizableObjectCompilationState::InProgress)
+	{
+		Compiler.Tick();
+		const double CurrentTime = FPlatformTime::Seconds();
+		if (CurrentTime > PrintTime)
+		{
+			PrintTime = CurrentTime + PrintFrequencySeconds;
+			UE_LOG( LogMutable, Display,
+				TEXT("Synchronously Compiling %s for %f seconds"),
+				*ObjectPath, CurrentTime - StartTime
+			);
+		}
+	}
+
+	const double CurrentTime = FPlatformTime::Seconds();
+	UE_LOG( LogMutable, Display,
+		TEXT("Synchronously Compiled %s %s in %f seconds"),
+		*ObjectPath, 
+		Compiler.GetCompilationState() == ECustomizableObjectCompilationState::Completed ? 
+		TEXT("successfully") : TEXT("unsuccessfully"),
+		CurrentTime - StartTime
+	);
+
+	ensure(CustomizableObject->IsCompiled());
+	return Compiler.GetCompilationState();
+}
diff --git a/Engine/Plugins/Experimental/Mutable/Source/CustomizableObjectEditor/Public/MuCOE/CustomizableObjectEditorFunctionLibrary.h b/Engine/Plugins/Experimental/Mutable/Source/CustomizableObjectEditor/Public/MuCOE/CustomizableObjectEditorFunctionLibrary.h
new file mode 100644
--- /dev/null
+++ b/Engine/Plugins/Experimental/Mutable/Source/CustomizableObjectEditor/Public/MuCOE/CustomizableObjectEditorFunctionLibrary.h
@@ -0,0 +1,28 @@
+// Copyright Epic Games, Inc. All Rights Reserved.
+#pragma once
+
+#include "Kismet/BlueprintFunctionLibrary.h"
+
+#include "CustomizableObjectEditorFunctionLibrary.generated.h"
+
+/**
+ * Functions we want to be able to call on CustomizableObjects at edit time - could
+ * be exposed to cook as well.
+ */
+UCLASS()
+class UCustomizableObjectEditorFunctionLibrary : public UBlueprintFunctionLibrary
+{
+	GENERATED_BODY()
+
+public:
+	/** 
+	 *	Synchronously compiles the provided CustomizableObject, LogMutable will contain intermittent updates on
+	 *	progress.
+	 * 
+	 * @param CustomizableObject	The CustomizableObject to compile
+	 * 
+	 * @return	The final ECustomizableObjectCompilationState - typically Completed or Failed
+	 */
+	UFUNCTION(BlueprintCallable, Category = "CustomizableObject")
+	static ECustomizableObjectCompilationState CompileCustomizableObjectSynchronously(UCustomizableObject* CustomizableObject);
+};
