[Mutable] Fixed wrong mesh skinning indices being set after skinning optimization.
- This happened when we were using a mesh with an index format different than UINT8 as the mesh provided by a Table mutable node.
- This issue does not seem to be manifested by any of the COs used in the current critical projects.

diff --git a/Engine/Plugins/Experimental/Mutable/Source/MutableRuntime/Private/MuR/OpMeshOptimizeSkinning.h b/Engine/Plugins/Experimental/Mutable/Source/MutableRuntime/Private/MuR/OpMeshOptimizeSkinning.h
--- a/Engine/Plugins/Experimental/Mutable/Source/MutableRuntime/Private/MuR/OpMeshOptimizeSkinning.h
+++ b/Engine/Plugins/Experimental/Mutable/Source/MutableRuntime/Private/MuR/OpMeshOptimizeSkinning.h
@@ -137,7 +137,7 @@ namespace mu
 				VertexBuffers.SetBuffer(BufferIndex, ElementSize, ChannelsCount, Semantics.GetData(), SemanticIndices.GetData(), Formats.GetData(), Components.GetData(), Offsets.GetData());
 				uint8* Data = VertexBuffers.GetBufferData(BufferIndex);
 
-				const uint16* SourceData = (const uint16*)SourceBuffer.m_data.GetData();
+				const uint8* SourceData = (const uint8*)SourceBuffer.m_data.GetData();
 				const int32 SourceBoneIndexSize = GetMeshFormatData(SourceBoneIndexFormat).m_size;
 
 				const int32 BoneIndexSize = GetMeshFormatData(mu::MBF_UINT8).m_size;
@@ -159,12 +159,16 @@ namespace mu
 
 						for (int32 ComponentIndex = 0; ComponentIndex < BoneIndexComponentCount; ++ComponentIndex)
 						{
-							const uint16 SourceIndex = *SourceData;
+							const uint16 SourceIndex = *((const uint16*)SourceData);
 							if (SourceIndex < NumBonesInBoneMap)
 							{
 								*Data = (uint8)SourceIndex;
 							}
-
+							else
+							{
+								*Data = 0;
+							}
+							
 							Data += BoneIndexSize;
 							SourceData += SourceBoneIndexSize;
 						}
