[mutable] Separate "DisableTextureStream" option into its own state flag instead of deducing it from other enums.
Add option to override this flag at compile time.
Force "None" optimization to disable streaming.
#rnx

diff --git a/Engine/Plugins/Experimental/Mutable/Source/CustomizableObject/Classes/MuCO/CustomizableObject.h b/Engine/Plugins/Experimental/Mutable/Source/CustomizableObject/Classes/MuCO/CustomizableObject.h
--- a/Engine/Plugins/Experimental/Mutable/Source/CustomizableObject/Classes/MuCO/CustomizableObject.h
+++ b/Engine/Plugins/Experimental/Mutable/Source/CustomizableObject/Classes/MuCO/CustomizableObject.h
@@ -1477,7 +1477,7 @@ private:
 	// This is a manual version number for the binary blobs in this asset.
 	// Increasing it invalidates all the previously compiled models.
 	// Warning: If while merging code both versions have changed, take the highest+1.
-	static const int32 CurrentSupportedVersion = 408;
+	static const int32 CurrentSupportedVersion = 409;
 
 public:
 
diff --git a/Engine/Plugins/Experimental/Mutable/Source/CustomizableObject/Private/MuCO/CustomizableObject.cpp b/Engine/Plugins/Experimental/Mutable/Source/CustomizableObject/Private/MuCO/CustomizableObject.cpp
--- a/Engine/Plugins/Experimental/Mutable/Source/CustomizableObject/Private/MuCO/CustomizableObject.cpp
+++ b/Engine/Plugins/Experimental/Mutable/Source/CustomizableObject/Private/MuCO/CustomizableObject.cpp
@@ -179,6 +179,17 @@ void UCustomizableObject::PostLoad()
 			: ECustomizableObjectTextureCompression::None
 			;
 	}
+
+	// Update state never-stream flag from deprecated enum
+	if (CustomizableObjectCustomVersion < FCustomizableObjectCustomVersion::CustomizableObjectStateHasSeparateNeverStreamFlag)
+	{
+		for (TPair<FString, FParameterUIData>& s : StateUIDataMap)
+		{
+			s.Value.bDisableTextureStreaming = s.Value.TextureCompressionStrategy != ETextureCompressionStrategy::None;
+		}
+	}
+
+
 #endif
 
 #if WITH_EDITORONLY_DATA
diff --git a/Engine/Plugins/Experimental/Mutable/Source/CustomizableObject/Private/MuCO/CustomizableObjectSystem.cpp b/Engine/Plugins/Experimental/Mutable/Source/CustomizableObject/Private/MuCO/CustomizableObjectSystem.cpp
--- a/Engine/Plugins/Experimental/Mutable/Source/CustomizableObject/Private/MuCO/CustomizableObjectSystem.cpp
+++ b/Engine/Plugins/Experimental/Mutable/Source/CustomizableObject/Private/MuCO/CustomizableObjectSystem.cpp
@@ -42,7 +42,6 @@
 #endif
 
 
-
 #include UE_INLINE_GENERATED_CPP_BY_NAME(CustomizableObjectSystem)
 
 class AActor;
@@ -889,14 +888,26 @@ void UpdateSkeletalMesh(const TSharedRef<FUpdateContextPrivate>& Context)
 
 void FCustomizableObjectSystemPrivate::GetMipStreamingConfig(const UCustomizableObjectInstance& Instance, bool& bOutNeverStream, int32& OutMipsToSkip) const
 {
+	bOutNeverStream = false;
+
+	// From user-controlled per-state flag?
 	const FString CurrentState = Instance.GetCurrentState();
-	const FParameterUIData* State = Instance.GetCustomizableObject()->StateUIDataMap.Find(CurrentState);
+	const FParameterUIData* State = Instance.GetCustomizableObject()->StateUIDataMap.Find(CurrentState);	
+	if (State)
+	{
+		bOutNeverStream = State->bDisableTextureStreaming;
+
+		// Was streaming disabled at object-compilation time? 
+		if (State->bDisableTextureStreamingOverride)
+		{
+			bOutNeverStream = true;
+		}
+	}
 
-	// \TODO: This should be controllable independently
-	bOutNeverStream = State ? State->TextureCompressionStrategy != ETextureCompressionStrategy::None : false;
 	bool bUseMipmapStreaming = !bOutNeverStream;
-	OutMipsToSkip = 0; // 0 means all mips
+	OutMipsToSkip = 0; // 0 means generate all mips
 
+	// Streaming disabled from platform settings?
 #if PLATFORM_SUPPORTS_TEXTURE_STREAMING
 	if (!IStreamingManager::Get().IsTextureStreamingEnabled())
 	{
@@ -906,6 +917,7 @@ void FCustomizableObjectSystemPrivate::GetMipStreamingConfig(const UCustomizable
 	bUseMipmapStreaming = false;
 #endif
 
+	// Streaming disabled from platform CustomizableObjectSystem properties?
 	if (bUseMipmapStreaming && EnableMutableProgressiveMipStreaming)
 	{
 		OutMipsToSkip = 255; // This means skip all possible mips until only UTexture::GetStaticMinTextureResidentMipCount() are left
@@ -1753,7 +1765,7 @@ namespace impl
 
 								if (Props.IsPassThrough)
 								{
-									// Since it's known it's a pass-through texture, the GetImage will be inexpensive enough to be run in the game thread
+									// Since it's known it's a pass-through texture there is no need to cache or convert it so we can generate it here already.
 									Image.Image = System->GetImage(OperationData->InstanceID, Image.ImageID, 0, 0);
 									check(Image.Image->IsReference());
 
diff --git a/Engine/Plugins/Experimental/Mutable/Source/CustomizableObject/Public/MuCO/CustomizableObjectCustomVersion.h b/Engine/Plugins/Experimental/Mutable/Source/CustomizableObject/Public/MuCO/CustomizableObjectCustomVersion.h
--- a/Engine/Plugins/Experimental/Mutable/Source/CustomizableObject/Public/MuCO/CustomizableObjectCustomVersion.h
+++ b/Engine/Plugins/Experimental/Mutable/Source/CustomizableObject/Public/MuCO/CustomizableObjectCustomVersion.h
@@ -101,9 +101,11 @@ struct CUSTOMIZABLEOBJECT_API FCustomizableObjectCustomVersion
 		NodeVariationSerializationIssue, // The parent class and child classes had a member with the same name. Unreal was serialized the member from the child class but deserialized it to the parent class.
 
 		RegenerateNodeObjectsIds,
-		
+
 		NodeMaterialPinDataImageDetails,
 
+		CustomizableObjectStateHasSeparateNeverStreamFlag,
+
 		// -----<new versions can be added above this line>-------------------------------------------------
 		VersionPlusOne,
 		LatestVersion = VersionPlusOne - 1
diff --git a/Engine/Plugins/Experimental/Mutable/Source/CustomizableObject/Public/MuCO/CustomizableObjectUIData.h b/Engine/Plugins/Experimental/Mutable/Source/CustomizableObject/Public/MuCO/CustomizableObjectUIData.h
--- a/Engine/Plugins/Experimental/Mutable/Source/CustomizableObject/Public/MuCO/CustomizableObjectUIData.h
+++ b/Engine/Plugins/Experimental/Mutable/Source/CustomizableObject/Public/MuCO/CustomizableObjectUIData.h
@@ -221,10 +221,7 @@ struct FParameterUIData
 	GENERATED_BODY()
 
 	FParameterUIData() 
-		: Type(EMutableParameterType::None)
-		, IntegerParameterGroupType(ECustomizableObjectGroupType::COGT_ONE_OR_NONE)
 	{
-
 	}
 
 	FParameterUIData(
@@ -233,8 +230,7 @@ struct FParameterUIData
 		EMutableParameterType ParamType) :
 		Name(ParamName),
 		ParamUIMetadata(InParamUIMetadata),
-		Type(ParamType), 
-		IntegerParameterGroupType(ECustomizableObjectGroupType::COGT_ONE_OR_NONE)
+		Type(ParamType)
 	{
 
 	}
@@ -248,7 +244,7 @@ struct FParameterUIData
 
 	/** Parameter type, using uint8 since the enum in declared in the class it is used */
 	UPROPERTY(BlueprintReadWrite, Category = UI)
-	EMutableParameterType Type;
+	EMutableParameterType Type = EMutableParameterType::None;
 
 	/** In the case of an integer parameter, store here all options */
 	UPROPERTY(BlueprintReadWrite, Category = UI)
@@ -256,13 +252,19 @@ struct FParameterUIData
 
 	/** In the case of an integer parameter, how are the different options selected (one, one or none, etc...) */
 	UPROPERTY(BlueprintReadWrite, Category = UI)
-	ECustomizableObjectGroupType IntegerParameterGroupType;
+	ECustomizableObjectGroupType IntegerParameterGroupType = ECustomizableObjectGroupType::COGT_ONE_OR_NONE;
 
-	/** Not really relevant for UI, but apparently bDontCompressRuntimeTextures_DEPRECATED was used to decided some texture properties at runtime.
-	* TODO: Try to handle the NeverStream without going through this. */
+	/** Not really relevant for UI, but apparently bDontCompressRuntimeTextures_DEPRECATED was used to decided some texture properties at runtime. */
 	UPROPERTY()
 	ETextureCompressionStrategy TextureCompressionStrategy = ETextureCompressionStrategy::None;
 
+	/** If this is enabled, texture streaming won't be used for this state, and full images will be generated when an instance is first updated. */
+	UPROPERTY()
+	bool bDisableTextureStreaming = false;
+
+	/** This is a non-user-controlled flag to disable streaming at object compilation time. */
+	bool bDisableTextureStreamingOverride = false;
+
 	/** In this mode instances and their temp data will be reused between updates. It will be much faster but spend as much as ten times the memory.
 	    Useful for customization lockers with few characters that are going to have their parameters changed many times, not for in-game */
 	UPROPERTY(BlueprintReadWrite, Category = CustomizableObject)
@@ -279,6 +281,8 @@ struct FParameterUIData
 		if (Name != Other.Name || ParamUIMetadata != Other.ParamUIMetadata || Type != Other.Type 
 			|| ArrayIntegerParameterOption != Other.ArrayIntegerParameterOption || IntegerParameterGroupType != Other.IntegerParameterGroupType 
 			|| bLiveUpdateMode != Other.bLiveUpdateMode || bReuseInstanceTextures != Other.bReuseInstanceTextures
+			|| bDisableTextureStreaming != Other.bDisableTextureStreaming
+			|| bDisableTextureStreamingOverride != Other.bDisableTextureStreamingOverride
 			|| !ForcedParameterValues.OrderIndependentCompareEqual(Other.ForcedParameterValues)
 			)
 		{
@@ -296,6 +300,8 @@ struct FParameterUIData
 		Ar << UIData.ArrayIntegerParameterOption;
 		Ar << UIData.IntegerParameterGroupType;
 		Ar << UIData.TextureCompressionStrategy;
+		Ar << UIData.bDisableTextureStreaming;
+		Ar << UIData.bDisableTextureStreamingOverride;
 		Ar << UIData.bLiveUpdateMode;
 		Ar << UIData.bReuseInstanceTextures;
 		Ar << UIData.ForcedParameterValues;
diff --git a/Engine/Plugins/Experimental/Mutable/Source/CustomizableObjectEditor/Private/MuCOE/CustomizableObjectCompiler.cpp b/Engine/Plugins/Experimental/Mutable/Source/CustomizableObjectEditor/Private/MuCOE/CustomizableObjectCompiler.cpp
--- a/Engine/Plugins/Experimental/Mutable/Source/CustomizableObjectEditor/Private/MuCOE/CustomizableObjectCompiler.cpp
+++ b/Engine/Plugins/Experimental/Mutable/Source/CustomizableObjectEditor/Private/MuCOE/CustomizableObjectCompiler.cpp
@@ -1171,7 +1171,25 @@ void FCustomizableObjectCompiler::CompileInternal(UCustomizableObject* Object, c
 
 		Object->GroupNodeMap = GenerationContext.GroupNodeMap;
 		Object->ParameterUIDataMap = GenerationContext.ParameterUIDataMap;
+
 		Object->StateUIDataMap = GenerationContext.StateUIDataMap;
+		if (GenerationContext.Options.OptimizationLevel == 0 )
+		{
+			// If the optimization level is "none" disable texture streaming, because textures are all referenced
+			// unreal assets and progressive generation is not supported.
+			for (TPair<FString, FParameterUIData>& StateData : Object->StateUIDataMap)
+			{
+				StateData.Value.bDisableTextureStreamingOverride = true;
+			}
+
+			// If there is no state data, we need to add the default one to be able to disable streaming.
+			if (Object->StateUIDataMap.IsEmpty())
+			{
+				FParameterUIData& Data = Object->StateUIDataMap.Add(TEXT("Default"));
+				Data.bDisableTextureStreamingOverride = true;
+			}
+		}
+
 		Object->AlwaysLoadedExtensionData = MoveTemp(GenerationContext.AlwaysLoadedExtensionData);
 
 		Object->StreamedExtensionData.Empty(GenerationContext.StreamedExtensionData.Num());
