[Mutable] Added new method to generate data tables on compilation time using a ScriptStruct and the asset registry.
+Fixed crash when there was an empty texture in a data table.
+Fixed method to get the Animation Instance of a data table row.

#rnx

diff --git a/Engine/Plugins/Experimental/Mutable/Source/CustomizableObjectEditor/Classes/MuCOE/Nodes/CustomizableObjectNodeTable.h b/Engine/Plugins/Experimental/Mutable/Source/CustomizableObjectEditor/Classes/MuCOE/Nodes/CustomizableObjectNodeTable.h
--- a/Engine/Plugins/Experimental/Mutable/Source/CustomizableObjectEditor/Classes/MuCOE/Nodes/CustomizableObjectNodeTable.h
+++ b/Engine/Plugins/Experimental/Mutable/Source/CustomizableObjectEditor/Classes/MuCOE/Nodes/CustomizableObjectNodeTable.h
@@ -40,6 +40,18 @@ enum class ETableMeshPinType : uint8
 };
 
 
+/** Enum to decide where the data comes from: Struct or Data Table*/
+UENUM()
+enum class ETableDataGatheringSource : uint8
+{
+	/** Gathers the information from a data table */
+	ETDGM_DataTable = 0 UMETA(DisplayName = "Data Table"),
+
+	/** When compiling the CO, it uses the asset registry to gather and generate a data table. It uses all the data tables found in the specified paths that are references of the selected structure.*/
+	ETDGM_AssetRegistry = 1 UMETA(DisplayName = " Struct + Asset Registry")
+};
+
+
 /** Base class for all Table Pins. */
 UCLASS()
 class CUSTOMIZABLEOBJECTEDITOR_API UCustomizableObjectNodeTableObjectPinData : public UCustomizableObjectNodePinData
@@ -150,14 +162,25 @@ public:
 	UPROPERTY(EditAnywhere, Category = TableProperties)
 	bool bAddNoneOption;
 
-	// Pointer to the Data Table Asset represented in this node
-	UPROPERTY(EditAnywhere, Category = TableProperties, meta = (DontUpdateWhileEditing))
-	TObjectPtr<UDataTable> Table = nullptr;
-	
 	/** If there is a bool column in the table, checked rows will not be compiled */
 	UPROPERTY(EditAnywhere, Category = TableProperties)
 	bool bDisableCheckedRows = true;
 
+	/** Source where table gathers the data */
+	UPROPERTY(EditAnywhere, Category = TableProperties)
+	ETableDataGatheringSource TableDataGatheringMode = ETableDataGatheringSource::ETDGM_DataTable;
+
+	// Pointer to the Data Table Asset represented in this node
+	UPROPERTY(EditAnywhere, Category = TableProperties, meta = (DontUpdateWhileEditing, EditCondition = "TableDataGatheringMode == ETableDataGatheringSource::ETDGM_DataTable", EditConditionHides))
+	TObjectPtr<UDataTable> Table = nullptr;
+
+	// Pointer to the Struct Asset represented in this node
+	UPROPERTY(EditAnywhere, Category = TableProperties, meta = (DontUpdateWhileEditing, EditCondition = "TableDataGatheringMode == ETableDataGatheringSource::ETDGM_AssetRegistry", EditConditionHides))
+	TObjectPtr<UScriptStruct> Structure = nullptr;
+
+	UPROPERTY(EditAnywhere, Category = TableProperties, meta = (EditCondition = "TableDataGatheringMode == ETableDataGatheringSource::ETDGM_AssetRegistry", EditConditionHides))
+	TArray<FName> FilterPaths;
+
 	/** Decides the default type of the texture pins (passtrhough or mutable)
 	*   Right click on a non-linked image pin to customize its image mode
 	*/
@@ -222,51 +245,56 @@ public:
 	T* GetColumnDefaultAssetByType(FString ColumnName) const
 	{
 		T* ObjectToReturn = nullptr;
+		const UScriptStruct* TableStruct = nullptr;
 
-		if (Table)
+		if (TableDataGatheringMode == ETableDataGatheringSource::ETDGM_AssetRegistry)
 		{
-			// Getting Struct Pointer
-			const UScriptStruct* TableStruct = Table->GetRowStruct();
-
-			if (TableStruct)
+			TableStruct = Structure;
+		}
+		else
+		{
+			if (Table)
 			{
-				// Getting Default Struct Values
-				uint8* DefaultRowData = (uint8*)FMemory::Malloc(TableStruct->GetStructureSize());
-				TableStruct->InitializeStruct(DefaultRowData);
+				TableStruct = Table->GetRowStruct();
+			}
+		}
 
-				if (DefaultRowData)
+		if (TableStruct)
+		{
+			// Getting Default Struct Values
+			// A Script Struct always has at leaset one property
+			TArray<int8> DeafaultDataArray;
+			DeafaultDataArray.SetNumZeroed(TableStruct->GetStructureSize());
+			TableStruct->InitializeStruct(DeafaultDataArray.GetData());
+
+			FProperty* Property = FindTableProperty(TableStruct, FName(*ColumnName));
+
+			if (Property)
+			{
+				if (const FSoftObjectProperty* SoftObjectProperty = CastField<FSoftObjectProperty>(Property))
 				{
-					FProperty* Property = Table->FindTableProperty(FName(*ColumnName));
+					UObject* Object = nullptr;
+					
+					// Getting default UObject
+					uint8* CellData = SoftObjectProperty->ContainerPtrToValuePtr<uint8>(DeafaultDataArray.GetData());
 
-					if (Property)
+					if (CellData)
 					{
-						if (const FSoftObjectProperty* SoftObjectProperty = CastField<FSoftObjectProperty>(Property))
+						Object = SoftObjectProperty->GetPropertyValue(CellData).LoadSynchronous();
+					}
+
+					if (Object)
+					{
+						if (Object->IsA(T::StaticClass()))
 						{
-							UObject* Object = nullptr;
-							
-							// Getting default UObject
-							uint8* CellData = SoftObjectProperty->ContainerPtrToValuePtr<uint8>(DefaultRowData);
-
-							if (CellData)
-							{
-								Object = SoftObjectProperty->GetPropertyValue(CellData).LoadSynchronous();
-							}
-
-							if (Object)
-							{
-								if (Object->IsA(T::StaticClass()))
-								{
-									ObjectToReturn = Cast <T>(Object);
-								}
-							}
+							ObjectToReturn = Cast <T>(Object);
 						}
 					}
-
-					// Cleaning Default Structure Pointer
-					TableStruct->DestroyStruct(DefaultRowData);
-					FMemory::Free(DefaultRowData);
 				}
 			}
+
+			// Cleaning Default Structure Pointer
+			TableStruct->DestroyStruct(DeafaultDataArray.GetData());
 		}
 
 		return ObjectToReturn;
@@ -288,15 +316,15 @@ public:
 		return nullptr;
 	}
 
-
+	// Generation Mutable Source Methods
 	// We should do this in a template!
-	USkeletalMesh* GetSkeletalMeshAt(const UEdGraphPin* Pin, const FName& RowName) const;
-	TSoftClassPtr<UAnimInstance> GetAnimInstanceAt(const UEdGraphPin* Pin, const FName& RowName) const;
+	USkeletalMesh* GetSkeletalMeshAt(const UEdGraphPin* Pin, const UDataTable* DataTable, const FName& RowName) const;
+	TSoftClassPtr<UAnimInstance> GetAnimInstanceAt(const UEdGraphPin* Pin, const UDataTable* DataTable, const FName& RowName) const;
 
 
 	// Return the name of the enabled rows in the data table.
 	// Returns the name if the row has a bool column set as false (true == disabled)
-	TArray<FName> GetRowNames() const;
+	TArray<FName> GetRowNames(const UDataTable* DataTable) const;
 
 	// Changes the image mode of a pin
 	// bSetDefault param: if true sets the pin to be equal to the default mode (same as node)
@@ -318,6 +346,16 @@ public:
 	FString GenerateSkeletalMeshMutableColumName(const FString& PinName, int32 LODIndex, int32 MaterialIndex) const;
 	FString GenerateStaticMeshMutableColumName(const FString& PinName, int32 MaterialIndex) const;
 
+	/** Returns the struct pointer used to gather data */
+	const UScriptStruct* GetTableNodeStruct() const;
+
+	// Methods from UDataTable but modified to support ScriptStructs
+	/** Get an array of all the column titles, using the friendly display name from the property */
+	TArray<FString> GetColumnTitles(const UScriptStruct* ScriptStruct) const;
+	
+	/** Returns the column property where PropertyName matches the name of the column property. Returns nullptr if no match is found or the match is not a supported table property */
+	FProperty* FindTableProperty(const UScriptStruct* ScriptStruct, const FName& PropertyName) const;
+
 private:
 
 	/** Number of properties to know when the node needs an update */
diff --git a/Engine/Plugins/Experimental/Mutable/Source/CustomizableObjectEditor/Private/MuCOE/GenerateMutableSource/GenerateMutableSource.h b/Engine/Plugins/Experimental/Mutable/Source/CustomizableObjectEditor/Private/MuCOE/GenerateMutableSource/GenerateMutableSource.h
--- a/Engine/Plugins/Experimental/Mutable/Source/CustomizableObjectEditor/Private/MuCOE/GenerateMutableSource/GenerateMutableSource.h
+++ b/Engine/Plugins/Experimental/Mutable/Source/CustomizableObjectEditor/Private/MuCOE/GenerateMutableSource/GenerateMutableSource.h
@@ -4,6 +4,7 @@
 
 #include "Animation/Skeleton.h"
 #include "Animation/AnimInstance.h"
+#include "Engine/DataTable.h"
 #include "Engine/SkeletalMesh.h"
 #include "MuCO/CustomizableObject.h"
 #include "MuCOE/ExtensionDataCompilerInterface.h"
@@ -29,6 +30,7 @@ struct FCustomizableObjectClothingAssetData;
 
 class FCustomizableObjectCompiler;
 class UAnimInstance;
+class UCompositeDataTable;
 class UCustomizableObjectNodeMaterial;
 class UCustomizableObjectNodeMeshMorph;
 class UCustomizableObjectNodeObjectGroup;
@@ -538,6 +540,21 @@ struct FMutableGraphGenerationContext
 	// Cache of generated Node Tables
 	TMap<FString, mu::TablePtr> GeneratedTables;
 
+	struct FGeneratedDataTablesData
+	{
+		UScriptStruct* ParentStruct = nullptr;
+		TArray<FName> FilterPaths;
+		UCompositeDataTable* GeneratedDataTable = nullptr;
+
+		bool operator==(const FGeneratedDataTablesData& Other) const
+		{
+			return ParentStruct == Other.ParentStruct && FilterPaths == Other.FilterPaths;
+		}
+	};
+
+	// Cache of generated Composited Data Tables
+	TArray<FGeneratedDataTablesData> GeneratedCompositeDataTables;
+
 	// Cache of generated images, because sometimes they are reused by LOD, we use this as a second
 	// level cache
 	TMap<FGeneratedImageKey, mu::NodeImagePtr> GeneratedImages;
@@ -795,6 +812,9 @@ struct FMutableGraphGenerationContext
 
 	/** See UCustomizableObject::ParticipatingObjects. */
 	TMap<TObjectPtr<const UObject>, FGuid> ParticipatingObjects;
+
+	/** Map to relate a Composite Data Table Row and its original DataTable */
+	TMap<UDataTable*,TMap<FName, TArray<UDataTable*>>> CompositeDataTableRowToOriginalDataTableMap;
 };
 
 /** Pin Data scope wrapper. Pops the pin data on scope exit. */
diff --git a/Engine/Plugins/Experimental/Mutable/Source/CustomizableObjectEditor/Private/MuCOE/GenerateMutableSource/GenerateMutableSourceColor.cpp b/Engine/Plugins/Experimental/Mutable/Source/CustomizableObjectEditor/Private/MuCOE/GenerateMutableSource/GenerateMutableSourceColor.cpp
--- a/Engine/Plugins/Experimental/Mutable/Source/CustomizableObjectEditor/Private/MuCOE/GenerateMutableSource/GenerateMutableSourceColor.cpp
+++ b/Engine/Plugins/Experimental/Mutable/Source/CustomizableObjectEditor/Private/MuCOE/GenerateMutableSource/GenerateMutableSourceColor.cpp
@@ -277,11 +277,12 @@ mu::NodeColourPtr GenerateMutableSourceColor(const UEdGraphPin* Pin, FMutableGra
 		Result = WhiteColorNode;
 
 		bool bSuccess = true;
+		UDataTable* DataTable = GetDataTable(TypedNodeTable, GenerationContext);
 
-		if (TypedNodeTable->Table)
+		if (DataTable)
 		{
 			FString ColumnName = Pin->PinFriendlyName.ToString();
-			FProperty* Property = TypedNodeTable->Table->FindTableProperty(FName(*ColumnName));
+			FProperty* Property = DataTable->FindTableProperty(FName(*ColumnName));
 
 			if (!Property)
 			{
@@ -295,7 +296,7 @@ mu::NodeColourPtr GenerateMutableSourceColor(const UEdGraphPin* Pin, FMutableGra
 			{
 				// Generating a new data table if not exists
 				mu::TablePtr Table;
-				Table = GenerateMutableSourceTable(TypedNodeTable->Table->GetName(), Pin, GenerationContext);
+				Table = GenerateMutableSourceTable(DataTable->GetName(), Pin, GenerationContext);
 
 				if (Table)
 				{
diff --git a/Engine/Plugins/Experimental/Mutable/Source/CustomizableObjectEditor/Private/MuCOE/GenerateMutableSource/GenerateMutableSourceFloat.cpp b/Engine/Plugins/Experimental/Mutable/Source/CustomizableObjectEditor/Private/MuCOE/GenerateMutableSource/GenerateMutableSourceFloat.cpp
--- a/Engine/Plugins/Experimental/Mutable/Source/CustomizableObjectEditor/Private/MuCOE/GenerateMutableSource/GenerateMutableSourceFloat.cpp
+++ b/Engine/Plugins/Experimental/Mutable/Source/CustomizableObjectEditor/Private/MuCOE/GenerateMutableSource/GenerateMutableSourceFloat.cpp
@@ -311,11 +311,12 @@ mu::NodeScalarPtr GenerateMutableSourceFloat(const UEdGraphPin* Pin, FMutableGra
 		}
 
 		bool bSuccess = true;
+		UDataTable* DataTable = GetDataTable(TypedNodeTable, GenerationContext);
 
-		if (TypedNodeTable->Table)
+		if (DataTable)
 		{
 			FString ColumnName = Pin->PinFriendlyName.ToString();
-			FProperty* Property = TypedNodeTable->Table->FindTableProperty(FName(*ColumnName));
+			FProperty* Property = DataTable->FindTableProperty(FName(*ColumnName));
 
 			if (!Property)
 			{
@@ -329,7 +330,7 @@ mu::NodeScalarPtr GenerateMutableSourceFloat(const UEdGraphPin* Pin, FMutableGra
 			{
 				// Generating a new data table if not exists
 				mu::TablePtr Table;
-				Table = GenerateMutableSourceTable(TypedNodeTable->Table->GetName(), Pin, GenerationContext);
+				Table = GenerateMutableSourceTable(DataTable->GetName(), Pin, GenerationContext);
 
 				if (Table)
 				{
diff --git a/Engine/Plugins/Experimental/Mutable/Source/CustomizableObjectEditor/Private/MuCOE/GenerateMutableSource/GenerateMutableSourceImage.cpp b/Engine/Plugins/Experimental/Mutable/Source/CustomizableObjectEditor/Private/MuCOE/GenerateMutableSource/GenerateMutableSourceImage.cpp
--- a/Engine/Plugins/Experimental/Mutable/Source/CustomizableObjectEditor/Private/MuCOE/GenerateMutableSource/GenerateMutableSourceImage.cpp
+++ b/Engine/Plugins/Experimental/Mutable/Source/CustomizableObjectEditor/Private/MuCOE/GenerateMutableSource/GenerateMutableSourceImage.cpp
@@ -1027,10 +1027,12 @@ mu::NodeImagePtr GenerateMutableSourceImage(const UEdGraphPin* Pin, FMutableGrap
 				bDoNotAddToGeneratedCache = true;
 			}
 
-			if (TypedNodeTable->Table)
+			UDataTable* DataTable = GetDataTable(TypedNodeTable, GenerationContext);
+
+			if (DataTable)
 			{
 				FString ColumnName = Pin->PinFriendlyName.ToString();
-				FProperty* Property = TypedNodeTable->Table->FindTableProperty(FName(*ColumnName));
+				FProperty* Property = DataTable->FindTableProperty(FName(*ColumnName));
 
 				if (!Property)
 				{
@@ -1052,7 +1054,7 @@ mu::NodeImagePtr GenerateMutableSourceImage(const UEdGraphPin* Pin, FMutableGrap
 				{
 					// Generating a new data table if not exists
 					mu::TablePtr Table = nullptr;
-					Table = GenerateMutableSourceTable(TypedNodeTable->Table->GetName(), Pin, GenerationContext);
+					Table = GenerateMutableSourceTable(DataTable->GetName(), Pin, GenerationContext);
 
 					if (Table)
 					{
diff --git a/Engine/Plugins/Experimental/Mutable/Source/CustomizableObjectEditor/Private/MuCOE/GenerateMutableSource/GenerateMutableSourceMesh.cpp b/Engine/Plugins/Experimental/Mutable/Source/CustomizableObjectEditor/Private/MuCOE/GenerateMutableSource/GenerateMutableSourceMesh.cpp
--- a/Engine/Plugins/Experimental/Mutable/Source/CustomizableObjectEditor/Private/MuCOE/GenerateMutableSource/GenerateMutableSourceMesh.cpp
+++ b/Engine/Plugins/Experimental/Mutable/Source/CustomizableObjectEditor/Private/MuCOE/GenerateMutableSource/GenerateMutableSourceMesh.cpp
@@ -2040,8 +2040,13 @@ mu::MeshPtr BuildMorphedMutableMesh(const UEdGraphPin* BaseSourcePin, const FStr
 
 	else if (const UCustomizableObjectNodeTable* TypedNodeTable = Cast<UCustomizableObjectNodeTable>(Node))
 	{
-		TypedNodeTable->GetPinLODAndSection(BaseSourcePin, LODIndexConnected, SectionIndexConnected);
-		SkeletalMesh = TypedNodeTable->GetSkeletalMeshAt(BaseSourcePin, RowName);
+		UDataTable* DataTable = GetDataTable(TypedNodeTable, GenerationContext);
+
+		if (DataTable)
+		{
+			TypedNodeTable->GetPinLODAndSection(BaseSourcePin, LODIndexConnected, SectionIndexConnected);
+			SkeletalMesh = TypedNodeTable->GetSkeletalMeshAt(BaseSourcePin, DataTable, RowName);
+		}
 	}
 
 	if (SkeletalMesh)
@@ -2161,7 +2166,7 @@ void GenerateMorphFactor(const UCustomizableObjectNode* Node, const UEdGraphPin&
 }
 
 TArray<TTuple<USkeletalMesh*, TSoftClassPtr<UAnimInstance>>> GetSkeletalMeshesInfoForReshapeSelection(
-		const UEdGraphNode* SkeletalMeshOrTableNode, const UEdGraphPin* SourceMeshPin)
+		const UEdGraphNode* SkeletalMeshOrTableNode, const UEdGraphPin* SourceMeshPin, FMutableGraphGenerationContext& GenerationContext)
 {
 	TArray<TTuple<USkeletalMesh*, TSoftClassPtr<UAnimInstance>>> SkeletalMeshesInfo;
 
@@ -2179,12 +2184,14 @@ TArray<TTuple<USkeletalMesh*, TSoftClassPtr<UAnimInstance>>> GetSkeletalMeshesIn
 	}
 	else if (const UCustomizableObjectNodeTable* TableNode = Cast<UCustomizableObjectNodeTable>(SkeletalMeshOrTableNode))
 	{
-		if (TableNode->Table)
+		UDataTable* DataTable = GetDataTable(TableNode, GenerationContext);
+
+		if (DataTable)
 		{
-			for (const FName& RowName : TableNode->GetRowNames())
+			for (const FName& RowName : TableNode->GetRowNames(DataTable))
 			{
-				USkeletalMesh* SkeletalMesh = TableNode->GetSkeletalMeshAt(SourceMeshPin, RowName);
-				TSoftClassPtr<UAnimInstance> MeshAnimInstance = TableNode->GetAnimInstanceAt(SourceMeshPin, RowName);
+				USkeletalMesh* SkeletalMesh = TableNode->GetSkeletalMeshAt(SourceMeshPin, DataTable, RowName);
+				TSoftClassPtr<UAnimInstance> MeshAnimInstance = TableNode->GetAnimInstanceAt(SourceMeshPin, DataTable, RowName);
 
 				if (SkeletalMesh)
 				{
@@ -2641,18 +2648,21 @@ mu::NodeMeshPtr GenerateMorphMesh(const UEdGraphPin* Pin,
 	
 	if (const UCustomizableObjectNodeTable* TypedNodeTable = Cast<UCustomizableObjectNodeTable>(Pin->GetOwningNode()))
 	{
+		UDataTable* DataTable = GetDataTable(TypedNodeTable, GenerationContext);
+
 		// Generate a new Column for each morph
-		int32 NumRows = TypedNodeTable->GetRowNames().Num();
+		const TArray<FName>& RowNames = TypedNodeTable->GetRowNames(DataTable);
+		int32 NumRows = RowNames.Num();
 
 		// Should exist
-		mu::TablePtr Table = GenerationContext.GeneratedTables[TypedNodeTable->Table->GetName()];
+		mu::TablePtr Table = GenerationContext.GeneratedTables[DataTable->GetName()];
 
 		FString ColumnName = TableColumnName + TypedNodeMorphs[MorphIndex].MorphTargetName;
 		int32 ColumnIndex = INDEX_NONE;
 
 		for (int32 RowIndex = 0; RowIndex < NumRows; ++RowIndex)
 		{
-			const FName RowName = TypedNodeTable->GetRowNames()[RowIndex];
+			const FName RowName = RowNames[RowIndex];
 
 			ColumnIndex = Table->FindColumn(ColumnName);
 
@@ -2714,7 +2724,7 @@ mu::NodeMeshPtr GenerateMorphMesh(const UEdGraphPin* Pin,
 					const UEdGraphNode* SkeletalMeshNode = SourceMeshPin ? SourceMeshPin->GetOwningNode() : nullptr;
 
 					TArray<TTuple<USkeletalMesh*, TSoftClassPtr<UAnimInstance>>> SkeletalMeshesToDeform = 
-							GetSkeletalMeshesInfoForReshapeSelection(SkeletalMeshNode, SourceMeshPin);
+							GetSkeletalMeshesInfoForReshapeSelection(SkeletalMeshNode, SourceMeshPin, GenerationContext);
 
 					bool bWarningFound = false;
 					if (TypedMorphNode->bReshapeSkeleton)
@@ -3457,7 +3467,7 @@ mu::NodeMeshPtr GenerateMutableSourceMesh(const UEdGraphPin* Pin,
 			const UEdGraphNode* SkeletalMeshNode = SourceMeshPin ? SourceMeshPin->GetOwningNode() : nullptr;
 
 			TArray<TTuple<USkeletalMesh*, TSoftClassPtr<UAnimInstance>>> SkeletalMeshesToDeform = 
-					GetSkeletalMeshesInfoForReshapeSelection(SkeletalMeshNode, SourceMeshPin);
+					GetSkeletalMeshesInfoForReshapeSelection(SkeletalMeshNode, SourceMeshPin, GenerationContext);
 
 			bool bWarningFound = false;
 			if (TypedNodeReshape->bReshapeSkeleton)
@@ -3596,11 +3606,13 @@ mu::NodeMeshPtr GenerateMutableSourceMesh(const UEdGraphPin* Pin,
 		Result = EmptyNode;
 		bool bSuccess = true;
 
-		if (TypedNodeTable->Table)
+		UDataTable* DataTable = GetDataTable(TypedNodeTable, GenerationContext);
+
+		if (DataTable)
 		{
 			// Getting the real name of the data table column
 			FString DataTableColumnName = TypedNodeTable->GetColumnNameByPin(Pin);
-			FProperty* Property = TypedNodeTable->Table->FindTableProperty(FName(*DataTableColumnName));
+			FProperty* Property = DataTable->FindTableProperty(FName(*DataTableColumnName));
 
 			if (!Property)
 			{
@@ -3625,7 +3637,7 @@ mu::NodeMeshPtr GenerateMutableSourceMesh(const UEdGraphPin* Pin,
 			{
 				// Generating a new data table if not exists
 				mu::TablePtr Table = nullptr;
-				Table = GenerateMutableSourceTable(TypedNodeTable->Table->GetName(), Pin, GenerationContext);
+				Table = GenerateMutableSourceTable(DataTable->GetName(), Pin, GenerationContext);
 
 				if (Table)
 				{
diff --git a/Engine/Plugins/Experimental/Mutable/Source/CustomizableObjectEditor/Private/MuCOE/GenerateMutableSource/GenerateMutableSourceTable.cpp b/Engine/Plugins/Experimental/Mutable/Source/CustomizableObjectEditor/Private/MuCOE/GenerateMutableSource/GenerateMutableSourceTable.cpp
--- a/Engine/Plugins/Experimental/Mutable/Source/CustomizableObjectEditor/Private/MuCOE/GenerateMutableSource/GenerateMutableSourceTable.cpp
+++ b/Engine/Plugins/Experimental/Mutable/Source/CustomizableObjectEditor/Private/MuCOE/GenerateMutableSource/GenerateMutableSourceTable.cpp
@@ -3,6 +3,9 @@
 #include "MuCOE/GenerateMutableSource/GenerateMutableSourceTable.h"
 
 #include "Animation/AnimInstance.h"
+#include "AssetRegistry/ARFilter.h"
+#include "AssetRegistry/AssetRegistryModule.h"
+#include "Engine/CompositeDataTable.h"
 #include "Engine/StaticMesh.h"
 #include "Engine/Texture2DArray.h"
 #include "GameplayTagContainer.h"
@@ -25,6 +28,7 @@ bool FillTableColumn(const UCustomizableObjectNodeTable* TableNode,	mu::TablePtr
 	const int LODIndexConnected, const int32 SectionIndexConnected, int32 LODIndex, int32 SectionIndex, const bool bOnlyConnectedLOD, FMutableGraphGenerationContext& GenerationContext)
 {
 	int32 CurrentColumn;
+	UDataTable* DataTablePtr = GetDataTable(TableNode, GenerationContext);
 
 	// Getting property type
 	if (const FSoftObjectProperty* SoftObjectProperty = CastField<FSoftObjectProperty>(ColumnProperty))
@@ -63,16 +67,16 @@ bool FillTableColumn(const UCustomizableObjectNodeTable* TableNode,	mu::TablePtr
 			{
 				if (!AnimSlot.IsEmpty())
 				{
-					if (TableNode->Table)
+					if (DataTablePtr)
 					{
-						uint8* AnimRowData = TableNode->Table->FindRowUnchecked(FName(*RowName));
+						uint8* AnimRowData = DataTablePtr->FindRowUnchecked(FName(*RowName));
 
 						if (AnimRowData)
 						{
 							FName SlotIndex;
 
 							// Getting animation slot row value from data table
-							if (FProperty* AnimSlotProperty = TableNode->Table->FindTableProperty(FName(*AnimSlot)))
+							if (FProperty* AnimSlotProperty = DataTablePtr->FindTableProperty(FName(*AnimSlot)))
 							{
 								uint8* AnimSlotData = AnimSlotProperty->ContainerPtrToValuePtr<uint8>(AnimRowData, 0);
 
@@ -97,7 +101,7 @@ bool FillTableColumn(const UCustomizableObjectNodeTable* TableNode,	mu::TablePtr
 							if (SlotIndex.GetStringLength() != 0)
 							{
 								// Getting animation instance soft class from data table
-								if (FProperty* AnimBPProperty = TableNode->Table->FindTableProperty(FName(*AnimBP)))
+								if (FProperty* AnimBPProperty = DataTablePtr->FindTableProperty(FName(*AnimBP)))
 								{
 									uint8* AnimBPData = AnimBPProperty->ContainerPtrToValuePtr<uint8>(AnimRowData, 0);
 
@@ -125,7 +129,7 @@ bool FillTableColumn(const UCustomizableObjectNodeTable* TableNode,	mu::TablePtr
 							else
 							{
 								FString msg = FString::Printf(TEXT("Could not find the Slot column of the animation blueprint column [%s] for the mesh column [%s] row [%s]."), *AnimBP, *ColumnName, *RowName);
-								GenerationContext.Compiler->CompilerLog(FText::FromString(msg), TableNode);
+								LogRowGenerationMessage(TableNode, DataTablePtr, GenerationContext, msg, RowName);
 							}
 						}
 					}
@@ -140,14 +144,14 @@ bool FillTableColumn(const UCustomizableObjectNodeTable* TableNode,	mu::TablePtr
 			// Getting Gameplay tags
 			if (!GameplayTag.IsEmpty())
 			{
-				if (TableNode->Table)
+				if (DataTablePtr)
 				{
-					uint8* GameplayRowData = TableNode->Table->FindRowUnchecked(FName(*RowName));
+					uint8* GameplayRowData = DataTablePtr->FindRowUnchecked(FName(*RowName));
 
 					if (GameplayRowData)
 					{
 						// Getting animation tag row value from data table
-						if (FProperty* GameplayTagProperty = TableNode->Table->FindTableProperty(FName(*GameplayTag)))
+						if (FProperty* GameplayTagProperty = DataTablePtr->FindTableProperty(FName(*GameplayTag)))
 						{
 							uint8* GameplayTagData = GameplayTagProperty->ContainerPtrToValuePtr<uint8>(GameplayRowData, 0);
 
@@ -193,7 +197,8 @@ bool FillTableColumn(const UCustomizableObjectNodeTable* TableNode,	mu::TablePtr
 
 					FString msg = FString::Printf(TEXT("Mesh from column [%s] row [%s] needs LOD %d but has less LODs than the reference mesh. LOD %d will be used instead. This can cause some performance penalties."),
 						*ColumnName, *RowName, LODIndex, LODIndex);
-					GenerationContext.Compiler->CompilerLog(FText::FromString(msg), TableNode);
+
+					LogRowGenerationMessage(TableNode, DataTablePtr, GenerationContext, msg, RowName);
 				}
 			}
 
@@ -210,7 +215,7 @@ bool FillTableColumn(const UCustomizableObjectNodeTable* TableNode,	mu::TablePtr
 					FString Dif_2 = NumSections > ReferenceNumMaterials ? "Some will be ignored" : "This can cause some compilation errors.";
 
 					FString msg = FString::Printf(TEXT("Mesh from column [%s] row [%s] has %s Sections than the reference mesh. %s"), *ColumnName, *RowName, *Dif_1, *Dif_2);
-					GenerationContext.Compiler->CompilerLog(FText::FromString(msg), TableNode);
+					LogRowGenerationMessage(TableNode, DataTablePtr, GenerationContext, msg, RowName);
 				}
 			}
 
@@ -299,7 +304,7 @@ bool FillTableColumn(const UCustomizableObjectNodeTable* TableNode,	mu::TablePtr
 			{
 				FString msg = FString::Printf(TEXT("Error converting skeletal mesh LOD %d, Section %d from column [%s] row [%s] to mutable."),
 					LODIndex, SectionIndex, *ColumnName, *RowName);
-				GenerationContext.Compiler->CompilerLog(FText::FromString(msg), TableNode);
+				LogRowGenerationMessage(TableNode, DataTablePtr, GenerationContext, msg, RowName);
 			}
 		}
 
@@ -334,7 +339,7 @@ bool FillTableColumn(const UCustomizableObjectNodeTable* TableNode,	mu::TablePtr
 
 				FString msg = FString::Printf(TEXT("Mesh from column [%s] row [%s] needs LOD %d but has less LODs than the reference mesh. LOD %d will be used instead. This can cause some performance penalties."),
 					*ColumnName, *RowName, LODIndex, CurrentLOD);
-				GenerationContext.Compiler->CompilerLog(FText::FromString(msg), TableNode);
+				LogRowGenerationMessage(TableNode, DataTablePtr, GenerationContext, msg, RowName);
 			}
 
 			int32 NumMaterials = StaticMesh->GetRenderData()->LODResources[CurrentLOD].Sections.Num();
@@ -346,7 +351,7 @@ bool FillTableColumn(const UCustomizableObjectNodeTable* TableNode,	mu::TablePtr
 				FString SecondTextOption = NumMaterials > ReferenceNumMaterials ? "Some will be ignored" : "This can cause some compilation errors.";
 
 				FString msg = FString::Printf(TEXT("Mesh from column [%s] row [%s] has %s Sections than the reference mesh. %s"), *ColumnName, *RowName, *FirstTextOption, *SecondTextOption);
-				GenerationContext.Compiler->CompilerLog(FText::FromString(msg), TableNode);
+				LogRowGenerationMessage(TableNode, DataTablePtr, GenerationContext, msg, RowName);
 			}
 
 			FString MutableColumnName = TableNode->GenerateStaticMeshMutableColumName(ColumnName, SectionIndex);
@@ -368,8 +373,7 @@ bool FillTableColumn(const UCustomizableObjectNodeTable* TableNode,	mu::TablePtr
 			{
 				FString msg = FString::Printf(TEXT("Error converting skeletal mesh LOD %d, Section %d from column [%s] row [%s] to mutable."),
 					LODIndex, SectionIndex, *ColumnName, *RowName);
-
-				GenerationContext.Compiler->CompilerLog(FText::FromString(msg), TableNode);
+				LogRowGenerationMessage(TableNode, DataTablePtr, GenerationContext, msg, RowName);
 			}
 		}
 
@@ -379,14 +383,17 @@ bool FillTableColumn(const UCustomizableObjectNodeTable* TableNode,	mu::TablePtr
 			UTexture2D* Texture2D = Cast<UTexture2D>(Object);
 			UTexture2DArray* TextureArray = Cast<UTexture2DArray>(Object);
 
+			// Removing encoding part
+			const FString PinName = ColumnName.Replace(TEXT("--PassThrough"), TEXT(""), ESearchCase::CaseSensitive);
+
 			if (!Texture2D && !TextureArray)
 			{
-				Texture2D = TableNode->GetColumnDefaultAssetByType<UTexture2D>(ColumnName);
-				TextureArray = TableNode->GetColumnDefaultAssetByType<UTexture2DArray>(ColumnName);
+				Texture2D = TableNode->GetColumnDefaultAssetByType<UTexture2D>(PinName);
+				TextureArray = TableNode->GetColumnDefaultAssetByType<UTexture2DArray>(PinName);
 
 				FString Message = Cast<UObject>(Object) ? "not a suported Texture" : "null";
-				FString WarningMessage = FString::Printf(TEXT("Texture from column [%s] row [%s] is %s. The default texture will be used instead."), *ColumnName, *RowName, *Message);
-				GenerationContext.Compiler->CompilerLog(FText::FromString(WarningMessage), TableNode);
+				FString WarningMessage = FString::Printf(TEXT("Texture from column [%s] row [%s] is %s. The default texture will be used instead."), *PinName, *RowName, *Message);
+				LogRowGenerationMessage(TableNode, DataTablePtr, GenerationContext, WarningMessage, RowName);
 			}
 
 			// There will be always one of the two options
@@ -401,9 +408,6 @@ bool FillTableColumn(const UCustomizableObjectNodeTable* TableNode,	mu::TablePtr
 				CurrentColumn = MutableTable->AddColumn(ColumnName, mu::ETableColumnType::Image);
 			}
 
-			// Removing encoding part
-			const FString PinName = ColumnName.Replace(TEXT("--PassThrough"),TEXT(""), ESearchCase::CaseSensitive);
-
 			if (TableNode->GetColumnImageMode(PinName) == ETableTextureType::PASSTHROUGH_TEXTURE)
 			{
 				FMutableGraphGenerationContext::FGeneratedPassThroughTexture* FoundIndex = GenerationContext.PassThroughTextureMap.Find(Texture);
@@ -494,7 +498,7 @@ bool FillTableColumn(const UCustomizableObjectNodeTable* TableNode,	mu::TablePtr
 
 				Material = ReferenceMaterial;
 
-				GenerationContext.Compiler->CompilerLog(Warning, TableNode);
+				LogRowGenerationMessage(TableNode, DataTablePtr, GenerationContext, Warning.ToString(), RowName);
 			}
 
 			GenerationContext.AddParticipatingObject(*Material);
@@ -558,7 +562,7 @@ bool FillTableColumn(const UCustomizableObjectNodeTable* TableNode,	mu::TablePtr
 					FString Message = Cast<UObject>(ParentParameterTexture) ? "not a Texture2D" : "null";
 					
 					FString msg = FString::Printf(TEXT("Parameter [%s] from Default Material Instance of column [%s] is %s. This parameter will be ignored."), *ParamName, *MaterialColumnName, *Message);
-					GenerationContext.Compiler->CompilerLog(FText::FromString(msg), TableNode);
+					LogRowGenerationMessage(TableNode, DataTablePtr, GenerationContext, msg, RowName);
 					 
 					 return false;
 				}
@@ -576,7 +580,7 @@ bool FillTableColumn(const UCustomizableObjectNodeTable* TableNode,	mu::TablePtr
 					FString Message = Cast<UObject>(TextureValue) ? "not a Texture2D" : "null";
 
 					FString msg = FString::Printf(TEXT("Parameter [%s] from material instance of column [%s] row [%s] is %s. The parameter texture of the default material will be used instead."), *ParamName, *MaterialColumnName, *RowName, *Message);
-					GenerationContext.Compiler->CompilerLog(FText::FromString(msg), TableNode);
+					LogRowGenerationMessage(TableNode, DataTablePtr, GenerationContext, msg, RowName);
 				}
 
 				GenerationContext.ArrayTextureUnrealToMutableTask.Add(FTextureUnrealToMutableTask(MutableTable, ParameterTexture, TableNode, ColumnIndex, RowIdx));
@@ -658,20 +662,27 @@ bool GenerateTableColumn(const UCustomizableObjectNodeTable* TableNode, const UE
 {
 	SCOPED_PIN_DATA(GenerationContext, Pin)
 
-	if (!TableNode || !TableNode->Table || !TableNode->Table->GetRowStruct())
+	if (!TableNode)
+	{
+		return false;
+	}
+
+	UDataTable* DataTable = GetDataTable(TableNode, GenerationContext);
+
+	if (!DataTable || !DataTable->GetRowStruct())
 	{
 		return false;
 	}
 	
-	GenerationContext.AddParticipatingObject(*TableNode->Table);
+	GenerationContext.AddParticipatingObject(*DataTable);
 
 	// Getting names of the rows to access the information
-	TArray<FName> RowNames = TableNode->GetRowNames();
+	TArray<FName> RowNames = TableNode->GetRowNames(DataTable);
 
 	for (int32 RowIndex = 0; RowIndex < RowNames.Num(); ++RowIndex)
 	{
 		// Getting Row Data
-		uint8* RowData = TableNode->Table->FindRowUnchecked(RowNames[RowIndex]);
+		uint8* RowData = DataTable->FindRowUnchecked(RowNames[RowIndex]);
 
 		if (RowData)
 		{
@@ -712,9 +723,9 @@ mu::TablePtr GenerateMutableSourceTable(const FString& TableName, const UEdGraph
 
 	if (const UCustomizableObjectNodeTable* TypedTable = Cast<UCustomizableObjectNodeTable>(Node))
 	{
-		UDataTable* Table = TypedTable->Table;
+		UDataTable* DataTable = GetDataTable(TypedTable, GenerationContext);
 
-		if (!Table)
+		if (!DataTable)
 		{
 			FString msg = "Couldn't find the Data Table asset in the Node.";
 			GenerationContext.Compiler->CompilerLog(FText::FromString(msg), Node);
@@ -722,12 +733,12 @@ mu::TablePtr GenerateMutableSourceTable(const FString& TableName, const UEdGraph
 			return nullptr;
 		}
 
-		const UScriptStruct* TableStruct = Table->GetRowStruct();
+		const UScriptStruct* TableStruct = DataTable->GetRowStruct();
 
 		if (TableStruct)
 		{
 			// Getting names of the rows to access the information
-			TArray<FName> RowNames = TypedTable->GetRowNames();
+			const TArray<FName>& RowNames = TypedTable->GetRowNames(DataTable);
 
 			// Adding and filling Name Column
 			MutableTable->AddColumn("Name", mu::ETableColumnType::String);
@@ -775,5 +786,195 @@ mu::TablePtr GenerateMutableSourceTable(const FString& TableName, const UEdGraph
 	return MutableTable;
 }
 
+
+UDataTable* GetDataTable(const UCustomizableObjectNodeTable* TableNode, FMutableGraphGenerationContext& GenerationContext)
+{
+	UDataTable* OutDataTable = nullptr;
+
+	if (TableNode->TableDataGatheringMode == ETableDataGatheringSource::ETDGM_AssetRegistry)
+	{
+		OutDataTable = GenerateDataTableFromStruct(TableNode, GenerationContext);
+	}
+	else
+	{
+		OutDataTable = TableNode->Table;
+	}
+
+	return OutDataTable;
+}
+
+
+UDataTable* GenerateDataTableFromStruct(const UCustomizableObjectNodeTable* TableNode, FMutableGraphGenerationContext& GenerationContext)
+{
+	if (!TableNode->Structure)
+	{
+		GenerationContext.Compiler->CompilerLog(LOCTEXT("EmptyStructureError", "Empty structure asset."), TableNode);
+		return nullptr;
+	}
+
+	FMutableGraphGenerationContext::FGeneratedDataTablesData DataTableData;
+	DataTableData.ParentStruct = TableNode->Structure;
+	DataTableData.FilterPaths = TableNode->FilterPaths;
+	
+	//Checking cache of generated data tables
+	int32 DataTableIndex = GenerationContext.GeneratedCompositeDataTables.Find(DataTableData);
+	if (DataTableIndex != INDEX_NONE)
+	{
+		// DataTable already generated
+		UCompositeDataTable* GeneratedDataTable = GenerationContext.GeneratedCompositeDataTables[DataTableIndex].GeneratedDataTable;
+		return Cast<UDataTable>(GeneratedDataTable);
+	}
+
+	FAssetRegistryModule& AssetRegistryModule = FModuleManager::LoadModuleChecked<FAssetRegistryModule>("AssetRegistry");
+	IAssetRegistry& AssetRegistry = AssetRegistryModule.GetRegistry();
+
+	if (TableNode->FilterPaths.IsEmpty())
+	{
+		// Preventing load all data tables of the project
+		GenerationContext.Compiler->CompilerLog(LOCTEXT("NoFilePathsError", "There are no filter paths selected. This is an error to prevent loading all data table of the project."), TableNode);
+
+		return nullptr;
+	}
+
+	TArray<FName> ReferencedTables;
+	AssetRegistry.Get()->GetReferencers(TableNode->Structure.GetPackage().GetFName(), ReferencedTables, UE::AssetRegistry::EDependencyCategory::Package, UE::AssetRegistry::EDependencyQuery::NoRequirements);
+	
+	FARFilter Filter;
+	Filter.ClassPaths.Add(FTopLevelAssetPath(UDataTable::StaticClass()));
+
+	for (const FName& ReferencedTable : ReferencedTables)
+	{
+		Filter.PackageNames.Add(ReferencedTable);
+	}
+
+	for (int32 PathIndex = 0; PathIndex < TableNode->FilterPaths.Num(); ++PathIndex)
+	{
+		Filter.PackagePaths.Add(TableNode->FilterPaths[PathIndex]);
+	}
+
+	Filter.bRecursivePaths = true;
+
+	TArray<FAssetData> DataTableAssets;
+	AssetRegistry.Get()->GetAssets(Filter, DataTableAssets);
+
+	UCompositeDataTable* CompositeDataTable = NewObject<UCompositeDataTable>();
+	CompositeDataTable->RowStruct = TableNode->Structure;
+
+	TArray<UDataTable*> ParentTables;
+
+	for (const FAssetData& DataTableAsset : DataTableAssets)
+	{
+		if (DataTableAsset.IsValid())
+		{
+			if (UDataTable* DataTable = Cast<UDataTable>(DataTableAsset.GetAsset()))
+			{
+				ParentTables.Add(DataTable);
+			}
+		}
+	}
+
+	if (ParentTables.IsEmpty())
+	{
+		GenerationContext.Compiler->CompilerLog(LOCTEXT("NoDataTablesFoundWarning", "Could not find a data table with the specified struct in the selected paths."), TableNode);
+
+		return nullptr;
+	}
+
+	// Map to find the original data table of a row
+	TMap<FName, TArray<UDataTable*>> OriginalTableRowsMap;
+
+	// Set to iterate faster the repeated rows inside the map
+	TSet<FName> RepeatedRowNamesArray;
+
+	// Checking if a row name is repeated in several tables
+	for (int32 ParentIndx = 0; ParentIndx < ParentTables.Num(); ++ParentIndx)
+	{
+		const TArray<FName>& RowNames = ParentTables[ParentIndx]->GetRowNames();
+
+		for (const FName& RowName : RowNames)
+		{
+			TArray<UDataTable*>* DataTablesNames = OriginalTableRowsMap.Find(RowName);
+
+			if (DataTablesNames == nullptr)
+			{
+				TArray<UDataTable*> ArrayTemp;
+				ArrayTemp.Add(ParentTables[ParentIndx]);
+				OriginalTableRowsMap.Add(RowName, ArrayTemp);
+			}
+			else
+			{
+				DataTablesNames->Add(ParentTables[ParentIndx]);
+				RepeatedRowNamesArray.Add(RowName);
+			}
+		}
+	}
+
+	for (const FName& RowName : RepeatedRowNamesArray)
+	{
+		const TArray<UDataTable*>& DataTablesNames = OriginalTableRowsMap[RowName];
+
+		FString TableNames;
+
+		for (int32 NameIndx = 0; NameIndx < DataTablesNames.Num(); ++NameIndx)
+		{
+			TableNames += DataTablesNames[NameIndx]->GetName();
+
+			if (NameIndx + 1 < DataTablesNames.Num())
+			{
+				TableNames += ", ";
+			}
+		}
+
+		FString Message = FString::Printf(TEXT("Row with name [%s] repeated in the following Data Tables: [%s]. The last row processed will be used [%s]."),
+			*RowName.ToString(), *TableNames, *DataTablesNames.Last()->GetName());
+		GenerationContext.Compiler->CompilerLog(FText::FromString(Message), TableNode);
+	}
+
+	CompositeDataTable->AppendParentTables(ParentTables);
+
+	// Adding Generated Data Table to the cache
+	DataTableData.GeneratedDataTable = CompositeDataTable;
+	GenerationContext.GeneratedCompositeDataTables.Add(DataTableData);
+	GenerationContext.CompositeDataTableRowToOriginalDataTableMap.Add(CompositeDataTable, OriginalTableRowsMap);
+	
+	return Cast<UDataTable>(CompositeDataTable);
+}
+
+
+
+void LogRowGenerationMessage(const UCustomizableObjectNodeTable* TableNode, const UDataTable* DataTable, FMutableGraphGenerationContext& GenerationContext, const FString& Message, const FString& RowName)
+{
+	FString FinalMessage = Message;
+
+	if (TableNode->TableDataGatheringMode == ETableDataGatheringSource::ETDGM_AssetRegistry)
+	{
+		TMap<FName, TArray<UDataTable*>>* ParameterDataTableMap = GenerationContext.CompositeDataTableRowToOriginalDataTableMap.Find(DataTable);
+
+		if (ParameterDataTableMap)
+		{
+			TArray<UDataTable*>* DataTables = ParameterDataTableMap->Find(FName(*RowName));
+
+			if (DataTables)
+			{
+				FString TableNames;
+
+				for (int32 NameIndx = 0; NameIndx < DataTables->Num(); ++NameIndx)
+				{
+					TableNames += (*DataTables)[NameIndx]->GetName();
+
+					if (NameIndx + 1 < DataTables->Num())
+					{
+						TableNames += ", ";
+					}
+				}
+
+				FinalMessage += " Row from Composite Data Table, original Data Table/s: " + TableNames;
+			}
+		}
+	}
+
+	GenerationContext.Compiler->CompilerLog(FText::FromString(FinalMessage), TableNode);
+}
+
 #undef LOCTEXT_NAMESPACE
 
diff --git a/Engine/Plugins/Experimental/Mutable/Source/CustomizableObjectEditor/Private/MuCOE/GenerateMutableSource/GenerateMutableSourceTable.h b/Engine/Plugins/Experimental/Mutable/Source/CustomizableObjectEditor/Private/MuCOE/GenerateMutableSource/GenerateMutableSourceTable.h
--- a/Engine/Plugins/Experimental/Mutable/Source/CustomizableObjectEditor/Private/MuCOE/GenerateMutableSource/GenerateMutableSourceTable.h
+++ b/Engine/Plugins/Experimental/Mutable/Source/CustomizableObjectEditor/Private/MuCOE/GenerateMutableSource/GenerateMutableSourceTable.h
@@ -6,8 +6,8 @@
 
 class FProperty;
 class FString;
-
 class UCustomizableObjectNodeTable;
+class UDataTable;
 class UEdGraphPin;
 struct FMutableGraphGenerationContext;
 
@@ -58,3 +58,12 @@ bool FillTableColumn(const UCustomizableObjectNodeTable* TableNode, mu::TablePtr
 
 
 mu::TablePtr GenerateMutableSourceTable(const FString& TableName, const UEdGraphPin* Pin, FMutableGraphGenerationContext& GenerationContext);
+
+/** Gets the data table needed during the compilation process */
+UDataTable* GetDataTable(const UCustomizableObjectNodeTable* TableNode, FMutableGraphGenerationContext& GenerationContext);
+
+/** Generates a Data Table from the Data Tables referenced in a Script Struct */
+UDataTable* GenerateDataTableFromStruct(const UCustomizableObjectNodeTable* TableNode, FMutableGraphGenerationContext& GenerationContext);
+
+/** This method adds the original Data Table(s) of the processed row at the end of the log message when the data table is a Composite Data table. */
+void LogRowGenerationMessage(const UCustomizableObjectNodeTable* TableNode, const UDataTable* DataTable, FMutableGraphGenerationContext& GenerationContext, const FString& Message, const FString& RowName);
diff --git a/Engine/Plugins/Experimental/Mutable/Source/CustomizableObjectEditor/Private/MuCOE/Nodes/CustomizableObjectNodeTable.cpp b/Engine/Plugins/Experimental/Mutable/Source/CustomizableObjectEditor/Private/MuCOE/Nodes/CustomizableObjectNodeTable.cpp
--- a/Engine/Plugins/Experimental/Mutable/Source/CustomizableObjectEditor/Private/MuCOE/Nodes/CustomizableObjectNodeTable.cpp
+++ b/Engine/Plugins/Experimental/Mutable/Source/CustomizableObjectEditor/Private/MuCOE/Nodes/CustomizableObjectNodeTable.cpp
@@ -4,6 +4,8 @@
 
 #include "Engine/StaticMesh.h"
 #include "Engine/Texture2DArray.h"
+#include "Engine/UserDefinedStruct.h"
+#include "Kismet2/StructureEditorUtils.h"
 #include "Materials/Material.h"
 #include "Materials/MaterialInstance.h"
 #include "MuCO/UnrealPortabilityHelpers.h"
@@ -145,7 +147,14 @@ FText UCustomizableObjectNodeTable::GetNodeTitle(ENodeTitleType::Type TitleType)
 		FFormatNamedArguments Args;
 		Args.Add(TEXT("TableName"), FText::FromString(Table->GetName()));
 
-		return FText::Format(LOCTEXT("TableNode_Title", "{TableName}\nTable"), Args);
+		return FText::Format(LOCTEXT("TableNode_Title", "{TableName}\nData Table"), Args);
+	}
+	else if (Structure)
+	{
+		FFormatNamedArguments Args;
+		Args.Add(TEXT("StructureName"), FText::FromString(Structure->GetName()));
+
+		return FText::Format(LOCTEXT("TableNode_Title", "{StructureName}\nScript Struct"), Args);
 	}
 	
 	return LOCTEXT("Mutable Table", "Table");
@@ -190,7 +199,7 @@ void UCustomizableObjectNodeTable::PostEditChangeProperty(FPropertyChangedEvent&
 
 	if (PropertyThatChanged)
 	{
-		if (PropertyThatChanged->GetName() == TEXT("Table"))
+		if (PropertyThatChanged->GetName() == TEXT("Table") || PropertyThatChanged->GetName() == TEXT("Structure"))
 		{
 			ReconstructNode();
 		}
@@ -217,38 +226,43 @@ void UCustomizableObjectNodeTable::PostEditChangeProperty(FPropertyChangedEvent&
 				ReconstructNode();
 			}
 		}
+		else if (PropertyThatChanged->GetName() == TEXT("TableDataGatheringMode"))
+		{
+			if (Table)
+			{
+				Table->OnDataTableChanged().Remove(OnTableChangedDelegateHandle);
+			}
+
+			Table = nullptr;
+			Structure = nullptr;
+
+			FilterPaths.Empty();
+
+			ReconstructNode();
+		}
 	}
 }
 
 
 void UCustomizableObjectNodeTable::AllocateDefaultPins(UCustomizableObjectNodeRemapPins* RemapPins)
 {
-	if (!Table)
-	{
-		return;
-	}
-
 	// Getting Struct Pointer
-	const UScriptStruct* TableStruct = Table->GetRowStruct();
+	const UScriptStruct* TableStruct = GetTableNodeStruct();
 
 	if (!TableStruct)
 	{
 		return;
 	}
 
-	NumProperties = Table->GetColumnTitles().Num();
+	NumProperties = GetColumnTitles(TableStruct).Num();
 
 	// Getting Default Struct Values
-	uint8* DefaultRowData = (uint8*)FMemory::Malloc(TableStruct->GetStructureSize());
-	
-	if (!DefaultRowData)
-	{
-		return;
-	}
+	// A Script Struct always has at leaset one property
+	TArray<int8> DefaultDataArray;
+	DefaultDataArray.SetNumZeroed(TableStruct->GetStructureSize());
+	TableStruct->InitializeStruct(DefaultDataArray.GetData());
 
 	TArray<UEdGraphPin*> OldPins(Pins);
-
-	TableStruct->InitializeStruct(DefaultRowData);
 	
 	for (TFieldIterator<FProperty> It(TableStruct); It; ++It)
 	{
@@ -271,7 +285,7 @@ void UCustomizableObjectNodeTable::AllocateDefaultPins(UCustomizableObjectNodeRe
 			UObject* Object = nullptr;
 
 			// Getting default UObject
-			uint8* CellData = SoftObjectProperty->ContainerPtrToValuePtr<uint8>(DefaultRowData);
+			uint8* CellData = SoftObjectProperty->ContainerPtrToValuePtr<uint8>(DefaultDataArray.GetData());
 
 			if (CellData)
 			{
@@ -376,9 +390,7 @@ void UCustomizableObjectNodeTable::AllocateDefaultPins(UCustomizableObjectNodeRe
 		}
 	}
 
-	// Cleaning Default Structure Pointer
-	TableStruct->DestroyStruct(DefaultRowData);
-	FMemory::Free(DefaultRowData);
+	TableStruct->DestroyStruct(DefaultDataArray.GetData());
 }
 
 
@@ -464,33 +476,24 @@ void UCustomizableObjectNodeTable::GenerateMeshPins(UObject* Mesh, const FString
 
 bool UCustomizableObjectNodeTable::IsNodeOutDatedAndNeedsRefresh()
 {
-	if (!Table)
-	{
-		return Pins.Num() > 0;
-	}
-
 	// Getting Struct Pointer
-	const UScriptStruct* TableStruct = Table->GetRowStruct();
+	const UScriptStruct* TableStruct = GetTableNodeStruct();
 
 	if (!TableStruct)
 	{
 		return Pins.Num() > 0;
 	}
 
-	if (NumProperties != Table->GetColumnTitles().Num())
+	if (NumProperties != GetColumnTitles(TableStruct).Num())
 	{
 		return true;
 	}
 
-	// Getting Default Struct 
-	uint8* DefaultRowData = (uint8*)FMemory::Malloc(TableStruct->GetStructureSize());
-
-	if (!DefaultRowData)
-	{
-		return Pins.Num() > 0;
-	}
-
-	TableStruct->InitializeStruct(DefaultRowData);
+	// Getting Default Struct Values
+	// A Script Struct always has at leaset one property
+	TArray<int8> DefaultDataArray;
+	DefaultDataArray.SetNumZeroed(TableStruct->GetStructureSize());
+	TableStruct->InitializeStruct(DefaultDataArray.GetData());
 
 	int32  NumPins = 0;
 
@@ -511,7 +514,7 @@ bool UCustomizableObjectNodeTable::IsNodeOutDatedAndNeedsRefresh()
 				UObject* Object = nullptr;
 
 				// Getting default UObject
-				uint8* CellData = SoftObjectProperty->ContainerPtrToValuePtr<uint8>(DefaultRowData);
+				uint8* CellData = SoftObjectProperty->ContainerPtrToValuePtr<uint8>(DefaultDataArray.GetData());
 
 				if (CellData)
 				{
@@ -642,9 +645,7 @@ bool UCustomizableObjectNodeTable::IsNodeOutDatedAndNeedsRefresh()
 		}
 	}
 
-	// Cleaning Default Structure Pointer
-	TableStruct->DestroyStruct(DefaultRowData);
-	FMemory::Free(DefaultRowData);
+	TableStruct->DestroyStruct(DefaultDataArray.GetData());
 
 	if (Pins.Num() != NumPins)
 	{
@@ -672,19 +673,21 @@ void UCustomizableObjectNodeTable::RemapPinsData(const TMap<UEdGraphPin*, UEdGra
 			UCustomizableObjectNodeTableMeshPinData* PinDataOldPin = Cast<UCustomizableObjectNodeTableMeshPinData>(GetPinData(*(Pair.Key)));
 			UCustomizableObjectNodeTableMeshPinData* PinDataNewPin = Cast<UCustomizableObjectNodeTableMeshPinData>(GetPinData(*(Pair.Value)));
 
-			if (PinDataOldPin && PinDataNewPin)
+			const UScriptStruct* ScriptStruct = GetTableNodeStruct();
+
+			if (PinDataOldPin && PinDataNewPin && ScriptStruct)
 			{
-				if (Table->FindTableProperty(FName(*PinDataOldPin->AnimInstanceColumnName)))
+				if (FindTableProperty(ScriptStruct, FName(*PinDataOldPin->AnimInstanceColumnName)))
 				{
 					PinDataNewPin->AnimInstanceColumnName = PinDataOldPin->AnimInstanceColumnName;
 				}
 
-				if (Table->FindTableProperty(FName(*PinDataOldPin->AnimSlotColumnName)))
+				if (FindTableProperty(ScriptStruct, FName(*PinDataOldPin->AnimSlotColumnName)))
 				{
 					PinDataNewPin->AnimSlotColumnName = PinDataOldPin->AnimSlotColumnName;
 				}
 
-				if (Table->FindTableProperty(FName(*PinDataOldPin->AnimTagColumnName)))
+				if (FindTableProperty(ScriptStruct, FName(*PinDataOldPin->AnimTagColumnName)))
 				{
 					PinDataNewPin->AnimTagColumnName = PinDataOldPin->AnimTagColumnName;
 				}
@@ -979,15 +982,14 @@ bool UCustomizableObjectNodeTable::CheckPinUpdated(const FString& PinName, const
 }
 
 
-USkeletalMesh* UCustomizableObjectNodeTable::GetSkeletalMeshAt(const UEdGraphPin* Pin, const FName& RowName) const
+USkeletalMesh* UCustomizableObjectNodeTable::GetSkeletalMeshAt(const UEdGraphPin* Pin, const UDataTable* DataTable, const FName& RowName) const
 {
-	if (!Table || !Table->GetRowStruct() || !Pin || !Table->GetRowNames().Contains(RowName))
+	if (!DataTable || !DataTable->GetRowStruct() || !Pin || !DataTable->GetRowNames().Contains(RowName))
 	{
 		return nullptr;
 	}
-
-	// Getting Struct Pointer
-	const UScriptStruct* TableStruct = Table->GetRowStruct();
+	
+	const UScriptStruct* TableStruct = DataTable->GetRowStruct();
 
 	FString ColumnName = GetColumnNameByPin(Pin);
 
@@ -1002,7 +1004,7 @@ USkeletalMesh* UCustomizableObjectNodeTable::GetSkeletalMeshAt(const UEdGraphPin
 
 		if (const FSoftObjectProperty* SoftObjectProperty = CastField<FSoftObjectProperty>(ColumnProperty))
 		{
-			if (uint8* RowData = Table->FindRowUnchecked(RowName))
+			if (uint8* RowData = DataTable->FindRowUnchecked(RowName))
 			{
 				if (uint8* CellData = ColumnProperty->ContainerPtrToValuePtr<uint8>(RowData, 0))
 				{
@@ -1021,15 +1023,15 @@ USkeletalMesh* UCustomizableObjectNodeTable::GetSkeletalMeshAt(const UEdGraphPin
 	return nullptr;
 }
 
-TSoftClassPtr<UAnimInstance> UCustomizableObjectNodeTable::GetAnimInstanceAt(const UEdGraphPin* Pin, const FName& RowName) const
+
+TSoftClassPtr<UAnimInstance> UCustomizableObjectNodeTable::GetAnimInstanceAt(const UEdGraphPin* Pin, const UDataTable* DataTable, const FName& RowName) const
 {
-	if (!Table || !Table->GetRowStruct() || !Pin || !Table->GetRowNames().Contains(RowName))
+	if (!DataTable || !DataTable->GetRowStruct() || !Pin || !DataTable->GetRowNames().Contains(RowName))
 	{
 		return TSoftClassPtr<UAnimInstance>();
 	}
 
-	// Getting Struct Pointer
-	const UScriptStruct* TableStruct = Table->GetRowStruct();
+	const UScriptStruct* TableStruct = DataTable->GetRowStruct();
 
 	FString ColumnName = GetColumnNameByPin(Pin);
 
@@ -1044,11 +1046,17 @@ TSoftClassPtr<UAnimInstance> UCustomizableObjectNodeTable::GetAnimInstanceAt(con
 
 		if (const FSoftClassProperty* SoftClassProperty = CastField<FSoftClassProperty>(ColumnProperty))
 		{
-			TSoftClassPtr<UAnimInstance> AnimInstance(SoftClassProperty->GetPropertyValue(SoftClassProperty).ToSoftObjectPath());
-
-			if (!AnimInstance.IsNull())
+			if (uint8* RowData = DataTable->FindRowUnchecked(RowName))
 			{
-				return AnimInstance;
+				if (uint8* CellData = ColumnProperty->ContainerPtrToValuePtr<uint8>(RowData, 0))
+				{
+					TSoftClassPtr<UAnimInstance> AnimInstance(SoftClassProperty->GetPropertyValue(CellData).ToSoftObjectPath());
+
+					if (!AnimInstance.IsNull())
+					{
+						return AnimInstance;
+					}
+				}
 			}
 		}
 	}
@@ -1056,20 +1064,21 @@ TSoftClassPtr<UAnimInstance> UCustomizableObjectNodeTable::GetAnimInstanceAt(con
 	return TSoftClassPtr<UAnimInstance>();
 }
 
-TArray<FName> UCustomizableObjectNodeTable::GetRowNames() const
+
+TArray<FName> UCustomizableObjectNodeTable::GetRowNames(const UDataTable* DataTable) const
 {
 	TArray<FName> RowNames;
 
-	if (Table)
+	if (DataTable)
 	{
-		const UScriptStruct* TableStruct = Table->GetRowStruct();
+		const UScriptStruct* TableStruct = DataTable->GetRowStruct();
 
 		if (!TableStruct)
 		{
 			return RowNames;
 		}
 
-		TArray<FName> TableRowNames = Table->GetRowNames();
+		TArray<FName> TableRowNames = DataTable->GetRowNames();
 		FBoolProperty* BoolProperty = nullptr;
 
 		for (TFieldIterator<FProperty> PropertyIt(TableStruct); PropertyIt && bDisableCheckedRows; ++PropertyIt)
@@ -1080,7 +1089,7 @@ TArray<FName> UCustomizableObjectNodeTable::GetRowNames() const
 			{
 				for (const FName& RowName : TableRowNames)
 				{
-					if (uint8* RowData = Table->FindRowUnchecked(RowName))
+					if (uint8* RowData = DataTable->FindRowUnchecked(RowName))
 					{
 						if (uint8* CellData = BoolProperty->ContainerPtrToValuePtr<uint8>(RowData, 0))
 						{
@@ -1191,10 +1200,13 @@ ETableTextureType UCustomizableObjectNodeTable::GetColumnImageMode(const FString
 
 ETableMeshPinType UCustomizableObjectNodeTable::GetPinMeshType(const UEdGraphPin* Pin) const
 {
-	if (Pin && Pin->PinType.PinCategory == UEdGraphSchema_CustomizableObject::PC_Mesh)
+	// Getting Struct Pointer
+	const UScriptStruct* TableStruct = GetTableNodeStruct();
+
+	if (TableStruct && Pin && Pin->PinType.PinCategory == UEdGraphSchema_CustomizableObject::PC_Mesh)
 	{
 		FString ColumnName = GetColumnNameByPin(Pin);
-		FProperty* Property = Table->FindTableProperty(FName(*ColumnName));
+		FProperty* Property = FindTableProperty(TableStruct, FName(*ColumnName));
 
 		if (const FSoftObjectProperty* SoftObjectProperty = CastField<FSoftObjectProperty>(Property))
 		{
@@ -1224,4 +1236,78 @@ FString UCustomizableObjectNodeTable::GenerateStaticMeshMutableColumName(const F
 	return PinName + FString::Printf(TEXT(" Mat_%d"), MaterialIndex);
 }
 
+
+const UScriptStruct* UCustomizableObjectNodeTable::GetTableNodeStruct() const
+{
+	const UScriptStruct* TableStruct = nullptr;
+
+	if (TableDataGatheringMode == ETableDataGatheringSource::ETDGM_AssetRegistry)
+	{
+		if (Structure)
+		{
+			TableStruct = Structure;
+		}
+	}
+	else
+	{
+		if (Table)
+		{
+			TableStruct = Table->GetRowStruct();
+		}
+	}
+
+	return TableStruct;
+}
+
+
+TArray<FString> UCustomizableObjectNodeTable::GetColumnTitles(const UScriptStruct* ScriptStruct) const
+{
+	TArray<FString> Result;
+	Result.Add(TEXT("Name"));
+	if (ScriptStruct)
+	{
+		for (TFieldIterator<FProperty> It(ScriptStruct); It; ++It)
+		{
+			FProperty* Prop = *It;
+			check(Prop != nullptr);
+			const FString DisplayName = DataTableUtils::GetPropertyExportName(Prop);
+			Result.Add(DisplayName);
+		}
+	}
+
+	return Result;
+}
+
+
+FProperty* UCustomizableObjectNodeTable::FindTableProperty(const UScriptStruct* ScriptStruct, const FName& PropertyName) const
+{
+	FProperty* Property = nullptr;
+
+	if (ScriptStruct)
+	{
+		Property = ScriptStruct->FindPropertyByName(PropertyName);
+		if (Property == nullptr && ScriptStruct->IsA<UUserDefinedStruct>())
+		{
+			const FString PropertyNameStr = PropertyName.ToString();
+
+			for (TFieldIterator<FProperty> It(ScriptStruct); It; ++It)
+			{
+				if (PropertyNameStr == ScriptStruct->GetAuthoredNameForField(*It))
+				{
+					Property = *It;
+					break;
+				}
+			}
+		}
+		if (!DataTableUtils::IsSupportedTableProperty(Property))
+		{
+			Property = nullptr;
+		}
+	}
+
+	return Property;
+}
+
+
+
 #undef LOCTEXT_NAMESPACE
diff --git a/Engine/Plugins/Experimental/Mutable/Source/CustomizableObjectEditor/Private/MuCOE/Nodes/CustomizableObjectNodeTableDetails.cpp b/Engine/Plugins/Experimental/Mutable/Source/CustomizableObjectEditor/Private/MuCOE/Nodes/CustomizableObjectNodeTableDetails.cpp
--- a/Engine/Plugins/Experimental/Mutable/Source/CustomizableObjectEditor/Private/MuCOE/Nodes/CustomizableObjectNodeTableDetails.cpp
+++ b/Engine/Plugins/Experimental/Mutable/Source/CustomizableObjectEditor/Private/MuCOE/Nodes/CustomizableObjectNodeTableDetails.cpp
@@ -272,52 +272,54 @@ void FCustomizableObjectNodeTableDetails::GenerateMeshColumnComboBoxOptions()
 	AnimMeshColumnOptionNames.Empty();
 	LayoutMeshColumnOptionNames.Empty();
 
-	if (Node->Table)
+	const UScriptStruct* TableStruct = Node->GetTableNodeStruct();
+
+	if (!TableStruct)
 	{
-		const UScriptStruct* TableStruct = Node->Table->GetRowStruct();
+		return;
+	}
 
-		// we just need the mesh columns
-		for (TFieldIterator<FProperty> It(TableStruct); It; ++It)
+	// we just need the mesh columns
+	for (TFieldIterator<FProperty> It(TableStruct); It; ++It)
+	{
+		FProperty* ColumnProperty = *It;
+
+		if (!ColumnProperty)
 		{
-			FProperty* ColumnProperty = *It;
+			continue;
+		}
 
-			if (!ColumnProperty)
+		if (const FSoftObjectProperty* SoftObjectProperty = CastField<FSoftObjectProperty>(ColumnProperty))
+		{
+			if (SoftObjectProperty->PropertyClass->IsChildOf(USkeletalMesh::StaticClass())
+				|| SoftObjectProperty->PropertyClass->IsChildOf(UStaticMesh::StaticClass()))
 			{
-				continue;
-			}
+				FString MeshColumnName = DataTableUtils::GetPropertyExportName(ColumnProperty);
+				AnimMeshColumnOptionNames.Add(MakeShareable(new FString(MeshColumnName)));
 
-			if (const FSoftObjectProperty* SoftObjectProperty = CastField<FSoftObjectProperty>(ColumnProperty))
-			{
-				if (SoftObjectProperty->PropertyClass->IsChildOf(USkeletalMesh::StaticClass())
-					|| SoftObjectProperty->PropertyClass->IsChildOf(UStaticMesh::StaticClass()))
+				for (const UEdGraphPin* Pin : Node->GetAllNonOrphanPins())
 				{
-					FString MeshColumnName = DataTableUtils::GetPropertyExportName(ColumnProperty);
-					AnimMeshColumnOptionNames.Add(MakeShareable(new FString(MeshColumnName)));
+					const UCustomizableObjectNodeTableMeshPinData* PinData = Cast<UCustomizableObjectNodeTableMeshPinData >(Node->GetPinData(*Pin));
 
-					for (const UEdGraphPin* Pin : Node->GetAllNonOrphanPins())
+					if (!PinData || PinData->ColumnName != MeshColumnName || Node->GetPinMeshType(Pin) != ETableMeshPinType::SKELETAL_MESH)
 					{
-						const UCustomizableObjectNodeTableMeshPinData* PinData = Cast<UCustomizableObjectNodeTableMeshPinData >(Node->GetPinData(*Pin));
-
-						if (!PinData || PinData->ColumnName != MeshColumnName || Node->GetPinMeshType(Pin) != ETableMeshPinType::SKELETAL_MESH)
-						{
-							continue;
-						}
+						continue;
+					}
 
-						if (PinData && PinData->ColumnName == MeshColumnName)
+					if (PinData && PinData->ColumnName == MeshColumnName)
+					{
+						if (PinData->Layouts.Num() > 1)
 						{
-							if (PinData->Layouts.Num() > 1)
-							{
-								for (int32 LayoutIndex = 0; LayoutIndex < PinData->Layouts.Num(); ++LayoutIndex)
-								{
-									FString LayoutName = Pin->PinFriendlyName.ToString() + FString::Printf(TEXT(" UV_%d"), LayoutIndex);
-									LayoutMeshColumnOptionNames.Add(MakeShareable(new FString(LayoutName)));
-								}
-							}
-							else
+							for (int32 LayoutIndex = 0; LayoutIndex < PinData->Layouts.Num(); ++LayoutIndex)
 							{
-								LayoutMeshColumnOptionNames.Add(MakeShareable(new FString(Pin->PinFriendlyName.ToString())));
+								FString LayoutName = Pin->PinFriendlyName.ToString() + FString::Printf(TEXT(" UV_%d"), LayoutIndex);
+								LayoutMeshColumnOptionNames.Add(MakeShareable(new FString(LayoutName)));
 							}
 						}
+						else
+						{
+							LayoutMeshColumnOptionNames.Add(MakeShareable(new FString(Pin->PinFriendlyName.ToString())));
+						}
 					}
 				}
 			}
@@ -345,7 +347,12 @@ void FCustomizableObjectNodeTableDetails::GenerateAnimInstanceComboBoxOptions()
 		ColumnName = *AnimMeshColumnComboBox->GetSelectedItem();
 	}
 
-	const UScriptStruct* TableStruct = Node->Table->GetRowStruct();
+	const UScriptStruct* TableStruct = Node->GetTableNodeStruct();
+
+	if(!TableStruct)
+	{
+		return;
+	}
 
 	// Fill in name option arrays and set the selected item if any
 	for (TFieldIterator<FProperty> It(TableStruct); It; ++It)
