#rnx

diff --git a/Engine/Plugins/Experimental/Mutable/Source/CustomizableObject/Private/MuCO/CustomizableInstanceLODManagement.cpp b/Engine/Plugins/Experimental/Mutable/Source/CustomizableObject/Private/MuCO/CustomizableInstanceLODManagement.cpp
--- a/Engine/Plugins/Experimental/Mutable/Source/CustomizableObject/Private/MuCO/CustomizableInstanceLODManagement.cpp
+++ b/Engine/Plugins/Experimental/Mutable/Source/CustomizableObject/Private/MuCO/CustomizableInstanceLODManagement.cpp
@@ -130,9 +130,16 @@ void UCustomizableInstanceLODManagement::UpdateInstanceDistsAndLODs(FMutableInst
 					continue;
 				}
 #endif
-				if (CustomizableObjectInstanceUsage && !CustomizableObjectInstanceUsage->IsTemplate())
+				if (!IsValid(*CustomizableObjectInstanceUsage))
 				{
-					UWorld* LocalWorld = CustomizableObjectInstanceUsage->GetWorld();
+					continue;
+				}
+
+				USkeletalMeshComponent* Parent = Cast<USkeletalMeshComponent>(CustomizableObjectInstanceUsage->GetAttachParent());
+
+				if (!CustomizableObjectInstanceUsage->IsTemplate())
+				{
+					UWorld* LocalWorld = Parent ? Parent->GetWorld() : nullptr;
 					APlayerController* Controller = LocalWorld ? LocalWorld->GetFirstPlayerController() : nullptr;
 					TWeakObjectPtr<const AActor> ViewCenter = Controller ? TWeakObjectPtr<const AActor>(Controller->GetPawn()) : nullptr;
 
@@ -143,12 +150,12 @@ void UCustomizableInstanceLODManagement::UpdateInstanceDistsAndLODs(FMutableInst
 					}
 
 	#if WITH_EDITOR
-					else if (CustomizableObjectInstanceUsage->GetWorld())
+					else if (Parent && Parent->GetWorld())
 					{
-						EWorldType::Type worldType = CustomizableObjectInstanceUsage->GetWorld()->WorldType;
+						EWorldType::Type WorldType = Parent->GetWorld()->WorldType;
 					
 						// Level Editor Instances (non PIE)
-						if (!bLevelEditorInstancesUpdated && worldType == EWorldType::Editor)
+						if (!bLevelEditorInstancesUpdated && WorldType == EWorldType::Editor)
 						{
 							for (FLevelEditorViewportClient* LevelVC : GEditor->GetLevelViewportClients())
 							{
@@ -162,7 +169,7 @@ void UCustomizableInstanceLODManagement::UpdateInstanceDistsAndLODs(FMutableInst
 						}
 
 						// Blueprint instances
-						else if (worldType == EWorldType::EditorPreview)
+						else if (WorldType == EWorldType::EditorPreview)
 						{
 							CustomizableObjectInstanceUsage->EditorUpdateComponent();
 						}
@@ -309,11 +316,13 @@ void UCustomizableInstanceLODManagement::UpdateInstanceDistsAndLODs(FMutableInst
 					continue;
 				}
 
+				USkeletalMeshComponent* Parent = Cast<USkeletalMeshComponent>(CustomizableObjectInstanceUsage->GetAttachParent());
+
 #if WITH_EDITOR
 
 				EWorldType::Type WorldType = EWorldType::Type::None;
 
-				UWorld* World = CustomizableObjectInstanceUsage->GetWorld();
+				UWorld* World = Parent ? Parent->GetWorld() : nullptr;
 				if (World)
 				{
 					WorldType = World->WorldType;
@@ -345,7 +354,6 @@ void UCustomizableInstanceLODManagement::UpdateInstanceDistsAndLODs(FMutableInst
 					LODTracker.bInitialized = true;
 				}
 
-				USkeletalMeshComponent* Parent = Cast<USkeletalMeshComponent>(CustomizableObjectInstanceUsage->GetAttachParent());
 				if (Parent)
 				{
 					COI->SetIsBeingUsedByComponentInPlay(true);
diff --git a/Engine/Plugins/Experimental/Mutable/Source/CustomizableObject/Private/MuCO/CustomizableObjectSystem.cpp b/Engine/Plugins/Experimental/Mutable/Source/CustomizableObject/Private/MuCO/CustomizableObjectSystem.cpp
--- a/Engine/Plugins/Experimental/Mutable/Source/CustomizableObject/Private/MuCO/CustomizableObjectSystem.cpp
+++ b/Engine/Plugins/Experimental/Mutable/Source/CustomizableObject/Private/MuCO/CustomizableObjectSystem.cpp
@@ -2799,9 +2799,11 @@ namespace impl
 			{
 				EWorldType::Type WorldType = EWorldType::Type::None;
 
-				if (CustomizableObjectInstanceUsage->GetWorld())
+				USkeletalMeshComponent* Parent = Cast<USkeletalMeshComponent>(CustomizableObjectInstanceUsage->GetAttachParent());
+
+				if (Parent && Parent->GetWorld())
 				{
-					WorldType = CustomizableObjectInstanceUsage->GetWorld()->WorldType;
+					WorldType = Parent->GetWorld()->WorldType;
 				}
 
 				switch (WorldType)
diff --git a/Engine/Plugins/Experimental/Mutable/Source/CustomizableObject/Private/MuCO/LogBenchmarkUtil.cpp b/Engine/Plugins/Experimental/Mutable/Source/CustomizableObject/Private/MuCO/LogBenchmarkUtil.cpp
--- a/Engine/Plugins/Experimental/Mutable/Source/CustomizableObject/Private/MuCO/LogBenchmarkUtil.cpp
+++ b/Engine/Plugins/Experimental/Mutable/Source/CustomizableObject/Private/MuCO/LogBenchmarkUtil.cpp
@@ -205,13 +205,16 @@ void FLogBenchmarkUtil::UpdateStats()
 			continue;
 		}
 
-		AActor* ParentActor = CustomizableObjectInstanceUsage->GetAttachParent() ? 
-			CustomizableObjectInstanceUsage->GetAttachParent()->GetAttachmentRootActor()
+		USkeletalMeshComponent* Parent = Cast<USkeletalMeshComponent>(CustomizableObjectInstanceUsage->GetAttachParent());
+		AActor* ParentActor = Parent ? 
+			Parent->GetAttachmentRootActor()
 			:  nullptr;
 		UCustomizableObjectInstance* Instance = CustomizableObjectInstanceUsage->GetCustomizableObjectInstance();
 
 		APawn* PlayerPawn = nullptr;
-		if (UWorld* World = CustomizableObjectInstanceUsage->GetWorld())
+		UWorld* World = Parent ? Parent->GetWorld() : nullptr;
+
+		if (World)
 		{
 			PlayerPawn = UGameplayStatics::GetPlayerPawn(World, 0);
 		}
