[Mutable] Added CO option to force generating only the MaxLOD on initial generation.

- Fixed InstanceDescriptorRuntimeHash not being updated.

#rnx

diff --git a/Engine/Plugins/Experimental/Mutable/Source/CustomizableObject/Classes/MuCO/CustomizableObject.h b/Engine/Plugins/Experimental/Mutable/Source/CustomizableObject/Classes/MuCO/CustomizableObject.h
--- a/Engine/Plugins/Experimental/Mutable/Source/CustomizableObject/Classes/MuCO/CustomizableObject.h
+++ b/Engine/Plugins/Experimental/Mutable/Source/CustomizableObject/Classes/MuCO/CustomizableObject.h
@@ -1214,6 +1214,10 @@ private:
 	UPROPERTY(EditAnywhere, Category = CustomizableObject)
 	bool bEnableUseRefSkeletalMeshAsPlaceholder = true;
 
+	/** Use the Instance MinLOD, MaxLOD and RequestedLODs in the descriptor when performing the initial generation (ignore LOD Management). */
+	UPROPERTY(Category = "CustomizableObject", EditAnywhere, DisplayName = "Preserve User LODs On First Generation")
+	bool bPreserveUserLODsOnFirstGeneration = false;
+	
 public:
 #if WITH_EDITORONLY_DATA
 	// Hide this property because it is not used yet.
@@ -1600,6 +1604,9 @@ public:
 
 	TSoftObjectPtr<USkeleton> GetReferencedSkeletonAssetPtr( uint32 Index );
 
+	/** See bPreserveUserLODsOnFirstGeneration. */
+	bool IsPreserveUserLODsOnFirstGeneration() const;
+
 	/** Stores all the parameter UI metadata information for all the dependencies of this Customizable Object */
 	UPROPERTY()
 	TMap<FString, FParameterUIData> ParameterUIDataMap;
diff --git a/Engine/Plugins/Experimental/Mutable/Source/CustomizableObject/Classes/MuCO/CustomizableObjectInstance.h b/Engine/Plugins/Experimental/Mutable/Source/CustomizableObject/Classes/MuCO/CustomizableObjectInstance.h
--- a/Engine/Plugins/Experimental/Mutable/Source/CustomizableObject/Classes/MuCO/CustomizableObjectInstance.h
+++ b/Engine/Plugins/Experimental/Mutable/Source/CustomizableObject/Classes/MuCO/CustomizableObjectInstance.h
@@ -732,7 +732,13 @@ public:
 	/** Array of RequestedLODs per component to generate if this candidate is chosen */
 	TArray<uint16> RequestedLODLevels;
 
-	FMutableUpdateCandidate(UCustomizableObjectInstance* InCustomizableObjectInstance) : CustomizableObjectInstance(InCustomizableObjectInstance) {}
+	FMutableUpdateCandidate(UCustomizableObjectInstance* InCustomizableObjectInstance) : CustomizableObjectInstance(InCustomizableObjectInstance)
+	{
+		const FCustomizableObjectInstanceDescriptor& Descriptor = InCustomizableObjectInstance->GetDescriptor();
+		MinLOD = Descriptor.GetMinLod();
+		MaxLOD = Descriptor.GetMaxLod();
+		RequestedLODLevels = Descriptor.GetRequestedLODLevels();
+	}
 
 	FMutableUpdateCandidate(const UCustomizableObjectInstance* InCustomizableObjectInstance, const int32 InMinLOD, const int32 InMaxLOD,
 		const TArray<uint16>& InRequestedLODLevels) :
diff --git a/Engine/Plugins/Experimental/Mutable/Source/CustomizableObject/Classes/MuCO/CustomizableObjectSystem.h b/Engine/Plugins/Experimental/Mutable/Source/CustomizableObject/Classes/MuCO/CustomizableObjectSystem.h
--- a/Engine/Plugins/Experimental/Mutable/Source/CustomizableObject/Classes/MuCO/CustomizableObjectSystem.h
+++ b/Engine/Plugins/Experimental/Mutable/Source/CustomizableObject/Classes/MuCO/CustomizableObjectSystem.h
@@ -43,6 +43,7 @@ extern TAutoConsoleVariable<bool> CVarClearWorkingMemoryOnUpdateEnd;
 
 extern TAutoConsoleVariable<bool> CVarReuseImagesBetweenInstances;
 
+extern TAutoConsoleVariable<bool> CVarPreserveUserLODsOnFirstGeneration;
 
 #if WITH_EDITOR
 
diff --git a/Engine/Plugins/Experimental/Mutable/Source/CustomizableObject/Private/MuCO/CustomizableObject.cpp b/Engine/Plugins/Experimental/Mutable/Source/CustomizableObject/Private/MuCO/CustomizableObject.cpp
--- a/Engine/Plugins/Experimental/Mutable/Source/CustomizableObject/Private/MuCO/CustomizableObject.cpp
+++ b/Engine/Plugins/Experimental/Mutable/Source/CustomizableObject/Private/MuCO/CustomizableObject.cpp
@@ -1971,6 +1971,12 @@ TArray<TSharedPtr<IAsyncReadFileHandle>> UCustomizableObjectBulk::GetAsyncReadFi
 }
 
 
+bool UCustomizableObject::IsPreserveUserLODsOnFirstGeneration() const
+{
+	return bPreserveUserLODsOnFirstGeneration;
+}
+
+
 #if WITH_EDITOR
 
 void UCustomizableObjectBulk::CookAdditionalFilesOverride(const TCHAR* PackageFilename,
diff --git a/Engine/Plugins/Experimental/Mutable/Source/CustomizableObject/Private/MuCO/CustomizableObjectInstance.cpp b/Engine/Plugins/Experimental/Mutable/Source/CustomizableObject/Private/MuCO/CustomizableObjectInstance.cpp
--- a/Engine/Plugins/Experimental/Mutable/Source/CustomizableObject/Private/MuCO/CustomizableObjectInstance.cpp
+++ b/Engine/Plugins/Experimental/Mutable/Source/CustomizableObject/Private/MuCO/CustomizableObjectInstance.cpp
@@ -2835,6 +2835,9 @@ void FMutableUpdateCandidate::ApplyLODUpdateParamsToInstance(FUpdateContextPriva
 	if (Context)
 	{
 		Context->InstanceDescriptorRuntimeHash = CustomizableObjectInstance->GetPrivate()->UpdateDescriptorRuntimeHash;
+		Context->CurrentMinLOD = MinLOD;
+		Context->CurrentMaxLOD = MaxLOD;
+		Context->RequestedLODs = RequestedLODLevels;
 	}
 }
 
@@ -6155,6 +6158,13 @@ void UCustomizableObjectInstance::SetRequestedLODs(int32 InMinLOD, int32 InMaxLO
 	{
 		return;
 	}
+
+	if (CVarPreserveUserLODsOnFirstGeneration.GetValueOnGameThread() &&
+		GetCustomizableObject()->IsPreserveUserLODsOnFirstGeneration() &&
+		GetPrivate()->GetSkeletalMeshStatus() != ESkeletalMeshStatus::Success)
+	{
+		return;
+	}
 	
 #if !UE_BUILD_SHIPPING
 	// Ignore Min/Max LOD limits. Mainly used for debug
diff --git a/Engine/Plugins/Experimental/Mutable/Source/CustomizableObject/Private/MuCO/CustomizableObjectSystem.cpp b/Engine/Plugins/Experimental/Mutable/Source/CustomizableObject/Private/MuCO/CustomizableObjectSystem.cpp
--- a/Engine/Plugins/Experimental/Mutable/Source/CustomizableObject/Private/MuCO/CustomizableObjectSystem.cpp
+++ b/Engine/Plugins/Experimental/Mutable/Source/CustomizableObject/Private/MuCO/CustomizableObjectSystem.cpp
@@ -89,7 +89,15 @@ static TAutoConsoleVariable<int32> CVarGeneratedResourcesCacheSize(
 	TEXT("Limit the number of resources (images and meshes) that will be tracked for reusal. Each tracked resource uses a small amout of memory for its key."),
 	ECVF_Scalability);
 
+TAutoConsoleVariable<bool> CVarPreserveUserLODsOnFirstGeneration(
+	TEXT("mutable.PreserveUserLODsOnFirstGeneration"),
+	true,
+	TEXT("If false, force disable UCustomizableObject::bPreserveUserLODsOnFirstGeneration."),
+	ECVF_Scalability);
+
+
 int32 FCustomizableObjectSystemPrivate::SkeletalMeshMinLodQualityLevel = -1;
+
 static void CVarMutableSinkFunction()
 {
 	if (UCustomizableObjectSystem::IsCreated())
@@ -116,6 +124,9 @@ FUpdateContextPrivate::FUpdateContextPrivate(UCustomizableObjectInstance& InInst
 	bBuildParameterRelevancy = InInstance.GetBuildParameterRelevancy();
 	Parameters = InInstance.GetDescriptor().GetParameters();
 	TextureParameters = InInstance.GetPrivate()->UpdateTextureParameters;
+	CurrentMinLOD = InInstance.GetCurrentMinLOD();
+	CurrentMaxLOD = InInstance.GetCurrentMinLOD();
+	RequestedLODs = InInstance.GetRequestedLODsPerComponent();
 	
 	InInstance.GetCustomizableObject()->ApplyStateForcedValuesToParameters(State, Parameters.get());
 
@@ -1627,7 +1638,7 @@ namespace impl
 			OperationData->CurrentMaxLOD = OperationData->NumLODsAvailable - 1;
 		}
 
-		if(!OperationData->RequestedLODs.IsEmpty())
+		if (!OperationData->RequestedLODs.IsEmpty())
 		{
 			// Initialize RequestedLODs to zero if not set
 			const int32 ComponentCount = Instance->GetComponentCount(OperationData->CurrentMinLOD);
@@ -1652,6 +1663,8 @@ namespace impl
 				}
 			}
 		}
+		
+		OperationData->InstanceDescriptorRuntimeHash.UpdateRequestedLODs(OperationData->RequestedLODs);
 
 		// Map SharedSurfaceId to surface index
 		TArray<int32> SurfacesSharedId;
@@ -2818,10 +2831,12 @@ namespace impl
 		}
 #endif // WITH_EDITOR
 		
-		if (System->IsOnlyGenerateRequestedLODsEnabled() && System->CurrentInstanceLODManagement->IsOnlyGenerateRequestedLODLevelsEnabled() && 
-			!bIsInEditorViewport)
+		if (!System->IsOnlyGenerateRequestedLODsEnabled() ||
+			!System->CurrentInstanceLODManagement->IsOnlyGenerateRequestedLODLevelsEnabled() ||
+			bIsInEditorViewport)
 		{
-			Operation->RequestedLODs = Operation->InstanceDescriptorRuntimeHash.GetRequestedLODs();
+			Operation->RequestedLODs = {};
+			Operation->InstanceDescriptorRuntimeHash.UpdateRequestedLODs({});
 		}
 
 #ifdef MUTABLE_USE_NEW_TASKGRAPH
