[Mutable] Fix crash when comparing different FDescriptorRuntimeHash.

#rnx

diff --git a/Engine/Plugins/Experimental/Mutable/Source/CustomizableObject/Private/MuCO/CustomizableObjectInstanceDescriptor.cpp b/Engine/Plugins/Experimental/Mutable/Source/CustomizableObject/Private/MuCO/CustomizableObjectInstanceDescriptor.cpp
--- a/Engine/Plugins/Experimental/Mutable/Source/CustomizableObject/Private/MuCO/CustomizableObjectInstanceDescriptor.cpp
+++ b/Engine/Plugins/Experimental/Mutable/Source/CustomizableObject/Private/MuCO/CustomizableObjectInstanceDescriptor.cpp
@@ -2560,6 +2560,11 @@ bool FDescriptorRuntimeHash::IsSubset(const FDescriptorRuntimeHash& Other) const
 		return true;
 	}
 
+	if (RequestedLODsPerComponent.Num() != Other.RequestedLODsPerComponent.Num())
+	{
+		return false;
+	}
+	
 	for (int32 ComponentIndex = 0; ComponentIndex < RequestedLODsPerComponent.Num(); ++ComponentIndex)
 	{
 		int32 RequestedLODs = RequestedLODsPerComponent[ComponentIndex];
