[mutable] Now, bCastShadow value is not ignored.

#rnx

diff --git a/Engine/Plugins/Experimental/Mutable/Source/CustomizableObject/Classes/MuCO/CustomizableObject.h b/Engine/Plugins/Experimental/Mutable/Source/CustomizableObject/Classes/MuCO/CustomizableObject.h
--- a/Engine/Plugins/Experimental/Mutable/Source/CustomizableObject/Classes/MuCO/CustomizableObject.h
+++ b/Engine/Plugins/Experimental/Mutable/Source/CustomizableObject/Classes/MuCO/CustomizableObject.h
@@ -1481,7 +1481,7 @@ private:
 	// This is a manual version number for the binary blobs in this asset.
 	// Increasing it invalidates all the previously compiled models.
 	// Warning: If while merging code both versions have changed, take the highest+1.
-	static const int32 CurrentSupportedVersion = 409;
+	static const int32 CurrentSupportedVersion = 410;
 
 public:
 
diff --git a/Engine/Plugins/Experimental/Mutable/Source/CustomizableObject/Private/MuCO/CustomizableObjectInstance.cpp b/Engine/Plugins/Experimental/Mutable/Source/CustomizableObject/Private/MuCO/CustomizableObjectInstance.cpp
--- a/Engine/Plugins/Experimental/Mutable/Source/CustomizableObject/Private/MuCO/CustomizableObjectInstance.cpp
+++ b/Engine/Plugins/Experimental/Mutable/Source/CustomizableObject/Private/MuCO/CustomizableObjectInstance.cpp
@@ -3625,7 +3625,7 @@ void UCustomizableInstancePrivateData::BuildOrCopyMorphTargetsData(const TShared
 				}
 
 				int FirstVertex, VerticesCount, FirstIndex, IndiciesCount;
-				Component.Mesh->GetSurface(Section, &FirstVertex, &VerticesCount, &FirstIndex, &IndiciesCount, nullptr, nullptr);
+				Component.Mesh->GetSurface(Section, &FirstVertex, &VerticesCount, &FirstIndex, &IndiciesCount, nullptr, nullptr, nullptr);
 
 				for (int32 VertexIdx = FirstVertex; VertexIdx < FirstVertex + VerticesCount; ++VertexIdx)
 				{
@@ -3789,7 +3789,7 @@ void UCustomizableInstancePrivateData::BuildOrCopyClothingData(const TSharedRef<
 				for (int32 Section = 0; Section < SurfaceCount; ++Section)
 				{
 					int FirstVertex, VerticesCount, FirstIndex, IndicesCount;
-					MutableMesh->GetSurface(Section, &FirstVertex, &VerticesCount, &FirstIndex, &IndicesCount, nullptr, nullptr);
+					MutableMesh->GetSurface(Section, &FirstVertex, &VerticesCount, &FirstIndex, &IndicesCount, nullptr, nullptr, nullptr);
 
 					if (VerticesCount == 0 || IndicesCount == 0)
 					{
@@ -4421,7 +4421,7 @@ void UCustomizableInstancePrivateData::BuildOrCopyClothingData(const TSharedRef<
 				{
 					// Check that is a valid surface.
 					int32 FirstVertex, VerticesCount, FirstIndex, IndicesCount;
-					MutableMesh->GetSurface(Section, &FirstVertex, &VerticesCount, &FirstIndex, &IndicesCount, nullptr, nullptr);
+					MutableMesh->GetSurface(Section, &FirstVertex, &VerticesCount, &FirstIndex, &IndicesCount, nullptr, nullptr, nullptr);
 
 					if (VerticesCount == 0 || IndicesCount == 0 )
 					{
@@ -6806,6 +6806,7 @@ void UCustomizableInstancePrivateData::RegenerateImportedModel(USkeletalMesh* Sk
 			SectionUserData.CorrespondClothAssetIndex = RenderSection.CorrespondClothAssetIndex;
 			SectionUserData.ClothingData.AssetGuid = RenderSection.ClothingData.AssetGuid;
 			SectionUserData.ClothingData.AssetLodIndex = RenderSection.ClothingData.AssetLodIndex;
+			SectionUserData.bCastShadow = RenderSection.bCastShadow;
 		}
 
 		ImportedModel->LODModels[LODIndex].SyncronizeUserSectionsDataArray();
diff --git a/Engine/Plugins/Experimental/Mutable/Source/CustomizableObject/Private/MuCO/UnrealConversionUtils.cpp b/Engine/Plugins/Experimental/Mutable/Source/CustomizableObject/Private/MuCO/UnrealConversionUtils.cpp
--- a/Engine/Plugins/Experimental/Mutable/Source/CustomizableObject/Private/MuCO/UnrealConversionUtils.cpp
+++ b/Engine/Plugins/Experimental/Mutable/Source/CustomizableObject/Private/MuCO/UnrealConversionUtils.cpp
@@ -156,7 +156,8 @@ namespace UnrealConversionUtils
 			int32 VertexCount;
 			int32 FirstBone;
 			int32 BoneCount;
-			InMutableMesh->GetSurface(SurfaceIndex, &FirstVertex, &VertexCount, &FirstIndex, &IndexCount, &FirstBone, &BoneCount);
+			bool bCastShadow;
+			InMutableMesh->GetSurface(SurfaceIndex, &FirstVertex, &VertexCount, &FirstIndex, &IndexCount, &FirstBone, &BoneCount, &bCastShadow);
 			FSkelMeshRenderSection& Section = OutSkeletalMesh->GetResourceForRendering()->LODRenderData[MeshLODIndex].RenderSections[SurfaceIndex];
 
 			Section.DuplicatedVerticesBuffer.Init(1, TMap<int, TArray<int32>>());
@@ -172,6 +173,7 @@ namespace UnrealConversionUtils
 			Section.BaseVertexIndex = FirstVertex;
 			Section.MaxBoneInfluences = NumBoneInfluences;
 			Section.NumVertices = VertexCount;
+			Section.bCastShadow = bCastShadow;
 
 			// InBoneMaps may contain bonemaps from other sections. Copy the bones belonging to this mesh.
 			FirstBone += InFirstBoneMapIndex;
@@ -473,6 +475,7 @@ namespace UnrealConversionUtils
 					 DestSection->MaxBoneInfluences = SrcSection.MaxBoneInfluences;
 					 DestSection->NumVertices = SrcSection.NumVertices;
 					 DestSection->BoneMap = SrcSection.BoneMap;
+					 DestSection->bCastShadow = SrcSection.bCastShadow;
 				 }
 			 }
 		 }
diff --git a/Engine/Plugins/Experimental/Mutable/Source/CustomizableObjectEditor/Private/MuCOE/GenerateMutableSource/GenerateMutableSourceMesh.cpp b/Engine/Plugins/Experimental/Mutable/Source/CustomizableObjectEditor/Private/MuCOE/GenerateMutableSource/GenerateMutableSourceMesh.cpp
--- a/Engine/Plugins/Experimental/Mutable/Source/CustomizableObjectEditor/Private/MuCOE/GenerateMutableSource/GenerateMutableSourceMesh.cpp
+++ b/Engine/Plugins/Experimental/Mutable/Source/CustomizableObjectEditor/Private/MuCOE/GenerateMutableSource/GenerateMutableSourceMesh.cpp
@@ -1513,6 +1513,14 @@ mu::MeshPtr ConvertSkeletalMeshToMutable(const USkeletalMesh* InSkeletalMesh, co
 		}
 	}
 
+	// Ensure Surface Data
+	mu::MESH_SURFACE MeshSurface;
+	MeshSurface.m_vertexCount = MutableMesh->m_VertexBuffers.GetElementCount();
+	MeshSurface.m_indexCount = MutableMesh->m_IndexBuffers.GetElementCount();
+	MeshSurface.BoneMapCount = MutableMesh->BoneMap.Num();
+	MeshSurface.bCastShadow = MeshSection.bCastShadow;
+	MutableMesh->m_surfaces.Add(MeshSurface);
+
 	if (!bIgnorePhysics && InSkeletalMesh->GetPhysicsAsset() && MutableMesh->GetSkeleton() && GenerationContext.Options.bPhysicsAssetMergeEnabled)
 	{
 		// Find BodySetups with relevant bones.
diff --git a/Engine/Plugins/Experimental/Mutable/Source/MutableRuntime/Private/MuR/Mesh.cpp b/Engine/Plugins/Experimental/Mutable/Source/MutableRuntime/Private/MuR/Mesh.cpp
--- a/Engine/Plugins/Experimental/Mutable/Source/MutableRuntime/Private/MuR/Mesh.cpp
+++ b/Engine/Plugins/Experimental/Mutable/Source/MutableRuntime/Private/MuR/Mesh.cpp
@@ -385,7 +385,8 @@ int Mesh::GetSurfaceCount() const
 void Mesh::GetSurface( int32 surfaceIndex,
                        int32* firstVertex, int32* vertexCount,
                        int32* firstIndex, int32* indexCount,
-					   int32* BoneIndex, int32* BoneCount) const
+					   int32* BoneIndex, int32* BoneCount,
+					   bool* bCastShadow) const
 {
     int count = GetSurfaceCount();
 
@@ -400,6 +401,7 @@ void Mesh::GetSurface( int32 surfaceIndex,
             if (indexCount) *indexCount = surf.m_indexCount;
             if (BoneIndex) *BoneIndex = surf.BoneMapIndex;
             if (BoneCount) *BoneCount = surf.BoneMapCount;
+            if (bCastShadow) *bCastShadow = surf.bCastShadow;
         }
         else
         {
@@ -410,6 +412,7 @@ void Mesh::GetSurface( int32 surfaceIndex,
             if (indexCount) *indexCount = GetIndexCount();
 			if (BoneIndex) *BoneIndex = 0;
 			if (BoneCount) *BoneCount = BoneMap.Num();
+			if (bCastShadow) *bCastShadow = false;
         }
     }
     else
@@ -421,6 +424,7 @@ void Mesh::GetSurface( int32 surfaceIndex,
         if (indexCount) *indexCount = 0;
 		if (BoneIndex) *BoneIndex = 0;
 		if (BoneCount) *BoneCount = 0;
+		if (bCastShadow) *bCastShadow = false;
     }
 }
 
@@ -1089,8 +1093,7 @@ void UnserialiseLegacySurfaces(InputArchive& arch, TArray<MESH_SURFACE>& OutMesh
 //-------------------------------------------------------------------------------------------------
 void MESH_SURFACE::Serialise(OutputArchive& arch) const
 {
-
-	const int32 ver = 0;
+	const int32 ver = 1;
 	arch << ver;
 
 	arch << m_firstVertex;
@@ -1099,9 +1102,9 @@ void MESH_SURFACE::Serialise(OutputArchive& arch) const
 	arch << m_indexCount;
 	arch << BoneMapIndex;
 	arch << BoneMapCount;
-	
-	arch << m_id;
+	arch << bCastShadow;
 
+	arch << m_id;
 }
 
 
@@ -1110,7 +1113,7 @@ void MESH_SURFACE::Unserialise(InputArchive& arch)
 {
 	int32 ver = 0;
 	arch >> ver;
-	check(ver == 0);
+	check(ver <= 1);
 
 	arch >> m_firstVertex;
 	arch >> m_vertexCount;
@@ -1118,7 +1121,12 @@ void MESH_SURFACE::Unserialise(InputArchive& arch)
 	arch >> m_indexCount;
 	arch >> BoneMapIndex;
 	arch >> BoneMapCount;
-		 
+
+	if (ver >= 1)
+	{
+		arch >> bCastShadow;
+	}
+
 	arch >> m_id;
 }
 
diff --git a/Engine/Plugins/Experimental/Mutable/Source/MutableRuntime/Private/MuR/OpMeshFormat.cpp b/Engine/Plugins/Experimental/Mutable/Source/MutableRuntime/Private/MuR/OpMeshFormat.cpp
--- a/Engine/Plugins/Experimental/Mutable/Source/MutableRuntime/Private/MuR/OpMeshFormat.cpp
+++ b/Engine/Plugins/Experimental/Mutable/Source/MutableRuntime/Private/MuR/OpMeshFormat.cpp
@@ -692,6 +692,8 @@ namespace mu
 		// A shallow copy is done here, it should not be a problem.
 		Result->AdditionalPhysicsBodies = pSource->AdditionalPhysicsBodies;
 
+		Result->m_surfaces = pSource->m_surfaces;
+
 		Result->ResetStaticFormatFlags();
 		Result->EnsureSurfaceData();
 	}
diff --git a/Engine/Plugins/Experimental/Mutable/Source/MutableRuntime/Private/MuR/OpMeshMerge.h b/Engine/Plugins/Experimental/Mutable/Source/MutableRuntime/Private/MuR/OpMeshMerge.h
--- a/Engine/Plugins/Experimental/Mutable/Source/MutableRuntime/Private/MuR/OpMeshMerge.h
+++ b/Engine/Plugins/Experimental/Mutable/Source/MutableRuntime/Private/MuR/OpMeshMerge.h
@@ -527,6 +527,9 @@ namespace mu
 				NewSurface.m_vertexCount = pFirst->GetVertexCount() + pSecond->GetVertexCount();
 				NewSurface.m_indexCount = pFirst->GetIndexCount() + pSecond->GetIndexCount();
 				NewSurface.BoneMapCount = Result->BoneMap.Num();
+				
+				//All merged surfaces will have the same bCastShadow value. Decided by the first merged mesh
+				NewSurface.bCastShadow = pFirst->m_surfaces.Last().bCastShadow;
 			}
 			else
 			{
diff --git a/Engine/Plugins/Experimental/Mutable/Source/MutableRuntime/Public/MuR/Mesh.h b/Engine/Plugins/Experimental/Mutable/Source/MutableRuntime/Public/MuR/Mesh.h
--- a/Engine/Plugins/Experimental/Mutable/Source/MutableRuntime/Public/MuR/Mesh.h
+++ b/Engine/Plugins/Experimental/Mutable/Source/MutableRuntime/Public/MuR/Mesh.h
@@ -53,6 +53,8 @@ namespace mu
 
 			BoneMapIndex = 0;
 			BoneMapCount = 0;
+
+			bCastShadow = false;
 		}
 
 		int32 m_firstVertex;
@@ -64,6 +66,8 @@ namespace mu
 		uint32 BoneMapIndex;
 		uint32 BoneMapCount;
 
+		bool bCastShadow;
+
 		//!
 		inline bool operator==(const MESH_SURFACE& o) const
 		{
@@ -73,7 +77,8 @@ namespace mu
 				&& m_indexCount == o.m_indexCount
 				&& m_id == o.m_id
 				&& BoneMapIndex == o.BoneMapIndex
-				&& BoneMapCount == o.BoneMapCount;
+				&& BoneMapCount == o.BoneMapCount
+				&& bCastShadow == o.bCastShadow;
 		}
 
 		inline void Serialise(OutputArchive& arch) const;
@@ -217,7 +222,8 @@ namespace mu
         void GetSurface( int32 surfaceIndex,
                          int32* FirstVertex, int32* VertexCount,
                          int32* FirstIndex, int32* IndexCount,
-						 int32* FirstBone, int32* BoneCount) const;
+						 int32* FirstBone, int32* BoneCount,
+						 bool* bCastShadow) const;
 
         //! Return an internal id that can be used to match mesh surfaces and instance surfaces.
         //! Only valid for meshes that are part of instances.
