[Mutable] Mesh Cache. Enable the reuse of USkeletalMeshes between CO.

Minimum changes to enable the reuse of USkeletalMeshes. CVar toggleable.

#rnx

diff --git a/Engine/Plugins/Experimental/Mutable/Source/CustomizableObject/Classes/MuCO/CustomizableObject.h b/Engine/Plugins/Experimental/Mutable/Source/CustomizableObject/Classes/MuCO/CustomizableObject.h
--- a/Engine/Plugins/Experimental/Mutable/Source/CustomizableObject/Classes/MuCO/CustomizableObject.h
+++ b/Engine/Plugins/Experimental/Mutable/Source/CustomizableObject/Classes/MuCO/CustomizableObject.h
@@ -1209,7 +1209,6 @@ public:
 
 private:
 	/** Use the SkeletalMesh of reference as a placeholder until the custom mesh is ready to use.
-	  * 
 	  * Note: If disabled, a null mesh will be used to replace the discarded mesh due to 'ReplaceDiscardedWithReferenceMesh' being enabled. */
 	UPROPERTY(EditAnywhere, Category = CustomizableObject)
 	bool bEnableUseRefSkeletalMeshAsPlaceholder = true;
@@ -1217,6 +1216,11 @@ private:
 	/** Use the Instance MinLOD, MaxLOD and RequestedLODs in the descriptor when performing the initial generation (ignore LOD Management). */
 	UPROPERTY(Category = "CustomizableObject", EditAnywhere, DisplayName = "Preserve User LODs On First Generation")
 	bool bPreserveUserLODsOnFirstGeneration = false;
+
+	/** If true, reuse previously generated USkeletalMesh (if still valid and the the number of LOD have not changed)
+	 * USkeletalMeshes are only reused between the same CO. */
+	UPROPERTY(EditAnywhere, Category = CustomizableObject)
+	bool bEnableMeshCache = false;
 	
 public:
 #if WITH_EDITORONLY_DATA
@@ -1721,6 +1725,9 @@ public:
 	/** See bEnableUseRefSkeletalMeshAsPlaceholder. */
 	bool IsEnableUseRefSkeletalMeshAsPlaceholder() const;
 
+	/** See bEnableUseMeshCache. */
+	bool IsMeshCacheEnabled() const;
+
 private:
 	
 	/** Stores the bone names of all the bones that can possibly use the generated meshes */
diff --git a/Engine/Plugins/Experimental/Mutable/Source/CustomizableObject/Classes/MuCO/CustomizableObjectInstance.h b/Engine/Plugins/Experimental/Mutable/Source/CustomizableObject/Classes/MuCO/CustomizableObjectInstance.h
--- a/Engine/Plugins/Experimental/Mutable/Source/CustomizableObject/Classes/MuCO/CustomizableObjectInstance.h
+++ b/Engine/Plugins/Experimental/Mutable/Source/CustomizableObject/Classes/MuCO/CustomizableObjectInstance.h
@@ -19,6 +19,7 @@ class UCustomizableObject;
 class UTexture2D;
 class FUpdateContextPrivate;
 class UCustomizableObjectInstanceUsage;
+class UMaterialInterface;
 struct FFrame;
 struct FGameplayTagContainer;
 struct FPropertyChangedEvent;
@@ -593,6 +594,8 @@ public:
 
 	void AdditionalAssetsAsyncLoaded( FGraphEventRef CompletionEvent );
 
+	const TArray<TObjectPtr<UMaterialInterface>>* GetOverrideMaterials(int32 ComponentIndex) const;
+	
 	// The following methods should only be used in an LOD management class
 	void SetIsBeingUsedByComponentInPlay(bool bIsUsedByComponent );
 	bool GetIsBeingUsedByComponentInPlay() const;
diff --git a/Engine/Plugins/Experimental/Mutable/Source/CustomizableObject/Classes/MuCO/CustomizableObjectSystem.h b/Engine/Plugins/Experimental/Mutable/Source/CustomizableObject/Classes/MuCO/CustomizableObjectSystem.h
--- a/Engine/Plugins/Experimental/Mutable/Source/CustomizableObject/Classes/MuCO/CustomizableObjectSystem.h
+++ b/Engine/Plugins/Experimental/Mutable/Source/CustomizableObject/Classes/MuCO/CustomizableObjectSystem.h
@@ -45,6 +45,8 @@ extern TAutoConsoleVariable<bool> CVarReuseImagesBetweenInstances;
 
 extern TAutoConsoleVariable<bool> CVarPreserveUserLODsOnFirstGeneration;
 
+extern TAutoConsoleVariable<bool> CVarEnableMeshCache;
+
 #if WITH_EDITOR
 
 // Struct used to keep a copy of the EditorSettings needed to compile Customizable Objects.
diff --git a/Engine/Plugins/Experimental/Mutable/Source/CustomizableObject/Private/MuCO/CustomizableInstancePrivateData.h b/Engine/Plugins/Experimental/Mutable/Source/CustomizableObject/Private/MuCO/CustomizableInstancePrivateData.h
--- a/Engine/Plugins/Experimental/Mutable/Source/CustomizableObject/Private/MuCO/CustomizableInstancePrivateData.h
+++ b/Engine/Plugins/Experimental/Mutable/Source/CustomizableObject/Private/MuCO/CustomizableInstancePrivateData.h
@@ -159,6 +159,9 @@ struct FCustomizableInstanceComponentData
 	 *  Size == NumLODsAvailable
 	 *  LODs without mesh will be set to the maximum value of FResourceID (Max_uint64). */
 	TArray<mu::FResourceID> LastMeshIdPerLOD;
+
+	UPROPERTY(Transient)
+	TArray<TObjectPtr<UMaterialInterface>> OverrideMaterials;
 };
 
 USTRUCT()
diff --git a/Engine/Plugins/Experimental/Mutable/Source/CustomizableObject/Private/MuCO/CustomizableObject.cpp b/Engine/Plugins/Experimental/Mutable/Source/CustomizableObject/Private/MuCO/CustomizableObject.cpp
--- a/Engine/Plugins/Experimental/Mutable/Source/CustomizableObject/Private/MuCO/CustomizableObject.cpp
+++ b/Engine/Plugins/Experimental/Mutable/Source/CustomizableObject/Private/MuCO/CustomizableObject.cpp
@@ -1866,6 +1866,12 @@ bool UCustomizableObject::IsEnableUseRefSkeletalMeshAsPlaceholder() const
 }
 
 
+bool UCustomizableObject::IsMeshCacheEnabled() const
+{
+	return bEnableMeshCache;
+}
+
+
 FGuid UCustomizableObject::GetCompilationGuid() const
 {
 	return CompilationGuid;
@@ -1876,6 +1882,33 @@ FGuid UCustomizableObject::GetCompilationGuid() const
 //-------------------------------------------------------------------------------------------------
 //-------------------------------------------------------------------------------------------------
 
+USkeletalMesh* FMeshCache::Get(const TArray<mu::FResourceID>& Key)
+{
+	const TWeakObjectPtr<USkeletalMesh>* Result = GeneratedMeshes.Find(Key);
+	return Result ? Result->Get() : nullptr;
+}
+
+
+void FMeshCache::Add(const TArray<mu::FResourceID>& Key, USkeletalMesh* Value)
+{
+	if (!Value)
+	{
+		return;
+	}
+	
+	GeneratedMeshes.Add(Key, Value);
+
+	// Remove invalid SkeletalMeshes from the cache.
+	for (auto MeshIterator = GeneratedMeshes.CreateIterator(); MeshIterator; ++MeshIterator)
+	{
+		if (MeshIterator.Value().IsStale())
+		{
+			MeshIterator.RemoveCurrent();
+		}
+	}	
+}
+
+
 void FCustomizableObjectPrivateData::SetModel(const TSharedPtr<mu::Model, ESPMode::ThreadSafe>& Model, const FGuid Id)
 {
 #if WITH_EDITOR
diff --git a/Engine/Plugins/Experimental/Mutable/Source/CustomizableObject/Private/MuCO/CustomizableObjectInstance.cpp b/Engine/Plugins/Experimental/Mutable/Source/CustomizableObject/Private/MuCO/CustomizableObjectInstance.cpp
--- a/Engine/Plugins/Experimental/Mutable/Source/CustomizableObject/Private/MuCO/CustomizableObjectInstance.cpp
+++ b/Engine/Plugins/Experimental/Mutable/Source/CustomizableObject/Private/MuCO/CustomizableObjectInstance.cpp
@@ -1614,6 +1614,16 @@ bool UCustomizableInstancePrivateData::DoComponentsNeedUpdate(UCustomizableObjec
 
 		for (int32 ComponentIndex = 0; ComponentIndex < NumComponents; ++ComponentIndex)
 		{
+			if (OperationData->bUseMeshCache)
+			{
+				if (CustomizableObject->GetPrivate()->MeshCache.Get(OperationData->MeshDescriptors[ComponentIndex]))
+				{
+					OutComponentNeedsUpdate[ComponentIndex] = true;
+					ComponentWithMesh[ComponentIndex] = true;
+					continue;					
+				}
+			}
+			
 			// Components with mesh must have valid geometry at CurrentMaxLOD
 			if (ComponentWithMesh[ComponentIndex] && MeshIDs[ComponentIndex * MAX_MESH_LOD_COUNT + OperationData->CurrentMaxLOD] == MAX_uint64)
 			{
@@ -1808,6 +1818,21 @@ bool UCustomizableInstancePrivateData::UpdateSkeletalMesh_PostBeginUpdate0(UCust
 			continue;
 		}
 
+		// Check if we have initialized the component
+		if (Public->SkeletalMeshes[Component.Id])
+		{
+			continue;
+		}
+
+		if (OperationData->bUseMeshCache)
+		{
+			if (USkeletalMesh* CachedMesh = CustomizableObject->GetPrivate()->MeshCache.Get(OperationData->MeshDescriptors[Component.Id]))
+			{
+				Public->SkeletalMeshes[Component.Id] = CachedMesh;
+				continue;
+			}
+		}
+
 		if (!Component.bGenerated || !Component.Mesh || Component.SurfaceCount == 0)
 		{
 			continue;
@@ -1820,12 +1845,6 @@ bool UCustomizableInstancePrivateData::UpdateSkeletalMesh_PostBeginUpdate0(UCust
 			break;
 		}
 
-		// Check if we have initialized the component
-		if (Public->SkeletalMeshes[Component.Id])
-		{
-			continue;
-		}
-
 		// Create and initialize the SkeletalMesh for this component
 		MUTABLE_CPUPROFILER_SCOPE(ConstructMesh);
 
@@ -2011,6 +2030,15 @@ bool UCustomizableInstancePrivateData::UpdateSkeletalMesh_PostBeginUpdate0(UCust
 		const int32 NumComponents = Public->SkeletalMeshes.Num();
 		for (int32 ComponentIndex = 0; bSuccess && ComponentIndex < NumComponents; ++ComponentIndex)
 		{
+			if (OperationData->bUseMeshCache)
+			{
+				const TArray<mu::FResourceID>& MeshId = OperationData->MeshDescriptors[ComponentIndex];
+				if (CustomizableObject->GetPrivate()->MeshCache.Get(MeshId))
+				{
+					continue;
+				}
+			}
+			
 			if (!ComponentNeedsUpdate[ComponentIndex])
 			{
 				continue;
@@ -2039,6 +2067,12 @@ bool UCustomizableInstancePrivateData::UpdateSkeletalMesh_PostBeginUpdate0(UCust
 				ensure(SkeletalMesh->GetResourceForRendering()->LODRenderData.Num() > 0);
 				ensure(SkeletalMesh->GetLODInfoArray().Num() > 0);
 			}
+
+			if (OperationData->bUseMeshCache)
+			{
+				const TArray<mu::FResourceID>& MeshId = OperationData->MeshDescriptors[ComponentIndex];
+				CustomizableObject->GetPrivate()->MeshCache.Add(MeshId, SkeletalMesh);
+			}
 		}
 	}
 
@@ -5025,6 +5059,13 @@ void UCustomizableObjectInstance::AdditionalAssetsAsyncLoaded( FGraphEventRef Co
 }
 
 
+const TArray<TObjectPtr<UMaterialInterface>>* UCustomizableObjectInstance::GetOverrideMaterials(int32 ComponentIndex) const
+{
+	FCustomizableInstanceComponentData* ComponentData = PrivateData->GetComponentData(ComponentIndex);
+	return ComponentData ? &ComponentData->OverrideMaterials : nullptr;
+}
+
+
 void UCustomizableInstancePrivateData::AdditionalAssetsAsyncLoaded(UCustomizableObjectInstance* Public)
 {
 	MUTABLE_CPUPROFILER_SCOPE(UCustomizableInstancePrivateData::AdditionalAssetsAsyncLoaded);
@@ -5326,6 +5367,8 @@ void UCustomizableInstancePrivateData::BuildMaterials(const TSharedRef<FUpdateCo
 			continue;
 		}
 
+		ComponentsData[ComponentIndex].OverrideMaterials.Reset();
+
 		{
 			// TEMP: Keep a reference to the previous materials for n frames to avoid GC of materials in use in the render thread.
 			// TODO: MTBL-1632 - Implement a proper fix to replace this hotfix
@@ -5856,6 +5899,7 @@ void UCustomizableInstancePrivateData::BuildMaterials(const TSharedRef<FUpdateCo
 					}
 
 					GeneratedMaterials.Add(Material);
+					ComponentsData[ComponentIndex].OverrideMaterials.Add(Material.MaterialInterface);
 				}
 
 				int32 LODMaterialIndex = SkeletalMesh->GetLODInfoArray()[LODIndex].LODMaterialMap.Add(MatIndex);
diff --git a/Engine/Plugins/Experimental/Mutable/Source/CustomizableObject/Private/MuCO/CustomizableObjectInstanceUsage.cpp b/Engine/Plugins/Experimental/Mutable/Source/CustomizableObject/Private/MuCO/CustomizableObjectInstanceUsage.cpp
--- a/Engine/Plugins/Experimental/Mutable/Source/CustomizableObject/Private/MuCO/CustomizableObjectInstanceUsage.cpp
+++ b/Engine/Plugins/Experimental/Mutable/Source/CustomizableObject/Private/MuCO/CustomizableObjectInstanceUsage.cpp
@@ -235,10 +235,26 @@ void UCustomizableObjectInstanceUsage::SetSkeletalMesh(USkeletalMesh* SkeletalMe
 #endif
 		}
 
-		if (Parent->GetNumOverrideMaterials() > 0)
+		const UCustomizableObjectInstance* Instance = GetCustomizableObjectInstance();
+		const UCustomizableObject* CustomizableObject = Instance ? Instance->GetCustomizableObject() : nullptr;
+
+		if (Parent->HasOverrideMaterials())
 		{
 			// For some reason the reference skeletal mesh materials are added as override materials, clear them if necessary
 			Parent->EmptyOverrideMaterials();
+		}
+		
+		if (CustomizableObject &&
+			CustomizableObject->IsMeshCacheEnabled() &&
+			CVarEnableMeshCache.GetValueOnAnyThread())
+		{
+			if (FCustomizableInstanceComponentData* ComponentData = Instance->GetPrivate()->GetComponentData(GetComponentIndex()))
+			{
+				for (int32 Index = 0; Index < ComponentData->OverrideMaterials.Num(); ++Index)
+				{
+					Parent->SetMaterial(Index, ComponentData->OverrideMaterials[Index]);
+				}	
+			}
 		}	
 	}
 }
@@ -520,7 +536,7 @@ void UCustomizableObjectInstanceUsage::Tick(float DeltaTime)
 		{
 			Parent->SetSkeletalMesh(SkeletalMesh);
 
-			if (Parent->OverrideMaterials.Num() > 0)
+			if (Parent->HasOverrideMaterials())
 			{
 #if !(UE_BUILD_SHIPPING || UE_BUILD_TEST)
 				if (Parent->GetClass()->GetFName() != FName(TEXT("SkeletalMeshComponentBudgeted"))) // Reduce unnecessary logging
@@ -532,6 +548,20 @@ void UCustomizableObjectInstanceUsage::Tick(float DeltaTime)
 				Parent->EmptyOverrideMaterials();
 			}
 
+			if (CustomizableObject &&
+				bInstanceGenerated &&
+				CustomizableObject->IsMeshCacheEnabled() &&
+				CVarEnableMeshCache.GetValueOnAnyThread())
+			{
+				if (FCustomizableInstanceComponentData* ComponentData = CustomizableObjectInstance->GetPrivate()->GetComponentData(GetComponentIndex()))
+				{
+					for (int32 Index = 0; Index < ComponentData->OverrideMaterials.Num(); ++Index)
+					{
+						Parent->SetMaterial(Index, ComponentData->OverrideMaterials[Index]);
+					}
+				}
+			}
+
 			SetPendingSetSkeletalMesh(false);
 		}
 	}
diff --git a/Engine/Plugins/Experimental/Mutable/Source/CustomizableObject/Private/MuCO/CustomizableObjectPrivate.h b/Engine/Plugins/Experimental/Mutable/Source/CustomizableObject/Private/MuCO/CustomizableObjectPrivate.h
--- a/Engine/Plugins/Experimental/Mutable/Source/CustomizableObject/Private/MuCO/CustomizableObjectPrivate.h
+++ b/Engine/Plugins/Experimental/Mutable/Source/CustomizableObject/Private/MuCO/CustomizableObjectPrivate.h
@@ -3,6 +3,8 @@
 #pragma once
 
 #include "Templates/SharedPointer.h"
+#include "UObject/WeakObjectPtr.h"
+#include "MuR/Types.h"
 
 #if WITH_EDITOR
 #include "Misc/Guid.h"
@@ -10,6 +12,19 @@
 
 namespace mu { class Model; }
 class UCustomizableObject;
+class USkeletalMesh;
+
+
+class FMeshCache
+{
+public:
+	USkeletalMesh* Get(const TArray<mu::FResourceID>& Key);
+
+	void Add(const TArray<mu::FResourceID>& Key, USkeletalMesh* Value);
+
+private:
+	TMap<TArray<mu::FResourceID>, TWeakObjectPtr<USkeletalMesh>> GeneratedMeshes;
+};
 
 
 class FCustomizableObjectPrivateData
@@ -35,5 +50,8 @@ public:
 	bool bModelCompiledForCook = false;
 	TArray<FString> CachedPlatformNames;
 #endif
+
+	/** Cache of generated SkeletalMeshes */
+	FMeshCache MeshCache;
 };
 
diff --git a/Engine/Plugins/Experimental/Mutable/Source/CustomizableObject/Private/MuCO/CustomizableObjectSystem.cpp b/Engine/Plugins/Experimental/Mutable/Source/CustomizableObject/Private/MuCO/CustomizableObjectSystem.cpp
--- a/Engine/Plugins/Experimental/Mutable/Source/CustomizableObject/Private/MuCO/CustomizableObjectSystem.cpp
+++ b/Engine/Plugins/Experimental/Mutable/Source/CustomizableObject/Private/MuCO/CustomizableObjectSystem.cpp
@@ -96,6 +96,13 @@ TAutoConsoleVariable<bool> CVarPreserveUserLODsOnFirstGeneration(
 	ECVF_Scalability);
 
 
+TAutoConsoleVariable<bool> CVarEnableMeshCache(
+	TEXT("mutable.EnableMeshCache"),
+	true,
+	TEXT("Enables or disables the reuse of meshes."),
+	ECVF_Scalability);
+
+
 int32 FCustomizableObjectSystemPrivate::SkeletalMeshMinLodQualityLevel = -1;
 
 static void CVarMutableSinkFunction()
@@ -140,6 +147,18 @@ FUpdateContextPrivate::~FUpdateContextPrivate()
 }
 
 
+FString FUpdateContextPrivate::GetReferencerName() const
+{
+	return TEXT("FUpdateContextPrivate");
+}
+
+
+void FUpdateContextPrivate::AddReferencedObjects(FReferenceCollector& Collector)
+{
+	Collector.AddReferencedObjects(Objects);
+}
+
+
 FMutablePendingInstanceUpdate::FMutablePendingInstanceUpdate(const TSharedRef<FUpdateContextPrivate>& InContext) :
 	Context(InContext)
 {
@@ -589,7 +608,6 @@ static FAutoConsoleVariableRef CVarEnableMutableAnimInfoDebugging(
 
 void FCustomizableObjectSystemPrivate::AddGameThreadTask(const FMutableTask& Task)
 {
-	check(IsInGameThread())
 	PendingTasks.Enqueue(Task);
 }
 
@@ -1569,6 +1587,58 @@ namespace impl
 	}
 
 
+	void CreateMutableInstance(const TSharedRef<FUpdateContextPrivate>& Operation)
+	{
+		UCustomizableObjectSystem* System = UCustomizableObjectSystem::GetInstanceChecked(); // Save since UCustomizableObjectSystem::BeginDestroy always waits for all tasks to finish
+		const FCustomizableObjectSystemPrivate* SystemPrivate = System->GetPrivateChecked();
+		
+		const mu::Ptr<mu::System> MutableSystem = SystemPrivate->MutableSystem;
+
+		const TSharedPtr<mu::Model> Model = Operation->Instance->GetCustomizableObject()->GetPrivate()->GetModel();
+
+		if (Operation->bLiveUpdateMode)
+		{
+			if (Operation->InstanceID == 0)
+			{
+				// It's the first update since the instance was put in LiveUpdate Mode, this ID will be reused from now on
+				Operation->InstanceID = MutableSystem->NewInstance(Model);
+				UE_LOG(LogMutable, Verbose, TEXT("Creating Mutable instance with id [%d] for reuse "), Operation->InstanceID);
+			}
+			else
+			{
+				// The instance was already in LiveUpdate Mode, the ID is reused
+				check(Operation->InstanceID);
+				UE_LOG(LogMutable, Verbose, TEXT("Reusing Mutable instance with id [%d] "), Operation->InstanceID);
+			}
+		}
+		else
+		{
+			// In non-LiveUpdate mode, we are forcing the recreation of mutable-side instances with every update.
+			check(Operation->InstanceID == 0);
+			Operation->InstanceID = MutableSystem->NewInstance(Model);
+			UE_LOG(LogMutable, Verbose, TEXT("Creating Mutable instance with id [%d] "), Operation->InstanceID);
+		}
+
+		Operation->MutableInstance = MutableSystem->BeginUpdate(Operation->InstanceID, Operation->Parameters, Operation->State, mu::System::AllLODs); // TODO GMT possible data race between GameThread and MutableThread
+	}
+
+	
+	void FixLODs(const TSharedRef<FUpdateContextPrivate>& Operation)
+	{
+		Operation->NumLODsAvailable = Operation->MutableInstance->GetLODCount();
+
+		if (Operation->CurrentMinLOD >= Operation->NumLODsAvailable)
+		{
+			Operation->CurrentMinLOD = Operation->NumLODsAvailable - 1;
+			Operation->CurrentMaxLOD = Operation->CurrentMinLOD;
+		}
+		else if (Operation->CurrentMaxLOD >= Operation->NumLODsAvailable)
+		{
+			Operation->CurrentMaxLOD = Operation->NumLODsAvailable - 1;
+		}
+	}
+	
+	
 	// This runs in the mutable thread.
 	void Subtask_Mutable_BeginUpdate_GetMesh(const TSharedRef<FUpdateContextPrivate>& OperationData, TSharedPtr<mu::Model> Model)
 	{
@@ -1587,57 +1657,26 @@ namespace impl
 
 		CustomizableObjectInstancePrivateData->PassThroughTexturesToLoad.Empty();
 
-		if (OperationData->bLiveUpdateMode)
-		{
-			if (OperationData->InstanceID == 0)
-			{
-				// It's the first update since the instance was put in LiveUpdate Mode, this ID will be reused from now on
-				OperationData->InstanceID = System->NewInstance(Model);
-				UE_LOG(LogMutable, Verbose, TEXT("Creating Mutable instance with id [%d] for reuse "), OperationData->InstanceID);
-			}
-			else
-			{
-				// The instance was already in LiveUpdate Mode, the ID is reused
-				check(OperationData->InstanceID);
-				UE_LOG(LogMutable, Verbose, TEXT("Reusing Mutable instance with id [%d] "), OperationData->InstanceID);
-			}
-		}
-		else
+		if (OperationData->PixelFormatOverride)
 		{
-			// In non-LiveUpdate mode, we are forcing the recreation of mutable-side instances with every update.
-			check(OperationData->InstanceID == 0);
-			OperationData->InstanceID = System->NewInstance(Model);
-			UE_LOG(LogMutable, Verbose, TEXT("Creating Mutable instance with id [%d] "), OperationData->InstanceID);
+			System->SetImagePixelConversionOverride( OperationData->PixelFormatOverride );
 		}
 
-		check(OperationData->InstanceID != 0);
-
-		if (OperationData->PixelFormatOverride)
+		if (!OperationData->bUseMeshCache)
 		{
-			System->SetImagePixelConversionOverride( OperationData->PixelFormatOverride );
+			CreateMutableInstance(OperationData);
+			FixLODs(OperationData);
 		}
 
 		// Main instance generation step
 		// LOD mask, set to all ones to build  all LODs
-		const mu::Instance* Instance = System->BeginUpdate(OperationData->InstanceID, OperationData->Parameters, OperationData->State, mu::System::AllLODs);
+		const mu::Instance* Instance = OperationData->MutableInstance; // TODO GMTFuture remove
 		if (!Instance)
 		{
 			UE_LOG(LogMutable, Warning, TEXT("An Instace update has failed."));
 			return;
 		}
 
-		OperationData->NumLODsAvailable = Instance->GetLODCount();
-
-		if (OperationData->CurrentMinLOD >= OperationData->NumLODsAvailable)
-		{
-			OperationData->CurrentMinLOD = OperationData->NumLODsAvailable - 1;
-			OperationData->CurrentMaxLOD = OperationData->CurrentMinLOD;
-		}
-		else if (OperationData->CurrentMaxLOD >= OperationData->NumLODsAvailable)
-		{
-			OperationData->CurrentMaxLOD = OperationData->NumLODsAvailable - 1;
-		}
-
 		if (!OperationData->RequestedLODs.IsEmpty())
 		{
 			// Initialize RequestedLODs to zero if not set
@@ -2225,6 +2264,7 @@ namespace impl
 			{
 				MutableSystem->ReleaseInstance(OperationData->InstanceID);
 				OperationData->InstanceID = 0;
+				OperationData->MutableInstance = nullptr;
 			}
 		}
 
@@ -2632,6 +2672,83 @@ namespace impl
 	}
 
 
+	void Task_Game_LockMeshCache(const TSharedRef<FUpdateContextPrivate>& Operation)
+	{
+		MUTABLE_CPUPROFILER_SCOPE(Task_Game_LockMeshCache);
+
+		UCustomizableObjectSystem* System = UCustomizableObjectSystem::GetInstanceChecked();
+		FCustomizableObjectSystemPrivate* SystemPrivate = System->GetPrivateChecked();
+
+		const UCustomizableObject* CustomizableObject = Operation->Instance->GetCustomizableObject();
+		FCustomizableObjectPrivateData* CustomizableObjectPrivate = CustomizableObject->GetPrivate();
+		
+		for (const TArray<mu::FResourceID>& MeshId : Operation->MeshDescriptors)
+		{
+			if (USkeletalMesh* CachedMesh = CustomizableObject->GetPrivate()->MeshCache.Get(MeshId))
+			{
+				Operation->Objects.Add(CachedMesh);
+			}
+		}
+
+		// Task inputs
+		TSharedPtr<mu::Model> Model = CustomizableObject->GetPrivate()->GetModel();
+
+#ifdef MUTABLE_USE_NEW_TASKGRAPH
+		UE::Tasks::FTask Dependency;
+#else
+		FGraphEventRef Dependency;
+#endif
+		
+		Dependency = SystemPrivate->MutableTaskGraph.AddMutableThreadTask(
+			TEXT("Task_Mutable_Update_GetMesh"),
+			[Operation, Model]()
+			{
+				impl::Task_Mutable_Update_GetMesh(Operation, Model);
+			});
+
+		SystemPrivate->AddGameThreadTask(
+			{
+			FMutableTaskDelegate::CreateLambda(
+				[Operation]()
+				{
+					impl::Task_Game_LockCache(Operation);
+				}),
+				Dependency,
+			});
+	}
+	
+
+	void Task_Mutable_GetMeshID(const TSharedRef<FUpdateContextPrivate>& Operation)
+	{
+		MUTABLE_CPUPROFILER_SCOPE(Task_Mutable_GetMeshID);
+
+
+		UCustomizableObjectSystem* System = UCustomizableObjectSystem::GetInstanceChecked(); // Save since UCustomizableObjectSystem::BeginDestroy always waits for all tasks to finish
+		FCustomizableObjectSystemPrivate* SystemPrivate = System->GetPrivateChecked();
+		
+		CreateMutableInstance(Operation);
+		FixLODs(Operation);
+
+		const int32 NumComponents = Operation->Instance->GetCustomizableObject()->GetComponentCount();
+
+		Operation->MeshDescriptors.SetNum(NumComponents);
+
+		for (int32 ComponentIndex = 0; ComponentIndex < NumComponents; ++ComponentIndex)
+		{
+			TArray<mu::FResourceID>& MeshId = Operation->MeshDescriptors[ComponentIndex];
+			MeshId.Init(MAX_uint64, MAX_MESH_LOD_COUNT);
+			
+			for (int32 LODIndex = Operation->CurrentMinLOD; LODIndex <= Operation->CurrentMaxLOD; ++LODIndex)
+			{
+				const bool bGenerateLOD = Operation->RequestedLODs.IsValidIndex(ComponentIndex) ? (Operation->RequestedLODs[ComponentIndex] & (1 << LODIndex)) != 0 : true;
+				if (bGenerateLOD)
+				{
+					MeshId[LODIndex] = Operation->MutableInstance->GetMeshId(LODIndex, ComponentIndex, 0);
+				}
+			}
+		}
+	}
+	
 	/** "Start Update" */
 	void Task_Game_StartUpdate(const TSharedRef<FUpdateContextPrivate>& Operation)
 	{
@@ -2788,6 +2905,7 @@ namespace impl
 		Operation->CurrentMinLOD = Operation->InstanceDescriptorRuntimeHash.GetMinLOD();
 		Operation->CurrentMaxLOD = Operation->InstanceDescriptorRuntimeHash.GetMaxLOD();
 		Operation->InstanceID = Operation->bLiveUpdateMode ? CandidateInstancePrivateData->LiveUpdateModeInstanceID : 0;
+		Operation->bUseMeshCache = CustomizableObject->IsMeshCacheEnabled() && !Operation->bLiveUpdateMode && CVarEnableMeshCache.GetValueOnGameThread();
 #if WITH_EDITOR
 		Operation->PixelFormatOverride = SystemPrivateData->ImageFormatOverrideFunc;
 #endif
@@ -2844,6 +2962,27 @@ namespace impl
 #else
 		FGraphEventRef Mutable_GetMeshTask;
 #endif
+		
+		if (Operation->bUseMeshCache)
+		{
+			Mutable_GetMeshTask = SystemPrivateData->MutableTaskGraph.AddMutableThreadTask(
+				TEXT("Task_Mutable_GetMeshID"),
+				[Operation]()
+				{
+					impl::Task_Mutable_GetMeshID(Operation);
+				});
+
+			SystemPrivateData->AddGameThreadTask( // TODO GMTFuture GameThread task instead of a tick task
+				{
+				FMutableTaskDelegate::CreateLambda(
+					[Operation]()
+					{
+						impl::Task_Game_LockMeshCache(Operation);
+					}),
+					Mutable_GetMeshTask
+				});
+		}
+		else
 		{
 			// Task inputs
 			TSharedPtr<mu::Model> Model = CustomizableObject->GetPrivate()->GetModel();
@@ -2854,22 +2993,21 @@ namespace impl
 				{
 					impl::Task_Mutable_Update_GetMesh(Operation, Model);
 				});
-		}
-
 
-		// Task: Lock cache
-		//-------------------------------------------------------------
-		{
-			// Task inputs
-			SystemPrivateData->AddGameThreadTask(
-				{
-				FMutableTaskDelegate::CreateLambda(
-					[Operation]()
+			// Task: Lock cache
+			//-------------------------------------------------------------
+			{
+				// Task inputs
+				SystemPrivateData->AddGameThreadTask(
 					{
-						impl::Task_Game_LockCache(Operation);
-					}),
-				Mutable_GetMeshTask
-				});
+					FMutableTaskDelegate::CreateLambda(
+						[Operation]()
+						{
+							impl::Task_Game_LockCache(Operation);
+						}),
+					Mutable_GetMeshTask
+					});
+			}
 		}
 	}
 } // namespace impl
diff --git a/Engine/Plugins/Experimental/Mutable/Source/CustomizableObject/Private/MuCO/CustomizableObjectSystemPrivate.h b/Engine/Plugins/Experimental/Mutable/Source/CustomizableObject/Private/MuCO/CustomizableObjectSystemPrivate.h
--- a/Engine/Plugins/Experimental/Mutable/Source/CustomizableObject/Private/MuCO/CustomizableObjectSystemPrivate.h
+++ b/Engine/Plugins/Experimental/Mutable/Source/CustomizableObject/Private/MuCO/CustomizableObjectSystemPrivate.h
@@ -400,12 +400,15 @@ struct FPendingTextureCoverageQuery
 /** Update Context.
  *
  * Alive from the start to the end of the update (both API and LOD update). */
-class FUpdateContextPrivate
+class FUpdateContextPrivate : public FGCObject
 {
 public:
 	FUpdateContextPrivate(UCustomizableObjectInstance& InInstance);
-	~FUpdateContextPrivate();
+	virtual ~FUpdateContextPrivate() override;
 
+	virtual FString GetReferencerName() const override;
+	virtual void AddReferencedObjects(FReferenceCollector& Collector) override;
+	
 	EQueuePriorityType PriorityType = EQueuePriorityType::Low;
 	
 	FInstanceUpdateDelegate UpdateCallback;
@@ -442,11 +445,14 @@ public:
 	/** When this option is enabled it will reuse the Mutable core instance and its temp data between updates.  */
 	bool bLiveUpdateMode = false;
 	bool bReuseInstanceTextures = false;
+	bool bUseMeshCache = false;
 	
 	/** This option comes from the operation request. It is used to reduce the number of mipmaps that mutable must generate for images.  */
 	int32 MipsToSkip = 0;
 
-	mu::Instance::ID InstanceID = 0;
+	mu::Instance::ID InstanceID = 0; // Redundant
+	const mu::Instance* MutableInstance = nullptr;
+
 	int32 CurrentMinLOD = 0;
 	int32 CurrentMaxLOD = 0;
 	int32 NumLODsAvailable = 0;
@@ -465,6 +471,9 @@ public:
 
 	mu::FImageOperator::FImagePixelFormatFunc PixelFormatOverride;
 
+	/** Mutable Meshes required for each component. Outermost index is the component, inner index is the LOD. */
+	TArray<TArray<mu::FResourceID>> MeshDescriptors;
+	
 	bool UpdateStarted = false;
 
 	// Update stats
@@ -481,6 +490,9 @@ public:
 	/** Used for profiling in the editor. */
 	uint32 MutableRuntimeCycles = 0;
 #endif
+
+	/** Hard references to objects. Avoids GC to collect them. */
+	TArray<TObjectPtr<const UObject>> Objects;
 };
 
 
@@ -525,7 +537,7 @@ public:
 	// The pending instance updates, discards or releases
 	FMutablePendingInstanceWork MutablePendingInstanceWork;
 
-	// Queue of game-thread tasks that need to be executed for the current operation
+	// Queue of game-thread tasks that need to be executed for the current operation. TQueue is thread safe
 	TQueue<FMutableTask> PendingTasks;
 
 	static int32 EnableMutableProgressiveMipStreaming;
