Make sure that optimization level and texture compression level are consistent when calling CompileCustomizableObjectSynchronously

diff --git a/Engine/Plugins/Experimental/Mutable/Source/CustomizableObjectEditor/Private/MuCOE/CustomizableObjectEditorFunctionLibrary.cpp b/Engine/Plugins/Experimental/Mutable/Source/CustomizableObjectEditor/Private/MuCOE/CustomizableObjectEditorFunctionLibrary.cpp
--- a/Engine/Plugins/Experimental/Mutable/Source/CustomizableObjectEditor/Private/MuCOE/CustomizableObjectEditorFunctionLibrary.cpp
+++ b/Engine/Plugins/Experimental/Mutable/Source/CustomizableObjectEditor/Private/MuCOE/CustomizableObjectEditorFunctionLibrary.cpp
@@ -6,7 +6,10 @@
 #include "MuCO/CustomizableObjectSystem.h"
 #include "MuCO/CustomizableObject.h"
 
-ECustomizableObjectCompilationState UCustomizableObjectEditorFunctionLibrary::CompileCustomizableObjectSynchronously(UCustomizableObject* CustomizableObject)
+ECustomizableObjectCompilationState UCustomizableObjectEditorFunctionLibrary::CompileCustomizableObjectSynchronously(
+	UCustomizableObject* CustomizableObject,
+	ECustomizableObjectOptimizationLevel InOptimizationLevel,
+	ECustomizableObjectTextureCompression InTextureCompression)
 {
 	if (!CustomizableObject)
 	{
@@ -31,6 +34,8 @@ ECustomizableObjectCompilationState UCustomizableObjectEditorFunctionLibrary::Co
 
 	FCustomizableObjectCompiler Compiler;
 	FCompilationOptions Options = CustomizableObject->CompileOptions;
+	Options.OptimizationLevel = static_cast<int32>(InOptimizationLevel);
+	Options.TextureCompression = InTextureCompression;
 	Options.bSilentCompilation = false;
 	Compiler.Compile(*CustomizableObject, Options, true);
 
diff --git a/Engine/Plugins/Experimental/Mutable/Source/CustomizableObjectEditor/Public/MuCOE/CustomizableObjectEditorFunctionLibrary.h b/Engine/Plugins/Experimental/Mutable/Source/CustomizableObjectEditor/Public/MuCOE/CustomizableObjectEditorFunctionLibrary.h
--- a/Engine/Plugins/Experimental/Mutable/Source/CustomizableObjectEditor/Public/MuCOE/CustomizableObjectEditorFunctionLibrary.h
+++ b/Engine/Plugins/Experimental/Mutable/Source/CustomizableObjectEditor/Public/MuCOE/CustomizableObjectEditorFunctionLibrary.h
@@ -2,9 +2,19 @@
 #pragma once
 
 #include "Kismet/BlueprintFunctionLibrary.h"
+#include "MuCO/CustomizableObject.h"
 
 #include "CustomizableObjectEditorFunctionLibrary.generated.h"
 
+// This mirrors the logic in CustomizableObjectEditor.cpp
+UENUM(BlueprintType)
+enum class ECustomizableObjectOptimizationLevel : uint8
+{
+	None,
+	Minimal,
+	Maximum
+};
+
 /**
  * Functions we want to be able to call on CustomizableObjects at edit time - could
  * be exposed to cook as well.
@@ -24,5 +34,8 @@ public:
 	 * @return	The final ECustomizableObjectCompilationState - typically Completed or Failed
 	 */
 	UFUNCTION(BlueprintCallable, Category = "CustomizableObject")
-	static ECustomizableObjectCompilationState CompileCustomizableObjectSynchronously(UCustomizableObject* CustomizableObject);
+	static ECustomizableObjectCompilationState CompileCustomizableObjectSynchronously(
+		UCustomizableObject* CustomizableObject, 
+		ECustomizableObjectOptimizationLevel OptimizationLevel = ECustomizableObjectOptimizationLevel::Minimal, 
+		ECustomizableObjectTextureCompression TextureCompression = ECustomizableObjectTextureCompression::Fast);
 };
