Fix warning V517: The use of 'if (A) {...} else if (A) {...}' pattern was detected. There is a probability of logical error presence.
#rnx

diff --git a/Engine/Plugins/Experimental/Mutable/Source/CustomizableObjectEditor/Private/MuCOE/EdGraphSchema_CustomizableObject.cpp b/Engine/Plugins/Experimental/Mutable/Source/CustomizableObjectEditor/Private/MuCOE/EdGraphSchema_CustomizableObject.cpp
--- a/Engine/Plugins/Experimental/Mutable/Source/CustomizableObjectEditor/Private/MuCOE/EdGraphSchema_CustomizableObject.cpp
+++ b/Engine/Plugins/Experimental/Mutable/Source/CustomizableObjectEditor/Private/MuCOE/EdGraphSchema_CustomizableObject.cpp
@@ -1264,7 +1264,7 @@ FText UEdGraphSchema_CustomizableObject::GetPinCategoryName(const FName& PinCate
 	{
 		return LOCTEXT("Material_Asset_Pin_Category", "Table Material");
 	}
-	else if (PinCategory == UEdGraphSchema_CustomizableObject::PC_MaterialAsset)
+	else if (PinCategory == UEdGraphSchema_CustomizableObject::PC_Wildcard)
 	{
 		return LOCTEXT("Wildcard_Pin_Category", "Wildcard");
 	}
