Silence false positive warning V781: The value of the 'ParameterIndexInInstance' variable is checked after it was used. Perhaps there is a mistake in program logic.
#rnx

diff --git a/Engine/Plugins/Experimental/Mutable/Source/CustomizableObject/Private/MuCO/CustomizableObjectInstanceDescriptor.cpp b/Engine/Plugins/Experimental/Mutable/Source/CustomizableObject/Private/MuCO/CustomizableObjectInstanceDescriptor.cpp
--- a/Engine/Plugins/Experimental/Mutable/Source/CustomizableObject/Private/MuCO/CustomizableObjectInstanceDescriptor.cpp
+++ b/Engine/Plugins/Experimental/Mutable/Source/CustomizableObject/Private/MuCO/CustomizableObjectInstanceDescriptor.cpp
@@ -1275,7 +1275,7 @@ void FCustomizableObjectInstanceDescriptor::SetIntParameterSelectedOption(const
 	check(CustomizableObject);
 	RETURN_ON_UNCOMPILED_CO(CustomizableObject, TEXT("Error: Cannot set Int parameter "));
 
-	const int32 ParameterIndexInObject = IntParameters.IsValidIndex(ParameterIndexInInstance) ? CustomizableObject->FindParameter(IntParameters[ParameterIndexInInstance].ParameterName) : INDEX_NONE;
+	const int32 ParameterIndexInObject = IntParameters.IsValidIndex(ParameterIndexInInstance) ? CustomizableObject->FindParameter(IntParameters[ParameterIndexInInstance].ParameterName) : INDEX_NONE; //-V781
 
 	if (ParameterIndexInObject < 0 || ParameterIndexInInstance < 0)
 	{
