Clear dirty flag, use explicit synchronous compile rather than rolling me own

diff --git a/Engine/Plugins/Experimental/Mutable/Source/CustomizableObjectEditor/Private/MuCOE/CustomizableObjectEditorFunctionLibrary.cpp b/Engine/Plugins/Experimental/Mutable/Source/CustomizableObjectEditor/Private/MuCOE/CustomizableObjectEditorFunctionLibrary.cpp
--- a/Engine/Plugins/Experimental/Mutable/Source/CustomizableObjectEditor/Private/MuCOE/CustomizableObjectEditorFunctionLibrary.cpp
+++ b/Engine/Plugins/Experimental/Mutable/Source/CustomizableObjectEditor/Private/MuCOE/CustomizableObjectEditorFunctionLibrary.cpp
@@ -27,32 +27,20 @@ ECustomizableObjectCompilationState UCustomizableObjectEditorFunctionLibrary::Co
 		return ECustomizableObjectCompilationState::Failed;
 	}
 
+	// store package dirty state so that we can restore it - compile is not an edit:
+	const bool bPackageWasDirty = CustomizableObject->GetOutermost()->IsDirty();
+
 	const double StartTime = FPlatformTime::Seconds();
-	const double PrintFrequencySeconds = 15.0;
-	double PrintTime = StartTime + PrintFrequencySeconds;
 
 	FCustomizableObjectCompiler Compiler;
 	FCompilationOptions Options = CustomizableObject->CompileOptions;
 	Options.OptimizationLevel = static_cast<int32>(InOptimizationLevel);
 	Options.TextureCompression = InTextureCompression;
 	Options.bSilentCompilation = false;
-	Compiler.Compile(*CustomizableObject, Options, true);
+	const bool bAsync = false;
+	Compiler.Compile(*CustomizableObject, Options, bAsync);
 
-	Compiler.Tick();
-	
-	while (Compiler.GetCompilationState() == ECustomizableObjectCompilationState::InProgress)
-	{
-		Compiler.Tick();
-		const double CurrentTime = FPlatformTime::Seconds();
-		if (CurrentTime > PrintTime)
-		{
-			PrintTime = CurrentTime + PrintFrequencySeconds;
-			UE_LOG( LogMutable, Display,
-				TEXT("Synchronously Compiling %s for %f seconds"),
-				*ObjectPath, CurrentTime - StartTime
-			);
-		}
-	}
+	CustomizableObject->GetOutermost()->SetDirtyFlag(bPackageWasDirty);
 
 	const double CurrentTime = FPlatformTime::Seconds();
 	UE_LOG( LogMutable, Display,
@@ -63,6 +51,10 @@ ECustomizableObjectCompilationState UCustomizableObjectEditorFunctionLibrary::Co
 		CurrentTime - StartTime
 	);
 
-	ensure(CustomizableObject->IsCompiled());
+	if (!CustomizableObject->IsCompiled())
+	{
+		UE_LOG(LogMutable, Warning, TEXT("CO not marked as compiled"));
+	}
+
 	return Compiler.GetCompilationState();
 }
