[mutable] Fix bug with "never-stream" flag not being propagated correctly.
#rnx

diff --git a/Engine/Plugins/Experimental/Mutable/Source/CustomizableObject/Public/MuCO/CustomizableObjectCustomVersion.h b/Engine/Plugins/Experimental/Mutable/Source/CustomizableObject/Public/MuCO/CustomizableObjectCustomVersion.h
--- a/Engine/Plugins/Experimental/Mutable/Source/CustomizableObject/Public/MuCO/CustomizableObjectCustomVersion.h
+++ b/Engine/Plugins/Experimental/Mutable/Source/CustomizableObject/Public/MuCO/CustomizableObjectCustomVersion.h
@@ -106,6 +106,8 @@ struct CUSTOMIZABLEOBJECT_API FCustomizableObjectCustomVersion
 
 		CustomizableObjectStateHasSeparateNeverStreamFlag,
 
+		CustomizableObjectNodeHasSeparateNeverStreamFlag,
+
 		// -----<new versions can be added above this line>-------------------------------------------------
 		VersionPlusOne,
 		LatestVersion = VersionPlusOne - 1
diff --git a/Engine/Plugins/Experimental/Mutable/Source/CustomizableObjectEditor/Classes/MuCOE/Nodes/CustomizableObjectNodeObject.h b/Engine/Plugins/Experimental/Mutable/Source/CustomizableObjectEditor/Classes/MuCOE/Nodes/CustomizableObjectNodeObject.h
--- a/Engine/Plugins/Experimental/Mutable/Source/CustomizableObjectEditor/Classes/MuCOE/Nodes/CustomizableObjectNodeObject.h
+++ b/Engine/Plugins/Experimental/Mutable/Source/CustomizableObjectEditor/Classes/MuCOE/Nodes/CustomizableObjectNodeObject.h
@@ -68,6 +68,10 @@ struct CUSTOMIZABLEOBJECTEDITOR_API FCustomizableObjectState
 	UPROPERTY(EditAnywhere, Category = CustomizableObject)
 	ETextureCompressionStrategy TextureCompressionStrategy = ETextureCompressionStrategy::None;
 
+	/** If this is enabled, texture streaming won't be used for this state, and full images will be generated when an instance is first updated. */
+	UPROPERTY(EditAnywhere, Category = CustomizableObject)
+	bool bDisableTextureStreaming = false;
+
 	/** LiveUpdateMode will reuse instance temp. data between updates and speed up update times, but spend much more memory. Good for customization screens, not for actual gameplay modes. */
 	UPROPERTY(EditAnywhere, Category = CustomizableObject)
 	bool bLiveUpdateMode = false;
diff --git a/Engine/Plugins/Experimental/Mutable/Source/CustomizableObjectEditor/Private/MuCOE/GenerateMutableSource/GenerateMutableSource.cpp b/Engine/Plugins/Experimental/Mutable/Source/CustomizableObjectEditor/Private/MuCOE/GenerateMutableSource/GenerateMutableSource.cpp
--- a/Engine/Plugins/Experimental/Mutable/Source/CustomizableObjectEditor/Private/MuCOE/GenerateMutableSource/GenerateMutableSource.cpp
+++ b/Engine/Plugins/Experimental/Mutable/Source/CustomizableObjectEditor/Private/MuCOE/GenerateMutableSource/GenerateMutableSource.cpp
@@ -818,6 +818,7 @@ mu::NodeObjectPtr GenerateMutableSource(const UEdGraphPin * Pin, FMutableGraphGe
 			// UI Data
 			FParameterUIData ParameterUIData(State.Name, State.StateUIMetadata, EMutableParameterType::None);
 			ParameterUIData.TextureCompressionStrategy = State.TextureCompressionStrategy;
+			ParameterUIData.bDisableTextureStreaming = State.bDisableTextureStreaming;
 			ParameterUIData.bLiveUpdateMode = State.bLiveUpdateMode;
 			ParameterUIData.bReuseInstanceTextures = State.bReuseInstanceTextures;
 			ParameterUIData.ForcedParameterValues = State.ForcedParameterValues;
diff --git a/Engine/Plugins/Experimental/Mutable/Source/CustomizableObjectEditor/Private/MuCOE/Nodes/CustomizableObjectNodeObject.cpp b/Engine/Plugins/Experimental/Mutable/Source/CustomizableObjectEditor/Private/MuCOE/Nodes/CustomizableObjectNodeObject.cpp
--- a/Engine/Plugins/Experimental/Mutable/Source/CustomizableObjectEditor/Private/MuCOE/Nodes/CustomizableObjectNodeObject.cpp
+++ b/Engine/Plugins/Experimental/Mutable/Source/CustomizableObjectEditor/Private/MuCOE/Nodes/CustomizableObjectNodeObject.cpp
@@ -65,6 +65,15 @@ void UCustomizableObjectNodeObject::BackwardsCompatibleFixup()
 		// It is safe to do this here as Node Object do not use its node guid to link themeselves to other nodes.
 		CreateNewGuid();
 	}
+
+	// Update state never-stream flag from deprecated enum
+	if (CustomizableObjectCustomVersion < FCustomizableObjectCustomVersion::CustomizableObjectStateHasSeparateNeverStreamFlag)
+	{
+		for (FCustomizableObjectState& s : States)
+		{
+			s.bDisableTextureStreaming = s.TextureCompressionStrategy != ETextureCompressionStrategy::None;
+		}
+	}
 }
 
 
@@ -301,6 +310,7 @@ void UCustomizableObjectNodeObject::PostDuplicate(bool bDuplicateForPIE)
 	Identifier = FGuid::NewGuid();
 }
 
+
 void UCustomizableObjectNodeObject::SetParentObject(UCustomizableObject* CustomizableParentObject)
 {
 	if (CustomizableParentObject != GetGraphEditor()->GetCustomizableObject())
