[mutable] Fixed text conflict from LOCTEXT macro in CustomizableObjectNodeTable.cpp
#rnx

diff --git a/Engine/Plugins/Experimental/Mutable/Source/CustomizableObjectEditor/Private/MuCOE/Nodes/CustomizableObjectNodeTable.cpp b/Engine/Plugins/Experimental/Mutable/Source/CustomizableObjectEditor/Private/MuCOE/Nodes/CustomizableObjectNodeTable.cpp
--- a/Engine/Plugins/Experimental/Mutable/Source/CustomizableObjectEditor/Private/MuCOE/Nodes/CustomizableObjectNodeTable.cpp
+++ b/Engine/Plugins/Experimental/Mutable/Source/CustomizableObjectEditor/Private/MuCOE/Nodes/CustomizableObjectNodeTable.cpp
@@ -147,14 +147,14 @@ FText UCustomizableObjectNodeTable::GetNodeTitle(ENodeTitleType::Type TitleType)
 		FFormatNamedArguments Args;
 		Args.Add(TEXT("TableName"), FText::FromString(Table->GetName()));
 
-		return FText::Format(LOCTEXT("TableNode_Title", "{TableName}\nData Table"), Args);
+		return FText::Format(LOCTEXT("TableNode_Title_DataTable", "{TableName}\nData Table"), Args);
 	}
 	else if (Structure)
 	{
 		FFormatNamedArguments Args;
 		Args.Add(TEXT("StructureName"), FText::FromString(Structure->GetName()));
 
-		return FText::Format(LOCTEXT("TableNode_Title", "{StructureName}\nScript Struct"), Args);
+		return FText::Format(LOCTEXT("TableNode_Title_ScriptedStruct", "{StructureName}\nScript Struct"), Args);
 	}
 	
 	return LOCTEXT("Mutable Table", "Table");
