[Mutable] Enable Mutable custom memory tracking commands on Test builds.

#rnx

diff --git a/Engine/Plugins/Experimental/Mutable/Source/MutableRuntime/Public/MuR/MemoryTrackingUtils.h b/Engine/Plugins/Experimental/Mutable/Source/MutableRuntime/Public/MuR/MemoryTrackingUtils.h
--- a/Engine/Plugins/Experimental/Mutable/Source/MutableRuntime/Public/MuR/MemoryTrackingUtils.h
+++ b/Engine/Plugins/Experimental/Mutable/Source/MutableRuntime/Public/MuR/MemoryTrackingUtils.h
@@ -7,7 +7,7 @@
 #include "HAL/CriticalSection.h"
 
 #ifndef UE_MUTABLE_TRACK_ALLOCATOR_MEMORY_PEAK
-	#if UE_BUILD_DEBUG || UE_BUILD_DEVELOPMENT
+	#if UE_BUILD_DEBUG || UE_BUILD_DEVELOPMENT || UE_BUILD_TEST
 		#define UE_MUTABLE_TRACK_ALLOCATOR_MEMORY_PEAK 1
 	#else
 		#define UE_MUTABLE_TRACK_ALLOCATOR_MEMORY_PEAK 0
