#rnx

diff --git a/Engine/Plugins/Experimental/Mutable/Source/CustomizableObjectEditor/Private/MuCOE/UnrealBakeHelpers.cpp b/Engine/Plugins/Experimental/Mutable/Source/CustomizableObjectEditor/Private/MuCOE/UnrealBakeHelpers.cpp
--- a/Engine/Plugins/Experimental/Mutable/Source/CustomizableObjectEditor/Private/MuCOE/UnrealBakeHelpers.cpp
+++ b/Engine/Plugins/Experimental/Mutable/Source/CustomizableObjectEditor/Private/MuCOE/UnrealBakeHelpers.cpp
@@ -202,19 +202,39 @@ UTexture2D* FUnrealBakeHelpers::BakeHelper_CreateAssetTexture(UTexture2D* SrcTex
 	Dest = nullptr;
 	DupTex->Source.UnlockMip(MipIndex);
 
-	if (PixelFormat == TSF_G8 || PixelFormat == TSF_G16)
+	bool bNeeds_TC_Grayscale = PixelFormat == TSF_G8 || PixelFormat == TSF_G16;
+	bool bDoNotCompress = SrcPixelFormat == PF_R8G8B8A8;
+	bool bIsNormalMap = SrcPixelFormat == PF_BC5;
+
+	if (bNeeds_TC_Grayscale || bDoNotCompress || bIsNormalMap)
 	{
-		// If compression settings are not set to TC_Grayscale the texture will get a DXT format
-		// instead of G8 or G16.
 		FTextureFormatSettings Settings;
-		Settings.CompressionSettings = TC_Grayscale;
-		DupTex->SetLayerFormatSettings(0, Settings);
 
-		DupTex->CompressionSettings = TC_Grayscale;
+		if (bNeeds_TC_Grayscale)
+		{
+			// If compression settings are not set to TC_Grayscale the texture will get a DXT format
+			// instead of G8 or G16.
+			Settings.CompressionSettings = TC_Grayscale;
+			DupTex->CompressionSettings = TC_Grayscale;
+		}
+
+		if (bDoNotCompress)
+		{
+			// In this case keep the RGBA format instead of compressing to DXT
+			Settings.CompressionNone = true;
+			DupTex->CompressionNone = true;
+		}
+
+		if (bIsNormalMap)
+		{
+			Settings.CompressionSettings = TC_Normalmap;
+			DupTex->CompressionSettings = TC_Normalmap;
+		}
+		
+		DupTex->SetLayerFormatSettings(0, Settings);
 	}
 
 	DupTex->UpdateResource();
 
 	return DupTex;
 }
-
