[mutable] Added logging of update performance data

diff --git a/Engine/Plugins/Experimental/Mutable/Source/CustomizableObject/Private/MuCO/LogBenchmarkUtil.cpp b/Engine/Plugins/Experimental/Mutable/Source/CustomizableObject/Private/MuCO/LogBenchmarkUtil.cpp
--- a/Engine/Plugins/Experimental/Mutable/Source/CustomizableObject/Private/MuCO/LogBenchmarkUtil.cpp
+++ b/Engine/Plugins/Experimental/Mutable/Source/CustomizableObject/Private/MuCO/LogBenchmarkUtil.cpp
@@ -31,9 +31,14 @@ static TAutoConsoleVariable<bool> CVarEnableBenchmark(
 
 namespace LogBenchmarkUtil
 {
-	void Write(FArchive& Archive, FStringView Text)
+	void Write(FArchive& Archive, const FString& Text )
 	{
-		Archive.Serialize(const_cast<ANSICHAR*>(StringCast<ANSICHAR>(Text.GetData(), Text.Len()).Get()), Text.Len() * sizeof(ANSICHAR));
+		const FString LogUID = TEXT("MUTABLE_BENCHMARK");
+		UE_LOG(LogMutable, Display, TEXT("%s : %s"),*LogUID, *Text);
+		
+		const FString ComposedString =  FString::Printf(TEXT("%s\n"),*Text);
+		const FStringView StringView = ComposedString;
+		Archive.Serialize(const_cast<ANSICHAR*>(StringCast<ANSICHAR>(StringView.GetData(), StringView.Len()).Get()), StringView.Len() * sizeof(ANSICHAR));
 	}
 
 	static FString LocalBenchmarkFilePath = TEXT("");
@@ -58,7 +63,8 @@ TSharedPtr<FArchive> CreateFile()
 	TSharedPtr<FArchive> Archive = MakeShareable(IFileManager::Get().CreateFileWriter(*Filename, FILEWRITE_AllowRead | FILEWRITE_NoFail));
 	check(Archive);
 
-	LogBenchmarkUtil::Write(*Archive, TEXT("ID_CO,ID_COI,ID_UpdateType,Time_Update,Time_TaskGetMesh,Time_TaskLockCache,Time_TaskGetImages,Time_TaskConvertResources,Time_TaskCallbacks\n"));		
+	const FString HeaderRow = TEXT("ID_CO,ID_COI,ID_UpdateType,Time_Update,Time_TaskGetMesh,Time_TaskLockCache,Time_TaskGetImages,Time_TaskConvertResources,Time_TaskCallbacks");
+	LogBenchmarkUtil::Write(*Archive, HeaderRow);		
 
 	return Archive;
 }
@@ -320,8 +326,10 @@ void FLogBenchmarkUtil::FinishUpdate(const TSharedRef<FUpdateContextPrivate>& Co
 	const double Time_TaskGetImages = Context->TaskGetImagesTime * 1000;
 	const double Time_TaskConvertResources = Context->TaskConvertResourcesTime * 1000;
 	const double Time_TaskCallbacks = Context->TaskCallbacksTime * 1000;
-	
-	LogBenchmarkUtil::Write(*Archive, FString::Printf(TEXT("%s,%s,%s,%f,%f,%f,%f,%f,%f\n"), *ID_CO, *ID_COI, *ID_UpdateType, Time_Update, Time_TaskGetMesh, Time_TaskLockCache, Time_TaskGetImages, Time_TaskConvertResources, Time_TaskCallbacks));
+
+	const FString UpdateString = FString::Printf(TEXT("%s,%s,%s,%f,%f,%f,%f,%f,%f"), *ID_CO, *ID_COI, *ID_UpdateType, Time_Update, Time_TaskGetMesh, Time_TaskLockCache, Time_TaskGetImages, Time_TaskConvertResources, Time_TaskCallbacks);
+	LogBenchmarkUtil::Write(*Archive, UpdateString);
 	Archive->Flush();
+	
 }
 
