[Mutable] Fix constant updates triggered in-between the SetModel and the Unlock of the CO.
- Check the participating object if the CO is compiled. Prevent discarding compilations when running the Game in standalone.
#rnx
#tests PIE
#localization none

diff --git a/Engine/Plugins/Experimental/Mutable/Source/CustomizableObject/Private/MuCO/CustomizableObjectInstance.cpp b/Engine/Plugins/Experimental/Mutable/Source/CustomizableObject/Private/MuCO/CustomizableObjectInstance.cpp
--- a/Engine/Plugins/Experimental/Mutable/Source/CustomizableObject/Private/MuCO/CustomizableObjectInstance.cpp
+++ b/Engine/Plugins/Experimental/Mutable/Source/CustomizableObject/Private/MuCO/CustomizableObjectInstance.cpp
@@ -4843,18 +4843,7 @@ FGraphEventRef UCustomizableInstancePrivateData::LoadAdditionalAssetsAsync(const
 		// Add Skeletons to merge
 		for (const uint32 SkeletonId : SkeletonData.SkeletonIds)
 		{
-			TSoftObjectPtr<USkeleton> AssetPtr;
-			if (SkeletonId == 0)
-			{
-				FMutableRefSkeletalMeshData* RefSkeletalMeshData = CustomizableObject->GetRefSkeletalMeshData(SkeletonData.ComponentIndex);
-				check(RefSkeletalMeshData);
-
-				AssetPtr = RefSkeletalMeshData->Skeleton;
-			}
-			else
-			{
-				AssetPtr = CustomizableObject->GetReferencedSkeletonAssetPtr(SkeletonId);
-			}
+			TSoftObjectPtr<USkeleton> AssetPtr = CustomizableObject->GetReferencedSkeletonAssetPtr(SkeletonId);
 
 			if (AssetPtr.IsNull())
 			{
@@ -4872,7 +4861,6 @@ FGraphEventRef UCustomizableInstancePrivateData::LoadAdditionalAssetsAsync(const
 			{
 				AssetsToStream.Add(AssetPtr.ToSoftObjectPath());
 			}
-
 		}
 	}
 
diff --git a/Engine/Plugins/Experimental/Mutable/Source/CustomizableObject/Private/MuCO/CustomizableObjectSystem.cpp b/Engine/Plugins/Experimental/Mutable/Source/CustomizableObject/Private/MuCO/CustomizableObjectSystem.cpp
--- a/Engine/Plugins/Experimental/Mutable/Source/CustomizableObject/Private/MuCO/CustomizableObjectSystem.cpp
+++ b/Engine/Plugins/Experimental/Mutable/Source/CustomizableObject/Private/MuCO/CustomizableObjectSystem.cpp
@@ -131,8 +131,9 @@ FUpdateContextPrivate::FUpdateContextPrivate(UCustomizableObjectInstance& InInst
 	bBuildParameterRelevancy = InInstance.GetBuildParameterRelevancy();
 	Parameters = InInstance.GetDescriptor().GetParameters();
 	TextureParameters = InInstance.GetPrivate()->UpdateTextureParameters;
+	NumComponents = InInstance.GetCustomizableObject()->GetComponentCount();
 	CurrentMinLOD = InInstance.GetCurrentMinLOD();
-	CurrentMaxLOD = InInstance.GetCurrentMinLOD();
+	CurrentMaxLOD = InInstance.GetCurrentMaxLOD();
 	RequestedLODs = InInstance.GetRequestedLODsPerComponent();
 	
 	InInstance.GetCustomizableObject()->ApplyStateForcedValuesToParameters(State, Parameters.get());
@@ -1024,7 +1025,7 @@ EUpdateRequired FCustomizableObjectSystemPrivate::IsUpdateRequired(const UCustom
 	const UCustomizableInstancePrivateData* const Private = Instance.GetPrivate();
 	
 	const UCustomizableObject* CustomizableObject = Instance.GetCustomizableObject();
-	if (!Instance.CanUpdateInstance() || CustomizableObject->IsLocked())
+	if (!Instance.CanUpdateInstance())
 	{
 		return EUpdateRequired::NoUpdate;
 	}
@@ -1635,6 +1636,20 @@ namespace impl
 		{
 			Operation->CurrentMaxLOD = Operation->NumLODsAvailable - 1;
 		}
+
+		// Initialize RequestedLODs to zero if not set
+		Operation->RequestedLODs.SetNumZeroed(Operation->NumComponents);
+
+		for (int32 ComponentIndex = 0; ComponentIndex < Operation->NumComponents; ++ComponentIndex)
+		{
+			// Ensure we're generating at least one LOD
+			for (int32 LODIndex = Operation->CurrentMaxLOD; LODIndex < MAX_MESH_LOD_COUNT; ++LODIndex)
+			{
+				Operation->RequestedLODs[ComponentIndex] |= (1 << LODIndex);
+			}
+		}
+
+		Operation->InstanceDescriptorRuntimeHash.UpdateRequestedLODs(Operation->RequestedLODs);
 	}
 	
 	
@@ -1676,34 +1691,6 @@ namespace impl
 			return;
 		}
 
-		if (!OperationData->RequestedLODs.IsEmpty())
-		{
-			// Initialize RequestedLODs to zero if not set
-			const int32 ComponentCount = Instance->GetComponentCount(OperationData->CurrentMinLOD);
-			OperationData->RequestedLODs.SetNumZeroed(ComponentCount);
-			
-			for (int32 ComponentIndex = 0; ComponentIndex < ComponentCount; ++ComponentIndex)
-			{
-				// Ensure we're generating at least one LOD
-				if (OperationData->RequestedLODs[ComponentIndex] == 0)
-				{
-					OperationData->RequestedLODs[ComponentIndex] |= (1 << OperationData->CurrentMaxLOD);
-				}
-
-				// Make sure we are not requesting a LOD that doesn't exist in this state (Essentially for states with bBuildOnlyFirstLOD 
-				// and NumExtraLODsToBuildPerPlatform when the ExtraLOD is not needed)
-				if ((OperationData->RequestedLODs[ComponentIndex] & (1 << OperationData->CurrentMaxLOD)) == 0)
-				{
-					OperationData->CurrentMaxLOD = OperationData->CurrentMinLOD;
-
-					// Ensure the fallback LOD actually exists
-					OperationData->RequestedLODs[ComponentIndex] |= (1 << OperationData->CurrentMaxLOD);
-				}
-			}
-		}
-		
-		OperationData->InstanceDescriptorRuntimeHash.UpdateRequestedLODs(OperationData->RequestedLODs);
-
 		// Map SharedSurfaceId to surface index
 		TArray<int32> SurfacesSharedId;
 
@@ -2728,7 +2715,7 @@ namespace impl
 		CreateMutableInstance(Operation);
 		FixLODs(Operation);
 
-		const int32 NumComponents = Operation->Instance->GetCustomizableObject()->GetComponentCount();
+		const int32 NumComponents = Operation->NumComponents;
 
 		Operation->MeshDescriptors.SetNum(NumComponents);
 
@@ -2952,8 +2939,8 @@ namespace impl
 			!System->CurrentInstanceLODManagement->IsOnlyGenerateRequestedLODLevelsEnabled() ||
 			bIsInEditorViewport)
 		{
-			Operation->RequestedLODs = {};
-			Operation->InstanceDescriptorRuntimeHash.UpdateRequestedLODs({});
+			Operation->RequestedLODs.Init(MAX_uint8, Operation->NumComponents);
+			Operation->InstanceDescriptorRuntimeHash.UpdateRequestedLODs(Operation->RequestedLODs);
 		}
 
 #ifdef MUTABLE_USE_NEW_TASKGRAPH
diff --git a/Engine/Plugins/Experimental/Mutable/Source/CustomizableObject/Private/MuCO/CustomizableObjectSystemPrivate.h b/Engine/Plugins/Experimental/Mutable/Source/CustomizableObject/Private/MuCO/CustomizableObjectSystemPrivate.h
--- a/Engine/Plugins/Experimental/Mutable/Source/CustomizableObject/Private/MuCO/CustomizableObjectSystemPrivate.h
+++ b/Engine/Plugins/Experimental/Mutable/Source/CustomizableObject/Private/MuCO/CustomizableObjectSystemPrivate.h
@@ -453,9 +453,11 @@ public:
 	mu::Instance::ID InstanceID = 0; // Redundant
 	const mu::Instance* MutableInstance = nullptr;
 
+	int32 NumComponents = 0;
+	int32 NumLODsAvailable = 0;
+
 	int32 CurrentMinLOD = 0;
 	int32 CurrentMaxLOD = 0;
-	int32 NumLODsAvailable = 0;
 
 	TArray<uint16> RequestedLODs;
 
