[Mutable] Serialize AssetUserData as a streamed resource owned by the CO to prevent pulling unused SkeletalMeshes into the build.
- Rename StreamedExtensionData classes and structs. Replaced ExtensionData with ResourceData.
- Reuse the StreamedExtensionData infrastructure for streamed resources.

#rnx

diff --git a/Engine/Plugins/Experimental/Mutable/Source/CustomizableObject/Classes/MuCO/CustomizableObject.h b/Engine/Plugins/Experimental/Mutable/Source/CustomizableObject/Classes/MuCO/CustomizableObject.h
--- a/Engine/Plugins/Experimental/Mutable/Source/CustomizableObject/Classes/MuCO/CustomizableObject.h
+++ b/Engine/Plugins/Experimental/Mutable/Source/CustomizableObject/Classes/MuCO/CustomizableObject.h
@@ -4,8 +4,9 @@
 
 #include "RHIDefinitions.h"
 #include "MuCO/CustomizableObjectClothingTypes.h"
-#include "MuCO/CustomizableObjectExtensionData.h"
-#include "MuCO/CustomizableObjectStreamedExtensionData.h"
+#include "MuCO/CustomizableObjectResourceData.h"
+#include "MuCO/CustomizableObjectResourceDataTypes.h"
+#include "MuCO/CustomizableObjectStreamedResourceData.h"
 #include "MuCO/CustomizableObjectIdentifier.h"
 #include "MuCO/CustomizableObjectParameterTypeDefinitions.h"
 #include "MuCO/CustomizableObjectUIData.h"
@@ -1201,11 +1202,15 @@ public:
 
 	// mu::ExtensionData::Index is an index into this array when mu::ExtensionData::Origin is ConstantAlwaysLoaded
 	UPROPERTY()
-	TArray<FCustomizableObjectExtensionData> AlwaysLoadedExtensionData;
+	TArray<FCustomizableObjectResourceData> AlwaysLoadedExtensionData;
 
 	// mu::ExtensionData::Index is an index into this array when mu::ExtensionData::Origin is ConstantStreamed
 	UPROPERTY()
-	TArray<FCustomizableObjectStreamedExtensionData> StreamedExtensionData;
+	TArray<FCustomizableObjectStreamedResourceData> StreamedExtensionData;
+
+	// Constant Resources streamed in on demand when generating meshes
+	UPROPERTY()
+	TArray<FCustomizableObjectStreamedResourceData> StreamedResourceData;
 
 private:
 	/** Use the SkeletalMesh of reference as a placeholder until the custom mesh is ready to use.
@@ -1485,7 +1490,7 @@ private:
 	// This is a manual version number for the binary blobs in this asset.
 	// Increasing it invalidates all the previously compiled models.
 	// Warning: If while merging code both versions have changed, take the highest+1.
-	static const int32 CurrentSupportedVersion = 410;
+	static const int32 CurrentSupportedVersion = 411;
 
 public:
 
@@ -1635,10 +1640,6 @@ public:
 	UPROPERTY()
 	TArray<FAnimBpOverridePhysicsAssetsInfo> AnimBpOverridePhysiscAssetsInfo;
 
-	/** Stores the UAssetUserData assets gathered from the SkeletalMesh nodes during compilation, to be used in mesh generation in-game */
-	UPROPERTY()
-	TMap<FString, TSoftObjectPtr<UAssetUserData>> AssetUserDataAssetsMap;
-
 	UPROPERTY()
 	/** Stores the sockets provided by the part skeletal meshes, to be merged in the generated meshes */
 	TArray<FMutableRefSocket> SocketArray;
diff --git a/Engine/Plugins/Experimental/Mutable/Source/CustomizableObject/Classes/MuCO/CustomizableObjectExtensionData.h b/Engine/Plugins/Experimental/Mutable/Source/CustomizableObject/Classes/MuCO/CustomizableObjectExtensionData.h
deleted file mode 100644
diff --git a/Engine/Plugins/Experimental/Mutable/Source/CustomizableObject/Classes/MuCO/CustomizableObjectResourceData.h b/Engine/Plugins/Experimental/Mutable/Source/CustomizableObject/Classes/MuCO/CustomizableObjectResourceData.h
new file mode 100644
--- /dev/null
+++ b/Engine/Plugins/Experimental/Mutable/Source/CustomizableObject/Classes/MuCO/CustomizableObjectResourceData.h
@@ -0,0 +1,44 @@
+// Copyright Epic Games, Inc. All Rights Reserved.
+
+#pragma once
+
+#include "InstancedStruct.h"
+#include "UObject/Object.h"
+
+#include "CustomizableObjectResourceData.generated.h"
+
+/** Result of all the checks just before beginning an update. */
+UENUM()
+enum class ECOResourceDataType : uint16
+{
+	None, // Default.
+	AssetUserData
+};
+
+/** ResourceData generated by a Customizable Object during compilation. */
+USTRUCT()
+struct FCustomizableObjectResourceData
+{
+	GENERATED_BODY()
+
+public:
+
+	UPROPERTY()
+	ECOResourceDataType Type = ECOResourceDataType::None;
+
+	/** The data itself. This could be any USTRUCT. */
+	UPROPERTY()
+	FInstancedStruct Data;
+};
+
+
+/** A simple container that's used to store the FCustomizableObjectStreamedResourceData in a package */
+UCLASS()
+class CUSTOMIZABLEOBJECT_API UCustomizableObjectResourceDataContainer : public UObject
+{
+	GENERATED_BODY()
+
+public:
+	UPROPERTY()
+	FCustomizableObjectResourceData Data;
+};
\ No newline at end of file
diff --git a/Engine/Plugins/Experimental/Mutable/Source/CustomizableObject/Classes/MuCO/CustomizableObjectStreamedExtensionData.h b/Engine/Plugins/Experimental/Mutable/Source/CustomizableObject/Classes/MuCO/CustomizableObjectStreamedExtensionData.h
deleted file mode 100644
diff --git a/Engine/Plugins/Experimental/Mutable/Source/CustomizableObject/Classes/MuCO/CustomizableObjectStreamedResourceData.h b/Engine/Plugins/Experimental/Mutable/Source/CustomizableObject/Classes/MuCO/CustomizableObjectStreamedResourceData.h
new file mode 100644
--- /dev/null
+++ b/Engine/Plugins/Experimental/Mutable/Source/CustomizableObject/Classes/MuCO/CustomizableObjectStreamedResourceData.h
@@ -0,0 +1,78 @@
+// Copyright Epic Games, Inc. All Rights Reserved.
+
+#pragma once
+
+#include "MuCO/CustomizableObjectResourceData.h"
+#include "UObject/Object.h"
+#include "UObject/SoftObjectPtr.h"
+
+#include "CustomizableObjectStreamedResourceData.generated.h"
+
+
+
+/**
+ * Used to store resource data that can be streamed in and out from an external package.
+ *
+ * This allows you to load large resource data on demand, instead of having it always loaded.
+ */
+USTRUCT()
+struct CUSTOMIZABLEOBJECT_API FCustomizableObjectStreamedResourceData
+{
+	GENERATED_BODY()
+
+public:
+	FCustomizableObjectStreamedResourceData() = default;
+
+#if WITH_EDITOR
+	/** Streamed data can only be constructed in the editor */
+	FCustomizableObjectStreamedResourceData(UCustomizableObjectResourceDataContainer* InContainer);
+
+	/**
+	 * Clears the hard reference to the container, so that it's not forced to load as soon as this
+	 * struct is loaded, and replaces the soft reference with the provided one.
+	 * 
+	 * Note that the container may not have been moved to its new path yet, so don't try to load
+	 * the container from the given path here.
+	 */
+	void ConvertToSoftReferenceForCooking(const TSoftObjectPtr<UCustomizableObjectResourceDataContainer>& NewContainerPath);
+#endif // WITH_EDITOR
+
+	/**
+	 * If this returns true, GetLoadedData can safely be called.
+	 *
+	 * The data will stay loaded until Unload is called.
+	 */
+	bool IsLoaded() const;
+
+	/**
+	 * Returns the data, as long as it's loaded.
+	 *
+	 * If the data isn't loaded, an assertion will fail.
+	 */
+	const FCustomizableObjectResourceData& GetLoadedData() const;
+
+	/**
+	 * Release this struct's hard reference to the loaded data.
+	 *
+	 * This doesn't immediately remove it from memory, but will allow it to be deleted by the
+	 * garbage collector if there are no other references.
+	 *
+	 * Note that in editor builds streamed Resource Data must always be loaded due to the package
+	 * structure, and therefore this function has no effect in editor builds.
+	 */
+	void Unload();
+
+	/**
+	 * To stream this data in, load the object at the path returned by GetPath and then pass it
+	 * into NotifyLoaded.
+	 */
+	const TSoftObjectPtr<UCustomizableObjectResourceDataContainer>& GetPath() const { return ContainerPath; }
+	void NotifyLoaded(const UCustomizableObjectResourceDataContainer* LoadedContainer);
+
+private:
+	UPROPERTY()
+	TSoftObjectPtr<UCustomizableObjectResourceDataContainer> ContainerPath;
+
+	UPROPERTY()
+	TObjectPtr<const UCustomizableObjectResourceDataContainer> Container;
+};
diff --git a/Engine/Plugins/Experimental/Mutable/Source/CustomizableObject/Private/MuCO/CustomizableInstancePrivateData.h b/Engine/Plugins/Experimental/Mutable/Source/CustomizableObject/Private/MuCO/CustomizableInstancePrivateData.h
--- a/Engine/Plugins/Experimental/Mutable/Source/CustomizableObject/Private/MuCO/CustomizableInstancePrivateData.h
+++ b/Engine/Plugins/Experimental/Mutable/Source/CustomizableObject/Private/MuCO/CustomizableInstancePrivateData.h
@@ -136,8 +136,8 @@ struct FCustomizableInstanceComponentData
 	UPROPERTY(Transient, Category = CustomizableObjectInstance, editfixedsize, VisibleAnywhere)
 	TSet<TObjectPtr<UAssetUserData>> AssetUserDataArray;
 
-	// Temp structure to store the AssetUserData that will be async loaded during an update
-	TSet<TSoftObjectPtr<UAssetUserData>> AssetUserDataToStream;
+	// Index of the resource in the StreamedResourceData array of the CustomizableObject.
+	TArray<int32> StreamedResourceIndex;
 
 #if WITH_EDITORONLY_DATA
 	// Just used for mutable.EnableMutableAnimInfoDebugging command
diff --git a/Engine/Plugins/Experimental/Mutable/Source/CustomizableObject/Private/MuCO/CustomizableObject.cpp b/Engine/Plugins/Experimental/Mutable/Source/CustomizableObject/Private/MuCO/CustomizableObject.cpp
--- a/Engine/Plugins/Experimental/Mutable/Source/CustomizableObject/Private/MuCO/CustomizableObject.cpp
+++ b/Engine/Plugins/Experimental/Mutable/Source/CustomizableObject/Private/MuCO/CustomizableObject.cpp
@@ -432,13 +432,118 @@ void UCustomizableObject::UpdateCompiledDataFromModel()
 }
 
 
-void UCustomizableObject::SaveCompiledData(FArchive& MemoryWriter, bool bSkipEditorOnlyData)
+void SerializeStreamedResources(FArchive& Ar, UObject* Object, TArray<FCustomizableObjectStreamedResourceData>& StreamedResources, bool bIsCooking)
+{
+	if (Ar.IsSaving())
+	{
+		int32 NumStreamedResources = StreamedResources.Num();
+		Ar << NumStreamedResources;
+
+		for (const FCustomizableObjectStreamedResourceData& ResourceData : StreamedResources)
+		{
+			const FCustomizableObjectResourceData& Data = ResourceData.GetLoadedData();
+			Ar << Data.Type;
+
+			switch (Data.Type)
+			{
+			case ECOResourceDataType::AssetUserData:
+			{
+				const FCustomizableObjectAssetUserData* AssetUserData = Data.Data.GetPtr<FCustomizableObjectAssetUserData>();
+				FString AssetUserDataPath = TSoftObjectPtr<UAssetUserData>(AssetUserData->AssetUserDataEditor).ToSoftObjectPath().ToString();
+				Ar << AssetUserDataPath;
+				break;
+			}
+			default:
+				check(false);
+				break;
+			}
+		}
+	}
+	else 
+	{
+		const FString CustomizableObjectName = GetNameSafe(Object) + TEXT("_");
+
+		int32 NumStreamedResources = 0;
+		Ar << NumStreamedResources;
+
+		// Initialize if not cooking. Otherwise, resources will be initialized at this point, and only their data will be updated.
+		if (!bIsCooking) 
+		{
+			StreamedResources.SetNumUninitialized(NumStreamedResources);
+		}			
+		
+		check(NumStreamedResources == StreamedResources.Num())
+
+		for (int32 ResourceIndex = 0; ResourceIndex < NumStreamedResources; ++ResourceIndex)
+		{
+			UCustomizableObjectResourceDataContainer* Container = nullptr;
+
+			if (bIsCooking)
+			{
+				// Override existing containers
+				Container = StreamedResources[ResourceIndex].GetPath().Get();
+			}
+
+			// Generate a deterministic name to help with deterministic cooking
+			if (!Container)
+			{
+				const FString ContainerName = CustomizableObjectName + FString::Printf(TEXT("SR_%d"), ResourceIndex);
+
+				UCustomizableObjectResourceDataContainer* ExistingContainer = Cast<UCustomizableObjectResourceDataContainer>(FindObject<UObject>(Object, *ContainerName));
+				Container = ExistingContainer ? ExistingContainer : NewObject<UCustomizableObjectResourceDataContainer>(
+					Object,
+					FName(*ContainerName),
+					RF_Public);
+
+				StreamedResources[ResourceIndex] = { Container };
+			}
+
+			check(Container);
+			Ar << Container->Data.Type;
+			switch (Container->Data.Type)
+			{
+				case ECOResourceDataType::AssetUserData:
+				{
+					FString AssetUserDataPath;
+					Ar << AssetUserDataPath;
+					
+					FCustomizableObjectAssetUserData ResourceData;
+					ResourceData.AssetUserDataEditor = TSoftObjectPtr<UAssetUserData>(AssetUserDataPath).LoadSynchronous();
+
+					if (bIsCooking)
+					{
+						// Rename the asset user data for duplicate
+						const FString AssetName = CustomizableObjectName + GetNameSafe(ResourceData.AssetUserDataEditor);
+						
+						// Find or duplicate the AUD replacing the outer
+						ResourceData.AssetUserData = FindObject<UAssetUserData>(Container, *AssetName);
+						if (!ResourceData.AssetUserData)
+						{
+							// AUD may be private objects within meshes. Duplicate changing the outer to avoid including meshes into the builds.
+							ResourceData.AssetUserData = DuplicateObject<UAssetUserData>(ResourceData.AssetUserDataEditor, Container, FName(*AssetName));
+						}
+					}
+
+					Container->Data.Data = FInstancedStruct::Make(ResourceData);
+					break;
+				}
+				default:
+					break;
+			}
+		}
+	}
+}
+
+
+void UCustomizableObject::SaveCompiledData(FArchive& MemoryWriter, bool bIsCooking)
 {
 	int32 InternalVersion = CurrentSupportedVersion;
 	MutableCompiledDataStreamHeader Header(InternalVersion, VersionId);
 	MemoryWriter << Header;
 
 	MemoryWriter << ReferenceSkeletalMeshesData;
+
+	SerializeStreamedResources(MemoryWriter, this, StreamedResourceData, bIsCooking);
 	
 	int32 NumReferencedMaterials = ReferencedMaterials.Num();
 	MemoryWriter << NumReferencedMaterials;
@@ -496,7 +601,7 @@ void UCustomizableObject::SaveCompiledData(FArchive& MemoryWriter, bool bSkipEdi
 	MemoryWriter << HashToStreamableBlock;
 
 	// All Editor Only data must be serialized here
-	if (!bSkipEditorOnlyData)
+	if (!bIsCooking)
 	{
 		MemoryWriter << CustomizableObjectPathMap;
 		MemoryWriter << GroupNodeMap;
@@ -511,7 +616,7 @@ void UCustomizableObject::SaveCompiledData(FArchive& MemoryWriter, bool bSkipEdi
 	MemoryWriter << LODSettings.bLODStreamingEnabled;
 }
 
-void UCustomizableObject::LoadCompiledData(FArchive& MemoryReader, const ITargetPlatform* InTargetPlatform, bool bSkipEditorOnlyData)
+void UCustomizableObject::LoadCompiledData(FArchive& MemoryReader, const ITargetPlatform* InTargetPlatform, bool bIsCooking)
 {
 	PrivateData->SetModel(nullptr, FGuid());
 	ClearCompiledData();
@@ -532,7 +637,8 @@ void UCustomizableObject::LoadCompiledData(FArchive& MemoryReader, const ITarget
 			ReferenceSkeletalMeshData.InitResources(this, InTargetPlatform);
 		}
 
-		// We can load
+		SerializeStreamedResources(MemoryReader, this, StreamedResourceData, bIsCooking);
+
 		int32 NumReferencedMaterials = 0;
 		MemoryReader << NumReferencedMaterials;
 
@@ -604,7 +710,7 @@ void UCustomizableObject::LoadCompiledData(FArchive& MemoryReader, const ITarget
 		MemoryReader << HashToStreamableBlock;
 
 		// All Editor Only data must be loaded here
-		if (!bSkipEditorOnlyData)
+		if (!bIsCooking)
 		{
 			MemoryReader << CustomizableObjectPathMap;
 			MemoryReader << GroupNodeMap;
diff --git a/Engine/Plugins/Experimental/Mutable/Source/CustomizableObject/Private/MuCO/CustomizableObjectInstance.cpp b/Engine/Plugins/Experimental/Mutable/Source/CustomizableObject/Private/MuCO/CustomizableObjectInstance.cpp
--- a/Engine/Plugins/Experimental/Mutable/Source/CustomizableObject/Private/MuCO/CustomizableObjectInstance.cpp
+++ b/Engine/Plugins/Experimental/Mutable/Source/CustomizableObject/Private/MuCO/CustomizableObjectInstance.cpp
@@ -194,13 +194,13 @@ void UCustomizableInstancePrivateData::PrepareForUpdate(const TSharedRef<FUpdate
 	{
 		ComponentData.AnimSlotToBP.Empty();
 		ComponentData.AssetUserDataArray.Empty();
-		ComponentData.AssetUserDataToStream.Empty();
 		ComponentData.Skeletons.Skeleton = nullptr;
 		ComponentData.Skeletons.SkeletonIds.Empty();
 		ComponentData.Skeletons.SkeletonsToMerge.Empty();
 		ComponentData.PhysicsAssets.PhysicsAssetToLoad.Empty();
 		ComponentData.PhysicsAssets.PhysicsAssetsToMerge.Empty();
 		ComponentData.ClothingPhysicsAssetsToStream.Empty();
+		ComponentData.StreamedResourceIndex.Empty();
 
 #if WITH_EDITORONLY_DATA
 		ComponentData.MeshPartPaths.Empty();
@@ -1972,7 +1972,7 @@ bool UCustomizableInstancePrivateData::UpdateSkeletalMesh_PostBeginUpdate0(UCust
 
 				TArray<FInputPinDataContainer>& ContainerArray = ExtensionToExtensionData.FindOrAdd(Extension);
 
-				const FCustomizableObjectExtensionData* ReferencedExtensionData = nullptr;
+				const FCustomizableObjectResourceData* ReferencedExtensionData = nullptr;
 				switch (ExtensionOutput.Data->Origin)
 				{
 					case mu::ExtensionData::EOrigin::ConstantAlwaysLoaded:
@@ -1986,7 +1986,7 @@ bool UCustomizableInstancePrivateData::UpdateSkeletalMesh_PostBeginUpdate0(UCust
 					{
 						check(CustomizableObject->StreamedExtensionData.IsValidIndex(ExtensionOutput.Data->Index));
 						
-						const FCustomizableObjectStreamedExtensionData& StreamedData =
+						const FCustomizableObjectStreamedResourceData& StreamedData =
 							CustomizableObject->StreamedExtensionData[ExtensionOutput.Data->Index];
 
 						if (!StreamedData.IsLoaded())
@@ -4881,6 +4881,26 @@ FGraphEventRef UCustomizableInstancePrivateData::LoadAdditionalAssetsAsync(const
 
 			FCustomizableInstanceComponentData* ComponentData = GetComponentData(Component.Id);
 
+			const TArray<int32>& StreamedResources = MutableMesh->GetStreamedResources();
+			const TArray<FCustomizableObjectStreamedResourceData>& StreamedResourcesData = CustomizableObject->StreamedResourceData;
+
+			for (int32 ResourceIndex : StreamedResources)
+			{
+				if (!StreamedResourcesData.IsValidIndex(ResourceIndex))
+				{
+					UE_LOG(LogMutable, Error, TEXT("Invalid streamed resource index. Max Index [%d]. Resource Index [%d]."), StreamedResourcesData.Num(), ResourceIndex);
+					continue; 
+				}
+
+				const FCustomizableObjectStreamedResourceData& StreamedResource = StreamedResourcesData[ResourceIndex];
+				if (!StreamedResource.IsLoaded())
+				{
+					AssetsToStream.AddUnique(StreamedResource.GetPath().ToSoftObjectPath());
+				}
+
+				ComponentData->StreamedResourceIndex.Add(ResourceIndex);
+			}
+
 			const bool bReplacePhysicsAssets = HasCOInstanceFlags(ReplacePhysicsAssets);
 
 			for (int32 TagIndex = 0; TagIndex < MutableMesh->GetTagCount(); ++TagIndex)
@@ -4981,25 +5001,6 @@ FGraphEventRef UCustomizableInstancePrivateData::LoadAdditionalAssetsAsync(const
 				{
 					AnimBPGameplayTags.AddTag(FGameplayTag::RequestGameplayTag(*Tag));
 				}
-				else if (Tag.RemoveFromStart("__AssetUserData:"))
-				{
-					const FString& AssetPath = Tag;
-
-					TSoftObjectPtr<UAssetUserData>* AssetUserData = CustomizableObject->AssetUserDataAssetsMap.Find(AssetPath);
-
-					if (AssetUserData && !AssetUserData->IsNull())
-					{
-						if (AssetUserData->Get())
-						{
-							ComponentData->AssetUserDataArray.Add(AssetUserData->Get());
-						}
-						else
-						{
-							ComponentData->AssetUserDataToStream.Add(*AssetUserData);
-							AssetsToStream.Add(AssetUserData->ToSoftObjectPath());
-						}
-					}
-				}
 #if WITH_EDITORONLY_DATA
 				else if (Tag.RemoveFromStart("__MeshPath:"))
 				{
@@ -5095,8 +5096,36 @@ void UCustomizableInstancePrivateData::AdditionalAssetsAsyncLoaded(UCustomizable
 #endif
 
 	
+	TArray<FCustomizableObjectStreamedResourceData>& StreamedResources = CustomizableObject->StreamedResourceData;
+
 	for (FCustomizableInstanceComponentData& ComponentData : ComponentsData)
 	{
+		for (int32 ResourceIndex : ComponentData.StreamedResourceIndex)
+		{
+			FCustomizableObjectStreamedResourceData& Resource = StreamedResources[ResourceIndex];
+			if (!Resource.IsLoaded())
+			{
+				Resource.NotifyLoaded(Resource.GetPath().Get());
+			}
+
+			const FCustomizableObjectResourceData& ResourceData = Resource.GetLoadedData();
+			switch (ResourceData.Type)
+			{
+				case ECOResourceDataType::AssetUserData:
+				{
+					const FCustomizableObjectAssetUserData* AUDResource = ResourceData.Data.GetPtr<FCustomizableObjectAssetUserData>();
+#if WITH_EDITORONLY_DATA
+					ComponentData.AssetUserDataArray.Add(AUDResource->AssetUserDataEditor);
+#else
+					ComponentData.AssetUserDataArray.Add(AUDResource->AssetUserData);
+#endif
+					break;
+				}
+				default:
+					break;
+			}
+		}
+
 		// Loaded Skeletons
 		FReferencedSkeletons& Skeletons = ComponentData.Skeletons;
 		for (int32 SkeletonIndex : Skeletons.SkeletonIds)
@@ -5172,28 +5201,6 @@ void UCustomizableInstancePrivateData::AdditionalAssetsAsyncLoaded(UCustomizable
 #endif
 		}
 
-		for (TSoftObjectPtr<UAssetUserData> LoadedAssetUserData : ComponentData.AssetUserDataToStream)
-		{
-			if (LoadedAssetUserData.IsValid())
-			{
-				ComponentData.AssetUserDataArray.Add(LoadedAssetUserData.Get());
-			}
-#if WITH_EDITOR
-			else
-			{
-				FMessageLogModule& MessageLogModule = FModuleManager::LoadModuleChecked<FMessageLogModule>("MessageLog");
-				MessageLogModule.RegisterLogListing(FName("Mutable"), FText::FromString(FString("Mutable")));
-				FMessageLog MessageLog("Mutable");
-
-				FString ErrorMsg = FString::Printf(TEXT("Mutable couldn't load the AssetUserData [%s]. If it has been deleted or renamed, please recompile all the mutable objects that use it."), *LoadedAssetUserData.GetAssetName());
-				UE_LOG(LogMutable, Error, TEXT("%s"), *ErrorMsg);
-				MessageLog.Notify(FText::FromString(ErrorMsg), EMessageSeverity::Error, true);
-			}
-#endif
-		}
-
-		ComponentData.AssetUserDataToStream.Empty();
-
 		const int32 AdditionalPhysicsNum = ComponentData.PhysicsAssets.AdditionalPhysicsAssetsToLoad.Num();
 		ComponentData.PhysicsAssets.AdditionalPhysicsAssets.Reserve(AdditionalPhysicsNum);
 		for (int32 I = 0; I < AdditionalPhysicsNum; ++I)
@@ -5215,6 +5222,15 @@ void UCustomizableInstancePrivateData::AdditionalAssetsAsyncLoaded(UCustomizable
 	}
 
 	PassThroughTexturesToLoad.Empty();
+
+	// Only Unload in cooked builds. Unloading them when in the editor will trigger an assert. 
+	if (FPlatformProperties::RequiresCookedData())
+	{
+		for (FCustomizableObjectStreamedResourceData& ResourceData : StreamedResources)
+		{
+			ResourceData.Unload();
+		}
+	}
 }
 
 
diff --git a/Engine/Plugins/Experimental/Mutable/Source/CustomizableObject/Private/MuCO/CustomizableObjectStreamedExtensionData.cpp b/Engine/Plugins/Experimental/Mutable/Source/CustomizableObject/Private/MuCO/CustomizableObjectStreamedExtensionData.cpp
deleted file mode 100644
diff --git a/Engine/Plugins/Experimental/Mutable/Source/CustomizableObject/Private/MuCO/CustomizableObjectStreamedResourceData.cpp b/Engine/Plugins/Experimental/Mutable/Source/CustomizableObject/Private/MuCO/CustomizableObjectStreamedResourceData.cpp
new file mode 100644
--- /dev/null
+++ b/Engine/Plugins/Experimental/Mutable/Source/CustomizableObject/Private/MuCO/CustomizableObjectStreamedResourceData.cpp
@@ -0,0 +1,80 @@
+// Copyright Epic Games, Inc. All Rights Reserved.
+
+#include "MuCO/CustomizableObjectStreamedResourceData.h"
+
+#include "ExternalPackageHelper.h"
+#include "MuCO/CustomizableObject.h"
+#include "UObject/Package.h"
+
+#include UE_INLINE_GENERATED_CPP_BY_NAME(CustomizableObjectStreamedResourceData)
+
+#if WITH_EDITOR
+FCustomizableObjectStreamedResourceData::FCustomizableObjectStreamedResourceData(
+	UCustomizableObjectResourceDataContainer* InContainer)
+{
+	check(IsInGameThread());
+	check(InContainer);
+
+	Container = InContainer;
+	ContainerPath = InContainer;
+}
+
+void FCustomizableObjectStreamedResourceData::ConvertToSoftReferenceForCooking(
+	const TSoftObjectPtr<UCustomizableObjectResourceDataContainer>& NewContainerPath)
+{
+	check(IsInGameThread());
+
+	// Remove the hard reference to the container, so that it can be unloaded
+	Container = nullptr;
+
+	ContainerPath = NewContainerPath;
+}
+#endif // WITH_EDITOR
+
+bool FCustomizableObjectStreamedResourceData::IsLoaded() const
+{
+	check(IsInGameThread());
+
+	return Container != nullptr;
+}
+
+const FCustomizableObjectResourceData& FCustomizableObjectStreamedResourceData::GetLoadedData() const
+{
+	check(IsInGameThread());
+
+	check(Container);
+	return Container->Data;
+}
+
+void FCustomizableObjectStreamedResourceData::Unload()
+{
+	check(IsInGameThread());
+
+	if (!Container)
+	{
+		// Already unloaded
+		return;
+	}
+
+	if (FPlatformProperties::RequiresCookedData())
+	{
+		// The container should have been cooked out to its own package
+		check(Container->GetOuter()->IsA<UPackage>());
+
+		Container = nullptr;
+	}
+	else
+	{
+		// In an uncooked build the container is stored in the CO package and the CO needs to
+		// maintain a hard reference to it to ensure that it isn't discarded on save.
+		check(Container->GetOuter()->IsA<UCustomizableObject>());
+	}
+}
+
+void FCustomizableObjectStreamedResourceData::NotifyLoaded(const UCustomizableObjectResourceDataContainer* LoadedContainer)
+{
+	check(IsInGameThread());
+	check(ContainerPath.Get() == LoadedContainer);
+
+	Container = LoadedContainer;
+}
diff --git a/Engine/Plugins/Experimental/Mutable/Source/CustomizableObject/Private/MuCO/UnrealExtensionDataStreamer.cpp b/Engine/Plugins/Experimental/Mutable/Source/CustomizableObject/Private/MuCO/UnrealExtensionDataStreamer.cpp
--- a/Engine/Plugins/Experimental/Mutable/Source/CustomizableObject/Private/MuCO/UnrealExtensionDataStreamer.cpp
+++ b/Engine/Plugins/Experimental/Mutable/Source/CustomizableObject/Private/MuCO/UnrealExtensionDataStreamer.cpp
@@ -178,7 +178,7 @@ TSharedPtr<FStreamableHandle> FUnrealExtensionDataStreamer::StartLoadOnGameThrea
 		return nullptr;
 	}
 
-	FCustomizableObjectStreamedExtensionData& StreamedData = Object->StreamedExtensionData[LoadHandle->Data->Index];
+	FCustomizableObjectStreamedResourceData& StreamedData = Object->StreamedExtensionData[LoadHandle->Data->Index];
 	if (StreamedData.IsLoaded())
 	{
 		// Already loaded
@@ -246,13 +246,13 @@ void FUnrealExtensionDataStreamer::NotifyLoadCompleted(
 		return;
 	}
 
-	FCustomizableObjectStreamedExtensionData& StreamedData = Object->StreamedExtensionData[LoadHandle->Data->Index];
+	FCustomizableObjectStreamedResourceData& StreamedData = Object->StreamedExtensionData[LoadHandle->Data->Index];
 
 	// The object could have been loaded by another request, in which case we can skip updating
 	// the StreamedData.
 	if (!StreamedData.IsLoaded())
 	{
-		const UCustomizableObjectExtensionDataContainer* LoadedObject = StreamedData.GetPath().Get();
+		const UCustomizableObjectResourceDataContainer* LoadedObject = StreamedData.GetPath().Get();
 		if (!LoadedObject)
 		{
 			// Object wasn't loaded for some reason
diff --git a/Engine/Plugins/Experimental/Mutable/Source/CustomizableObject/Public/MuCO/CustomizableObjectResourceDataTypes.h b/Engine/Plugins/Experimental/Mutable/Source/CustomizableObject/Public/MuCO/CustomizableObjectResourceDataTypes.h
new file mode 100644
--- /dev/null
+++ b/Engine/Plugins/Experimental/Mutable/Source/CustomizableObject/Public/MuCO/CustomizableObjectResourceDataTypes.h
@@ -0,0 +1,30 @@
+// Copyright Epic Games, Inc. All Rights Reserved.
+
+#pragma once
+
+#include "UObject/Object.h"
+#include "UObject/ObjectPtr.h"
+#include "UObject/SoftObjectPtr.h"
+
+#include "CustomizableObjectResourceDataTypes.generated.h"
+
+class UAssetUserData;
+
+/** Simple struct used to store streamed resources of type AssetUserData */
+USTRUCT()
+struct FCustomizableObjectAssetUserData
+{
+	GENERATED_BODY()
+
+public:
+	
+	// Valid in Cooked builds
+	UPROPERTY()
+	TObjectPtr<UAssetUserData> AssetUserData;
+
+#if WITH_EDITORONLY_DATA
+	// Valid in the Editor
+	UPROPERTY(Transient)
+	TObjectPtr<UAssetUserData> AssetUserDataEditor;
+#endif
+};
\ No newline at end of file
diff --git a/Engine/Plugins/Experimental/Mutable/Source/CustomizableObjectEditor/Private/MuCOE/CustomizableObjectCompiler.cpp b/Engine/Plugins/Experimental/Mutable/Source/CustomizableObjectEditor/Private/MuCOE/CustomizableObjectCompiler.cpp
--- a/Engine/Plugins/Experimental/Mutable/Source/CustomizableObjectEditor/Private/MuCOE/CustomizableObjectCompiler.cpp
+++ b/Engine/Plugins/Experimental/Mutable/Source/CustomizableObjectEditor/Private/MuCOE/CustomizableObjectCompiler.cpp
@@ -1193,7 +1193,7 @@ void FCustomizableObjectCompiler::CompileInternal(UCustomizableObject* Object, c
 		Object->AlwaysLoadedExtensionData = MoveTemp(GenerationContext.AlwaysLoadedExtensionData);
 
 		Object->StreamedExtensionData.Empty(GenerationContext.StreamedExtensionData.Num());
-		for (UCustomizableObjectExtensionDataContainer* Container : GenerationContext.StreamedExtensionData)
+		for (UCustomizableObjectResourceDataContainer* Container : GenerationContext.StreamedExtensionData)
 		{
 			Object->StreamedExtensionData.Emplace(Container);
 		}
@@ -1233,16 +1233,7 @@ void FCustomizableObjectCompiler::CompileInternal(UCustomizableObject* Object, c
 	
 		Object->AnimBPAssetsMap = GenerationContext.AnimBPAssetsMap;
 
-		// Mark the object as modified, used to avoid missing assets in packages. 
-		if (!Object->AssetUserDataAssetsMap.OrderIndependentCompareEqual(GenerationContext.AssetUserDataAssetsMap))
-		{
-			if (!ParamNamesToSelectedOptions.Num()) // Don't mark the objects as modified because of a partial compilation
-			{
-				Object->MarkPackageDirty();
-			}
-		}
-
-		Object->AssetUserDataAssetsMap = GenerationContext.AssetUserDataAssetsMap;
+		Object->StreamedResourceData = MoveTemp(GenerationContext.StreamedResourceData);
 
 		// Mark the object as modified, used to avoid missing assets in packages.
 		if (Object->SocketArray != GenerationContext.SocketArray)
diff --git a/Engine/Plugins/Experimental/Mutable/Source/CustomizableObjectEditor/Private/MuCOE/CustomizableObjectCookPackageSplitter.cpp b/Engine/Plugins/Experimental/Mutable/Source/CustomizableObjectEditor/Private/MuCOE/CustomizableObjectCookPackageSplitter.cpp
--- a/Engine/Plugins/Experimental/Mutable/Source/CustomizableObjectEditor/Private/MuCOE/CustomizableObjectCookPackageSplitter.cpp
+++ b/Engine/Plugins/Experimental/Mutable/Source/CustomizableObjectEditor/Private/MuCOE/CustomizableObjectCookPackageSplitter.cpp
@@ -11,16 +11,17 @@ REGISTER_COOKPACKAGE_SPLITTER(FCustomizableObjectCookPackageSplitter, UCustomiza
 
 namespace
 {
-// Look up a streamed Extension Data constant by name on a Customizable Object.
+// Look up a streamed Resource Data constant by name on a Customizable Object.
 //
 // Returns nullptr if not found.
-FCustomizableObjectStreamedExtensionData* FindStreamedExtensionData(
-	UCustomizableObject* Object,
-	const FString& ContainerName)
+FCustomizableObjectStreamedResourceData* FindStreamedResourceData(
+	TArray<FCustomizableObjectStreamedResourceData>& StreamedResources,
+	const FString& ContainerName
+	)
 {
 	return Algo::FindByPredicate(
-		Object->StreamedExtensionData,
-			[&ContainerName](const FCustomizableObjectStreamedExtensionData& StreamedData)
+		StreamedResources,
+			[&ContainerName](const FCustomizableObjectStreamedResourceData& StreamedData)
 			{
 				const FSoftObjectPath& Path = StreamedData.GetPath().ToSoftObjectPath();
 				
@@ -36,29 +37,22 @@ FCustomizableObjectStreamedExtensionData* FindStreamedExtensionData(
 			});
 }
 
-// Finds a streamed Extension Data container by name and moves it to the given Outer
+// Moves the StreamedResourceData's data container to the given Outer.
 bool MoveContainerToNewOuter(
-	UCustomizableObject* Object,
-	const FString& ContainerName,
 	UObject* NewOuter,
-	UCustomizableObjectExtensionDataContainer*& OutContainer)
+	const FCustomizableObjectStreamedResourceData* StreamedResourceData,
+	UCustomizableObjectResourceDataContainer*& OutContainer
+)
 {
+	check(StreamedResourceData);
+	
 	OutContainer = nullptr;
 
-	const FCustomizableObjectStreamedExtensionData* FoundData = FindStreamedExtensionData(Object, ContainerName);
-	if (!FoundData)
-	{
-		UE_LOG(LogMutable, Error, TEXT("Couldn't find streamed Extension Data container with name %s in array of %d entries"),
-			*ContainerName, Object->StreamedExtensionData.Num());
-
-		return false;
-	}
-
-	UCustomizableObjectExtensionDataContainer* Container = FoundData->GetPath().LoadSynchronous();
+	UCustomizableObjectResourceDataContainer* Container = StreamedResourceData->GetPath().LoadSynchronous();
 	if (!Container)
 	{
-		UE_LOG(LogMutable, Error, TEXT("Failed to load streamed Extension Data container %s"),
-			*FoundData->GetPath().ToString());
+		UE_LOG(LogMutable, Error, TEXT("Failed to load streamed Resource Data container %s"),
+			*StreamedResourceData->GetPath().ToString());
 
 		return false;
 	}
@@ -69,7 +63,7 @@ bool MoveContainerToNewOuter(
 	// The Rename function moves the object into the given package
 	if (!Container->Rename(nullptr, NewOuter, REN_DontCreateRedirectors))
 	{
-		UE_LOG(LogMutable, Error, TEXT("Failed to move streamed Extension Data container %s into Outer %s"),
+		UE_LOG(LogMutable, Error, TEXT("Failed to move streamed Resource Data container %s into Outer %s"),
 			*Container->GetPathName(), *NewOuter->GetPathName());
 
 		return false;
@@ -78,71 +72,90 @@ bool MoveContainerToNewOuter(
 	OutContainer = Container;
 	return true;
 }
+
+
+void GenerateNewPackage(const FCustomizableObjectStreamedResourceData& StreamedData,
+	const UPackage* OwnerPackage,
+	const UObject* OwnerObject,
+	TArray<ICookPackageSplitter::FGeneratedPackage>& Result)
+{
+	// The StreamedData container path should be of the form
+	// OwnerPackageName.OwnerObjectName:ContainerName
+	const FSoftObjectPath& StreamedDataPath = StreamedData.GetPath().ToSoftObjectPath();
+
+	// Check that the StreamedData container has the OwnerObject as its Outer
+	check(StreamedDataPath.GetWithoutSubPath() == FSoftObjectPath(OwnerObject));
+
+	// Check that the ContainerName is valid and that there isn't another Outer level between
+	// the OwnerObject and the container.
+	check(StreamedDataPath.GetSubPathString().Len() > 0);
+	check(!StreamedDataPath.GetSubPathString().Contains(SUBOBJECT_DELIMITER));
+
+	ICookPackageSplitter::FGeneratedPackage& Package = Result.AddDefaulted_GetRef();
+	// Because of the checks above, the container name must be unique within this Customizable
+	// Object, so it's safe to use as a package path.
+	Package.RelativePath = StreamedDataPath.GetSubPathString();
+	Package.SetCreateAsMap(false);
+
+	// To support iterative cooking, GenerationHash should only change when OwnerPackage
+	// changes.
+	//
+	// The simplest and fastest way to do this is to set it to OwnerPackage's GUID.
+	{
+		// Zero the hash, as we won't be writing all bytes of it below
+		Package.GenerationHash.Reset();
+
+		// Ensure there's no padding within FGuid, as the padding could be uninitialized
+		static_assert(sizeof(FGuid) == sizeof(uint32) * 4);
+		// Ensure an FGuid will fit within the hash buffer
+		static_assert(sizeof(FGuid) <= sizeof(FBlake3Hash::ByteArray));
+
+		// Create a copy that we can get the address of
+		//
+		// Note that UPackage::GetGuid is deprecated but there's no replacement yet. The engine
+		// team has advised me to use it anyway, so it's necessary to suppress the deprecation
+		// warning here.
+		PRAGMA_DISABLE_DEPRECATION_WARNINGS
+		const FGuid OwnerGuid = OwnerPackage->GetGuid();
+		PRAGMA_ENABLE_DEPRECATION_WARNINGS
+
+		FMemory::Memcpy(Package.GenerationHash.GetBytes(), &OwnerGuid, sizeof(FGuid));
+
+		// If OwnerGuid is non-zero, the hash should now be non-zero
+		check(!OwnerGuid.IsValid() || !Package.GenerationHash.IsZero());
+	}
+}
+
 }
 
 bool FCustomizableObjectCookPackageSplitter::ShouldSplit(UObject* SplitData)
 {
 	const UCustomizableObject* Object = CastChecked<UCustomizableObject>(SplitData);
 
-	return Object->StreamedExtensionData.Num() > 0;
+	return Object->StreamedResourceData.Num() > 0 || Object->StreamedExtensionData.Num() > 0;
 }
 
 TArray<ICookPackageSplitter::FGeneratedPackage> FCustomizableObjectCookPackageSplitter::GetGenerateList(
 	const UPackage* OwnerPackage,
 	const UObject* OwnerObject)
 {
+	// Keep a strong reference to the CO.
+	StrongObject.Reset(OwnerObject);
+	
 	const UCustomizableObject* Object = CastChecked<UCustomizableObject>(OwnerObject);
 
 	TArray<ICookPackageSplitter::FGeneratedPackage> Result;
 
+	// Generate a new package for each streamed Resource Data
+	for (const FCustomizableObjectStreamedResourceData& StreamedData : Object->StreamedResourceData)
+	{
+		GenerateNewPackage(StreamedData, OwnerPackage, OwnerObject, Result);
+	}
+
 	// Generate a new package for each streamed Extension Data
-	for (const FCustomizableObjectStreamedExtensionData& StreamedData : Object->StreamedExtensionData)
+	for (const FCustomizableObjectStreamedResourceData& StreamedData : Object->StreamedExtensionData)
 	{
-		// The StreamedData container path should be of the form
-		// OwnerPackageName.OwnerObjectName:ContainerName
-		const FSoftObjectPath& StreamedDataPath = StreamedData.GetPath().ToSoftObjectPath();
-		
-		// Check that the StreamedData container has the OwnerObject as its Outer
-		check(StreamedDataPath.GetWithoutSubPath() == FSoftObjectPath(OwnerObject));
-
-		// Check that the ContainerName is valid and that there isn't another Outer level between
-		// the OwnerObject and the container.
-		check(StreamedDataPath.GetSubPathString().Len() > 0);
-		check(!StreamedDataPath.GetSubPathString().Contains(SUBOBJECT_DELIMITER));
-
-		ICookPackageSplitter::FGeneratedPackage& Package = Result.AddDefaulted_GetRef();
-		// Because of the checks above, the container name must be unique within this Customizable
-		// Object, so it's safe to use as a package path.
-		Package.RelativePath = StreamedDataPath.GetSubPathString();
-		Package.SetCreateAsMap(false);
-
-		// To support iterative cooking, GenerationHash should only change when OwnerPackage
-		// changes.
-		//
-		// The simplest and fastest way to do this is to set it to OwnerPackage's GUID.
-		{
-			// Zero the hash, as we won't be writing all bytes of it below
-			Package.GenerationHash.Reset();
-			
-			// Ensure there's no padding within FGuid, as the padding could be uninitialized
-			static_assert(sizeof(FGuid) == sizeof(uint32) * 4);
-			// Ensure an FGuid will fit within the hash buffer
-			static_assert(sizeof(FGuid) <= sizeof(FBlake3Hash::ByteArray));
-
-			// Create a copy that we can get the address of
-			//
-			// Note that UPackage::GetGuid is deprecated but there's no replacement yet. The engine
-			// team has advised me to use it anyway, so it's necessary to suppress the deprecation
-			// warning here.
-			PRAGMA_DISABLE_DEPRECATION_WARNINGS
-			const FGuid OwnerGuid = OwnerPackage->GetGuid();
-			PRAGMA_ENABLE_DEPRECATION_WARNINGS
-
-			FMemory::Memcpy(Package.GenerationHash.GetBytes(), &OwnerGuid, sizeof(FGuid));
-
-			// If OwnerGuid is non-zero, the hash should now be non-zero
-			check(!OwnerGuid.IsValid() || !Package.GenerationHash.IsZero());
-		}
+		GenerateNewPackage(StreamedData, OwnerPackage, OwnerObject, Result);
 	}
 
 	return Result;
@@ -157,47 +170,74 @@ bool FCustomizableObjectCookPackageSplitter::PreSaveGeneratorPackage(
 	// The CO is just about to be saved (i.e. produce the cooked version of the asset), so this
 	// function needs to:
 	// 
-	// 1.	Move the streamed Extension Data out of the CO's package, so that it doesn't get saved
+	// 1.	Move the streamed Data out of the CO's package, so that it doesn't get saved
 	//		into the cooked package.
 	// 
 	// 2.	Remove hard references to the streamed data, so that it doesn't get loaded as soon as
 	//		the CO is loaded
 	//
 	// 3.	Fix up the soft references to point to the streamed data's new location
+	const auto& PreSavePackage = [] (const ICookPackageSplitter::FGeneratedPackageForPreSave& GeneratedPackage,
+		TArray<FCustomizableObjectStreamedResourceData>& StreamedResources
+		) -> bool
+	{
+		FCustomizableObjectStreamedResourceData* FoundData = FindStreamedResourceData(StreamedResources, GeneratedPackage.RelativePath);
+		if (!FoundData)
+		{
+			UE_LOG(LogMutable, Error, TEXT("Couldn't find streamed Resource Data container with name %s in array of %d entries"),
+				*GeneratedPackage.RelativePath, StreamedResources.Num());
+
+			return false;
+		}
+
+		// Move the streamed data to the generated package
+		UCustomizableObjectResourceDataContainer* Container = nullptr;
+		if (!MoveContainerToNewOuter(GeneratedPackage.Package, FoundData, Container))
+		{
+			return false;
+		}
+
+		// Remove the hard reference and set the soft reference to the streamed data's new location
+		FoundData->ConvertToSoftReferenceForCooking(TSoftObjectPtr<UCustomizableObjectResourceDataContainer>(Container));
+
+		return true;
+	};
 
 	UCustomizableObject* Object = CastChecked<UCustomizableObject>(OwnerObject);
 
-	// There should be one generated package per streamed Extension Data
-	check(Object->StreamedExtensionData.Num() == PlaceholderPackages.Num());
+	// There should be one generated package per streamed Resource Data
+	const int32 NumStreamedData = Object->StreamedResourceData.Num();
+	const int32 NumStreamedExtensionData = Object->StreamedExtensionData.Num();
+	
+	check(NumStreamedData + NumStreamedExtensionData == PlaceholderPackages.Num());
+
 
 	// After the CO has been saved, the contract for ICookPackageSplitter states that we need to
 	// restore the CO back to how it was before, so we need to save some information to help with
 	// this.
 	SavedContainerNames.Reset();
+	SavedExtensionContainerNames.Reset();
 
-	for (const ICookPackageSplitter::FGeneratedPackageForPreSave& GeneratedPackage : PlaceholderPackages)
+	for (int32 Index = 0; Index < NumStreamedData; ++Index)
 	{
-		// Find the streamed data based on the name of the generated package
-		FCustomizableObjectStreamedExtensionData* FoundData = FindStreamedExtensionData(Object, GeneratedPackage.RelativePath);
-		if (!FoundData)
+		const ICookPackageSplitter::FGeneratedPackageForPreSave& GeneratedPackage = PlaceholderPackages[Index];
+		if (!PreSavePackage(GeneratedPackage, Object->StreamedResourceData))
 		{
-			UE_LOG(LogMutable, Error, TEXT("Couldn't find streamed Extension Data container with name %s in array of %d entries"),
-				*GeneratedPackage.RelativePath, Object->StreamedExtensionData.Num());
-
 			return false;
 		}
 
-		// Move the streamed data to the generated package
-		UCustomizableObjectExtensionDataContainer* Container = nullptr;
-		if (!MoveContainerToNewOuter(Object, GeneratedPackage.RelativePath, GeneratedPackage.Package, Container))
+		SavedContainerNames.Add(GeneratedPackage.RelativePath);
+	}
+
+	for (int32 Index = 0; Index < NumStreamedExtensionData; ++Index)
+	{
+		const ICookPackageSplitter::FGeneratedPackageForPreSave& GeneratedPackage = PlaceholderPackages[NumStreamedData + Index];
+		if (!PreSavePackage(GeneratedPackage, Object->StreamedExtensionData))
 		{
 			return false;
 		}
 
-		SavedContainerNames.Add(GeneratedPackage.RelativePath);
-
-		// Remove the hard reference and set the soft reference to the streamed data's new location
-		FoundData->ConvertToSoftReferenceForCooking(TSoftObjectPtr<UCustomizableObjectExtensionDataContainer>(Container));
+		SavedExtensionContainerNames.Add(GeneratedPackage.RelativePath);
 	}
 
 	return true;
@@ -205,18 +245,48 @@ bool FCustomizableObjectCookPackageSplitter::PreSaveGeneratorPackage(
 
 void FCustomizableObjectCookPackageSplitter::PostSaveGeneratorPackage(UPackage* OwnerPackage, UObject* OwnerObject)
 {
-	// Move the streamed data back into the CO's package and restore the StreamedExtensionData
+	// Move the streamed data back into the CO's package and restore the StreamedResourceData and StreamedExtensionData
 	// array on the CO to how it was before PreSaveGeneratorPackage.
 
 	UCustomizableObject* Object = CastChecked<UCustomizableObject>(OwnerObject);
 
-	TArray<FCustomizableObjectStreamedExtensionData> NewArray;
-	NewArray.Reserve(SavedContainerNames.Num());
+	TArray<FCustomizableObjectStreamedResourceData> NewArray;
+	NewArray.Reset(SavedContainerNames.Num());
 
 	for (const FString& ContainerName : SavedContainerNames)
 	{
-		UCustomizableObjectExtensionDataContainer* Container = nullptr;
-		MoveContainerToNewOuter(Object, ContainerName, OwnerObject, Container);
+		FCustomizableObjectStreamedResourceData* ResourceData = FindStreamedResourceData(Object->StreamedResourceData, ContainerName);
+		if (!ResourceData)
+		{
+			UE_LOG(LogMutable, Error, TEXT("Couldn't find streamed Resource Data container with name %s in array of %d entries"),
+				*ContainerName, Object->StreamedResourceData.Num());
+
+			continue;
+		}
+
+		UCustomizableObjectResourceDataContainer* Container = nullptr;
+		MoveContainerToNewOuter(OwnerObject, ResourceData, Container);
+
+		NewArray.Emplace(Container);
+	}
+
+	Object->StreamedResourceData = NewArray;
+
+	NewArray.Reset(SavedExtensionContainerNames.Num());
+
+	for (const FString& ContainerName : SavedExtensionContainerNames)
+	{
+		FCustomizableObjectStreamedResourceData* ResourceData = FindStreamedResourceData(Object->StreamedExtensionData, ContainerName);
+		if (!ResourceData)
+		{
+			UE_LOG(LogMutable, Error, TEXT("Couldn't find streamed Extension Data container with name %s in array of %d entries"),
+				*ContainerName, Object->StreamedExtensionData.Num());
+
+			continue;
+		}
+
+		UCustomizableObjectResourceDataContainer* Container = nullptr;
+		MoveContainerToNewOuter(OwnerObject, ResourceData, Container);
 
 		NewArray.Emplace(Container);
 	}
@@ -234,9 +304,23 @@ bool FCustomizableObjectCookPackageSplitter::PopulateGeneratedPackage(
 	// Move the container into its newly generated package
 
 	UCustomizableObject* Object = CastChecked<UCustomizableObject>(OwnerObject);
+	
+	FCustomizableObjectStreamedResourceData* ResourceData = FindStreamedResourceData(Object->StreamedResourceData, GeneratedPackage.RelativePath);
+	if (!ResourceData)
+	{
+		ResourceData = FindStreamedResourceData(Object->StreamedExtensionData, GeneratedPackage.RelativePath);
+	}
 
-	UCustomizableObjectExtensionDataContainer* Container = nullptr;
-	if (!MoveContainerToNewOuter(Object, GeneratedPackage.RelativePath, GeneratedPackage.Package, Container))
+	if (!ResourceData)
+	{
+		UE_LOG(LogMutable, Error, TEXT("Couldn't find streamed resource Data container with name %s in arrays of %d and %d entries"),
+			*GeneratedPackage.RelativePath, Object->StreamedResourceData.Num(), Object->StreamedExtensionData.Num());
+
+		return false;
+	}
+
+	UCustomizableObjectResourceDataContainer* Container = nullptr;
+	if (!MoveContainerToNewOuter(GeneratedPackage.Package, ResourceData, Container))
 	{
 		return false;
 	}
@@ -256,6 +340,26 @@ void FCustomizableObjectCookPackageSplitter::PostSaveGeneratedPackage(
 
 	UCustomizableObject* Object = CastChecked<UCustomizableObject>(OwnerObject);
 
-	UCustomizableObjectExtensionDataContainer* Container = nullptr;
-	MoveContainerToNewOuter(Object, GeneratedPackage.RelativePath, OwnerObject, Container);
+	FCustomizableObjectStreamedResourceData* ResourceData = FindStreamedResourceData(Object->StreamedResourceData, GeneratedPackage.RelativePath);
+	if (!ResourceData)
+	{
+		ResourceData = FindStreamedResourceData(Object->StreamedExtensionData, GeneratedPackage.RelativePath);
+	}
+
+	if (!ResourceData)
+	{
+		UE_LOG(LogMutable, Error, TEXT("Couldn't find streamed resource Data container with name %s in arrays of %d and %d entries"),
+			*GeneratedPackage.RelativePath, Object->StreamedResourceData.Num(), Object->StreamedExtensionData.Num());
+
+		return;
+	}
+
+	UCustomizableObjectResourceDataContainer* Container = nullptr;
+	MoveContainerToNewOuter(OwnerObject, ResourceData, Container);
+
+}
+
+void FCustomizableObjectCookPackageSplitter::Teardown(ETeardown Status)
+{
+	StrongObject.Reset();
 }
diff --git a/Engine/Plugins/Experimental/Mutable/Source/CustomizableObjectEditor/Private/MuCOE/CustomizableObjectCookPackageSplitter.h b/Engine/Plugins/Experimental/Mutable/Source/CustomizableObjectEditor/Private/MuCOE/CustomizableObjectCookPackageSplitter.h
--- a/Engine/Plugins/Experimental/Mutable/Source/CustomizableObjectEditor/Private/MuCOE/CustomizableObjectCookPackageSplitter.h
+++ b/Engine/Plugins/Experimental/Mutable/Source/CustomizableObjectEditor/Private/MuCOE/CustomizableObjectCookPackageSplitter.h
@@ -3,9 +3,10 @@
 #pragma once
 
 #include "CookPackageSplitter.h"
+#include "UObject/StrongObjectPtr.h"
 #include "MuCO/CustomizableObject.h"
 
-/** Handles splitting the streamable Extension Data constants into their own packages */
+/** Handles splitting the streamable Data constants into their own packages */
 class FCustomizableObjectCookPackageSplitter : public ICookPackageSplitter
 {
 public:
@@ -37,7 +38,20 @@ public:
 		UObject* OwnerObject,
 		const FGeneratedPackageForPopulate& GeneratedPackage) override;
 
+	/** Do teardown actions after all packages have saved, or when the cook is cancelled. Always called before destruction. */
+	virtual void Teardown(ETeardown Status) override;
+
+	/**
+	 * If true, this splitter forces the Generator package objects it needs to remain referenced, and the cooker
+	 * should expect them to still be in memory after a garbage collect so long as the splitter is alive.
+	 */
+	virtual bool UseInternalReferenceToAvoidGarbageCollect() override { return true; }
+
 private:
-	TArray<TSoftObjectPtr<UCustomizableObjectExtensionDataContainer>> SavedStreamedExtensionData;
+
 	TArray<FString> SavedContainerNames;
+	TArray<FString> SavedExtensionContainerNames;
+
+	// Keep a strong reference to the CO to protect it from garbage collector.
+	TStrongObjectPtr<const UObject> StrongObject;
 };
diff --git a/Engine/Plugins/Experimental/Mutable/Source/CustomizableObjectEditor/Private/MuCOE/ExtensionDataCompilerInterface.cpp b/Engine/Plugins/Experimental/Mutable/Source/CustomizableObjectEditor/Private/MuCOE/ExtensionDataCompilerInterface.cpp
--- a/Engine/Plugins/Experimental/Mutable/Source/CustomizableObjectEditor/Private/MuCOE/ExtensionDataCompilerInterface.cpp
+++ b/Engine/Plugins/Experimental/Mutable/Source/CustomizableObjectEditor/Private/MuCOE/ExtensionDataCompilerInterface.cpp
@@ -3,7 +3,7 @@
 #include "MuCOE/ExtensionDataCompilerInterface.h"
 
 #include "InstancedStruct.h"
-#include "MuCO/CustomizableObjectStreamedExtensionData.h"
+#include "MuCO/CustomizableObjectStreamedResourceData.h"
 #include "MuCOE/GenerateMutableSource/GenerateMutableSource.h"
 #include "MuCOE/CustomizableObjectCompiler.h"
 #include "MuR/ExtensionData.h"
@@ -13,7 +13,7 @@ FExtensionDataCompilerInterface::FExtensionDataCompilerInterface(FMutableGraphGe
 {
 }
 
-mu::ExtensionDataPtrConst FExtensionDataCompilerInterface::MakeStreamedExtensionData(UCustomizableObjectExtensionDataContainer*& OutContainer)
+mu::ExtensionDataPtrConst FExtensionDataCompilerInterface::MakeStreamedExtensionData(UCustomizableObjectResourceDataContainer*& OutContainer)
 {
 	mu::ExtensionDataPtr Result = new mu::ExtensionData;
 	Result->Origin = mu::ExtensionData::EOrigin::ConstantStreamed;
@@ -35,7 +35,7 @@ mu::ExtensionDataPtrConst FExtensionDataCompilerInterface::MakeStreamedExtension
 	}
 
 	check(GenerationContext.Object);
-	OutContainer = NewObject<UCustomizableObjectExtensionDataContainer>(
+	OutContainer = NewObject<UCustomizableObjectResourceDataContainer>(
 		GenerationContext.Object,
 		FName(*ContainerName),
 		RF_Public);
@@ -51,7 +51,7 @@ mu::ExtensionDataPtrConst FExtensionDataCompilerInterface::MakeAlwaysLoadedExten
 	Result->Origin = mu::ExtensionData::EOrigin::ConstantAlwaysLoaded;
 	Result->Index = GenerationContext.AlwaysLoadedExtensionData.Num();
 
-	FCustomizableObjectExtensionData* CompileTimeExtensionData = &GenerationContext.AlwaysLoadedExtensionData.AddDefaulted_GetRef();
+	FCustomizableObjectResourceData* CompileTimeExtensionData = &GenerationContext.AlwaysLoadedExtensionData.AddDefaulted_GetRef();
 	CompileTimeExtensionData->Data = MoveTemp(Data);
 
 	return Result;
diff --git a/Engine/Plugins/Experimental/Mutable/Source/CustomizableObjectEditor/Private/MuCOE/GenerateMutableSource/GenerateMutableSource.cpp b/Engine/Plugins/Experimental/Mutable/Source/CustomizableObjectEditor/Private/MuCOE/GenerateMutableSource/GenerateMutableSource.cpp
--- a/Engine/Plugins/Experimental/Mutable/Source/CustomizableObjectEditor/Private/MuCOE/GenerateMutableSource/GenerateMutableSource.cpp
+++ b/Engine/Plugins/Experimental/Mutable/Source/CustomizableObjectEditor/Private/MuCOE/GenerateMutableSource/GenerateMutableSource.cpp
@@ -155,6 +155,65 @@ mu::MeshPtr FMutableGraphGenerationContext::FindGeneratedMesh( const FGeneratedM
 }
 
 
+int32 FMutableGraphGenerationContext::AddStreamedResource(uint32 InResourceHash, UCustomizableObjectResourceDataContainer*& OutNewResource)
+{
+	OutNewResource = nullptr;
+
+	// Return resource index if found.
+	if (int32* ResourceIndex = StreamedResourceIndices.Find(InResourceHash))
+	{
+		return *ResourceIndex;
+	}
+
+	int32 NewResourceIndex = StreamedResourceData.Num();
+	const FString ContainerName = FString::Printf(TEXT("SR_%d"), NewResourceIndex);
+
+	UObject* ExistingObject = FindObject<UObject>(Object, *ContainerName);
+	if (ExistingObject)
+	{
+		// This must have been left behind from a previous compilation and hasn't been deleted by 
+		// GC yet.
+		//
+		// Move it into the transient package to get it out of the way.
+		ExistingObject->Rename(nullptr, GetTransientPackage(), REN_DontCreateRedirectors);
+
+		check(!FindObject<UObject>(Object, *ContainerName));
+	}
+
+	OutNewResource = NewObject<UCustomizableObjectResourceDataContainer>(
+		Object,
+		FName(*ContainerName),
+		RF_Public);
+
+	StreamedResourceData.Add(OutNewResource);
+	StreamedResourceIndices.Add({ InResourceHash, NewResourceIndex });
+
+	return NewResourceIndex;
+}
+
+
+int32 FMutableGraphGenerationContext::AddAssetUserDataToStreamedResources(UAssetUserData* AssetUserData)
+{
+	check(AssetUserData);
+	const uint32 AssetIdentifier = AssetUserData->GetUniqueID();
+
+	UCustomizableObjectResourceDataContainer* NewResource = nullptr;
+	const int32 ResourceIndex = AddStreamedResource(AssetIdentifier, NewResource);
+
+	if (NewResource) // Nullptr if not new
+	{
+		FCustomizableObjectAssetUserData ResourceData;
+		ResourceData.AssetUserDataEditor = AssetUserData;
+
+		NewResource->Data.Type = ECOResourceDataType::AssetUserData;
+		NewResource->Data.Data = FInstancedStruct::Make(ResourceData);
+	}
+
+	check(StreamedResourceData[ResourceIndex].GetLoadedData().Type == ECOResourceDataType::AssetUserData);
+	return ResourceIndex;
+}
+
+
 /** Adds to ParameterNamesMap the node Node to the array of elements with name Name */
 void FMutableGraphGenerationContext::AddParameterNameUnique(const UCustomizableObjectNode* Node, FString Name)
 {
@@ -1418,12 +1477,6 @@ FString GenerateGameplayTag(const FString& GameplayTag)
 }
 
 
-FString GenerateAssetUserDataTag(const FString& AssetUserData)
-{
-	return FString("__AssetUserData:") + AssetUserData;
-}
-
-
 void PopulateReferenceSkeletalMeshesData(FMutableGraphGenerationContext& GenerationContext)
 {
 	const FString PlatformName = GenerationContext.Options.TargetPlatform->IniPlatformName();
diff --git a/Engine/Plugins/Experimental/Mutable/Source/CustomizableObjectEditor/Private/MuCOE/GenerateMutableSource/GenerateMutableSource.h b/Engine/Plugins/Experimental/Mutable/Source/CustomizableObjectEditor/Private/MuCOE/GenerateMutableSource/GenerateMutableSource.h
--- a/Engine/Plugins/Experimental/Mutable/Source/CustomizableObjectEditor/Private/MuCOE/GenerateMutableSource/GenerateMutableSource.h
+++ b/Engine/Plugins/Experimental/Mutable/Source/CustomizableObjectEditor/Private/MuCOE/GenerateMutableSource/GenerateMutableSource.h
@@ -632,6 +632,15 @@ struct FMutableGraphGenerationContext
 	/** Find a mesh if already generated for a given source and flags. */
 	mu::Ptr<mu::Mesh> FindGeneratedMesh(const FGeneratedMeshData::FKey& Key);
 
+	/** Add a resource to the streamed resources array.
+	  * OutStreamedResourceContainer - Container to store the streamed resource.
+	  * Returns resource index in the array of streamed resources. */
+	int32 AddStreamedResource(const uint32 InResourceHash, UCustomizableObjectResourceDataContainer*& OutStreamedResourceContainer);
+
+	/** Adds a streamed resource of type AssetUserData.
+	  * Returns resource index in the array of streamed resources. */
+	int32 AddAssetUserDataToStreamedResources(UAssetUserData* AssetUserData);
+
 	/** Adds to ParameterNamesMap the node Node to the array of elements with name Name */
 	void AddParameterNameUnique(const UCustomizableObjectNode* Node, FString Name);
 
@@ -750,9 +759,6 @@ struct FMutableGraphGenerationContext
 	/** Stores the anim BP assets gathered from the SkeletalMesh nodes during compilation, to be used in mesh generation in-game */
 	TMap<FString, TSoftClassPtr<UAnimInstance>> AnimBPAssetsMap;
 
-	/** Stores the AssetUserData assets gathered from the SkeletalMesh nodes during compilation, to be used in mesh generation in-game */
-	TMap<FString, TSoftObjectPtr<UAssetUserData>> AssetUserDataAssetsMap;
-
 	/** Stores the sockets provided by the part skeletal meshes, to be merged in the generated meshes */
 	TArray<FMutableRefSocket> SocketArray;
 
@@ -805,10 +811,14 @@ struct FMutableGraphGenerationContext
 	// UCustomizableObjectNodeMaterial material to SharedSurfaceId
 	TMap<UCustomizableObjectNodeMaterial*, TArray<FSharedSurface>> SharedSurfaceIds;
 
+	/** Resource Data constants */
+	TMap<uint32, int32> StreamedResourceIndices;
+	TArray<FCustomizableObjectStreamedResourceData> StreamedResourceData;
+
 	/** Extension Data constants are collected here */
 	FExtensionDataCompilerInterface ExtensionDataCompilerInterface;
-	TArray<FCustomizableObjectExtensionData> AlwaysLoadedExtensionData;
-	TArray<UCustomizableObjectExtensionDataContainer*> StreamedExtensionData;
+	TArray<FCustomizableObjectResourceData> AlwaysLoadedExtensionData;
+	TArray<UCustomizableObjectResourceDataContainer*> StreamedExtensionData;
 
 	/** See UCustomizableObject::ParticipatingObjects. */
 	TMap<TObjectPtr<const UObject>, FGuid> ParticipatingObjects;
diff --git a/Engine/Plugins/Experimental/Mutable/Source/CustomizableObjectEditor/Private/MuCOE/GenerateMutableSource/GenerateMutableSourceMesh.cpp b/Engine/Plugins/Experimental/Mutable/Source/CustomizableObjectEditor/Private/MuCOE/GenerateMutableSource/GenerateMutableSourceMesh.cpp
--- a/Engine/Plugins/Experimental/Mutable/Source/CustomizableObjectEditor/Private/MuCOE/GenerateMutableSource/GenerateMutableSourceMesh.cpp
+++ b/Engine/Plugins/Experimental/Mutable/Source/CustomizableObjectEditor/Private/MuCOE/GenerateMutableSource/GenerateMutableSourceMesh.cpp
@@ -2878,22 +2878,25 @@ mu::NodeMeshPtr GenerateMutableSourceMesh(const UEdGraphPin* Pin,
 				MeshUniqueTags += AnimBPTag;
 			}
 
-			TArray<FString> AssetUserDataTags;
+			TArray<int32> StreamedResources;
 
 			if (GenerationContext.Object->bEnableAssetUserDataMerge)
 			{
 				const TArray<UAssetUserData*>* AssetUserDataArray = TypedNodeSkel->SkeletalMesh->GetAssetUserDataArray();
 
-				if (AssetUserDataArray && !AssetUserDataArray->IsEmpty())
+				if (AssetUserDataArray)
 				{
-					for (const UAssetUserData* AssetUserData : *AssetUserDataArray)
+					for (UAssetUserData* AssetUserData : *AssetUserDataArray)
 					{
-						FString AuxString = AssetUserData->GetPathName();
-						GenerationContext.AssetUserDataAssetsMap.Add(AuxString, TSoftObjectPtr<UAssetUserData>(AssetUserData));
+						if (!AssetUserData)
+						{
+							continue;
+						}
+
+						const int32 ResourceIndex = GenerationContext.AddAssetUserDataToStreamedResources(AssetUserData);
+						StreamedResources.Add(ResourceIndex);
 
-						FString AssetUserDataTag = GenerateAssetUserDataTag(AuxString);
-						AssetUserDataTags.Add(AssetUserDataTag);
-						MeshUniqueTags += AssetUserDataTag;
+						MeshUniqueTags += AssetUserData->GetPathName();
 					}
 				}
 			}
@@ -2971,9 +2974,9 @@ mu::NodeMeshPtr GenerateMutableSourceMesh(const UEdGraphPin* Pin,
 					AddTagToMutableMeshUnique(*MutableMesh, GamePlayTag);
 				}
 
-				for (const FString& AssetUserDataTag : AssetUserDataTags)
+				for (int32 ResourceIndex : StreamedResources)
 				{
-					AddTagToMutableMeshUnique(*MutableMesh, AssetUserDataTag);
+					MutableMesh->AddStreamedResource(ResourceIndex);
 				}
 
 				AddSocketTagsToMesh(TypedNodeSkel->SkeletalMesh, MutableMesh, GenerationContext);
diff --git a/Engine/Plugins/Experimental/Mutable/Source/CustomizableObjectEditor/Private/MuCOE/GenerateMutableSource/GenerateMutableSourceTable.cpp b/Engine/Plugins/Experimental/Mutable/Source/CustomizableObjectEditor/Private/MuCOE/GenerateMutableSource/GenerateMutableSourceTable.cpp
--- a/Engine/Plugins/Experimental/Mutable/Source/CustomizableObjectEditor/Private/MuCOE/GenerateMutableSource/GenerateMutableSourceTable.cpp
+++ b/Engine/Plugins/Experimental/Mutable/Source/CustomizableObjectEditor/Private/MuCOE/GenerateMutableSource/GenerateMutableSourceTable.cpp
@@ -243,6 +243,29 @@ bool FillTableColumn(const UCustomizableObjectNodeTable* TableNode,	mu::TablePtr
 			{
 				MeshUniqueTags += GenerateGameplayTag(Tag.ToString());
 			}
+			
+			TArray<int32> StreamedResources;
+
+			if (GenerationContext.Object->bEnableAssetUserDataMerge)
+			{
+				const TArray<UAssetUserData*>* AssetUserDataArray = SkeletalMesh->GetAssetUserDataArray();
+
+				if (AssetUserDataArray)
+				{
+					for (UAssetUserData* AssetUserData : *AssetUserDataArray)
+					{
+						if (!AssetUserData)
+						{
+							continue;
+						}
+
+						const int32 ResourceIndex = GenerationContext.AddAssetUserDataToStreamedResources(AssetUserData);
+						StreamedResources.Add(ResourceIndex);
+
+						MeshUniqueTags += AssetUserData->GetPathName();
+					}
+				}
+			}
 
 			//TODO: Add AnimBp physics to Tables.
 			mu::Ptr<mu::Mesh> MutableMesh = GenerateMutableMesh(SkeletalMesh, TSoftClassPtr<UAnimInstance>(), LODIndexConnected, SectionIndexConnected, LODIndex, SectionIndex, MeshUniqueTags, GenerationContext, TableNode);
@@ -271,21 +294,9 @@ bool FillTableColumn(const UCustomizableObjectNodeTable* TableNode,	mu::TablePtr
 					AddTagToMutableMeshUnique(*MutableMesh, GenerateGameplayTag(Tag.ToString()));
 				}
 
-				if (GenerationContext.Object->bEnableAssetUserDataMerge)
+				for (int32 ResourceIndex : StreamedResources)
 				{
-					const TArray<UAssetUserData*>* AssetUserDataArray = SkeletalMesh->GetAssetUserDataArray();
-
-					if (AssetUserDataArray && !AssetUserDataArray->IsEmpty())
-					{
-						for (const UAssetUserData* AssetUserData : *AssetUserDataArray)
-						{
-							FString AuxString = AssetUserData->GetPathName();
-							GenerationContext.AssetUserDataAssetsMap.Add(AuxString, TSoftObjectPtr<UAssetUserData>(AssetUserData));
-
-							FString AssetUserDataTag = GenerateAssetUserDataTag(AuxString);
-							AddTagToMutableMeshUnique(*MutableMesh, AssetUserDataTag);
-						}
-					}
+					MutableMesh->AddStreamedResource(ResourceIndex);
 				}
 
 				AddSocketTagsToMesh(SkeletalMesh, MutableMesh, GenerationContext);
diff --git a/Engine/Plugins/Experimental/Mutable/Source/CustomizableObjectEditor/Public/MuCOE/ExtensionDataCompilerInterface.h b/Engine/Plugins/Experimental/Mutable/Source/CustomizableObjectEditor/Public/MuCOE/ExtensionDataCompilerInterface.h
--- a/Engine/Plugins/Experimental/Mutable/Source/CustomizableObjectEditor/Public/MuCOE/ExtensionDataCompilerInterface.h
+++ b/Engine/Plugins/Experimental/Mutable/Source/CustomizableObjectEditor/Public/MuCOE/ExtensionDataCompilerInterface.h
@@ -5,7 +5,7 @@
 #include "MuR/ExtensionData.h"
 
 struct FInstancedStruct;
-class UCustomizableObjectExtensionDataContainer;
+class UCustomizableObjectResourceDataContainer;
 struct FMutableGraphGenerationContext;
 
 /**
@@ -40,7 +40,7 @@ public:
 	 * TSoftObjectPtr, because the container will be moved to a different package during cooking
 	 * and your paths won't be automatically updated.
 	 */
-	mu::ExtensionDataPtrConst MakeStreamedExtensionData(UCustomizableObjectExtensionDataContainer*& OutContainer);
+	mu::ExtensionDataPtrConst MakeStreamedExtensionData(UCustomizableObjectResourceDataContainer*& OutContainer);
 
 	/**
 	 * Register a new Extension Data compile-time constant that will always be loaded.
diff --git a/Engine/Plugins/Experimental/Mutable/Source/MutableRuntime/Private/MuR/Mesh.cpp b/Engine/Plugins/Experimental/Mutable/Source/MutableRuntime/Private/MuR/Mesh.cpp
--- a/Engine/Plugins/Experimental/Mutable/Source/MutableRuntime/Private/MuR/Mesh.cpp
+++ b/Engine/Plugins/Experimental/Mutable/Source/MutableRuntime/Private/MuR/Mesh.cpp
@@ -51,6 +51,7 @@ MeshPtr Mesh::Clone() const
 	pResult->m_pSkeleton = m_pSkeleton;
 	pResult->m_pPhysicsBody = m_pPhysicsBody;
 	pResult->m_tags = m_tags;
+	pResult->StreamedResources = StreamedResources;
 
     // Clone the main buffers
     pResult->m_VertexBuffers = m_VertexBuffers;
@@ -109,6 +110,11 @@ MeshPtr Mesh::Clone(EMeshCopyFlags Flags) const
 		pResult->m_tags = m_tags;
 	}
 
+	if (EnumHasAnyFlags(Flags, EMeshCopyFlags::WithStreamedResources))
+	{
+		pResult->StreamedResources = StreamedResources;
+	}
+
     // Clone the main buffers
 	if (EnumHasAnyFlags(Flags, EMeshCopyFlags::WithVertexBuffers))
     {
@@ -196,6 +202,11 @@ void Mesh::CopyFrom(const Mesh& From, EMeshCopyFlags Flags)
 		m_tags = From.m_tags;
 	}
 
+	if (EnumHasAnyFlags(Flags, EMeshCopyFlags::WithStreamedResources))
+	{
+		StreamedResources = From.StreamedResources;
+	}
+
     // Copy the main buffers
 	if (EnumHasAnyFlags(Flags, EMeshCopyFlags::WithVertexBuffers))
     {
@@ -520,6 +531,23 @@ void Mesh::SetTag( int tagIndex, const FString& Name )
 }
 
 
+//---------------------------------------------------------------------------------------------
+void Mesh::AddStreamedResource(int32 ResourceIndex)
+{
+	if (ensure(ResourceIndex >= 0))
+	{
+		StreamedResources.AddUnique(ResourceIndex);
+	}
+}
+
+
+//---------------------------------------------------------------------------------------------
+const TArray<int32>& Mesh::GetStreamedResources() const
+{
+	return StreamedResources;
+}
+
+
 //---------------------------------------------------------------------------------------------
 int32 Mesh::FindBonePose(const uint16 BoneId) const
 {
@@ -1180,7 +1208,7 @@ void Mesh::FBonePose::Unserialise(InputArchive& arch)
 //-------------------------------------------------------------------------------------------------
 void Mesh::Serialise(OutputArchive& arch) const
 {
-	uint32 ver = 17;
+	uint32 ver = 18;
 	arch << ver;
 
 	arch << m_IndexBuffers;
@@ -1198,6 +1226,7 @@ void Mesh::Serialise(OutputArchive& arch) const
 	arch << m_surfaces;
 
 	arch << m_tags;
+	arch << StreamedResources;
 
 	arch << BonePoses;
 	arch << BoneMap;
@@ -1211,7 +1240,7 @@ void Mesh::Unserialise(InputArchive& arch)
 {
 	uint32 ver;
 	arch >> ver;
-	check(ver <= 17);
+	check(ver <= 18);
 
 	arch >> m_IndexBuffers;
 	arch >> m_VertexBuffers;
@@ -1280,6 +1309,11 @@ void Mesh::Unserialise(InputArchive& arch)
 		arch >> m_tags;
 	}
 
+	if (ver >= 18)
+	{
+		arch >> StreamedResources;
+	}
+
 	if (ver >= 13)
 	{
 		arch >> BonePoses;
diff --git a/Engine/Plugins/Experimental/Mutable/Source/MutableRuntime/Private/MuR/OpMeshFormat.cpp b/Engine/Plugins/Experimental/Mutable/Source/MutableRuntime/Private/MuR/OpMeshFormat.cpp
--- a/Engine/Plugins/Experimental/Mutable/Source/MutableRuntime/Private/MuR/OpMeshFormat.cpp
+++ b/Engine/Plugins/Experimental/Mutable/Source/MutableRuntime/Private/MuR/OpMeshFormat.cpp
@@ -681,6 +681,7 @@ namespace mu
 		}
 
 		Result->m_tags = pSource->m_tags;
+		Result->StreamedResources = pSource->StreamedResources;
 
 		Result->AdditionalBuffers = pSource->AdditionalBuffers;
 
diff --git a/Engine/Plugins/Experimental/Mutable/Source/MutableRuntime/Private/MuR/OpMeshMerge.h b/Engine/Plugins/Experimental/Mutable/Source/MutableRuntime/Private/MuR/OpMeshMerge.h
--- a/Engine/Plugins/Experimental/Mutable/Source/MutableRuntime/Private/MuR/OpMeshMerge.h
+++ b/Engine/Plugins/Experimental/Mutable/Source/MutableRuntime/Private/MuR/OpMeshMerge.h
@@ -1144,6 +1144,16 @@ namespace mu
 		{
 			Result->m_tags.AddUnique(SecondTag);
 		}
+
+
+		// Streamed Resources
+		Result->StreamedResources = pFirst->StreamedResources;
+
+		const int32 NumStreamedResources = pSecond->StreamedResources.Num();
+		for (int32 Index = 0; Index < NumStreamedResources; ++Index)
+		{
+			Result->StreamedResources.AddUnique(pSecond->StreamedResources[Index]);
+		}
 	}
 	
 
diff --git a/Engine/Plugins/Experimental/Mutable/Source/MutableRuntime/Public/MuR/Mesh.h b/Engine/Plugins/Experimental/Mutable/Source/MutableRuntime/Public/MuR/Mesh.h
--- a/Engine/Plugins/Experimental/Mutable/Source/MutableRuntime/Public/MuR/Mesh.h
+++ b/Engine/Plugins/Experimental/Mutable/Source/MutableRuntime/Public/MuR/Mesh.h
@@ -150,6 +150,7 @@ namespace mu
 		WithBoneMap = 1 << 13,
 		WithSkeletonIDs = 1 << 14,
 		WithAdditionalPhysics = 1 << 15,
+		WithStreamedResources = 1 << 16,
 
 		AllFlags = 0xFFFFFFFF
 	};
@@ -283,6 +284,12 @@ namespace mu
         //!
         void SetTag( int tagIndex, const FString& Name );
 
+		//!
+		void AddStreamedResource(int32 ResourceIndex);
+
+		//!
+		const TArray<int32>& GetStreamedResources() const;
+
 		//!
 		int32 FindBonePose(uint16 BoneId) const;
 		
@@ -379,6 +386,9 @@ namespace mu
 		//!
 		TArray<FString> m_tags;
 
+		//!
+		TArray<int32> StreamedResources;
+
 		struct FBonePose
 		{
 			// Index of the bone in the CO BoneNames array
@@ -434,6 +444,7 @@ namespace mu
 					equal = false;
 				}
 			}
+			if (equal) equal = (StreamedResources == o.StreamedResources);
 			if (equal) equal = (m_surfaces == o.m_surfaces);
 			if (equal) equal = (m_tags == o.m_tags);
 			if (equal) equal = (SkeletonIDs == o.SkeletonIDs);
