[Mutable] Fix crashes when trying to load a deleted/renamed asset without checking if it exists.
#rnx

diff --git a/Engine/Plugins/Experimental/Mutable/Source/CustomizableObject/Private/MuCO/CustomizableObject.cpp b/Engine/Plugins/Experimental/Mutable/Source/CustomizableObject/Private/MuCO/CustomizableObject.cpp
--- a/Engine/Plugins/Experimental/Mutable/Source/CustomizableObject/Private/MuCO/CustomizableObject.cpp
+++ b/Engine/Plugins/Experimental/Mutable/Source/CustomizableObject/Private/MuCO/CustomizableObject.cpp
@@ -449,7 +449,13 @@ void SerializeStreamedResources(FArchive& Ar, UObject* Object, TArray<FCustomiza
 			case ECOResourceDataType::AssetUserData:
 			{
 				const FCustomizableObjectAssetUserData* AssetUserData = Data.Data.GetPtr<FCustomizableObjectAssetUserData>();
-				FString AssetUserDataPath = TSoftObjectPtr<UAssetUserData>(AssetUserData->AssetUserDataEditor).ToSoftObjectPath().ToString();
+				FString AssetUserDataPath;
+
+				if (AssetUserData->AssetUserDataEditor)
+				{
+					AssetUserDataPath = TSoftObjectPtr<UAssetUserData>(AssetUserData->AssetUserDataEditor).ToSoftObjectPath().ToString();
+				}
+
 				Ar << AssetUserDataPath;
 				break;
 			}
@@ -508,7 +514,14 @@ void SerializeStreamedResources(FArchive& Ar, UObject* Object, TArray<FCustomiza
 					Ar << AssetUserDataPath;
 					
 					FCustomizableObjectAssetUserData ResourceData;
-					ResourceData.AssetUserDataEditor = TSoftObjectPtr<UAssetUserData>(AssetUserDataPath).LoadSynchronous();
+
+					TSoftObjectPtr<UAssetUserData> SoftAssetUserData(AssetUserDataPath);
+					ResourceData.AssetUserDataEditor = !SoftAssetUserData.IsNull() ? SoftAssetUserData.LoadSynchronous() : nullptr;
+
+					if (!ResourceData.AssetUserDataEditor)
+					{
+						UE_LOG(LogMutable, Warning, TEXT("Failed to load streamed resource of type AssetUserData. Resource name: [%s]"), *AssetUserDataPath);
+					}
 
 					if (bIsCooking)
 					{
