[Mutable] Fix crashes due to a mem free of an invalid data container.
#rnx

diff --git a/Engine/Plugins/Experimental/Mutable/Source/CustomizableObject/Classes/MuCO/CustomizableObject.h b/Engine/Plugins/Experimental/Mutable/Source/CustomizableObject/Classes/MuCO/CustomizableObject.h
--- a/Engine/Plugins/Experimental/Mutable/Source/CustomizableObject/Classes/MuCO/CustomizableObject.h
+++ b/Engine/Plugins/Experimental/Mutable/Source/CustomizableObject/Classes/MuCO/CustomizableObject.h
@@ -1490,7 +1490,7 @@ private:
 	// This is a manual version number for the binary blobs in this asset.
 	// Increasing it invalidates all the previously compiled models.
 	// Warning: If while merging code both versions have changed, take the highest+1.
-	static const int32 CurrentSupportedVersion = 411;
+	static const int32 CurrentSupportedVersion = 412;
 
 public:
 
diff --git a/Engine/Plugins/Experimental/Mutable/Source/CustomizableObject/Private/MuCO/CustomizableObject.cpp b/Engine/Plugins/Experimental/Mutable/Source/CustomizableObject/Private/MuCO/CustomizableObject.cpp
--- a/Engine/Plugins/Experimental/Mutable/Source/CustomizableObject/Private/MuCO/CustomizableObject.cpp
+++ b/Engine/Plugins/Experimental/Mutable/Source/CustomizableObject/Private/MuCO/CustomizableObject.cpp
@@ -442,7 +442,8 @@ void SerializeStreamedResources(FArchive& Ar, UObject* Object, TArray<FCustomiza
 		for (const FCustomizableObjectStreamedResourceData& ResourceData : StreamedResources)
 		{
 			const FCustomizableObjectResourceData& Data = ResourceData.GetLoadedData();
-			Ar << Data.Type;
+			uint32 Type = (uint32)Data.Type;
+			Ar << Type;
 
 			switch (Data.Type)
 			{
@@ -475,7 +476,7 @@ void SerializeStreamedResources(FArchive& Ar, UObject* Object, TArray<FCustomiza
 		// Initialize if not cooking. Otherwise, resources will be initialized at this point, and only their data will be updated.
 		if (!bIsCooking) 
 		{
-			StreamedResources.SetNumUninitialized(NumStreamedResources);
+			StreamedResources.SetNum(NumStreamedResources);
 		}			
 		
 		check(NumStreamedResources == StreamedResources.Num())
@@ -505,7 +506,10 @@ void SerializeStreamedResources(FArchive& Ar, UObject* Object, TArray<FCustomiza
 			}
 
 			check(Container);
-			Ar << Container->Data.Type;
+			uint32 Type = 0;
+			Ar << Type;
+			
+			Container->Data.Type = (ECOResourceDataType)Type;
 			switch (Container->Data.Type)
 			{
 				case ECOResourceDataType::AssetUserData:
@@ -541,6 +545,7 @@ void SerializeStreamedResources(FArchive& Ar, UObject* Object, TArray<FCustomiza
 					break;
 				}
 				default:
+					check(false);
 					break;
 			}
 		}
