#rnx
#tests PIE

diff --git a/Engine/Plugins/Experimental/Mutable/Source/CustomizableObject/Private/MuCO/CustomizableInstanceLODManagement.cpp b/Engine/Plugins/Experimental/Mutable/Source/CustomizableObject/Private/MuCO/CustomizableInstanceLODManagement.cpp
--- a/Engine/Plugins/Experimental/Mutable/Source/CustomizableObject/Private/MuCO/CustomizableInstanceLODManagement.cpp
+++ b/Engine/Plugins/Experimental/Mutable/Source/CustomizableObject/Private/MuCO/CustomizableInstanceLODManagement.cpp
@@ -73,13 +73,13 @@ void UpdatePawnToInstancesDistances(const class UCustomizableObjectInstance* Onl
 	for (TObjectIterator<UCustomizableObjectInstanceUsage> CustomizableObjectInstanceUsage; CustomizableObjectInstanceUsage; ++CustomizableObjectInstanceUsage)
 	{
 #if WITH_EDITOR
-		if (CustomizableObjectInstanceUsage && CustomizableObjectInstanceUsage->IsNetMode(NM_DedicatedServer))
+		if (IsValid(*CustomizableObjectInstanceUsage) && CustomizableObjectInstanceUsage->IsNetMode(NM_DedicatedServer))
 		{
 			continue;
 		}
 #endif
 
-		if (CustomizableObjectInstanceUsage->IsValidLowLevel() && (OnlyForInstance == nullptr || CustomizableObjectInstanceUsage->GetCustomizableObjectInstance() == OnlyForInstance))
+		if (IsValid(*CustomizableObjectInstanceUsage) && (OnlyForInstance == nullptr || CustomizableObjectInstanceUsage->GetCustomizableObjectInstance() == OnlyForInstance))
 		{
 			CustomizableObjectInstanceUsage->UpdateDistFromComponentToPlayer(ViewCenter.IsValid() ? ViewCenter.Get() : nullptr, OnlyForInstance != nullptr);
 		}
@@ -93,12 +93,12 @@ void UpdateCameraToInstancesDistance(const FVector CameraPosition)
 	for (TObjectIterator<UCustomizableObjectInstanceUsage> CustomizableObjectInstanceUsage; CustomizableObjectInstanceUsage; ++CustomizableObjectInstanceUsage)
 	{
 #if WITH_EDITOR
-		if (CustomizableObjectInstanceUsage && CustomizableObjectInstanceUsage->IsNetMode(NM_DedicatedServer))
+		if (IsValid(*CustomizableObjectInstanceUsage) && CustomizableObjectInstanceUsage->IsNetMode(NM_DedicatedServer))
 		{
 			continue;
 		}
 #endif
-		if (CustomizableObjectInstanceUsage->IsValidLowLevel() && !CustomizableObjectInstanceUsage->IsTemplate())
+		if (IsValid(*CustomizableObjectInstanceUsage) && !CustomizableObjectInstanceUsage->IsTemplate())
 		{
 			CustomizableObjectInstanceUsage->UpdateDistFromComponentToLevelEditorCamera(CameraPosition);
 		}
@@ -125,7 +125,7 @@ void UCustomizableInstanceLODManagement::UpdateInstanceDistsAndLODs(FMutableInst
 			for (TObjectIterator<UCustomizableObjectInstanceUsage> CustomizableObjectInstanceUsage; CustomizableObjectInstanceUsage; ++CustomizableObjectInstanceUsage)
 			{
 #if WITH_EDITOR
-				if (CustomizableObjectInstanceUsage && CustomizableObjectInstanceUsage->IsNetMode(NM_DedicatedServer))
+				if (IsValid(*CustomizableObjectInstanceUsage) && CustomizableObjectInstanceUsage->IsNetMode(NM_DedicatedServer))
 				{
 					continue;
 				}
@@ -200,7 +200,7 @@ void UCustomizableInstanceLODManagement::UpdateInstanceDistsAndLODs(FMutableInst
 
 		for (TObjectIterator<UCustomizableObjectInstance> CustomizableObjectInstance; CustomizableObjectInstance; ++CustomizableObjectInstance)
 		{
-			if (IsValidChecked(*CustomizableObjectInstance) &&
+			if (IsValid(*CustomizableObjectInstance) &&
 				CustomizableObjectInstance->GetPrivate() &&
 				CustomizableObjectInstance->GetIsBeingUsedByComponentInPlay())
 			{
@@ -302,13 +302,13 @@ void UCustomizableInstanceLODManagement::UpdateInstanceDistsAndLODs(FMutableInst
 		for (TObjectIterator<UCustomizableObjectInstanceUsage> CustomizableObjectInstanceUsage; CustomizableObjectInstanceUsage; ++CustomizableObjectInstanceUsage)
 		{
 #if WITH_EDITOR
-			if (CustomizableObjectInstanceUsage && CustomizableObjectInstanceUsage->IsNetMode(NM_DedicatedServer))
+			if (IsValid(*CustomizableObjectInstanceUsage) && CustomizableObjectInstanceUsage->IsNetMode(NM_DedicatedServer))
 			{
 				continue;
 			}
 #endif
 
-			if (CustomizableObjectInstanceUsage && !CustomizableObjectInstanceUsage->IsTemplate())
+			if (IsValid(*CustomizableObjectInstanceUsage) && !CustomizableObjectInstanceUsage->IsTemplate())
 			{
 				UCustomizableObjectInstance* COI = CustomizableObjectInstanceUsage->GetCustomizableObjectInstance();
 				if (!COI || !COI->GetCustomizableObject())
diff --git a/Engine/Plugins/Experimental/Mutable/Source/CustomizableObject/Private/MuCO/CustomizableObjectInstance.cpp b/Engine/Plugins/Experimental/Mutable/Source/CustomizableObject/Private/MuCO/CustomizableObjectInstance.cpp
--- a/Engine/Plugins/Experimental/Mutable/Source/CustomizableObject/Private/MuCO/CustomizableObjectInstance.cpp
+++ b/Engine/Plugins/Experimental/Mutable/Source/CustomizableObject/Private/MuCO/CustomizableObjectInstance.cpp
@@ -2755,13 +2755,13 @@ void UCustomizableInstancePrivateData::DiscardResourcesAndSetReferenceSkeletalMe
 		UCustomizableObjectInstanceUsage* CustomizableObjectInstanceUsage = *It;
 
 #if WITH_EDITOR
-		if (CustomizableObjectInstanceUsage && CustomizableObjectInstanceUsage->IsNetMode(NM_DedicatedServer))
+		if (IsValid(CustomizableObjectInstanceUsage) && CustomizableObjectInstanceUsage->IsNetMode(NM_DedicatedServer))
 		{
 			continue;
 		}
 #endif
 
-		if (CustomizableObjectInstanceUsage && CustomizableObjectInstanceUsage->GetCustomizableObjectInstance() == Instance)
+		if (IsValid(CustomizableObjectInstanceUsage) && CustomizableObjectInstanceUsage->GetCustomizableObjectInstance() == Instance)
 		{
 			UCustomizableObject* CustomizableObject = Instance->GetCustomizableObject();
 			bool bReplaceDiscardedWithReferenceMesh = UCustomizableObjectSystem::GetInstance()->GetPrivate()->IsReplaceDiscardedWithReferenceMeshEnabled();
diff --git a/Engine/Plugins/Experimental/Mutable/Source/CustomizableObject/Private/MuCO/CustomizableObjectModule.cpp b/Engine/Plugins/Experimental/Mutable/Source/CustomizableObject/Private/MuCO/CustomizableObjectModule.cpp
--- a/Engine/Plugins/Experimental/Mutable/Source/CustomizableObject/Private/MuCO/CustomizableObjectModule.cpp
+++ b/Engine/Plugins/Experimental/Mutable/Source/CustomizableObject/Private/MuCO/CustomizableObjectModule.cpp
@@ -227,13 +227,13 @@ UCustomizableObjectInstanceUsage* GetPlayerCustomizableObjectInstanceUsage(const
 		for (TObjectIterator<UCustomizableObjectInstanceUsage> CustomizableObjectInstanceUsage; CustomizableObjectInstanceUsage; ++CustomizableObjectInstanceUsage)
 		{
 #if WITH_EDITOR
-			if (CustomizableObjectInstanceUsage && CustomizableObjectInstanceUsage->IsNetMode(NM_DedicatedServer))
+			if (IsValid(*CustomizableObjectInstanceUsage) && CustomizableObjectInstanceUsage->IsNetMode(NM_DedicatedServer))
 			{
 				continue;
 			}
 #endif
 
-			if (CustomizableObjectInstanceUsage->IsValidLowLevel() && !CustomizableObjectInstanceUsage->IsTemplate() 
+			if (IsValid(*CustomizableObjectInstanceUsage) && !CustomizableObjectInstanceUsage->IsTemplate()
 				&& CustomizableObjectInstanceUsage->GetAttachParent())
 			{
 				AActor* CustomizableActor = CustomizableObjectInstanceUsage->GetAttachParent()->GetAttachmentRootActor();
@@ -259,13 +259,13 @@ UCustomizableObjectInstanceUsage* GetPlayerCustomizableObjectInstanceUsage(const
 		for (TObjectIterator<UCustomizableObjectInstanceUsage> CustomizableObjectInstanceUsage; CustomizableObjectInstanceUsage; ++CustomizableObjectInstanceUsage)
 		{
 #if WITH_EDITOR
-			if (CustomizableObjectInstanceUsage && CustomizableObjectInstanceUsage->IsNetMode(NM_DedicatedServer))
+			if (IsValid(*CustomizableObjectInstanceUsage) && CustomizableObjectInstanceUsage->IsNetMode(NM_DedicatedServer))
 			{
 				continue;
 			}
 #endif
 
-			if (CustomizableObjectInstanceUsage->IsValidLowLevel() && !CustomizableObjectInstanceUsage->IsTemplate())
+			if (IsValid(*CustomizableObjectInstanceUsage) && !CustomizableObjectInstanceUsage->IsTemplate())
 			{
 				++IndexFound;
 				SelectedCustomizableObjectInstanceUsage = *CustomizableObjectInstanceUsage;
diff --git a/Engine/Plugins/Experimental/Mutable/Source/CustomizableObject/Private/MuCO/CustomizableObjectSystem.cpp b/Engine/Plugins/Experimental/Mutable/Source/CustomizableObject/Private/MuCO/CustomizableObjectSystem.cpp
--- a/Engine/Plugins/Experimental/Mutable/Source/CustomizableObject/Private/MuCO/CustomizableObjectSystem.cpp
+++ b/Engine/Plugins/Experimental/Mutable/Source/CustomizableObject/Private/MuCO/CustomizableObjectSystem.cpp
@@ -361,13 +361,13 @@ void UCustomizableObjectSystem::LogShowData(bool bFullInfo, bool ShowMaterialInf
 		const UCustomizableObjectInstanceUsage* CustomizableObjectInstanceUsage = *It;
 
 #if WITH_EDITOR
-		if (CustomizableObjectInstanceUsage && CustomizableObjectInstanceUsage->IsNetMode(NM_DedicatedServer))
+		if (IsValid(CustomizableObjectInstanceUsage) && CustomizableObjectInstanceUsage->IsNetMode(NM_DedicatedServer))
 		{
 			continue;
 		}
 #endif
 
-		if (CustomizableObjectInstanceUsage && CustomizableObjectInstanceUsage->GetCustomizableObjectInstance() 
+		if (IsValid(CustomizableObjectInstanceUsage) && CustomizableObjectInstanceUsage->GetCustomizableObjectInstance()
 			&& CustomizableObjectInstanceUsage->GetAttachParent())
 		{
 			const AActor* ParentActor = CustomizableObjectInstanceUsage->GetAttachParent()->GetAttachmentRootActor();
@@ -814,14 +814,14 @@ void FinishUpdateGlobal(const TSharedRef<FUpdateContextPrivate>& Context)
 		for (TObjectIterator<UCustomizableObjectInstanceUsage> It; It; ++It) // Since iterating objects is expensive, for now CustomizableObjectInstanceUsage does not have a FinishUpdate function.
 		{
 #if WITH_EDITOR
-			if (It && It->IsNetMode(NM_DedicatedServer))
+			if (IsValid(*It) && It->IsNetMode(NM_DedicatedServer))
 			{
 				continue;
 			}
 #endif
 
 			if (const UCustomizableObjectInstanceUsage* CustomizableObjectInstanceUsage = *It;
-				CustomizableObjectInstanceUsage &&
+				IsValid(CustomizableObjectInstanceUsage) &&
 				CustomizableObjectInstanceUsage->GetCustomizableObjectInstance() == Instance)
 			{
 				CustomizableObjectInstanceUsage->Callbacks();
@@ -884,13 +884,13 @@ void UpdateSkeletalMesh(const TSharedRef<FUpdateContextPrivate>& Context)
 		UCustomizableObjectInstanceUsage* CustomizableObjectInstanceUsage = *It;
 
 #if WITH_EDITOR
-		if (CustomizableObjectInstanceUsage && CustomizableObjectInstanceUsage->IsNetMode(NM_DedicatedServer))
+		if (IsValid(CustomizableObjectInstanceUsage) && CustomizableObjectInstanceUsage->IsNetMode(NM_DedicatedServer))
 		{
 			continue;
 		}
 #endif
 
-		if (CustomizableObjectInstanceUsage &&
+		if (IsValid(CustomizableObjectInstanceUsage) &&
 			(CustomizableObjectInstanceUsage->GetCustomizableObjectInstance() == CustomizableObjectInstance) &&
 			CustomizableObjectInstance->SkeletalMeshes.IsValidIndex(CustomizableObjectInstanceUsage->GetComponentIndex())
 		   )
@@ -1387,7 +1387,7 @@ bool UCustomizableObjectSystem::CheckIfDiskOrMipUpdateOperationsPending(const UC
 {
 	for (TObjectIterator<UCustomizableObjectInstance> CustomizableObjectInstance; CustomizableObjectInstance; ++CustomizableObjectInstance)
 	{
-		if (CustomizableObjectInstance->GetCustomizableObject() == &Object)
+		if (IsValid(*CustomizableObjectInstance) && CustomizableObjectInstance->GetCustomizableObject() == &Object)
 		{
 			for (const FGeneratedTexture& GeneratedTexture : CustomizableObjectInstance->GetPrivate()->GeneratedTextures)
 			{
@@ -2906,12 +2906,12 @@ namespace impl
 #if WITH_EDITOR
 		for (TObjectIterator<UCustomizableObjectInstanceUsage> CustomizableObjectInstanceUsage; CustomizableObjectInstanceUsage && !bIsInEditorViewport; ++CustomizableObjectInstanceUsage)
 		{
-			if (CustomizableObjectInstanceUsage && CustomizableObjectInstanceUsage->IsNetMode(NM_DedicatedServer))
+			if (IsValid(*CustomizableObjectInstanceUsage) && CustomizableObjectInstanceUsage->IsNetMode(NM_DedicatedServer))
 			{
 				continue;
 			}
 
-			if (CustomizableObjectInstanceUsage &&
+			if (IsValid(*CustomizableObjectInstanceUsage) &&
 				CustomizableObjectInstanceUsage->GetCustomizableObjectInstance() == CandidateInstance)
 			{
 				EWorldType::Type WorldType = EWorldType::Type::None;
@@ -3097,7 +3097,7 @@ bool UCustomizableObjectSystem::Tick(float DeltaTime)
 
 		for (TObjectIterator<UCustomizableObjectInstance> CustomizableObjectInstance; CustomizableObjectInstance; ++CustomizableObjectInstance)
 		{
-			if (IsValidChecked(*CustomizableObjectInstance) && CustomizableObjectInstance->GetPrivate())
+			if (IsValid(*CustomizableObjectInstance) && CustomizableObjectInstance->GetPrivate())
 			{
 				UCustomizableInstancePrivateData* ObjectInstancePrivateData = CustomizableObjectInstance->GetPrivate();
 
@@ -3230,7 +3230,7 @@ bool UCustomizableObjectSystem::Tick(float DeltaTime)
 		{
 			for (TObjectIterator<UCustomizableObjectInstance> CustomizableObjectInstance; CustomizableObjectInstance; ++CustomizableObjectInstance)
 			{
-				if (IsValidChecked(*CustomizableObjectInstance) && CustomizableObjectInstance->GetPrivate())
+				if (IsValid(*CustomizableObjectInstance) && CustomizableObjectInstance->GetPrivate())
 				{
 					CustomizableObjectInstance->GetPrivate()->LastMinSquareDistFromComponentToPlayer = CustomizableObjectInstance->GetPrivate()->MinSquareDistFromComponentToPlayer;
 					CustomizableObjectInstance->GetPrivate()->MinSquareDistFromComponentToPlayer = FLT_MAX;
@@ -3444,7 +3444,7 @@ int32 UCustomizableObjectSystem::GetTotalInstances() const
 	
 	for (TObjectIterator<UCustomizableObjectInstance> Instance; Instance; ++Instance)
 	{
-		if (!Instance ||
+		if (!IsValid(*Instance) ||
 			Instance->HasAnyFlags(RF_ClassDefaultObject))
 		{
 			continue;
@@ -3849,7 +3849,7 @@ void FCustomizableObjectSystemPrivate::UpdateStats()
 	
 	for (TObjectIterator<UCustomizableObjectInstance> Instance; Instance; ++Instance)
 	{
-		if (!IsValidChecked(*Instance))
+		if (!IsValid(*Instance))
 		{
 			continue;
 		}
diff --git a/Engine/Plugins/Experimental/Mutable/Source/CustomizableObject/Private/MuCO/LogBenchmarkUtil.cpp b/Engine/Plugins/Experimental/Mutable/Source/CustomizableObject/Private/MuCO/LogBenchmarkUtil.cpp
--- a/Engine/Plugins/Experimental/Mutable/Source/CustomizableObject/Private/MuCO/LogBenchmarkUtil.cpp
+++ b/Engine/Plugins/Experimental/Mutable/Source/CustomizableObject/Private/MuCO/LogBenchmarkUtil.cpp
@@ -103,7 +103,7 @@ void FLogBenchmarkUtil::GetInstancesStats(int32& OutNumInstances, int32& OutNumB
 	
 	for (TObjectIterator<UCustomizableObjectInstance> Instance; Instance; ++Instance)
 	{
-		if (!IsValidChecked(*Instance) ||
+		if (!IsValid(*Instance) ||
 			Instance->HasAnyFlags(RF_ClassDefaultObject))
 		{
 			continue;
@@ -206,7 +206,7 @@ void FLogBenchmarkUtil::UpdateStats()
 
 	for (TObjectIterator<UCustomizableObjectInstanceUsage> CustomizableObjectInstanceUsage; CustomizableObjectInstanceUsage; ++CustomizableObjectInstanceUsage)
 	{
-		if (CustomizableObjectInstanceUsage && CustomizableObjectInstanceUsage->IsNetMode(NM_DedicatedServer))
+		if (!IsValid(*CustomizableObjectInstanceUsage) || CustomizableObjectInstanceUsage->IsNetMode(NM_DedicatedServer))
 		{
 			continue;
 		}
diff --git a/Engine/Plugins/Experimental/Mutable/Source/CustomizableObject/Private/MuCO/UnrealMutableImageProvider.cpp b/Engine/Plugins/Experimental/Mutable/Source/CustomizableObject/Private/MuCO/UnrealMutableImageProvider.cpp
--- a/Engine/Plugins/Experimental/Mutable/Source/CustomizableObject/Private/MuCO/UnrealMutableImageProvider.cpp
+++ b/Engine/Plugins/Experimental/Mutable/Source/CustomizableObject/Private/MuCO/UnrealMutableImageProvider.cpp
@@ -109,7 +109,7 @@ bool FUnrealMutableImageProvider::Tick(float DeltaTime)
 	UCustomizableObject* CO = nullptr;
 	for (TObjectIterator<UCustomizableObject> It; It; ++It)
 	{
-		if (It->GetModel().Get() == Request->ModelPtr)
+		if (IsValid(*It) && It->GetModel().Get() == Request->ModelPtr)
 		{
 			CO = *It;
 			break;
diff --git a/Engine/Plugins/Experimental/Mutable/Source/CustomizableObjectEditor/Private/MuCOE/CustomizableObjectCompiler.cpp b/Engine/Plugins/Experimental/Mutable/Source/CustomizableObjectEditor/Private/MuCOE/CustomizableObjectCompiler.cpp
--- a/Engine/Plugins/Experimental/Mutable/Source/CustomizableObjectEditor/Private/MuCOE/CustomizableObjectCompiler.cpp
+++ b/Engine/Plugins/Experimental/Mutable/Source/CustomizableObjectEditor/Private/MuCOE/CustomizableObjectCompiler.cpp
@@ -236,7 +236,7 @@ void FCustomizableObjectCompiler::Compile(UCustomizableObject& Object, const FCo
 	for (TObjectIterator<UCustomizableObjectInstance> It; It; ++It)
 	{
 		UCustomizableObjectInstance* Instance = *It;
-		if (Instance &&
+		if (IsValid(Instance) &&
 			Instance->GetCustomizableObject() == &Object)
 		{
 			Instance->DestroyLiveUpdateInstance();
@@ -1523,7 +1523,7 @@ void FCustomizableObjectCompiler::FinishCompilation()
 	for (TObjectIterator<UCustomizableObjectInstance> It; It; ++It)
 	{
 		UCustomizableObjectInstance* Instance = *It;
-		if (Instance &&
+		if (IsValid(Instance) &&
 			Instance->GetCustomizableObject() == CurrentObject)
 		{
 			Instance->SetObject(CurrentObject);
diff --git a/Engine/Plugins/Experimental/Mutable/Source/CustomizableObjectEditor/Private/MuCOE/EdGraphSchema_CustomizableObject.cpp b/Engine/Plugins/Experimental/Mutable/Source/CustomizableObjectEditor/Private/MuCOE/EdGraphSchema_CustomizableObject.cpp
--- a/Engine/Plugins/Experimental/Mutable/Source/CustomizableObjectEditor/Private/MuCOE/EdGraphSchema_CustomizableObject.cpp
+++ b/Engine/Plugins/Experimental/Mutable/Source/CustomizableObjectEditor/Private/MuCOE/EdGraphSchema_CustomizableObject.cpp
@@ -525,7 +525,7 @@ void UEdGraphSchema_CustomizableObject::GetGraphContextActions(FGraphContextMenu
 	for (TObjectIterator<UCustomizableObjectNode> It(RF_NoFlags); It; ++It)
 	{
 		const UCustomizableObjectNode* Node = *It;
-		if (!Node->HasAllFlags(RF_ClassDefaultObject) || Node->GetClass()->HasAnyClassFlags(CLASS_Abstract))
+		if (!IsValid(Node) || !Node->HasAllFlags(RF_ClassDefaultObject) || Node->GetClass()->HasAnyClassFlags(CLASS_Abstract))
 		{
 			// Only interested in non-abstract CDOs
 			continue;
@@ -738,7 +738,7 @@ void UEdGraphSchema_CustomizableObject::GetContextMenuActionsReconstructAllChild
 	for (TObjectIterator<UCustomizableObjectNode> It(RF_NoFlags); It; ++It)
 	{
 		const UCustomizableObjectNode* Node = *It;
-		if (!Node->HasAllFlags(RF_ClassDefaultObject))
+		if (!IsValid(Node) || !Node->HasAllFlags(RF_ClassDefaultObject))
 		{
 			continue; // Only interested in CDOs
 		}
diff --git a/Engine/Plugins/Experimental/Mutable/Source/CustomizableObjectEditor/Private/MuCOE/GraphTraversal.cpp b/Engine/Plugins/Experimental/Mutable/Source/CustomizableObjectEditor/Private/MuCOE/GraphTraversal.cpp
--- a/Engine/Plugins/Experimental/Mutable/Source/CustomizableObjectEditor/Private/MuCOE/GraphTraversal.cpp
+++ b/Engine/Plugins/Experimental/Mutable/Source/CustomizableObjectEditor/Private/MuCOE/GraphTraversal.cpp
@@ -77,7 +77,7 @@ TArray<UEdGraphPin*> FollowPinArray(const UEdGraphPin& Pin, bool* bOutCycleDetec
 				const UCustomizableObjectNodeExternalPin* LinkedNode = nullptr;
 				for (TObjectIterator<UCustomizableObjectNodeExternalPin> It; It; ++It)
 				{
-					if ((*It)->GetNodeExposePin() == ExposePinNode)
+					if (IsValid(*It) && (*It)->GetNodeExposePin() == ExposePinNode)
 					{
 						LinkedNode = *It;
 						break;
