[Mutable] Block compressed images mipmap and resize memory optimizations. Disabled by default.

- Removed advanced mip generation settings sharpen and address mode. No longer used.
- Addeded Cvars mutable.EnableImageResizeMemoryOptimizations and mutable.EnableCompressedMipGenerationMemoryOptimizations defaulted to false.

#rnx

diff --git a/Engine/Plugins/Experimental/Mutable/Source/MutableRuntime/Private/MuR/BlockCompression/Miro/Miro.cpp b/Engine/Plugins/Experimental/Mutable/Source/MutableRuntime/Private/MuR/BlockCompression/Miro/Miro.cpp
--- a/Engine/Plugins/Experimental/Mutable/Source/MutableRuntime/Private/MuR/BlockCompression/Miro/Miro.cpp
+++ b/Engine/Plugins/Experimental/Mutable/Source/MutableRuntime/Private/MuR/BlockCompression/Miro/Miro.cpp
@@ -9326,8 +9326,892 @@ namespace miro
 	}
 
 #endif // MIRO_INCLUDE_ASTC
+}
+
+namespace miro::SubImageDecompression
+{
+
+#if MIRO_INCLUDE_ASTC
+	template<uint32 BlockSize>
+	void GenericSubImage_ASTCRGBL_To_RGB(FImageSize FromSize, FImageSize ToSize, FImageSize SubSize, const uint8* From, uint8* To)
+	{
+		init_astc_decompress();
+
+		for (uint32 y = 0; y < SubSize.Y; y += BlockSize)
+		{
+			for (uint32 x = 0; x < SubSize.X; x += BlockSize)
+			{
+				constexpr bool bIsSRGB = false;
+				uint8 Block[BlockSize * BlockSize * 4];
+				const bool bSuccess = astcdec::decompress(Block, From, bIsSRGB, BlockSize, BlockSize);
+				miro_check(bSuccess);
+
+				for (uint32 py = 0; py < BlockSize; py++)
+				{
+					for (uint32 px = 0; px < BlockSize; px++)
+					{
+						uint32 xi = x + px;
+						uint32 yi = y + py;
+
+						if (LIKELY((xi < ToSize.X) & (yi < ToSize.Y)))
+						{
+							uint8* ToPixel = To + (yi * ToSize.X + xi) * 3;
+							ToPixel[0] = Block[py * BlockSize * 4 + px * 4 + 0];
+							ToPixel[1] = Block[py * BlockSize * 4 + px * 4 + 1];
+							ToPixel[2] = Block[py * BlockSize * 4 + px * 4 + 2];
+						}
+					}
+				}
+
+				From += 16;
+			}
+		}
+
+	}
+
+	template<uint32 BlockSize>
+	void GenericSubImage_ASTCRGBL_To_RGBA(FImageSize FromSize, FImageSize ToSize, FImageSize SubSize, const uint8* From, uint8* To)
+	{
+		init_astc_decompress();
+
+		for (uint32 y = 0; y < SubSize.Y; y += BlockSize)
+		{
+			for (uint32 x = 0; x < SubSize.X; x += BlockSize)
+			{
+				bool bIsSRGB = false;
+				uint8 Block[BlockSize * BlockSize * 4];
+				bool bSuccess = astcdec::decompress(Block, From, bIsSRGB, BlockSize, BlockSize);
+				miro_check(bSuccess);
+
+				for (uint32 py = 0; py < BlockSize; py++)
+				{
+					for (uint32 px = 0; px < BlockSize; px++)
+					{
+						uint32 xi = x + px;
+						uint32 yi = y + py;
+
+						if (LIKELY((xi < ToSize.X) & (yi < ToSize.Y)))
+						{
+							uint8* ToPixel = To + (yi * ToSize.X + xi) * 4;
+							ToPixel[0] = Block[py * BlockSize * 4 + px * 4 + 0];
+							ToPixel[1] = Block[py * BlockSize * 4 + px * 4 + 1];
+							ToPixel[2] = Block[py * BlockSize * 4 + px * 4 + 2];
+							ToPixel[3] = 255;
+						}
+					}
+				}
+
+				From += 16;
+			}
+		}
+	}
+
+	template<uint32 BlockSize>
+	void GenericSubImage_ASTCRGBAL_To_RGBA(FImageSize FromSize, FImageSize ToSize, FImageSize SubSize, const uint8* From, uint8* To)
+	{
+		init_astc_decompress();
+	
+		const uint32 NumBlocksX = FMath::DivideAndRoundUp(uint32(FromSize.X), BlockSize);
+		const uint32 NumSubBlocksX = FMath::DivideAndRoundUp(uint32(SubSize.X), BlockSize);
+		const uint32 NumSubBlocksY = FMath::DivideAndRoundUp(uint32(SubSize.Y), BlockSize);
+
+		if (SubSize.X <= 0 || SubSize.Y <= 0)
+		{
+			return;
+		}
+
+		for (uint32 BlockY = 0; BlockY < NumSubBlocksY; ++BlockY)
+		{
+			for (uint32 BlockX = 0; BlockX < NumSubBlocksX; ++BlockX)
+			{
+				constexpr bool bIsSRGB = false;
+				uint8 Block[BlockSize * BlockSize * 4];
 
+				constexpr int32 CompressedBlockSize = 16;
+				const uint8* SrcBlockPtr = From + (BlockY * NumBlocksX + BlockX) * CompressedBlockSize;
+
+				bool bSuccess = astcdec::decompress(Block, SrcBlockPtr, bIsSRGB, BlockSize, BlockSize);
+				miro_check(bSuccess);
+
+				for (uint32 py = 0; py < BlockSize; py++)
+				{
+					for (uint32 px = 0; px < BlockSize; px++)
+					{
+						uint32 xi = FMath::Min(BlockX * BlockSize + px, uint32(SubSize.X) - 1);
+						uint32 yi = FMath::Min(BlockY * BlockSize + py, uint32(SubSize.Y) - 1);
+
+						uint8* ToPixel = To + (yi * ToSize.X + xi) * 4;
+						ToPixel[0] = Block[py * BlockSize * 4 + px * 4 + 0];
+						ToPixel[1] = Block[py * BlockSize * 4 + px * 4 + 1];
+						ToPixel[2] = Block[py * BlockSize * 4 + px * 4 + 2];
+						ToPixel[3] = Block[py * BlockSize * 4 + px * 4 + 3];
+					}
+				}
+			}
+		}
+	}
+
+
+	template<uint32 BlockSize>
+	void GenericSubImage_ASTCRGBAL_To_RGB(FImageSize FromSize, FImageSize ToSize, FImageSize SubSize, const uint8* From, uint8* To)
+	{
+		init_astc_decompress();
+
+		const uint32 NumBlocksX = FMath::DivideAndRoundUp(uint32(FromSize.X), BlockSize);
+		const uint32 NumSubBlocksX = FMath::DivideAndRoundUp(uint32(SubSize.X), BlockSize);
+		const uint32 NumSubBlocksY = FMath::DivideAndRoundUp(uint32(SubSize.Y), BlockSize);
+
+		for (uint32 BlockY = 0; BlockY < NumSubBlocksY; ++BlockY)
+		{
+			for (uint32 BlockX = 0; BlockX < NumSubBlocksX; ++BlockX)
+			{
+				constexpr bool bIsSRGB = false;
+				uint8 Block[BlockSize * BlockSize * 4];
+
+				constexpr int32 CompressedBlockSize = 16;
+				const uint8* SrcBlockPtr = From + (BlockY * NumBlocksX + BlockX) * CompressedBlockSize;
+
+				bool bSuccess = astcdec::decompress(Block, SrcBlockPtr, bIsSRGB, BlockSize, BlockSize);
+				miro_check(bSuccess);
+
+				for (uint32 py = 0; py < BlockSize; py++)
+				{
+					for (uint32 px = 0; px < BlockSize; px++)
+					{
+						uint32 xi = BlockX * BlockSize + px;
+						uint32 yi = BlockY * BlockSize + py;
+
+						if (LIKELY((xi < ToSize.X) & (yi < ToSize.Y)))
+						{
+							uint8* ToPixel = To + (yi * ToSize.X + xi) * 3;
+							ToPixel[0] = Block[py * BlockSize * 4 + px * 4 + 0];
+							ToPixel[1] = Block[py * BlockSize * 4 + px * 4 + 1];
+							ToPixel[2] = Block[py * BlockSize * 4 + px * 4 + 2];
+						}
+					}
+				}
+			}
+		}
+	}
+
+	template<uint32 BlockSize>
+	void GenericSubImage_ASTCRGL_To_RGB(FImageSize FromSize, FImageSize ToSize, FImageSize SubSize, const uint8* From, uint8* To)
+	{
+		init_astc_decompress();
+
+		const uint32 NumBlocksX = FMath::DivideAndRoundUp(uint32(FromSize.X), BlockSize);
+		const uint32 NumSubBlocksX = FMath::DivideAndRoundUp(uint32(SubSize.X), BlockSize);
+		const uint32 NumSubBlocksY = FMath::DivideAndRoundUp(uint32(SubSize.Y), BlockSize);
+
+		for (uint32 by = 0; by < NumSubBlocksY; ++by)
+		{
+			for (uint32 bx = 0; bx < NumSubBlocksX; ++bx)
+			{
+				constexpr int32 CompressedBlockSize = 16;
+				const uint8* SrcBlockPtr = From + (by * NumBlocksX + bx) * CompressedBlockSize;
+
+				constexpr bool bIsSRGB = false;
+				uint8 Block[BlockSize * BlockSize * 4];
+				const bool bSuccess = astcdec::decompress(Block, SrcBlockPtr, bIsSRGB, BlockSize, BlockSize);
+				miro_check(bSuccess);
+
+				for (uint32 BlockY = 0; BlockY < BlockSize; BlockY++)
+				{
+					for (uint32 BlockX = 0; BlockX < BlockSize; BlockX++)
+					{
+						uint32 xi = bx * BlockSize + BlockX;
+						uint32 yi = by * BlockSize + BlockY;
+
+						if (LIKELY((xi < ToSize.X) & (yi < ToSize.Y)))
+						{
+							uint8* ToPixel = To + (yi * ToSize.X + xi) * 3;
+							ToPixel[0] = Block[BlockY * BlockSize * 4 + BlockX * 4 + 0];
+							ToPixel[1] = Block[BlockY * BlockSize * 4 + BlockX * 4 + 3];
+							ToPixel[2] = 255;
+						}
+					}
+				}
+			}
+		}
+	}
+
+	template<uint32 BlockSize>
+	void GenericSubImage_ASTCRGL_To_RGBA(FImageSize FromSize, FImageSize ToSize, FImageSize SubSize, const uint8* From, uint8* To)
+	{
+		init_astc_decompress();
+		
+		const uint32 NumBlocksX = FMath::DivideAndRoundUp(uint32(FromSize.X), BlockSize);
+		const uint32 NumSubBlocksX = FMath::DivideAndRoundUp(uint32(SubSize.X), BlockSize);
+		const uint32 NumSubBlocksY = FMath::DivideAndRoundUp(uint32(SubSize.Y), BlockSize);
 
+		for (uint32 by = 0; by < NumSubBlocksY; ++by)
+		{
+			for (uint32 bx = 0; bx < NumSubBlocksX; ++bx)
+			{
+				constexpr int32 CompressedBlockSize = 16;
+				const uint8* SrcBlockPtr = From + (by * NumBlocksX + bx) * CompressedBlockSize;
+
+				constexpr bool bIsSRGB = false;
+				uint8 Block[BlockSize * BlockSize * 4];
+				bool bSuccess = astcdec::decompress(Block, SrcBlockPtr, bIsSRGB, BlockSize, BlockSize);
+				miro_check(bSuccess);
+
+				for (uint32 BlockY = 0; BlockY < BlockSize; BlockY++)
+				{
+					for (uint32 BlockX = 0; BlockX < BlockSize; BlockX++)
+					{
+						uint32 xi = bx * BlockSize + BlockX;
+						uint32 yi = by * BlockSize + BlockY;
+
+						if (LIKELY((xi < ToSize.X) & (yi < ToSize.Y)))
+						{
+							uint8* ToPixel = To + (yi * ToSize.X + xi) * 4;
+							ToPixel[0] = Block[BlockY * BlockSize * 4 + BlockX * 4 + 0];
+							ToPixel[1] = Block[BlockY * BlockSize * 4 + BlockX * 4 + 3];
+							ToPixel[2] = 255;
+							ToPixel[3] = 255;
+						}
+					}
+				}
+			}
+		}
+	}
+
+
+	void ASTC4x4RGBAL_To_RGBASubImage(FImageSize FromSize, FImageSize ToSize, FImageSize SubSize, const uint8* From, uint8* To)
+	{
+		GenericSubImage_ASTCRGBAL_To_RGBA<4>(FromSize, ToSize, SubSize, From, To);
+	}
+
+	void ASTC4x4RGBAL_To_RGBSubImage(FImageSize FromSize, FImageSize ToSize, FImageSize SubSize, const uint8* From, uint8* To)
+	{
+		GenericSubImage_ASTCRGBAL_To_RGB<4>(FromSize, ToSize, SubSize, From, To);
+	}
+
+	void ASTC4x4RGBL_To_RGBSubImage(FImageSize FromSize, FImageSize ToSize, FImageSize SubSize, const uint8* From, uint8* To)
+	{
+		GenericSubImage_ASTCRGBL_To_RGB<4>(FromSize, ToSize, SubSize, From, To);
+	}
+
+	void ASTC4x4RGBL_To_RGBASubImage(FImageSize FromSize, FImageSize ToSize, FImageSize SubSize, const uint8* From, uint8* To)
+	{
+		GenericSubImage_ASTCRGBL_To_RGBA<4>(FromSize, ToSize, SubSize, From, To);
+	}
+	
+	void ASTC4x4RGL_To_RGBSubImage(FImageSize FromSize, FImageSize ToSize, FImageSize SubSize, const uint8* From, uint8* To)
+	{
+		GenericSubImage_ASTCRGL_To_RGB<4>(FromSize, ToSize, SubSize, From, To);
+	}
+	
+	void ASTC4x4RGL_To_RGBASubImage(FImageSize FromSize, FImageSize ToSize, FImageSize SubSize, const uint8* From, uint8* To)
+	{
+		GenericSubImage_ASTCRGL_To_RGBA<4>(FromSize, ToSize, SubSize, From, To);
+	}
+
+
+	//---------------------------------------------------------------------------------------------
+	// 6x6
+	//---------------------------------------------------------------------------------------------
+	void ASTC6x6RGBAL_To_RGBASubImage(FImageSize FromSize, FImageSize ToSize, FImageSize SubSize, const uint8* From, uint8* To)
+	{
+		GenericSubImage_ASTCRGBAL_To_RGBA<6>(FromSize, ToSize, SubSize, From, To);
+	}
+
+	void ASTC6x6RGBAL_To_RGBSubImage(FImageSize FromSize, FImageSize ToSize, FImageSize SubSize, const uint8* From, uint8* To)
+	{
+		GenericSubImage_ASTCRGBAL_To_RGB<6>(FromSize, ToSize, SubSize, From, To);
+	}
+
+	void ASTC6x6RGBL_To_RGBSubImage(FImageSize FromSize, FImageSize ToSize, FImageSize SubSize, const uint8* From, uint8* To)
+	{
+		GenericSubImage_ASTCRGBL_To_RGB<6>(FromSize, ToSize, SubSize, From, To);
+	}
+
+	void ASTC6x6RGBL_To_RGBASubImage(FImageSize FromSize, FImageSize ToSize, FImageSize SubSize, const uint8* From, uint8* To)
+	{
+		GenericSubImage_ASTCRGBL_To_RGBA<6>(FromSize, ToSize, SubSize, From, To);
+	}
+
+	void ASTC6x6RGL_To_RGBSubImage(FImageSize FromSize, FImageSize ToSize, FImageSize SubSize, const uint8* From, uint8* To)
+	{
+		GenericSubImage_ASTCRGL_To_RGB<6>(FromSize, ToSize, SubSize, From, To);
+	}
+
+	void ASTC6x6RGL_To_RGBASubImage(FImageSize FromSize, FImageSize ToSize, FImageSize SubSize, const uint8* From, uint8* To)
+	{
+		GenericSubImage_ASTCRGL_To_RGBA<6>(FromSize, ToSize, SubSize, From, To);
+	}
+
+	//---------------------------------------------------------------------------------------------
+	// 8x8
+	//---------------------------------------------------------------------------------------------
+
+	void ASTC8x8RGBAL_To_RGBASubImage(FImageSize FromSize, FImageSize ToSize, FImageSize SubSize, const uint8* From, uint8* To)
+	{
+		GenericSubImage_ASTCRGBAL_To_RGBA<8>(FromSize, ToSize, SubSize, From, To);
+	}
+
+	void ASTC8x8RGBAL_To_RGBSubImage(FImageSize FromSize, FImageSize ToSize, FImageSize SubSize, const uint8* From, uint8* To)
+	{
+		GenericSubImage_ASTCRGBAL_To_RGB<8>(FromSize, ToSize, SubSize, From, To);
+	}
+
+	void ASTC8x8RGBL_To_RGBSubImage(FImageSize FromSize, FImageSize ToSize, FImageSize SubSize, const uint8* From, uint8* To)
+	{
+		GenericSubImage_ASTCRGBL_To_RGB<8>(FromSize, ToSize, SubSize, From, To);
+	}
+
+	void ASTC8x8RGBL_To_RGBASubImage(FImageSize FromSize, FImageSize ToSize, FImageSize SubSize, const uint8* From, uint8* To)
+	{
+		GenericSubImage_ASTCRGBL_To_RGBA<8>(FromSize, ToSize, SubSize, From, To);
+	}
+
+	void ASTC8x8RGL_To_RGBSubImage(FImageSize FromSize, FImageSize ToSize, FImageSize SubSize, const uint8* From, uint8* To)
+	{
+		GenericSubImage_ASTCRGL_To_RGB<8>(FromSize, ToSize, SubSize, From, To);
+	}
+
+	void ASTC8x8RGL_To_RGBASubImage(FImageSize FromSize, FImageSize ToSize, FImageSize SubSize, const uint8* From, uint8* To)
+	{
+		GenericSubImage_ASTCRGL_To_RGBA<8>(FromSize, ToSize, SubSize, From, To);
+	}
+
+
+	//---------------------------------------------------------------------------------------------
+	// 10x10
+	//---------------------------------------------------------------------------------------------
+	void ASTC10x10RGBAL_To_RGBASubImage(FImageSize FromSize, FImageSize ToSize, FImageSize SubSize, const uint8* From, uint8* To)
+	{
+		GenericSubImage_ASTCRGBAL_To_RGBA<10>(FromSize, ToSize, SubSize, From, To);
+	}
+
+	void ASTC10x10RGBAL_To_RGBSubImage(FImageSize FromSize, FImageSize ToSize, FImageSize SubSize, const uint8* From, uint8* To)
+	{
+		GenericSubImage_ASTCRGBAL_To_RGB<10>(FromSize, ToSize, SubSize, From, To);
+	}
+
+	void ASTC10x10RGBL_To_RGBSubImage(FImageSize FromSize, FImageSize ToSize, FImageSize SubSize, const uint8* From, uint8* To)
+	{
+		GenericSubImage_ASTCRGBL_To_RGB<10>(FromSize, ToSize, SubSize, From, To);
+	}
+
+	void ASTC10x10RGBL_To_RGBASubImage(FImageSize FromSize, FImageSize ToSize, FImageSize SubSize, const uint8* From, uint8* To)
+	{
+		GenericSubImage_ASTCRGBL_To_RGBA<10>(FromSize, ToSize, SubSize, From, To);
+	}
+
+	void ASTC10x10RGL_To_RGBSubImage(FImageSize FromSize, FImageSize ToSize, FImageSize SubSize, const uint8* From, uint8* To)
+	{
+		GenericSubImage_ASTCRGL_To_RGB<10>(FromSize, ToSize, SubSize, From, To);
+	}
+
+	void ASTC10x10RGL_To_RGBASubImage(FImageSize FromSize, FImageSize ToSize, FImageSize SubSize, const uint8* From, uint8* To)
+	{
+		GenericSubImage_ASTCRGL_To_RGBA<10>(FromSize, ToSize, SubSize, From, To);
+	}
+
+
+	//---------------------------------------------------------------------------------------------
+	// 12x12
+	//---------------------------------------------------------------------------------------------
+	void ASTC12x12RGBAL_To_RGBASubImage(FImageSize FromSize, FImageSize ToSize, FImageSize SubSize, const uint8* From, uint8* To)
+	{
+		GenericSubImage_ASTCRGBAL_To_RGBA<12>(FromSize, ToSize, SubSize, From, To);
+	}
+
+	void ASTC12x12RGBAL_To_RGBSubImage(FImageSize FromSize, FImageSize ToSize, FImageSize SubSize, const uint8* From, uint8* To)
+	{
+		GenericSubImage_ASTCRGBAL_To_RGB<12>(FromSize, ToSize, SubSize, From, To);
+	}
+
+	void ASTC12x12RGBL_To_RGBSubImage(FImageSize FromSize, FImageSize ToSize, FImageSize SubSize, const uint8* From, uint8* To)
+	{
+		GenericSubImage_ASTCRGBL_To_RGB<12>(FromSize, ToSize, SubSize, From, To);
+	}
+
+	void ASTC12x12RGBL_To_RGBASubImage(FImageSize FromSize, FImageSize ToSize, FImageSize SubSize, const uint8* From, uint8* To)
+	{
+		GenericSubImage_ASTCRGBL_To_RGBA<12>(FromSize, ToSize, SubSize, From, To);
+	}
+
+	void ASTC12x12RGL_To_RGBSubImage(FImageSize FromSize, FImageSize ToSize, FImageSize SubSize, const uint8* From, uint8* To)
+	{
+		GenericSubImage_ASTCRGL_To_RGB<12>(FromSize, ToSize, SubSize, From, To);
+	}
+
+	void ASTC12x12RGL_To_RGBASubImage(FImageSize FromSize, FImageSize ToSize, FImageSize SubSize, const uint8* From, uint8* To)
+	{
+		GenericSubImage_ASTCRGL_To_RGBA<12>(FromSize, ToSize, SubSize, From, To);
+	}
+
+#endif // MIRO_INCLUDE_ASTC
+	// BC Formats.
+
+#if MIRO_INCLUDE_BC
+	void BC1_To_RGBSubImage(FImageSize FromSize, FImageSize ToSize, FImageSize SubSize, const uint8* RESTRICT From, uint8* RESTRICT To)
+	{
+		constexpr uint32 BlockSize = 4;
+
+		const uint32 NumBlocksX = FMath::DivideAndRoundUp(uint32(FromSize.X), BlockSize);
+		const uint32 NumSubBlocksX = FMath::DivideAndRoundUp(uint32(SubSize.X), BlockSize);
+		const uint32 NumSubBlocksY = FMath::DivideAndRoundUp(uint32(SubSize.Y), BlockSize);
+
+		for (uint32 by = 0; by < NumSubBlocksY; ++by)
+		{
+			const uint8* RowFrom = From + 8 * NumBlocksX * by;
+
+			for (uint32 bx = 0; bx < NumSubBlocksX; ++bx)
+			{
+				uint8 Block[4 * 4 * 4];
+
+				bcdec::bcdec_bc1(RowFrom, Block, 4 * 4);
+
+				int32 b = 0;
+				for (int32 j = 0; j < 4; ++j)
+				{
+					for (int32 i = 0; i < 4; ++i)
+					{
+						if (LIKELY((by * 4 + j < ToSize.Y) & (bx * 4 + i < ToSize.X)))
+						{
+							uint8* pPixelDest = To + 3 * (ToSize.X * (by * 4 + j) + (bx * 4 + i));
+
+							pPixelDest[0] = Block[b + 0];
+							pPixelDest[1] = Block[b + 1];
+							pPixelDest[2] = Block[b + 2];
+						}
+
+						b += 4;
+					}
+				}
+
+				RowFrom += 8;
+			}
+		}
+	}
+
+
+	void BC1_To_RGBASubImage(FImageSize FromSize, FImageSize ToSize, FImageSize SubSize, const uint8* RESTRICT From, uint8* RESTRICT To)
+	{
+		constexpr uint32 BlockSize = 4;
+
+		const uint32 NumBlocksX = FMath::DivideAndRoundUp(uint32(FromSize.X), BlockSize);
+		const uint32 NumSubBlocksX = FMath::DivideAndRoundUp(uint32(SubSize.X), BlockSize);
+		const uint32 NumSubBlocksY = FMath::DivideAndRoundUp(uint32(SubSize.Y), BlockSize);
+
+		for (uint32 by = 0; by < NumSubBlocksY; ++by)
+		{
+			const uint8* RowFrom = From + 8 * NumBlocksX * by;
+
+			for (uint32 bx = 0; bx < NumSubBlocksX; ++bx)
+			{
+				uint8 Block[4 * 4 * 4];
+				bcdec::bcdec_bc1(RowFrom, Block, 4 * 4);
+
+				int32 b = 0;
+				for (int32 j = 0; j < 4; ++j)
+				{
+					for (int32 i = 0; i < 4; ++i)
+					{
+						if (LIKELY((by * 4 + j < ToSize.Y) & (bx * 4 + i < ToSize.X)))
+						{
+							uint8* pPixelDest = To + 4 * (ToSize.X * (by * 4 + j) + (bx * 4 + i));
+
+							pPixelDest[0] = Block[b + 0];
+							pPixelDest[1] = Block[b + 1];
+							pPixelDest[2] = Block[b + 2];
+							pPixelDest[3] = Block[b + 3];
+						}
+
+						b += 4;
+					}
+				}
+
+				RowFrom += 8;
+			}
+		}
+	}
+
+	void BC2_To_RGBASubImage(FImageSize FromSize, FImageSize ToSize, FImageSize SubSize, const uint8* RESTRICT From, uint8* RESTRICT To)
+	{
+		constexpr uint32 BlockSize = 4;
+
+		const uint32 NumBlocksX = FMath::DivideAndRoundUp(uint32(FromSize.X), BlockSize);
+		const uint32 NumSubBlocksX = FMath::DivideAndRoundUp(uint32(SubSize.X), BlockSize);
+		const uint32 NumSubBlocksY = FMath::DivideAndRoundUp(uint32(SubSize.Y), BlockSize);
+
+		for (uint32 by = 0; by < NumSubBlocksY; ++by)
+		{
+			const uint8* RowFrom = From + 16 * NumBlocksX * by;
+
+			for (uint32 bx = 0; bx < NumSubBlocksY; ++bx)
+			{
+				uint8 Block[4 * 4 * 4];
+
+				//uint8 const* ColourBlock = RowFrom + 8;
+				//uint8 const* AlphaBlock = RowFrom;
+
+				bcdec::bcdec_bc2(RowFrom, Block, 4 * 4);
+
+				int32 b = 0;
+				for (int32 j = 0; j < 4; ++j)
+				{
+					for (int32 i = 0; i < 4; ++i)
+					{
+						if (LIKELY((by * 4 + j < ToSize.Y) & (bx * 4 + i < ToSize.X)))
+						{
+							uint8* pPixelDest = To + 4 * (ToSize.X * (by * 4 + j) + (bx * 4 + i));
+
+							pPixelDest[0] = Block[b + 0];
+							pPixelDest[1] = Block[b + 1];
+							pPixelDest[2] = Block[b + 2];
+							pPixelDest[3] = Block[b + 3];
+						}
+
+						b += 4;
+					}
+				}
+
+				RowFrom += 16;
+			}
+		}
+	}
+
+	void BC2_To_RGBSubImage(FImageSize FromSize, FImageSize ToSize, FImageSize SubSize, const uint8* RESTRICT From, uint8* RESTRICT To)
+	{
+		constexpr uint32 BlockSize = 4;
+
+		const uint32 NumBlocksX = FMath::DivideAndRoundUp(uint32(FromSize.X), BlockSize);
+		const uint32 NumSubBlocksX = FMath::DivideAndRoundUp(uint32(SubSize.X), BlockSize);
+		const uint32 NumSubBlocksY = FMath::DivideAndRoundUp(uint32(SubSize.Y), BlockSize);
+
+		for (uint32 by = 0; by < NumSubBlocksX; ++by)
+		{
+			const uint8* RowFrom = From + 16 * NumBlocksX * by;
+
+			for (uint32 bx = 0; bx < NumSubBlocksY; ++bx)
+			{
+				uint8 Block[4 * 4 * 4];
+
+				//uint8 const* ColourBlock = from + 8;
+
+				// TODO, can skip alpha part
+				bcdec::bcdec_bc2(RowFrom, Block, 4 * 4);
+
+				int32 b = 0;
+				for (int32 j = 0; j < 4; ++j)
+				{
+					for (int32 i = 0; i < 4; ++i)
+					{
+						if (LIKELY((by * 4 + j < ToSize.Y) & (bx * 4 + i < ToSize.X)))
+						{
+							uint8* pPixelDest = To + 3 * (ToSize.X * (by * 4 + j) + (bx * 4 + i));
+
+							pPixelDest[0] = Block[b + 0];
+							pPixelDest[1] = Block[b + 1];
+							pPixelDest[2] = Block[b + 2];
+						}
+
+						b += 4;
+					}
+				}
+
+				RowFrom += 16;
+			}
+		}
+	}
+
+	void BC3_To_RGBASubImage(FImageSize FromSize, FImageSize ToSize, FImageSize SubSize, const uint8* RESTRICT From, uint8* RESTRICT To)
+	{
+		constexpr uint32 BlockSize = 4;
+
+		const uint32 NumBlocksX = FMath::DivideAndRoundUp(uint32(FromSize.X), BlockSize);
+		const uint32 NumSubBlocksX = FMath::DivideAndRoundUp(uint32(SubSize.X), BlockSize);
+		const uint32 NumSubBlocksY = FMath::DivideAndRoundUp(uint32(SubSize.Y), BlockSize);
+
+		for (uint32 by = 0; by < NumSubBlocksY; ++by)
+		{
+			const uint8* RowFrom = From + 16 * NumBlocksX * by;
+
+			for (uint32 bx = 0; bx < NumSubBlocksX; ++bx)
+			{
+				FMiroPixelBlock4x4 Block;
+
+				//uint8 const* ColourBlock = From + 8;
+				//uint8 const* AlphaBlock = From;
+
+				bcdec::bcdec_bc3(RowFrom, Block.bytes, 4 * 4);
+
+				int32 b = 0;
+				for (int32 j = 0; j < 4; ++j)
+				{
+					for (int32 i = 0; i < 4; ++i)
+					{
+						if (LIKELY((by * 4 + j < ToSize.Y) & (bx * 4 + i < ToSize.X)))
+						{
+							uint8* pPixelDest = To + 4 * (ToSize.X * (by * 4 + j) + (bx * 4 + i));
+
+							pPixelDest[0] = Block.bytes[b + 0];
+							pPixelDest[1] = Block.bytes[b + 1];
+							pPixelDest[2] = Block.bytes[b + 2];
+							pPixelDest[3] = Block.bytes[b + 3];
+						}
+
+						b += 4;
+					}
+				}
+
+				RowFrom += 16;
+			}
+		}
+	}
+
+	void BC3_To_RGBSubImage(FImageSize FromSize, FImageSize ToSize, FImageSize SubSize, const uint8* RESTRICT From, uint8* RESTRICT To)
+	{
+		constexpr uint32 BlockSize = 4;
+
+		const uint32 NumBlocksX = FMath::DivideAndRoundUp(uint32(FromSize.X), BlockSize);
+		const uint32 NumSubBlocksX = FMath::DivideAndRoundUp(uint32(SubSize.X), BlockSize);
+		const uint32 NumSubBlocksY = FMath::DivideAndRoundUp(uint32(SubSize.Y), BlockSize);
+
+		for (uint32 by = 0; by < NumSubBlocksY; ++by)
+		{
+			const uint8* RowFrom = From + 16 * NumBlocksX * by;
+
+			for (uint32 bx = 0; bx < NumSubBlocksY; ++bx)
+			{
+				uint8 Block[4 * 4 * 4];
+
+				bcdec::bcdec_bc3(RowFrom, Block, 4 * 4);
+
+				int32 b = 0;
+				for (int32 j = 0; j < 4; ++j)
+				{
+					for (int32 i = 0; i < 4; ++i)
+					{
+						if (LIKELY((by * 4 + j < ToSize.Y) & (bx * 4 + i < ToSize.X)))
+						{
+							uint8* pPixelDest = To + 3 * (ToSize.X * (by * 4 + j) + (bx * 4 + i));
+
+							pPixelDest[0] = Block[b + 0];
+							pPixelDest[1] = Block[b + 1];
+							pPixelDest[2] = Block[b + 2];
+						}
+
+						b += 4;
+					}
+				}
+
+				RowFrom += 16;
+			}
+		}
+	}
+
+	void BC4_To_LSubImage(FImageSize FromSize, FImageSize ToSize, FImageSize SubSize, const uint8* RESTRICT From, uint8* RESTRICT To)
+	{
+		constexpr uint32 BlockSize = 4;
+
+		const uint32 NumBlocksX = FMath::DivideAndRoundUp(uint32(FromSize.X), BlockSize);
+		const uint32 NumSubBlocksX = FMath::DivideAndRoundUp(uint32(SubSize.X), BlockSize);
+		const uint32 NumSubBlocksY = FMath::DivideAndRoundUp(uint32(SubSize.Y), BlockSize);
+		
+		for (uint32 by = 0; by < NumSubBlocksY; ++by)
+		{
+			const uint8* RowFrom = From + 8 * NumBlocksX * by;
+
+			for (uint32 bx = 0; bx < NumBlocksX; ++bx)
+			{
+				uint8 Block[4 * 4];
+
+				bcdec::bcdec_bc4(RowFrom, Block, 4);
+
+				int32 b = 0;
+				for (int32 j = 0; j < 4; ++j)
+				{
+					for (int32 i = 0; i < 4; ++i)
+					{
+						if (LIKELY((by * 4 + j < ToSize.Y) & (bx * 4 + i < ToSize.X)))
+						{
+							uint8* pPixelDest = To + 1 * (ToSize.X * (by * 4 + j) + (bx * 4 + i));
+
+							pPixelDest[0] = Block[b];
+						}
+
+						++b;
+					}
+				}
+
+				RowFrom += 8;
+			}
+		}
+	}
+
+
+	void BC4_To_RGBSubImage(FImageSize FromSize, FImageSize ToSize, FImageSize SubSize, const uint8* RESTRICT From, uint8* RESTRICT To)
+	{
+		constexpr uint32 BlockSize = 4;
+
+		const uint32 NumBlocksX = FMath::DivideAndRoundUp(uint32(FromSize.X), BlockSize);
+		const uint32 NumSubBlocksX = FMath::DivideAndRoundUp(uint32(SubSize.X), BlockSize);
+		const uint32 NumSubBlocksY = FMath::DivideAndRoundUp(uint32(SubSize.Y), BlockSize);
+
+		for (uint32 by = 0; by < NumSubBlocksY; ++by)
+		{
+			const uint8* RowFrom = From + 8 * NumBlocksX * by;
+
+			for (uint32 bx = 0; bx < NumSubBlocksX; ++bx)
+			{
+				uint8 Block[4 * 4];
+
+				bcdec::bcdec_bc4(RowFrom, Block, 4);
+
+				int32 b = 0;
+				for (int32 j = 0; j < 4; ++j)
+				{
+					for (int32 i = 0; i < 4; ++i)
+					{
+						if (LIKELY((by * 4 + j < ToSize.Y) & (bx * 4 + i < ToSize.X)))
+						{
+							uint8* pPixelDest = To + 3 * (ToSize.X * (by * 4 + j) + (bx * 4 + i));
+
+							pPixelDest[0] = Block[b];
+							pPixelDest[1] = Block[b];
+							pPixelDest[2] = Block[b];
+						}
+
+						++b;
+					}
+				}
+
+				RowFrom += 8;
+			}
+		}
+	}
+
+
+	void BC4_To_RGBASubImage(FImageSize FromSize, FImageSize ToSize, FImageSize SubSize, const uint8* RESTRICT From, uint8* RESTRICT To)
+	{
+		constexpr uint32 BlockSize = 4;
+
+		const uint32 NumBlocksX = FMath::DivideAndRoundUp(uint32(FromSize.X), BlockSize);
+		const uint32 NumSubBlocksX = FMath::DivideAndRoundUp(uint32(SubSize.X), BlockSize);
+		const uint32 NumSubBlocksY = FMath::DivideAndRoundUp(uint32(SubSize.Y), BlockSize);
+
+		for (uint32 by = 0; by < NumSubBlocksY; ++by)
+		{
+			const uint8* RowFrom = From + 8 * NumBlocksX * by;
+
+			for (uint32 bx = 0; bx < NumSubBlocksX; ++bx)
+			{
+				uint8 Block[4 * 4];
+
+				bcdec::bcdec_bc4(RowFrom, Block, 4);
+
+				int32 b = 0;
+				for (int32 j = 0; j < 4; ++j)
+				{
+					for (int32 i = 0; i < 4; ++i)
+					{
+						if (LIKELY((by * 4 + j < ToSize.Y) & (bx * 4 + i < ToSize.X)))
+						{
+							uint8* pPixelDest = To + 4 * (ToSize.X * (by * 4 + j) + (bx * 4 + i));
+
+							pPixelDest[0] = Block[b];
+							pPixelDest[1] = Block[b];
+							pPixelDest[2] = Block[b];
+							pPixelDest[3] = 255;
+						}
+
+						++b;
+					}
+				}
+
+				RowFrom += 8;
+			}
+		}
+	}
+
+
+	void BC5_To_RGBASubImage(FImageSize FromSize, FImageSize ToSize, FImageSize SubSize, const uint8* RESTRICT From, uint8* RESTRICT To)
+	{
+		constexpr uint32 BlockSize = 4;
+
+		const uint32 NumBlocksX = FMath::DivideAndRoundUp(uint32(FromSize.X), BlockSize);
+		const uint32 NumSubBlocksX = FMath::DivideAndRoundUp(uint32(SubSize.X), BlockSize);
+		const uint32 NumSubBlocksY = FMath::DivideAndRoundUp(uint32(SubSize.Y), BlockSize);
+
+		for (uint32 by = 0; by < NumSubBlocksY; ++by)
+		{
+			const uint8* RowFrom = From + 16 * NumBlocksX * by;
+
+			for (uint32 bx = 0; bx < NumSubBlocksX; ++bx)
+			{
+				uint8 Block[4 * 4 * 2];
+				bcdec::bcdec_bc5(RowFrom, Block, 4 * 2);
+
+				int32 b = 0;
+				for (int32 j = 0; j < 4; ++j)
+				{
+					for (int32 i = 0; i < 4; ++i)
+					{
+						if (LIKELY((by * 4 + j < ToSize.Y) & (bx * 4 + i < ToSize.X)))
+						{
+							uint8* pPixelDest = To + 4 * (ToSize.X * (by * 4 + j) + (bx * 4 + i));
+
+							pPixelDest[0] = Block[b + 0];
+							pPixelDest[1] = Block[b + 1];
+							pPixelDest[2] = 255;
+							pPixelDest[3] = 255;
+						}
+
+						b += 2;
+					}
+				}
+
+				RowFrom += 16;
+			}
+		}
+	}
+
+	void BC5_To_RGBSubImage(FImageSize FromSize, FImageSize ToSize, FImageSize SubSize, const uint8* RESTRICT From, uint8* RESTRICT To)
+	{
+		constexpr uint32 BlockSize = 4;
+
+		const uint32 NumBlocksX = FMath::DivideAndRoundUp(uint32(FromSize.X), BlockSize);
+		const uint32 NumSubBlocksX = FMath::DivideAndRoundUp(uint32(SubSize.X), BlockSize);
+		const uint32 NumSubBlocksY = FMath::DivideAndRoundUp(uint32(SubSize.Y), BlockSize);
+
+		for (uint32 by = 0; by < NumSubBlocksY; ++by)
+		{
+			const uint8* RowFrom = From + 16 * NumBlocksX * by;
+
+			for (uint32 bx = 0; bx < NumSubBlocksX; ++bx)
+			{
+				uint8 Block[4 * 4 * 2];
+				bcdec::bcdec_bc5(RowFrom, Block, 4 * 2);
+
+				int32 b = 0;
+				for (int32 j = 0; j < 4; ++j)
+				{
+					for (int32 i = 0; i < 4; ++i)
+					{
+						if (LIKELY((by * 4 + j < ToSize.Y) & (bx * 4 + i < ToSize.X)))
+						{
+							uint8* pPixelDest = To + 3 * (ToSize.X * (by * 4 + j) + (bx * 4 + i));
+
+							pPixelDest[0] = Block[b + 0];
+							pPixelDest[1] = Block[b + 1];
+							pPixelDest[2] = 255;
+						}
+
+						b += 2;
+					}
+				}
+
+				RowFrom += 16;
+			}
+		}
+	}
+#endif
 }
 
 #ifdef UE_MIRO_DEBUG
diff --git a/Engine/Plugins/Experimental/Mutable/Source/MutableRuntime/Private/MuR/BlockCompression/Miro/Miro.h b/Engine/Plugins/Experimental/Mutable/Source/MutableRuntime/Private/MuR/BlockCompression/Miro/Miro.h
--- a/Engine/Plugins/Experimental/Mutable/Source/MutableRuntime/Private/MuR/BlockCompression/Miro/Miro.h
+++ b/Engine/Plugins/Experimental/Mutable/Source/MutableRuntime/Private/MuR/BlockCompression/Miro/Miro.h
@@ -3,7 +3,8 @@
 #pragma once
 
 #include "HAL/Platform.h"
-
+#include "Math/IntVector.h"
+#include "Templates/Function.h"
 // Configure the compression library with these defines:
 
 // Support for BC1 to BC7, also known as DXTC1 to 5
@@ -33,11 +34,11 @@ namespace miro
     // 3 - Best for tools
     // 4 - Maximum, with no time limits.
 
+    using FImageSize = UE::Math::TIntVector2<uint16>;
 
     extern void initialize();
     extern void finalize();
 
-
 #if MIRO_INCLUDE_BC
 
     //! BC1 support
@@ -148,3 +149,62 @@ namespace miro
 
 }
 
+
+namespace miro::SubImageDecompression
+{
+    using FuncRefType = TFunctionRef<void(miro::FImageSize, miro::FImageSize, miro::FImageSize, const uint8*, uint8*)>;
+
+#if MIRO_INCLUDE_ASTC
+	 void ASTC4x4RGBAL_To_RGBASubImage(FImageSize FromSize, FImageSize ToSize, FImageSize SubSize, const uint8* RESTRICT From, uint8* RESTRICT To);
+	 void ASTC4x4RGBAL_To_RGBSubImage (FImageSize FromSize, FImageSize ToSize, FImageSize SubSize, const uint8* RESTRICT From, uint8* RESTRICT To);
+	 void ASTC4x4RGBL_To_RGBSubImage  (FImageSize FromSize, FImageSize ToSize, FImageSize SubSize, const uint8* RESTRICT From, uint8* RESTRICT To);
+	 void ASTC4x4RGBL_To_RGBASubImage (FImageSize FromSize, FImageSize ToSize, FImageSize SubSize, const uint8* RESTRICT From, uint8* RESTRICT To);
+	 void ASTC4x4RGL_To_RGBSubImage   (FImageSize FromSize, FImageSize ToSize, FImageSize SubSize, const uint8* RESTRICT From, uint8* RESTRICT To);
+	 void ASTC4x4RGL_To_RGBASubImage  (FImageSize FromSize, FImageSize ToSize, FImageSize SubSize, const uint8* RESTRICT From, uint8* RESTRICT To);
+
+	 void ASTC6x6RGBAL_To_RGBASubImage(FImageSize FromSize, FImageSize ToSize, FImageSize SubSize, const uint8* RESTRICT From, uint8* RESTRICT To);
+	 void ASTC6x6RGBAL_To_RGBSubImage (FImageSize FromSize, FImageSize ToSize, FImageSize SubSize, const uint8* RESTRICT From, uint8* RESTRICT To);
+	 void ASTC6x6RGBL_To_RGBSubImage  (FImageSize FromSize, FImageSize ToSize, FImageSize SubSize, const uint8* RESTRICT From, uint8* RESTRICT To);
+	 void ASTC6x6RGBL_To_RGBASubImage (FImageSize FromSize, FImageSize ToSize, FImageSize SubSize, const uint8* RESTRICT From, uint8* RESTRICT To);
+	 void ASTC6x6RGL_To_RGBSubImage   (FImageSize FromSize, FImageSize ToSize, FImageSize SubSize, const uint8* RESTRICT From, uint8* RESTRICT To);
+	 void ASTC6x6RGL_To_RGBASubImage  (FImageSize FromSize, FImageSize ToSize, FImageSize SubSize, const uint8* RESTRICT From, uint8* RESTRICT To);
+
+	 void ASTC8x8RGBAL_To_RGBASubImage(FImageSize FromSize, FImageSize ToSize, FImageSize SubSize, const uint8* RESTRICT From, uint8* RESTRICT To);
+	 void ASTC8x8RGBAL_To_RGBSubImage (FImageSize FromSize, FImageSize ToSize, FImageSize SubSize, const uint8* RESTRICT From, uint8* RESTRICT To);
+	 void ASTC8x8RGBL_To_RGBSubImage  (FImageSize FromSize, FImageSize ToSize, FImageSize SubSize, const uint8* RESTRICT From, uint8* RESTRICT To);
+	 void ASTC8x8RGBL_To_RGBASubImage (FImageSize FromSize, FImageSize ToSize, FImageSize SubSize, const uint8* RESTRICT From, uint8* RESTRICT To);
+	 void ASTC8x8RGL_To_RGBSubImage   (FImageSize FromSize, FImageSize ToSize, FImageSize SubSize, const uint8* RESTRICT From, uint8* RESTRICT To);
+	 void ASTC8x8RGL_To_RGBASubImage  (FImageSize FromSize, FImageSize ToSize, FImageSize SubSize, const uint8* RESTRICT From, uint8* RESTRICT To);
+
+	 void ASTC10x10RGBAL_To_RGBASubImage(FImageSize FromSize, FImageSize ToSize, FImageSize SubSize, const uint8* RESTRICT From, uint8* RESTRICT To);
+	 void ASTC10x10RGBAL_To_RGBSubImage (FImageSize FromSize, FImageSize ToSize, FImageSize SubSize, const uint8* RESTRICT From, uint8* RESTRICT To);
+	 void ASTC10x10RGBL_To_RGBSubImage  (FImageSize FromSize, FImageSize ToSize, FImageSize SubSize, const uint8* RESTRICT From, uint8* RESTRICT To);
+	 void ASTC10x10RGBL_To_RGBASubImage (FImageSize FromSize, FImageSize ToSize, FImageSize SubSize, const uint8* RESTRICT From, uint8* RESTRICT To);
+	 void ASTC10x10RGL_To_RGBSubImage   (FImageSize FromSize, FImageSize ToSize, FImageSize SubSize, const uint8* RESTRICT From, uint8* RESTRICT To);
+	 void ASTC10x10RGL_To_RGBASubImage  (FImageSize FromSize, FImageSize ToSize, FImageSize SubSize, const uint8* RESTRICT From, uint8* RESTRICT To);
+
+	 void ASTC12x12RGBAL_To_RGBASubImage(FImageSize FromSize, FImageSize ToSize, FImageSize SubSize, const uint8* RESTRICT From, uint8* RESTRICT To);
+	 void ASTC12x12RGBAL_To_RGBSubImage (FImageSize FromSize, FImageSize ToSize, FImageSize SubSize, const uint8* RESTRICT From, uint8* RESTRICT To);
+	 void ASTC12x12RGBL_To_RGBSubImage  (FImageSize FromSize, FImageSize ToSize, FImageSize SubSize, const uint8* RESTRICT From, uint8* RESTRICT To);
+	 void ASTC12x12RGBL_To_RGBASubImage (FImageSize FromSize, FImageSize ToSize, FImageSize SubSize, const uint8* RESTRICT From, uint8* RESTRICT To);
+	 void ASTC12x12RGL_To_RGBSubImage   (FImageSize FromSize, FImageSize ToSize, FImageSize SubSize, const uint8* RESTRICT From, uint8* RESTRICT To);
+	 void ASTC12x12RGL_To_RGBASubImage  (FImageSize FromSize, FImageSize ToSize, FImageSize SubSize, const uint8* RESTRICT From, uint8* RESTRICT To);
+
+#endif // MIRO_INCLUDE_ASTC
+	// BC Formats.
+
+#if MIRO_INCLUDE_BC
+	 void BC1_To_RGBSubImage (FImageSize FromSize, FImageSize ToSize, FImageSize SubSize, const uint8* RESTRICT From, uint8* RESTRICT To);
+	 void BC1_To_RGBASubImage(FImageSize FromSize, FImageSize ToSize, FImageSize SubSize, const uint8* RESTRICT From, uint8* RESTRICT To);
+	 void BC2_To_RGBASubImage(FImageSize FromSize, FImageSize ToSize, FImageSize SubSize, const uint8* RESTRICT From, uint8* RESTRICT To);
+	 void BC2_To_RGBSubImage (FImageSize FromSize, FImageSize ToSize, FImageSize SubSize, const uint8* RESTRICT From, uint8* RESTRICT To);
+	 void BC3_To_RGBASubImage(FImageSize FromSize, FImageSize ToSize, FImageSize SubSize, const uint8* RESTRICT From, uint8* RESTRICT To);
+	 void BC3_To_RGBSubImage (FImageSize FromSize, FImageSize ToSize, FImageSize SubSize, const uint8* RESTRICT From, uint8* RESTRICT To);
+	 void BC4_To_LSubImage   (FImageSize FromSize, FImageSize ToSize, FImageSize SubSize, const uint8* RESTRICT From, uint8* RESTRICT To);
+	 void BC4_To_RGBSubImage (FImageSize FromSize, FImageSize ToSize, FImageSize SubSize, const uint8* RESTRICT From, uint8* RESTRICT To);
+	 void BC4_To_RGBASubImage(FImageSize FromSize, FImageSize ToSize, FImageSize SubSize, const uint8* RESTRICT From, uint8* RESTRICT To);
+	 void BC5_To_RGBASubImage(FImageSize FromSize, FImageSize ToSize, FImageSize SubSize, const uint8* RESTRICT From, uint8* RESTRICT To);
+	 void BC5_To_RGBSubImage (FImageSize FromSize, FImageSize ToSize, FImageSize SubSize, const uint8* RESTRICT From, uint8* RESTRICT To);
+#endif
+}
+
diff --git a/Engine/Plugins/Experimental/Mutable/Source/MutableRuntime/Private/MuR/ImagePrivate.h b/Engine/Plugins/Experimental/Mutable/Source/MutableRuntime/Private/MuR/ImagePrivate.h
--- a/Engine/Plugins/Experimental/Mutable/Source/MutableRuntime/Private/MuR/ImagePrivate.h
+++ b/Engine/Plugins/Experimental/Mutable/Source/MutableRuntime/Private/MuR/ImagePrivate.h
@@ -8,6 +8,8 @@
 #include "MuR/SerialisationPrivate.h"
 #include "MuR/MutableMath.h"
 
+#include "MuR/BlockCompression/Miro/Miro.h"
+
 namespace mu
 {
 	MUTABLE_DEFINE_ENUM_SERIALISABLE( EBlendType );
@@ -21,56 +23,75 @@ namespace mu
 	//---------------------------------------------------------------------------------------------
 	//!
 	//---------------------------------------------------------------------------------------------
-	MUTABLERUNTIME_API inline EImageFormat GetUncompressedFormat(EImageFormat f )
+	MUTABLERUNTIME_API inline EImageFormat GetUncompressedFormat(EImageFormat Format)
 	{
-		check(f < EImageFormat::IF_COUNT);
+		check(Format < EImageFormat::IF_COUNT);
 
-		EImageFormat r = f;
+		EImageFormat Result = Format;
 
-		switch ( r )
+		switch (Result)
 		{
-		case EImageFormat::IF_L_UBIT_RLE: r = EImageFormat::IF_L_UBYTE; break;
-		case EImageFormat::IF_L_UBYTE_RLE: r = EImageFormat::IF_L_UBYTE; break;
-		case EImageFormat::IF_RGB_UBYTE_RLE: r = EImageFormat::IF_RGB_UBYTE; break;
-        case EImageFormat::IF_RGBA_UBYTE_RLE: r = EImageFormat::IF_RGBA_UBYTE; break;
-        case EImageFormat::IF_BC1: r = EImageFormat::IF_RGBA_UBYTE; break;
-        case EImageFormat::IF_BC2: r = EImageFormat::IF_RGBA_UBYTE; break;
-        case EImageFormat::IF_BC3: r = EImageFormat::IF_RGBA_UBYTE; break;
-        case EImageFormat::IF_BC4: r = EImageFormat::IF_L_UBYTE; break;
-        case EImageFormat::IF_BC5: r = EImageFormat::IF_RGB_UBYTE; break;
-		case EImageFormat::IF_ASTC_4x4_RGB_LDR: r = EImageFormat::IF_RGB_UBYTE; break;
-		case EImageFormat::IF_ASTC_4x4_RGBA_LDR: r = EImageFormat::IF_RGBA_UBYTE; break;
-		case EImageFormat::IF_ASTC_4x4_RG_LDR: r = EImageFormat::IF_RGB_UBYTE; break;
-		case EImageFormat::IF_ASTC_6x6_RGB_LDR: r = EImageFormat::IF_RGB_UBYTE; break;
-		case EImageFormat::IF_ASTC_6x6_RGBA_LDR: r = EImageFormat::IF_RGBA_UBYTE; break;
-		case EImageFormat::IF_ASTC_6x6_RG_LDR: r = EImageFormat::IF_RGB_UBYTE; break;
-		case EImageFormat::IF_ASTC_8x8_RGB_LDR: r = EImageFormat::IF_RGB_UBYTE; break;
-		case EImageFormat::IF_ASTC_8x8_RGBA_LDR: r = EImageFormat::IF_RGBA_UBYTE; break;
-		case EImageFormat::IF_ASTC_8x8_RG_LDR: r = EImageFormat::IF_RGB_UBYTE; break;
-		case EImageFormat::IF_ASTC_10x10_RGB_LDR: r = EImageFormat::IF_RGB_UBYTE; break;
-		case EImageFormat::IF_ASTC_10x10_RGBA_LDR: r = EImageFormat::IF_RGBA_UBYTE; break;
-		case EImageFormat::IF_ASTC_10x10_RG_LDR: r = EImageFormat::IF_RGB_UBYTE; break;
-		case EImageFormat::IF_ASTC_12x12_RGB_LDR: r = EImageFormat::IF_RGB_UBYTE; break;
-		case EImageFormat::IF_ASTC_12x12_RGBA_LDR: r = EImageFormat::IF_RGBA_UBYTE; break;
-		case EImageFormat::IF_ASTC_12x12_RG_LDR: r = EImageFormat::IF_RGB_UBYTE; break;
+		case EImageFormat::IF_L_UBIT_RLE: Result = EImageFormat::IF_L_UBYTE; break;
+		case EImageFormat::IF_L_UBYTE_RLE: Result = EImageFormat::IF_L_UBYTE; break;
+		case EImageFormat::IF_RGB_UBYTE_RLE: Result = EImageFormat::IF_RGB_UBYTE; break;
+        case EImageFormat::IF_RGBA_UBYTE_RLE: Result = EImageFormat::IF_RGBA_UBYTE; break;
+        case EImageFormat::IF_BC1: Result = EImageFormat::IF_RGBA_UBYTE; break;
+        case EImageFormat::IF_BC2: Result = EImageFormat::IF_RGBA_UBYTE; break;
+        case EImageFormat::IF_BC3: Result = EImageFormat::IF_RGBA_UBYTE; break;
+        case EImageFormat::IF_BC4: Result = EImageFormat::IF_L_UBYTE; break;
+        case EImageFormat::IF_BC5: Result = EImageFormat::IF_RGB_UBYTE; break;
+		case EImageFormat::IF_ASTC_4x4_RGB_LDR: Result = EImageFormat::IF_RGB_UBYTE; break;
+		case EImageFormat::IF_ASTC_4x4_RGBA_LDR: Result = EImageFormat::IF_RGBA_UBYTE; break;
+		case EImageFormat::IF_ASTC_4x4_RG_LDR: Result = EImageFormat::IF_RGB_UBYTE; break;
+		case EImageFormat::IF_ASTC_6x6_RGB_LDR: Result = EImageFormat::IF_RGB_UBYTE; break;
+		case EImageFormat::IF_ASTC_6x6_RGBA_LDR: Result = EImageFormat::IF_RGBA_UBYTE; break;
+		case EImageFormat::IF_ASTC_6x6_RG_LDR: Result = EImageFormat::IF_RGB_UBYTE; break;
+		case EImageFormat::IF_ASTC_8x8_RGB_LDR: Result = EImageFormat::IF_RGB_UBYTE; break;
+		case EImageFormat::IF_ASTC_8x8_RGBA_LDR: Result = EImageFormat::IF_RGBA_UBYTE; break;
+		case EImageFormat::IF_ASTC_8x8_RG_LDR: Result = EImageFormat::IF_RGB_UBYTE; break;
+		case EImageFormat::IF_ASTC_10x10_RGB_LDR: Result = EImageFormat::IF_RGB_UBYTE; break;
+		case EImageFormat::IF_ASTC_10x10_RGBA_LDR: Result = EImageFormat::IF_RGBA_UBYTE; break;
+		case EImageFormat::IF_ASTC_10x10_RG_LDR: Result = EImageFormat::IF_RGB_UBYTE; break;
+		case EImageFormat::IF_ASTC_12x12_RGB_LDR: Result = EImageFormat::IF_RGB_UBYTE; break;
+		case EImageFormat::IF_ASTC_12x12_RGBA_LDR: Result = EImageFormat::IF_RGBA_UBYTE; break;
+		case EImageFormat::IF_ASTC_12x12_RG_LDR: Result = EImageFormat::IF_RGB_UBYTE; break;
 		default: break;
 		}
 
-		return r;
+		return Result;
 	}
 
     //---------------------------------------------------------------------------------------------
     //!
     //---------------------------------------------------------------------------------------------
-	inline EImageFormat GetMostGenericFormat(EImageFormat a, EImageFormat b)
+	inline EImageFormat GetMostGenericFormat(EImageFormat FormatA, EImageFormat FormatB)
     {
-        if (a==b) return a;
-        if ( GetImageFormatData(a).Channels>GetImageFormatData(b).Channels ) return a;
-        if ( GetImageFormatData(b).Channels>GetImageFormatData(a).Channels ) return b;
-        if (a== EImageFormat::IF_BC2 || a== EImageFormat::IF_BC3 || a== EImageFormat::IF_ASTC_4x4_RGBA_LDR) return a;
-        if (b== EImageFormat::IF_BC2 || b== EImageFormat::IF_BC3 || b== EImageFormat::IF_ASTC_4x4_RGBA_LDR) return b;
+        if (FormatA == FormatB) 
+		{	
+			return FormatA;
+		}
+
+        if (GetImageFormatData(FormatA).Channels > GetImageFormatData(FormatB).Channels)
+		{	
+			return FormatA;
+		}
+
+        if (GetImageFormatData(FormatB).Channels > GetImageFormatData(FormatA).Channels)
+		{
+			return FormatB;
+		}
+
+        if (FormatA == EImageFormat::IF_BC2 || FormatA == EImageFormat::IF_BC3 || FormatA == EImageFormat::IF_ASTC_4x4_RGBA_LDR)
+		{	
+			return FormatA;
+		}
+
+        if (FormatB == EImageFormat::IF_BC2 || FormatB == EImageFormat::IF_BC3 || FormatB == EImageFormat::IF_ASTC_4x4_RGBA_LDR)
+		{
+			return FormatB;
+		}
 
-        return a;
+        return FormatA;
     }
 
 
@@ -110,12 +131,301 @@ namespace mu
     //---------------------------------------------------------------------------------------------
     //!
     //---------------------------------------------------------------------------------------------
-    inline bool IsCompressedFormat(EImageFormat f )
+    inline bool IsCompressedFormat(EImageFormat Format)
     {
-        return f!=GetUncompressedFormat(f);
+        return Format != GetUncompressedFormat(Format);
     }
 
 
+	inline bool IsBlockCompressedFormat(EImageFormat Format)
+	{
+		return GetImageFormatData(Format).PixelsPerBlockX > 1 && Format != GetUncompressedFormat(Format);
+	}
+
+	namespace Private
+	{
+		inline void DecompressionFuncNotFoundFunc(miro::FImageSize, miro::FImageSize, miro::FImageSize, const uint8*, uint8*)
+		{
+			check(false);
+		}
+	}
+
+	inline miro::SubImageDecompression::FuncRefType SelectDecompressionFunction(EImageFormat DestFormat, EImageFormat SrcFormat)
+	{
+		using DecompressionFuncRefType = miro::SubImageDecompression::FuncRefType;
+
+		switch(SrcFormat)
+		{
+			case EImageFormat::IF_BC1:
+			{
+				if (DestFormat == EImageFormat::IF_RGBA_UBYTE)
+				{
+					return DecompressionFuncRefType(miro::SubImageDecompression::BC1_To_RGBASubImage);
+				}
+				else if (DestFormat == EImageFormat::IF_RGB_UBYTE)
+				{
+					return DecompressionFuncRefType(miro::SubImageDecompression::BC1_To_RGBSubImage);
+				}
+				break;
+			}
+
+			case EImageFormat::IF_BC2:
+			{
+				if (DestFormat == EImageFormat::IF_RGBA_UBYTE)
+				{
+					return DecompressionFuncRefType(miro::SubImageDecompression::BC2_To_RGBASubImage);
+				}
+				else if (DestFormat == EImageFormat::IF_RGB_UBYTE)
+				{
+					return DecompressionFuncRefType(miro::SubImageDecompression::BC2_To_RGBSubImage);
+				}
+				break;
+			}
+
+			case EImageFormat::IF_BC3:
+			{
+				if (DestFormat == EImageFormat::IF_RGBA_UBYTE)
+				{
+					return DecompressionFuncRefType(miro::SubImageDecompression::BC3_To_RGBASubImage);
+				}
+				else if (DestFormat == EImageFormat::IF_RGB_UBYTE)
+				{
+					return DecompressionFuncRefType(miro::SubImageDecompression::BC3_To_RGBSubImage);
+				}
+				break;
+			}
+
+			case EImageFormat::IF_BC4:
+			{
+				if (DestFormat == EImageFormat::IF_RGBA_UBYTE)
+				{
+					return DecompressionFuncRefType(miro::SubImageDecompression::BC4_To_RGBASubImage);
+				}
+				else if (DestFormat == EImageFormat::IF_RGB_UBYTE)
+				{
+					return DecompressionFuncRefType(miro::SubImageDecompression::BC4_To_RGBSubImage);
+				}
+				if (DestFormat == EImageFormat::IF_L_UBYTE)
+				{
+					return DecompressionFuncRefType(miro::SubImageDecompression::BC4_To_LSubImage);
+				}
+				break;
+			}
+
+			case EImageFormat::IF_BC5:
+			{
+				if (DestFormat == EImageFormat::IF_RGBA_UBYTE)
+				{
+					return DecompressionFuncRefType(miro::SubImageDecompression::BC5_To_RGBASubImage);
+				}
+				else if (DestFormat == EImageFormat::IF_RGB_UBYTE)
+				{
+					return DecompressionFuncRefType(miro::SubImageDecompression::BC5_To_RGBSubImage);
+				}
+				break;
+			}
+	
+			case EImageFormat::IF_ASTC_4x4_RGBA_LDR:
+			{
+				if (DestFormat == EImageFormat::IF_RGBA_UBYTE)
+				{
+					return DecompressionFuncRefType(miro::SubImageDecompression::ASTC4x4RGBAL_To_RGBASubImage);
+				}
+				else if (DestFormat == EImageFormat::IF_RGB_UBYTE)
+				{
+					return DecompressionFuncRefType(miro::SubImageDecompression::ASTC4x4RGBAL_To_RGBSubImage);
+				}
+				break;
+			}
+
+			case EImageFormat::IF_ASTC_4x4_RGB_LDR:
+			{
+				if (DestFormat == EImageFormat::IF_RGBA_UBYTE)
+				{
+					return DecompressionFuncRefType(miro::SubImageDecompression::ASTC4x4RGBL_To_RGBASubImage);
+				}
+				else if (DestFormat == EImageFormat::IF_RGB_UBYTE)
+				{
+					return DecompressionFuncRefType(miro::SubImageDecompression::ASTC4x4RGBL_To_RGBSubImage);
+				}
+				break;
+			}
+
+			case EImageFormat::IF_ASTC_4x4_RG_LDR:
+			{
+				if (DestFormat == EImageFormat::IF_RGBA_UBYTE)
+				{
+					return DecompressionFuncRefType(miro::SubImageDecompression::ASTC4x4RGL_To_RGBASubImage);
+				}
+				else if (DestFormat == EImageFormat::IF_RGB_UBYTE)
+				{
+					return DecompressionFuncRefType(miro::SubImageDecompression::ASTC4x4RGL_To_RGBSubImage);
+				}
+				break;
+			}
+
+			case EImageFormat::IF_ASTC_6x6_RGBA_LDR:
+			{
+				if (DestFormat == EImageFormat::IF_RGBA_UBYTE)
+				{
+					return DecompressionFuncRefType(miro::SubImageDecompression::ASTC6x6RGBAL_To_RGBASubImage);
+				}
+				else if (DestFormat == EImageFormat::IF_RGB_UBYTE)
+				{
+					return DecompressionFuncRefType(miro::SubImageDecompression::ASTC6x6RGBAL_To_RGBSubImage);
+				}
+				break;
+			}
+
+			case EImageFormat::IF_ASTC_6x6_RGB_LDR:
+			{
+				if (DestFormat == EImageFormat::IF_RGBA_UBYTE)
+				{
+					return DecompressionFuncRefType(miro::SubImageDecompression::ASTC6x6RGBL_To_RGBASubImage);
+				}
+				else if (DestFormat == EImageFormat::IF_RGB_UBYTE)
+				{
+					return DecompressionFuncRefType(miro::SubImageDecompression::ASTC6x6RGBL_To_RGBSubImage);
+				}
+				break;
+			}
+
+			case EImageFormat::IF_ASTC_6x6_RG_LDR:
+			{
+				if (DestFormat == EImageFormat::IF_RGBA_UBYTE)
+				{
+					return DecompressionFuncRefType(miro::SubImageDecompression::ASTC6x6RGL_To_RGBASubImage);
+				}
+				else if (DestFormat == EImageFormat::IF_RGB_UBYTE)
+				{
+					return DecompressionFuncRefType(miro::SubImageDecompression::ASTC6x6RGL_To_RGBSubImage);
+				}
+				break;
+			}
+
+			case EImageFormat::IF_ASTC_8x8_RGBA_LDR:
+			{
+				if (DestFormat == EImageFormat::IF_RGBA_UBYTE)
+				{
+					return DecompressionFuncRefType(miro::SubImageDecompression::ASTC8x8RGBAL_To_RGBASubImage);
+				}
+				else if (DestFormat == EImageFormat::IF_RGB_UBYTE)
+				{
+					return DecompressionFuncRefType(miro::SubImageDecompression::ASTC8x8RGBAL_To_RGBSubImage);
+				}
+				break;
+			}
+
+			case EImageFormat::IF_ASTC_8x8_RGB_LDR:
+			{
+				if (DestFormat == EImageFormat::IF_RGBA_UBYTE)
+				{
+					return DecompressionFuncRefType(miro::SubImageDecompression::ASTC8x8RGBL_To_RGBASubImage);
+				}
+				else if (DestFormat == EImageFormat::IF_RGB_UBYTE)
+				{
+					return DecompressionFuncRefType(miro::SubImageDecompression::ASTC8x8RGBL_To_RGBSubImage);
+				}
+				break;
+			}
+
+			case EImageFormat::IF_ASTC_8x8_RG_LDR:
+			{
+				if (DestFormat == EImageFormat::IF_RGBA_UBYTE)
+				{
+					return DecompressionFuncRefType(miro::SubImageDecompression::ASTC8x8RGL_To_RGBASubImage);
+				}
+				else if (DestFormat == EImageFormat::IF_RGB_UBYTE)
+				{
+					return DecompressionFuncRefType(miro::SubImageDecompression::ASTC8x8RGL_To_RGBSubImage);
+				}
+				break;
+			}
+
+			case EImageFormat::IF_ASTC_10x10_RGBA_LDR:
+			{
+				if (DestFormat == EImageFormat::IF_RGBA_UBYTE)
+				{
+					return DecompressionFuncRefType(miro::SubImageDecompression::ASTC10x10RGBAL_To_RGBASubImage);
+				}
+				else if (DestFormat == EImageFormat::IF_RGB_UBYTE)
+				{
+					return DecompressionFuncRefType(miro::SubImageDecompression::ASTC10x10RGBAL_To_RGBSubImage);
+				}
+				break;
+			}
+
+			case EImageFormat::IF_ASTC_10x10_RGB_LDR:
+			{
+				if (DestFormat == EImageFormat::IF_RGBA_UBYTE)
+				{
+					return DecompressionFuncRefType(miro::SubImageDecompression::ASTC10x10RGBL_To_RGBASubImage);
+				}
+				else if (DestFormat == EImageFormat::IF_RGB_UBYTE)
+				{
+					return DecompressionFuncRefType(miro::SubImageDecompression::ASTC10x10RGBL_To_RGBSubImage);
+				}
+				break;
+			}
+
+			case EImageFormat::IF_ASTC_10x10_RG_LDR:
+			{
+				if (DestFormat == EImageFormat::IF_RGBA_UBYTE)
+				{
+					return DecompressionFuncRefType(miro::SubImageDecompression::ASTC10x10RGL_To_RGBASubImage);
+				}
+				else if (DestFormat == EImageFormat::IF_RGB_UBYTE)
+				{
+					return DecompressionFuncRefType(miro::SubImageDecompression::ASTC10x10RGL_To_RGBSubImage);
+				}
+				break;
+			}
+
+			case EImageFormat::IF_ASTC_12x12_RGBA_LDR:
+			{
+				if (DestFormat == EImageFormat::IF_RGBA_UBYTE)
+				{
+					return DecompressionFuncRefType(miro::SubImageDecompression::ASTC12x12RGBAL_To_RGBASubImage);
+				}
+				else if (DestFormat == EImageFormat::IF_RGB_UBYTE)
+				{
+					return DecompressionFuncRefType(miro::SubImageDecompression::ASTC12x12RGBAL_To_RGBSubImage);
+				}
+				break;
+			}
+
+			case EImageFormat::IF_ASTC_12x12_RGB_LDR:
+			{
+				if (DestFormat == EImageFormat::IF_RGBA_UBYTE)
+				{
+					return DecompressionFuncRefType(miro::SubImageDecompression::ASTC12x12RGBL_To_RGBASubImage);
+				}
+				else if (DestFormat == EImageFormat::IF_RGB_UBYTE)
+				{
+					return DecompressionFuncRefType(miro::SubImageDecompression::ASTC12x12RGBL_To_RGBSubImage);
+				}
+				break;
+			}
+
+			case EImageFormat::IF_ASTC_12x12_RG_LDR:
+			{
+				if (DestFormat == EImageFormat::IF_RGBA_UBYTE)
+				{
+					return DecompressionFuncRefType(miro::SubImageDecompression::ASTC12x12RGL_To_RGBASubImage);
+				}
+				else if (DestFormat == EImageFormat::IF_RGB_UBYTE)
+				{
+					return DecompressionFuncRefType(miro::SubImageDecompression::ASTC12x12RGL_To_RGBSubImage);
+				}
+				break;
+			}
+		}
+
+		checkf(false, TEXT("Decompression not supported."));
+		return DecompressionFuncRefType(Private::DecompressionFuncNotFoundFunc);
+	}
+
+
     //---------------------------------------------------------------------------------------------
     //! Use with care.
     //---------------------------------------------------------------------------------------------
diff --git a/Engine/Plugins/Experimental/Mutable/Source/MutableRuntime/Private/MuR/OpImageMipmap.cpp b/Engine/Plugins/Experimental/Mutable/Source/MutableRuntime/Private/MuR/OpImageMipmap.cpp
--- a/Engine/Plugins/Experimental/Mutable/Source/MutableRuntime/Private/MuR/OpImageMipmap.cpp
+++ b/Engine/Plugins/Experimental/Mutable/Source/MutableRuntime/Private/MuR/OpImageMipmap.cpp
@@ -5,514 +5,403 @@
 #include "Async/ParallelFor.h"
 #include "HAL/UnrealMemory.h"
 
-namespace mu
+namespace
 {
 
-	namespace OpImageMipmap_Detail
-	{
-
-		struct SharpenKernelStorage8
-		{
-			constexpr static size_t SIZE = 8;
-			float m_storage[SIZE * SIZE];
-		};
-
-		template<size_t ITERS, size_t N>
-		inline void BlurKernel(float(&inOutKernelStorage)[N])
-		{
-			static_assert(N > 2);
-
-			float temp[N];
-
-			float* inBuf = &inOutKernelStorage[0];
-			float* outBuf = &temp[0];
-
-			for (int32 i = 0; i < ITERS; ++i)
-			{
-				constexpr float oneOverThree = 1.0f / 3.0f;
+bool bEnableCompressedMipGenerationMemoryOptimizations = false;
+static FAutoConsoleVariableRef CVarEnableCompressedMipGenerationMemoryOptimizations (
+	TEXT("mutable.EnableCompressedMipGenerationMemoryOptimizations"),
+	bEnableCompressedMipGenerationMemoryOptimizations,
+	TEXT("If set to true, enables memory optimizations for mip generation on compressed images."),
+	ECVF_Default);
+}
 
-				outBuf[0] = (inBuf[0] + inBuf[1]) * oneOverThree;
+namespace mu
+{
 
-				for (int32 j = 1; j < N - 1; ++j)
-				{
-					outBuf[j] = (inBuf[j - 1] + inBuf[j] + inBuf[j + 1]) * oneOverThree;
-				}
 
-				outBuf[N - 1] = (inBuf[N - 2] + inBuf[N - 1]) * oneOverThree;
+namespace OpImageMipmap_Detail
+{
 
-				// Ping-pong buffers.
-				std::swap(outBuf, inBuf);
-			}
+	template<int32 PIXEL_SIZE>
+	inline void GenerateMipmapUint8Unfiltered(
+		int mips,
+		const uint8* pSource, uint8* Dest,
+		FIntVector2 sourceSize)
+	{
+		for (; mips >= 0; --mips)
+		{
+			FIntVector2 destSize = FIntVector2(FMath::DivideAndRoundUp(sourceSize.X, 2), FMath::DivideAndRoundUp(sourceSize.Y, 2));
 
-			if (ITERS % 2 == 1)
+			for (int32 y = 0; y < destSize.Y; ++y)
 			{
-				for (uint32 i = 0; i < N; ++i)
+				for (int32 x = 0; x < destSize.X; ++x)
 				{
-					inOutKernelStorage[i] = temp[i];
+					for (int32 c = 0; c < PIXEL_SIZE; ++c)
+					{
+						Dest[(y * destSize.X + x) * PIXEL_SIZE + c] =
+							pSource[((y << 1) * sourceSize.X + (x << 1)) * PIXEL_SIZE + c];
+					}
 				}
 			}
-		}
-
-		template<size_t N>
-		inline void BuildGaussianBlurKernel(float(&kernel)[N], float variance)
-		{
-			constexpr float sqrtTau = 2.5066282746310002f;
-			const auto normalDist = [sdSqrtTauInv = 1.0f / (FMath::Sqrt(variance) * sqrtTau),
-				twoVInv = 1.0f / (2.0f * variance)]
-				(float x)
-				-> float
-			{
-				return FMath::Exp(-x * x * twoVInv) * sdSqrtTauInv;
-			};
-
-			const float center = static_cast<float>(N) * 0.5f;
-			float sum = 0;
-			for (int32 i = 0; i < N; ++i)
-			{
-				const float dist = normalDist(static_cast<float>(i) - center + 0.5f);
-				kernel[i] = dist;
-				sum += dist;
-			}
 
-			const float normFactor = 1.0f / sum;
-			for (int32 i = 0; i < N; ++i)
-			{
-				kernel[i] *= normFactor;
-			}
+			sourceSize = destSize;
+			pSource = Dest;
+			Dest = Dest + destSize.X * destSize.Y * PIXEL_SIZE;
 		}
+	}
 
-		template<size_t N>
-		inline void BuildSharpenKernel(float(&kernel)[N], float factor)
-		{
-			for (uint32 i = 0; i < N; ++i)
-			{
-				kernel[i] = 0.0f;
-			}
+	template<int32 PIXEL_SIZE>
+	inline void GenerateMipmapsUint8SimpleAverage(
+		int mips,
+		const uint8* pSource, uint8* Dest,
+		FIntVector2 sourceSize)
+	{
+		
+		const uint8* pMipSource = pSource;
+		uint8* pMipDest = Dest;
 
-			kernel[(N >> 1) - 1] = factor * 0.5f;
-			kernel[(N >> 1)] = factor * 0.5f;
-		}
+		FIntVector2 destSize = sourceSize;
 
-		template<size_t N>
-		inline void BuildKernel2DFromSeparable(float(&kernel1D)[N], SharpenKernelStorage8& outKernel)
+		for (int m = 0; m < mips; ++m)
 		{
-			static_assert(N == SharpenKernelStorage8::SIZE);
+			check(destSize[0] > 1 || destSize[1] > 1);
 
-			for (int32 y = 0; y < N; ++y)
-			{
-				for (int32 x = 0; x < N; ++x)
-				{
-					outKernel.m_storage[y * N + x] = kernel1D[y] * kernel1D[x];
-				}
-			}
-		}
-
-		inline SharpenKernelStorage8 MakeMipGaussianSharpenKernel(float factor)
-		{
-			SharpenKernelStorage8 kernel;
-
-			float kernel1DStorage0[SharpenKernelStorage8::SIZE];
-			float kernel1DStorage1[SharpenKernelStorage8::SIZE];
+			sourceSize = destSize;
 
-			// negative factors indicate blur where the -factor is the variance.
-			if (factor < 0)
-			{
-				BuildGaussianBlurKernel(kernel1DStorage0, -factor);
-
-				BuildKernel2DFromSeparable(kernel1DStorage0, kernel);
-
-				return kernel;
-			}
+			int fullColumns = destSize[0] / 2;
+			bool strayColumn = (destSize[0] % 2) != 0;
+			int fullRows = destSize[1] / 2;
+			bool strayRow = (destSize[1] % 2) != 0;
 
-			BuildSharpenKernel(kernel1DStorage0, 1.0f + factor);
-			BuildSharpenKernel(kernel1DStorage1, -factor);
+			destSize[0] = FMath::DivideAndRoundUp(destSize[0], 2);
+			destSize[1] = FMath::DivideAndRoundUp(destSize[1], 2);
 
-			BlurKernel<1>(kernel1DStorage0);
-			BlurKernel<3>(kernel1DStorage1);
+			int sourceStride = sourceSize[0] * PIXEL_SIZE;
+			int destStride = destSize[0] * PIXEL_SIZE;
 
-			for (int32 i = 0; i < SharpenKernelStorage8::SIZE; ++i)
+			const auto ProcessRow = [
+				pMipDest, pMipSource, fullColumns, strayColumn, sourceStride, destStride
+			] (uint32 y)
 			{
-				kernel1DStorage0[i] += kernel1DStorage1[i];
-			}
-
-			BuildKernel2DFromSeparable(kernel1DStorage0, kernel);
-
-			return kernel;
-		}
+				const uint8* pSourceRow0 = pMipSource + 2 * y * sourceStride;
+				const uint8* pSourceRow1 = pSourceRow0 + sourceStride;
+				uint8* pDestRow = pMipDest + y * destStride;
 
-		template<int32 PIXEL_SIZE, typename CHANNEL_TYPE, EAddressMode AD_MODE = EAddressMode::Wrap>
-		class ImageAccessor
-		{
-			const FIntVector2 m_dim;
-			CHANNEL_TYPE const* const m_pImageBuf;
-
-		public:
-			ImageAccessor(const CHANNEL_TYPE* const pImageBuf, FIntVector2 dim)
-				: m_dim(dim)
-				, m_pImageBuf(pImageBuf)
-			{
-			}
-
-			uint8 operator()(int32 x, int32 y, int32 c) const
-			{
-				if constexpr (AD_MODE == EAddressMode::Wrap)
-				{
-					// (y & (m_dim.Y - 1)) * m_dim.Y + (x & (m_dim.X - 1))
-					return m_pImageBuf[((y % m_dim.Y) * m_dim.X + (x % m_dim.X)) * PIXEL_SIZE + c];
-				}
-				else if (AD_MODE == EAddressMode::ClampToEdge)
-				{
-					return m_pImageBuf[(mu::clamp(y, 0, m_dim.Y - 1) * m_dim.X +
-						mu::clamp(x, 0, m_dim.X - 1))
-						* PIXEL_SIZE + c];
-				}
-				else if (AD_MODE == EAddressMode::ClampToBlack)
-				{
-					return ((x < m_dim.X) && (y < m_dim.Y) && (x >= 0) && (y >= 0))
-						? m_pImageBuf[(y * m_dim.X + x) * PIXEL_SIZE + c]
-						: 0;
-				}
-				else
+				for (int x = 0; x < fullColumns; ++x)
 				{
-					return m_pImageBuf[(y * m_dim.X + x) * PIXEL_SIZE + c];
-				}
-			}
-		};
-
-
-		template<int32 PIXEL_SIZE, EAddressMode AD, typename PIXEL_TYPE>
-		void GenerateMipSharpenedRegion(
-			const SharpenKernelStorage8& kernel,
-			const ImageAccessor<PIXEL_SIZE, PIXEL_TYPE, AD>& sourceAccessor,
-			PIXEL_TYPE* Dest, FIntVector2 destSize,
-			FIntVector2 regionStart, FIntVector2 regionEnd)
-		{
-			// Only kernels of size 8 supported.
-			constexpr int32 kernelSize = SharpenKernelStorage8::SIZE;
-			static_assert(kernelSize == 8);
-
-			constexpr int32 kernelCenter = (kernelSize >> 1) - 1;
-			for (int32 y = regionStart.Y; y < regionEnd.Y; ++y)
-			{
-				for (int32 x = regionStart.X; x < regionEnd.X; ++x)
-				{
-					float stagingStorage[kernelSize * kernelSize * PIXEL_SIZE];
-
-					const FIntVector2 sourceCoord = FIntVector2(x << 1, y << 1);
-
-
-					// Convert Image kernel stamp to float. This conversion will be done multiple times.
-					// and should be done once before applying the kernel, or operate directly in unorm8/snorm8.
-					for (int32 ky = 0; ky < kernelSize; ++ky)
+					if constexpr (PIXEL_SIZE == 4)
 					{
-						for (int32 kx = 0; kx < kernelSize; ++kx)
-						{
-							for (int32 c = 0; c < PIXEL_SIZE; ++c)
-							{
-								stagingStorage[(ky * kernelSize + kx) + (kernelSize * kernelSize * c)] = static_cast<float>(
-									sourceAccessor(sourceCoord.X + kx - kernelCenter,
-										sourceCoord.Y + ky - kernelCenter,
-										c)) / 255.0f;
-							}
-						}
-					}
+						//const uint64 Row0Bits = *reinterpret_cast<const uint64*>(pSourceRow0);
+						//const uint64 Row1Bits = *reinterpret_cast<const uint64*>(pSourceRow1);
 
-					// Apply kernel. Although the kernel is separable, the cost would not differ much from the not separable one.
-					// That is because to make it separable a partial convolution of the source is needed. 
-					// Since here we are only applying to half the image and the kernel used is small (8x8) the benefits are not 
-					// clear.
+						// Use memcpy to avoid any possible but improbable UB. memcpy should be optimized away by the compiler.
+						uint64 Row0Bits; 
+						uint64 Row1Bits;
 
-					for (int32 c = 0; c < PIXEL_SIZE; ++c)
+						FMemory::Memcpy(&Row0Bits, pSourceRow0, sizeof(uint64));
+						FMemory::Memcpy(&Row1Bits, pSourceRow1, sizeof(uint64));
+						
+						const uint64 XorRow0Row1Bits = Row0Bits ^ Row1Bits;
+
+						// Average of 2 unsigned integers without overflow extended to work on multiple bytes.
+						constexpr uint64 ShiftMask = 0xFEFEFEFEFEFEFEFE;
+						const uint64 ErrorCorrection = XorRow0Row1Bits & 0x0101010101010101;
+						const uint64 AvgLowBits = (Row0Bits & Row1Bits) + ((XorRow0Row1Bits & ShiftMask) >> 1) + ErrorCorrection;
+						const uint64 AvgHighBits = AvgLowBits >> 32;
+						const uint32 Result = (AvgLowBits & AvgHighBits) + (((AvgLowBits ^ AvgHighBits) & ShiftMask) >> 1);
+					
+						//*reinterpret_cast<uint32*>(pDestRow) = Result;
+						FMemory::Memcpy(pDestRow, &Result, sizeof(uint32));
+					}
+					else
 					{
-						float sampleValue = 0.0f;
-						float sum = 0.0f;
-						for (int32 k = 0; k < kernelSize * kernelSize; ++k)
+						for (int32 C = 0; C < PIXEL_SIZE; ++C)
 						{
-							sampleValue += kernel.m_storage[k] * stagingStorage[k + (kernelSize * kernelSize * c)];
-							sum += kernel.m_storage[k];
+							int32 PixelSum = pSourceRow0[C] + pSourceRow0[PIXEL_SIZE + C] + pSourceRow1[C] + pSourceRow1[PIXEL_SIZE + C];
+							pDestRow[C] = (uint8)(PixelSum >> 2);
 						}
-
-						Dest[(y * destSize.X + x) * PIXEL_SIZE + c] = static_cast<uint8>(mu::clamp(sampleValue, 0.0f, 1.0f) * 255.0f);
 					}
+
+					pSourceRow0 += 2 * PIXEL_SIZE;
+					pSourceRow1 += 2 * PIXEL_SIZE;
+					pDestRow += PIXEL_SIZE;
 				}
-			}
-		}
 
-		template<size_t PIXEL_SIZE, EAddressMode AD_MODE>
-		inline void GenerateMipmapUint8Sharpen(int mips,
-			const uint8* pSource, uint8* Dest,
-			FIntVector2 sourceSize,
-			float sharpeningFactor)
-		{
-			SharpenKernelStorage8 kernel = MakeMipGaussianSharpenKernel(sharpeningFactor);
+				if (strayColumn)
+				{
+					if constexpr (PIXEL_SIZE == 4)
+					{	
+						//const uint32 Row0Bits = *reinterpret_cast<const uint32*>(pSourceRow0);
+						//const uint32 Row1Bits = *reinterpret_cast<const uint32*>(pSourceRow1);
 
-			size_t totalDestSize = 0;
+						uint32 Row0Bits; 
+						uint32 Row1Bits;
 
-			for (; mips >= 0; --mips)
-			{
-				ImageAccessor<PIXEL_SIZE, uint8, AD_MODE> sourceSampler(pSource, sourceSize);
-				ImageAccessor<PIXEL_SIZE, uint8, EAddressMode::None> sourceSamplerAddressModeNone(pSource, sourceSize);
-
-				FIntVector2 destSize = FIntVector2(FMath::DivideAndRoundUp(sourceSize.X, 2), FMath::DivideAndRoundUp(sourceSize.Y, 2));
-
-				// Core image, will never sample outside the image so we don't need to care about address mode.
-				// Only sample source pixels where the kernel can be fully applied.
-				// Border cases are treated separately.
-				GenerateMipSharpenedRegion(
-					kernel,
-					sourceSamplerAddressModeNone,
-					Dest, destSize,
-					FIntVector2(2, 2), FIntVector2(destSize.X - 2, destSize.Y - 2));
-
-				// Horizontal borders.
-				GenerateMipSharpenedRegion(
-					kernel,
-					sourceSampler,
-					Dest, destSize,
-					FIntVector2(0, 0), FIntVector2(destSize.X, FMath::Min(2, destSize.Y)));
-
-				GenerateMipSharpenedRegion(
-					kernel,
-					sourceSampler,
-					Dest, destSize,
-					FIntVector2(0, FMath::Max(destSize.Y - 2, 0)), FIntVector2(destSize.X, destSize.Y));
-
-				// Vertical borders
-				GenerateMipSharpenedRegion(
-					kernel,
-					sourceSampler,
-					Dest, destSize,
-					FIntVector2(0, FMath::Min(2, destSize.Y)), FIntVector2(FMath::Min(2, destSize.X), destSize.Y - 2));
-
-				GenerateMipSharpenedRegion(
-					kernel,
-					sourceSampler,
-					Dest, destSize,
-					FIntVector2(destSize.X - 2, FMath::Min(2, destSize.Y)), FIntVector2(destSize.X, destSize.Y - 2));
-
-				sourceSize = destSize;
-
-				pSource = Dest;
-				Dest = Dest + destSize.X * destSize.Y * PIXEL_SIZE;
-
-				totalDestSize += sourceSize.X * sourceSize.Y * PIXEL_SIZE;
-			}
-		}
+						FMemory::Memcpy(&Row0Bits, pSourceRow0, sizeof(uint32));
+						FMemory::Memcpy(&Row1Bits, pSourceRow1, sizeof(uint32));
 
-		template<int32 PIXEL_SIZE>
-		inline void GenerateMipmapUint8Unfiltered(
-			int mips,
-			const uint8* pSource, uint8* Dest,
-			FIntVector2 sourceSize)
-		{
-			for (; mips >= 0; --mips)
-			{
-				FIntVector2 destSize = FIntVector2(FMath::DivideAndRoundUp(sourceSize.X, 2), FMath::DivideAndRoundUp(sourceSize.Y, 2));
+						// Average of 2 unsigned integers without overflow extended to work on multiple bytes.
+						constexpr uint32 ShiftMask = 0xFEFEFEFE;
+						const uint32 Result = (Row0Bits & Row1Bits) + (((Row0Bits ^ Row1Bits) & ShiftMask) >> 1);
 
-				for (int32 y = 0; y < destSize.Y; ++y)
-				{
-					for (int32 x = 0; x < destSize.X; ++x)
+						//*reinterpret_cast<uint32*>(pDestRow) = Result;
+						FMemory::Memcpy(pDestRow, &Result, sizeof(uint32));
+					}
+					else
 					{
-						for (int32 c = 0; c < PIXEL_SIZE; ++c)
+						for (int32 C = 0; C < PIXEL_SIZE; ++C)
 						{
-							Dest[(y * destSize.X + x) * PIXEL_SIZE + c] =
-								pSource[((y << 1) * sourceSize.X + (x << 1)) * PIXEL_SIZE + c];
+							int32 PixelSum = pSourceRow0[C] + pSourceRow1[C];
+							pDestRow[C] = (uint8)(PixelSum >> 1);
 						}
 					}
 				}
+			};
 
-				sourceSize = destSize;
-				pSource = Dest;
-				Dest = Dest + destSize.X * destSize.Y * PIXEL_SIZE;
+			constexpr int PixelConcurrencyThreshold = 0xffff;
+			if (destSize[0] * destSize[1] < PixelConcurrencyThreshold)
+			{
+				for (int y = 0; y < fullRows; ++y)
+				{
+					ProcessRow(y);
+				}
 			}
-		}
-
-		template<int32 PIXEL_SIZE>
-		inline void GenerateMipmapsUint8SimpleAverage(
-			int mips,
-			const uint8* pSource, uint8* Dest,
-			FIntVector2 sourceSize)
-		{
-			
-			const uint8* pMipSource = pSource;
-			uint8* pMipDest = Dest;
-
-			FIntVector2 destSize = sourceSize;
-
-			for (int m = 0; m < mips; ++m)
+			else
 			{
-				check(destSize[0] > 1 || destSize[1] > 1);
-
-				sourceSize = destSize;
-
-				int fullColumns = destSize[0] / 2;
-				bool strayColumn = (destSize[0] % 2) != 0;
-				int fullRows = destSize[1] / 2;
-				bool strayRow = (destSize[1] % 2) != 0;
-
-				destSize[0] = FMath::DivideAndRoundUp(destSize[0], 2);
-				destSize[1] = FMath::DivideAndRoundUp(destSize[1], 2);
+				ParallelFor(fullRows, ProcessRow);
+			}
 
-				int sourceStride = sourceSize[0] * PIXEL_SIZE;
-				int destStride = destSize[0] * PIXEL_SIZE;
+			if (strayRow)
+			{
+				const uint8* pSourceRow0 = pMipSource + 2 * fullRows * sourceStride;
+				const uint8* pSourceRow1 = pSourceRow0 + sourceStride;
+				uint8* pDestRow = pMipDest + fullRows * destStride;
 
-				const auto ProcessRow = [
-					pMipDest, pMipSource, fullColumns, strayColumn, sourceStride, destStride
-				] (uint32 y)
+				for (int x = 0; x < fullColumns; ++x)
 				{
-					const uint8* pSourceRow0 = pMipSource + 2 * y * sourceStride;
-					const uint8* pSourceRow1 = pSourceRow0 + sourceStride;
-					uint8* pDestRow = pMipDest + y * destStride;
-
-					for (int x = 0; x < fullColumns; ++x)
-					{
-						if constexpr (PIXEL_SIZE == 4)
-						{
-							//const uint64 Row0Bits = *reinterpret_cast<const uint64*>(pSourceRow0);
-							//const uint64 Row1Bits = *reinterpret_cast<const uint64*>(pSourceRow1);
-
-							// Use memcpy to avoid any possible but improbable UB. memcpy should be optimized away by the compiler.
-							uint64 Row0Bits; 
-							uint64 Row1Bits;
-
-							FMemory::Memcpy(&Row0Bits, pSourceRow0, sizeof(uint64));
-							FMemory::Memcpy(&Row1Bits, pSourceRow1, sizeof(uint64));
-							
-							const uint64 XorRow0Row1Bits = Row0Bits ^ Row1Bits;
-
-							// Average of 2 unsigned integers without overflow extended to work on multiple bytes.
-							constexpr uint64 ShiftMask = 0xFEFEFEFEFEFEFEFE;
-							const uint64 ErrorCorrection = XorRow0Row1Bits & 0x0101010101010101;
-							const uint64 AvgLowBits = (Row0Bits & Row1Bits) + ((XorRow0Row1Bits & ShiftMask) >> 1) + ErrorCorrection;
-							const uint64 AvgHighBits = AvgLowBits >> 32;
-							const uint32 Result = (AvgLowBits & AvgHighBits) + (((AvgLowBits ^ AvgHighBits) & ShiftMask) >> 1);
-						
-							//*reinterpret_cast<uint32*>(pDestRow) = Result;
-							FMemory::Memcpy(pDestRow, &Result, sizeof(uint32));
-						}
-						else
-						{
-							for (int32 C = 0; C < PIXEL_SIZE; ++C)
-							{
-								int32 PixelSum = pSourceRow0[C] + pSourceRow0[PIXEL_SIZE + C] + pSourceRow1[C] + pSourceRow1[PIXEL_SIZE + C];
-								pDestRow[C] = (uint8)(PixelSum >> 2);
-							}
-						}
-
-						pSourceRow0 += 2 * PIXEL_SIZE;
-						pSourceRow1 += 2 * PIXEL_SIZE;
-						pDestRow += PIXEL_SIZE;
-					}
-
-					if (strayColumn)
+					if constexpr (PIXEL_SIZE == 4)
 					{
-						if constexpr (PIXEL_SIZE == 4)
-						{	
-							//const uint32 Row0Bits = *reinterpret_cast<const uint32*>(pSourceRow0);
-							//const uint32 Row1Bits = *reinterpret_cast<const uint32*>(pSourceRow1);
+						//const uint32 Col0Bits = *reinterpret_cast<const uint32*>(pSourceRow0);
+						//const uint32 Col1Bits = *reinterpret_cast<const uint32*>(pSourceRow0 + 4);
 
-							uint32 Row0Bits; 
-							uint32 Row1Bits;
+						uint32 Col0Bits; 
+						uint32 Col1Bits;
 
-							FMemory::Memcpy(&Row0Bits, pSourceRow0, sizeof(uint32));
-							FMemory::Memcpy(&Row1Bits, pSourceRow1, sizeof(uint32));
+						FMemory::Memcpy(&Col0Bits, pSourceRow0, sizeof(uint32));
+						FMemory::Memcpy(&Col1Bits, pSourceRow0 + 4, sizeof(uint32));
 
-							// Average of 2 unsigned integers without overflow extended to work on multiple bytes.
-							constexpr uint32 ShiftMask = 0xFEFEFEFE;
-							const uint32 Result = (Row0Bits & Row1Bits) + (((Row0Bits ^ Row1Bits) & ShiftMask) >> 1);
+						// Average of 2 unsigned integers without overflow extended to work on multiple bytes. 
+						// In this case we use the ceil variant to be consistent with the method used for 4 pixel average. 
+						constexpr uint32 ShiftMask = 0xFEFEFEFE;
+						const uint32 Result = (Col0Bits & Col1Bits) + (((Col0Bits ^ Col1Bits) & ShiftMask) >> 1);
 
-							//*reinterpret_cast<uint32*>(pDestRow) = Result;
-							FMemory::Memcpy(pDestRow, &Result, sizeof(uint32));
-						}
-						else
+						//*reinterpret_cast<uint32*>(pDestRow) = Result;
+						FMemory::Memcpy(pDestRow, &Result, sizeof(uint32));
+					}
+					else
+					{
+						for (int32 C = 0; C < PIXEL_SIZE; ++C)
 						{
-							for (int32 C = 0; C < PIXEL_SIZE; ++C)
-							{
-								int32 PixelSum = pSourceRow0[C] + pSourceRow1[C];
-								pDestRow[C] = (uint8)(PixelSum >> 1);
-							}
+							int32 p = pSourceRow0[C] + pSourceRow0[PIXEL_SIZE + C];
+							pDestRow[C] = (uint8)(p >> 1);
 						}
 					}
-				};
 
-				constexpr int PixelConcurrencyThreshold = 0xffff;
-				if (destSize[0] * destSize[1] < PixelConcurrencyThreshold)
+					pSourceRow0 += 2 * PIXEL_SIZE;
+					pDestRow += PIXEL_SIZE;
+				}
+
+				if (strayColumn)
 				{
-					for (int y = 0; y < fullRows; ++y)
+					if constexpr (PIXEL_SIZE == 4)
 					{
-						ProcessRow(y);
+						//*reinterpret_cast<uint32*>(pDestRow) = *reinterpret_cast<const uint32*>(pSourceRow0);
+						FMemory::Memcpy(pDestRow, pSourceRow0, 4);
+					}
+					else
+					{
+						for (int32 C = 0; C < PIXEL_SIZE; ++C)
+						{
+							pDestRow[C] = pSourceRow0[C];
+						}
 					}
 				}
-				else
-				{
-					ParallelFor(fullRows, ProcessRow);
-				}
+			}
 
-				if (strayRow)
-				{
-					const uint8* pSourceRow0 = pMipSource + 2 * fullRows * sourceStride;
-					const uint8* pSourceRow1 = pSourceRow0 + sourceStride;
-					uint8* pDestRow = pMipDest + fullRows * destStride;
+			// Reset the source pointer for the next mip, to use the dest that we have just
+			// generated.
+			pMipSource = pMipDest;
+			pMipDest += destSize[0] * destSize[1] * PIXEL_SIZE;
+		}
+	}
 
-					for (int x = 0; x < fullColumns; ++x)
+	// Generate next mip decompressed form block compressed image. 
+	template<int32 NumChannels, EMipmapFilterType Filter>
+	inline void GenerateNextMipBlockCompressed(
+		const uint8* Src, uint8* Dest, FIntVector2 SrcSize, EImageFormat SrcFormat, EImageFormat DestFormat)
+	{
+		MUTABLE_CPUPROFILER_SCOPE(GenerateNextMipBlockCompressed);
+
+		const FImageFormatData& DestFormatData = GetImageFormatData(DestFormat);
+		const FImageFormatData& SrcFormatData = GetImageFormatData(SrcFormat);
+
+		check(NumChannels == DestFormatData.Channels);
+		check(DestFormatData.PixelsPerBlockX == 1 && DestFormatData.PixelsPerBlockY == 1);
+
+		const int32 DestChannelCount = DestFormatData.Channels; 
+		const FIntVector2 PixelsPerBlock = FIntVector2(SrcFormatData.PixelsPerBlockX, SrcFormatData.PixelsPerBlockY);
+		const int32 BlockSizeInBytes = SrcFormatData.BytesPerBlock;
+
+		const FIntVector2 DestSize = FIntVector2(
+				FMath::DivideAndRoundUp(SrcSize.X, 2),
+				FMath::DivideAndRoundUp(SrcSize.Y, 2));
+
+		const FIntVector2 NumBlocks = FIntVector2(
+				FMath::DivideAndRoundUp(SrcSize.X, PixelsPerBlock.X),
+				FMath::DivideAndRoundUp(SrcSize.Y, PixelsPerBlock.Y));
+
+		constexpr int32 BatchSizeInBlocksX = 1 << 5;
+		constexpr int32 BatchSizeInBlocksY = 1 << 4;
+
+		FIntVector2 NumBatches = FIntVector2(
+				FMath::DivideAndRoundUp(NumBlocks.X, BatchSizeInBlocksX),
+				FMath::DivideAndRoundUp(NumBlocks.Y, BatchSizeInBlocksY));
+
+
+		// Limit the parallel job num based on actual num workers. Here we cannot rely on ParallelFor
+		// balancing the load as we need to allocate memory for every job. Make sure there is always 1 job.
+		// TODO: Consider balancing work on using a 2D grid.
+		const int32 MaxParallelJobs = FMath::Max(1, FMath::Min(int32(LowLevelTasks::FScheduler::Get().GetNumWorkers()), 8));
+		
+		constexpr int32 MinRowBatchesPerJob = 1;
+
+		const int32 NumRowBatchesPerJob = 
+			 FMath::Min(NumBatches.Y, FMath::Max(MinRowBatchesPerJob, FMath::DivideAndRoundUp(NumBatches.Y, MaxParallelJobs)));
+
+		const int32 NumParallelJobs = FMath::DivideAndRoundUp(NumBatches.Y, NumRowBatchesPerJob); 
+
+		// Use the tracking allocator policy on the image counter, this will not count for preventing memory peaks 
+		// but will show if it happens. This allocation should be small enough so it is not a problem to get over-budget
+		// by this amount. 
+		TArray<uint8, FDefaultMemoryTrackingAllocator<MemoryCounters::FImageMemoryCounter>> StagingMemory;
+
+		const miro::FImageSize StagingSize = miro::FImageSize(
+				uint16(BatchSizeInBlocksX*PixelsPerBlock.X), 
+				uint16(BatchSizeInBlocksY*PixelsPerBlock.Y));
+		
+		// Allocate extra memory so the mip computation can work on all possible pixels sizes. 
+		// Also add some extra padding so different threads do not share cache lines.
+		const int32 PerJobStagingBytes = StagingSize.X*StagingSize.Y*NumChannels + 8 + 64;
+		
+		StagingMemory.SetNum(PerJobStagingBytes*NumParallelJobs);
+		uint8 * const StagingMemoryData = StagingMemory.GetData();
+
+		miro::SubImageDecompression::FuncRefType DecompressionFunc = SelectDecompressionFunction(DestFormat, SrcFormat);
+
+		auto ProcessJob = 
+			[
+				NumParallelJobs, NumRowBatchesPerJob, 
+				StagingMemoryData, PerJobStagingBytes, 
+				NumBatches, NumBlocks, PixelsPerBlock, BlockSizeInBytes, DecompressionFunc,
+				Src, SrcSize, Dest, DestSize
+			](int32 JobId)
+		{
+			const int32 JobRowBegin = JobId*NumRowBatchesPerJob;
+			const int32 JobRowEnd   = FMath::Min(JobRowBegin + NumRowBatchesPerJob, NumBatches.Y);
+			uint8 * const JobStagingMemoryData = StagingMemoryData + JobId*PerJobStagingBytes;
+
+			for (int32 BatchY = JobRowBegin; BatchY < JobRowEnd; ++BatchY)
+			{
+				for (int32 BatchX = 0; BatchX < NumBatches.X; ++BatchX)
+				{
+					const FIntVector2 BatchBeginInBlocks = FIntVector2(BatchX*BatchSizeInBlocksX, BatchY*BatchSizeInBlocksY);
+					const FIntVector2 BatchEndInBlocks = FIntVector2(
+							FMath::Min(BatchBeginInBlocks.X + BatchSizeInBlocksX, NumBlocks.X),
+							FMath::Min(BatchBeginInBlocks.Y + BatchSizeInBlocksY, NumBlocks.Y));
+
+					const uint8* const SrcBatchData = Src + (BatchBeginInBlocks.Y * NumBlocks.X + BatchBeginInBlocks.X)*BlockSizeInBytes;
+
+					// Assume the decompressed size is always multiple of the block size. Trim unused bytes when copying to 
+					// the final destination.
+					const FIntVector2 BatchDecSizeInPixels = FIntVector2(
+							(BatchEndInBlocks.X - BatchBeginInBlocks.X)*PixelsPerBlock.X,
+							(BatchEndInBlocks.Y - BatchBeginInBlocks.Y)*PixelsPerBlock.Y);
+
+					const miro::FImageSize FromSize = miro::FImageSize(uint16(SrcSize.X), uint16(SrcSize.Y)); 
+					const miro::FImageSize SubSize  = miro::FImageSize(uint16(BatchDecSizeInPixels.X), uint16(BatchDecSizeInPixels.Y));
+					DecompressionFunc(FromSize, SubSize, SubSize, SrcBatchData, JobStagingMemoryData);
+
+					const FIntVector2 BatchOutBeginInPixels = FIntVector2(
+							(BatchBeginInBlocks.X*PixelsPerBlock.X) >> 1, 
+							(BatchBeginInBlocks.Y*PixelsPerBlock.Y) >> 1);
+
+					const FIntVector2 BatchOutEndInPixels = FIntVector2(
+							FMath::Min(BatchOutBeginInPixels.X + ((BatchSizeInBlocksX*PixelsPerBlock.X) >> 1), DestSize.X), 
+							FMath::Min(BatchOutBeginInPixels.Y + ((BatchSizeInBlocksY*PixelsPerBlock.Y) >> 1), DestSize.Y));
+		
+					// Generate partial next mip to dest.
+					// This works for all pixel sizes because we have preallocated more memory than needed.
+					for (int32 Y = BatchOutBeginInPixels.Y; Y < BatchOutEndInPixels.Y; ++Y)
 					{
-						if constexpr (PIXEL_SIZE == 4)
+						for (int32 X = BatchOutBeginInPixels.X; X < BatchOutEndInPixels.X; ++X)
 						{
-							//const uint32 Col0Bits = *reinterpret_cast<const uint32*>(pSourceRow0);
-							//const uint32 Col1Bits = *reinterpret_cast<const uint32*>(pSourceRow0 + 4);
+							uint8* const DestPixel = Dest + (Y*DestSize.X + X) * NumChannels;
 
-							uint32 Col0Bits; 
-							uint32 Col1Bits;
+							const FIntVector2 Row0Offset = FIntVector2(
+									(X - BatchOutBeginInPixels.X) << 1, (Y - BatchOutBeginInPixels.Y) << 1);
 
-							FMemory::Memcpy(&Col0Bits, pSourceRow0, sizeof(uint32));
-							FMemory::Memcpy(&Col1Bits, pSourceRow0 + 4, sizeof(uint32));
+							uint8 const * const SrcRow0 = JobStagingMemoryData + (Row0Offset.Y*BatchDecSizeInPixels.X + Row0Offset.X) * NumChannels;
 
-							// Average of 2 unsigned integers without overflow extended to work on multiple bytes. 
-							// In this case we use the ceil variant to be consistent with the method used for 4 pixel average. 
-							constexpr uint32 ShiftMask = 0xFEFEFEFE;
-							const uint32 Result = (Col0Bits & Col1Bits) + (((Col0Bits ^ Col1Bits) & ShiftMask) >> 1);
-
-							//*reinterpret_cast<uint32*>(pDestRow) = Result;
-							FMemory::Memcpy(pDestRow, &Result, sizeof(uint32));
-						}
-						else
-						{
-							for (int32 C = 0; C < PIXEL_SIZE; ++C)
+							if constexpr (Filter == EMipmapFilterType::MFT_SimpleAverage)
 							{
-								int32 p = pSourceRow0[C] + pSourceRow0[PIXEL_SIZE + C];
-								pDestRow[C] = (uint8)(p >> 1);
+								// Use memcpy to avoid any possible but improbable UB. memcpy should be optimized away by the compiler.
+								uint64 Row0Bits; 
+								FMemory::Memcpy(&Row0Bits, SrcRow0, sizeof(uint64));
+								
+								uint8 const * const SrcRow1 = JobStagingMemoryData + 
+										(FMath::Min(Row0Offset.Y + 1, BatchDecSizeInPixels.Y - 1)*BatchDecSizeInPixels.X + Row0Offset.X) * NumChannels;
+
+								uint64 Row1Bits;
+								FMemory::Memcpy(&Row1Bits, SrcRow1, sizeof(uint64));
+								
+								const uint64 OutOfBoundsMask = Row0Offset.X + 1 >= BatchDecSizeInPixels.X ? 0 : ~0;
+
+								constexpr uint64 ShiftMask = 0xFEFEFEFEFEFEFEFE;
+								
+								const uint64 XorRow0Row1Bits = Row0Bits ^ Row1Bits;
+								const uint64 ErrorCorrection = XorRow0Row1Bits & 0x0101010101010101;
+								
+								// Average of 2 unsigned integers without overflow extended to work on multiple bytes.
+								const uint64 AvgLowBits = (Row0Bits & Row1Bits) + ((XorRow0Row1Bits & ShiftMask) >> 1) + ErrorCorrection;
+								const uint64 AvgHighBits = OutOfBoundsMask ? AvgLowBits : (AvgLowBits >> 32);
+								const uint32 Result = (AvgLowBits & AvgHighBits) + (((AvgLowBits ^ AvgHighBits) & ShiftMask) >> 1);
+								
+								FMemory::Memcpy(DestPixel, &Result, NumChannels);
 							}
-						}
-
-						pSourceRow0 += 2 * PIXEL_SIZE;
-						pDestRow += PIXEL_SIZE;
-					}
-
-					if (strayColumn)
-					{
-						if constexpr (PIXEL_SIZE == 4)
-						{
-							//*reinterpret_cast<uint32*>(pDestRow) = *reinterpret_cast<const uint32*>(pSourceRow0);
-							FMemory::Memcpy(pDestRow, pSourceRow0, 4);
-						}
-						else
-						{
-							for (int32 C = 0; C < PIXEL_SIZE; ++C)
+							else // constexpr Filter == EMipmapFilterType::MFT_Unfiltered
 							{
-								pDestRow[C] = pSourceRow0[C];
+								FMemory::Memcpy(DestPixel, SrcRow0, NumChannels);	
 							}
+							static_assert(
+								Filter == EMipmapFilterType::MFT_SimpleAverage || 
+								Filter == EMipmapFilterType::MFT_Unfiltered);
 						}
 					}
 				}
-
-				// Reset the source pointer for the next mip, to use the dest that we have just
-				// generated.
-				pMipSource = pMipDest;
-				pMipDest += destSize[0] * destSize[1] * PIXEL_SIZE;
 			}
-		}
+		};
 
-	} // namespace OpImageMipmap_Detail
+		if (NumParallelJobs == 1)
+		{
+			ProcessJob(0);
+		}
+		else if (NumParallelJobs > 1)
+		{	
+			ParallelFor(NumParallelJobs, ProcessJob);
+		}
+	}
+} // namespace OpImageMipmap_Detail
 
 
 	/** Generate the mipmaps for byte-based images of whatever number of channels.
@@ -533,41 +422,66 @@ namespace mu
 			GenerateMipmapsUint8SimpleAverage<PIXEL_SIZE>(mips, pSource, Dest, sourceSize);
 			break;
 		}
-		case EMipmapFilterType::MFT_Sharpen:
+		case EMipmapFilterType::MFT_Unfiltered:
 		{
-			switch (settings.m_addressMode)
-			{
-			case EAddressMode::ClampToEdge:
-			{
-				GenerateMipmapUint8Sharpen<PIXEL_SIZE, EAddressMode::ClampToEdge>(
-					mips, pSource, Dest, sourceSize, settings.m_sharpenFactor);
-				break;
-			}
-			case EAddressMode::Wrap:
-			{
-				GenerateMipmapUint8Sharpen<PIXEL_SIZE, EAddressMode::Wrap>(
-					mips, pSource, Dest, sourceSize, settings.m_sharpenFactor);
-				break;
-			}
-			case EAddressMode::ClampToBlack:
-			{
-				GenerateMipmapUint8Sharpen<PIXEL_SIZE, EAddressMode::ClampToBlack>(
-					mips, pSource, Dest, sourceSize, settings.m_sharpenFactor);
-				break;
-			}
-			default:
+			GenerateMipmapUint8Unfiltered<PIXEL_SIZE>(mips, pSource, Dest, sourceSize);
+			break;
+		}
+		default:
+		{
+			check(false);
+			break;
+		}
+		}
+	}
+
+	/** Generate the mipmaps for Block Comporessed images of whatever number of channels.
+	 *  The result is a non compressed image of the next mip with its tail.
+	* \param mips number of additional levels to build from the source.
+	*/
+	template<int32 PixelSize>
+	inline void GenerateMipmapsBlockCompressed(int32 Mips,
+		const uint8* SourceData, uint8* DestData,
+		FIntVector2 SourceSize, EImageFormat SrcFormat, EImageFormat DestFormat,
+		const FMipmapGenerationSettings& Settings)
+	{
+		using namespace OpImageMipmap_Detail;
+
+		switch (Settings.m_filterType)
+		{
+		case EMipmapFilterType::MFT_SimpleAverage:
+		{
+			GenerateNextMipBlockCompressed<PixelSize, EMipmapFilterType::MFT_SimpleAverage>(
+				SourceData, DestData, SourceSize, SrcFormat, DestFormat);
+
+			const FIntVector2 CurrentMipSize = FIntVector2(
+				FMath::DivideAndRoundUp(SourceSize[0], 2),
+				FMath::DivideAndRoundUp(SourceSize[1], 2));
+
+			if (CurrentMipSize.X > 1 || CurrentMipSize.Y > 1)
 			{
-				GenerateMipmapUint8Sharpen<PIXEL_SIZE, EAddressMode::None>(
-					mips, pSource, Dest, sourceSize, settings.m_sharpenFactor);
-				break;
-			}
+				uint8* CurrentMipData = DestData;
+				uint8* NextMipData = CurrentMipData + CurrentMipSize.X*CurrentMipSize.Y*PixelSize;
+				GenerateMipmapsUint8SimpleAverage<PixelSize>(Mips - 1, CurrentMipData, NextMipData, CurrentMipSize);
 			}
 
 			break;
 		}
 		case EMipmapFilterType::MFT_Unfiltered:
 		{
-			GenerateMipmapUint8Unfiltered<PIXEL_SIZE>(mips, pSource, Dest, sourceSize);
+			GenerateNextMipBlockCompressed<PixelSize, EMipmapFilterType::MFT_Unfiltered>(
+					SourceData, DestData, SourceSize, SrcFormat, DestFormat);
+
+			const FIntVector2 CurrentMipSize = FIntVector2(
+					FMath::DivideAndRoundUp(SourceSize[0], 2), 
+					FMath::DivideAndRoundUp(SourceSize[1], 2));
+
+			if (CurrentMipSize.X > 1 || CurrentMipSize.Y > 1)
+			{
+				uint8* CurrentMipData = DestData;
+				uint8* NextMipData = CurrentMipData + CurrentMipSize.Y*CurrentMipSize.X*PixelSize;
+				GenerateMipmapUint8Unfiltered<PixelSize>(Mips + 1, CurrentMipData, NextMipData, CurrentMipSize);
+			}
 			break;
 		}
 		default:
@@ -579,7 +493,6 @@ namespace mu
 	}
 
 
-
     //---------------------------------------------------------------------------------------------
     void FImageOperator::ImageMipmap_PrepareScratch(Image* Dest, const Image* Base, int32 LevelCount, FScratchImageMipmap& Scratch )
     {
@@ -596,12 +509,15 @@ namespace mu
 			// Is it a block format?
 			if (mu::GetImageFormatData(BaseFormat).PixelsPerBlockX > 1)
 			{
-				// Uncompress the last mip that we already have
-				FIntVector2 UncompressedSize = Base->CalculateMipSize(StartLevel);
-				Scratch.Uncompressed = CreateImage(
-					(uint16)UncompressedSize[0], (uint16)UncompressedSize[1],
-					1,
-					EImageFormat::IF_RGBA_UBYTE, EInitializationType::NotInitialized);
+				if (!bEnableCompressedMipGenerationMemoryOptimizations)
+				{
+					// Uncompress the last mip that we already have
+					FIntVector2 UncompressedSize = Base->CalculateMipSize(StartLevel);
+					Scratch.Uncompressed = CreateImage(
+						(uint16)UncompressedSize[0], (uint16)UncompressedSize[1],
+						1,
+						EImageFormat::IF_RGBA_UBYTE, EInitializationType::NotInitialized);
+				}
 
 				FIntVector2 UncompressedMipsSize = Base->CalculateMipSize(StartLevel + 1);
 				// Generate the mipmaps from there on
@@ -710,13 +626,54 @@ namespace mu
 
 		uint8* pDestBuf = !bGenerateOnlyTail ? Dest->GetMipData(StartLevel + 1) : Dest->GetMipData(0);
 
-		FIntVector2 sourceSize = Base->CalculateMipSize(StartLevel);
+		FIntVector2 SourceSize = Base->CalculateMipSize(StartLevel);
 
 		EImageFormat BaseFormat = Base->GetFormat();
-		if (mu::IsCompressedFormat(BaseFormat))
+		const bool bIsBlockCompressedFormat = mu::IsBlockCompressedFormat(BaseFormat);
+		const bool bIsCompressedFormat = mu::IsCompressedFormat(BaseFormat);
+		
+		check(!bIsBlockCompressedFormat || bIsCompressedFormat);
+
+		if (bIsBlockCompressedFormat && bEnableCompressedMipGenerationMemoryOptimizations)
 		{
-			// Bad case.
+			const EImageFormat DestFormat = Scratch.UncompressedMips->GetFormat();
+			uint8* DestData = Scratch.UncompressedMips->GetData();
+			switch (DestFormat)
+			{
+			case EImageFormat::IF_L_UBYTE:
+			{
+				constexpr int32 PixelSize = 1;
+				GenerateMipmapsBlockCompressed<PixelSize>(
+						LevelCount - StartLevel - 1, pSourceBuf, DestData, SourceSize, BaseFormat, DestFormat, Settings);
+				break;
+			}
+			case EImageFormat::IF_RGB_UBYTE:
+			{
+				constexpr int32 PixelSize = 3;
+				GenerateMipmapsBlockCompressed<PixelSize>(
+						LevelCount - StartLevel - 1, pSourceBuf, DestData, SourceSize, BaseFormat, DestFormat, Settings);
+				break;
+			}
+			case EImageFormat::IF_RGBA_UBYTE:
+			{
+				constexpr int32 PixelSize = 4;
+				GenerateMipmapsBlockCompressed<PixelSize>(
+						LevelCount - StartLevel - 1, pSourceBuf, DestData, SourceSize, BaseFormat, DestFormat, Settings);
+				break;
+			}
+			default: check(false);
+			}
 
+			bool bSuccess = false;
+			constexpr int32 OnlyLod = -1;
+			ImagePixelFormat(bSuccess, CompressionQuality, Scratch.CompressedMips.get(), Scratch.UncompressedMips.get(), OnlyLod); 
+			check(bSuccess);
+			
+			FMemory::Memcpy(pDestBuf, Scratch.CompressedMips->GetData(), Scratch.CompressedMips->GetDataSize());
+		}
+		else if (bIsCompressedFormat)
+		{
+			// Bad case.
 			// Uncompress the last mip that we already have
 			bool bSuccess = false;
 			ImagePixelFormat(bSuccess, CompressionQuality, Scratch.Uncompressed.get(), Base, StartLevel);
@@ -760,16 +717,16 @@ namespace mu
 			switch (Base->GetFormat())
 			{
 			case EImageFormat::IF_L_UBYTE:
-				GenerateMipmapsUint8<1>(LevelCount - StartLevel - 1, pSourceBuf, pDestBuf, sourceSize, Settings);
+				GenerateMipmapsUint8<1>(LevelCount - StartLevel - 1, pSourceBuf, pDestBuf, SourceSize, Settings);
 				break;
 
 			case EImageFormat::IF_RGB_UBYTE:
-				GenerateMipmapsUint8<3>(LevelCount - StartLevel - 1, pSourceBuf, pDestBuf, sourceSize, Settings);
+				GenerateMipmapsUint8<3>(LevelCount - StartLevel - 1, pSourceBuf, pDestBuf, SourceSize, Settings);
 				break;
 
 			case EImageFormat::IF_BGRA_UBYTE:
 			case EImageFormat::IF_RGBA_UBYTE:
-				GenerateMipmapsUint8<4>(LevelCount - StartLevel - 1, pSourceBuf, pDestBuf, sourceSize, Settings);
+				GenerateMipmapsUint8<4>(LevelCount - StartLevel - 1, pSourceBuf, pDestBuf, SourceSize, Settings);
 				break;
 
 			default:
@@ -778,7 +735,6 @@ namespace mu
 		}
 	}
 
-
 	void FImageOperator::ImageMipmap(int32 CompressionQuality, Image* Dest, const Image* Base,
 		int32 LevelCount,
 		const FMipmapGenerationSettings& Settings, bool bGenerateOnlyTail)
diff --git a/Engine/Plugins/Experimental/Mutable/Source/MutableRuntime/Private/MuR/OpImageResize.cpp b/Engine/Plugins/Experimental/Mutable/Source/MutableRuntime/Private/MuR/OpImageResize.cpp
--- a/Engine/Plugins/Experimental/Mutable/Source/MutableRuntime/Private/MuR/OpImageResize.cpp
+++ b/Engine/Plugins/Experimental/Mutable/Source/MutableRuntime/Private/MuR/OpImageResize.cpp
@@ -4,6 +4,17 @@
 #include "MuR/MutableTrace.h"
 #include "Async/ParallelFor.h"
 
+namespace
+{
+
+bool bEnableImageResizeMemoryOptimizations = false;
+static FAutoConsoleVariableRef CVarEnableImageResizeMemoryOptimizations (
+	TEXT("mutable.EnableImageResizeMemoryOptimizations"),
+	bEnableImageResizeMemoryOptimizations,
+	TEXT("If set to true, enables image resize optimizations form small resizes and block compressed images."),
+	ECVF_Default);
+}
+
 namespace mu
 {
 	/** */
@@ -700,94 +711,458 @@ namespace mu
 
 	}
 
-
-	void FImageOperator::ImageResizeLinear(Image* Dest, int32 ImageCompressionQuality, const Image* InBase)
+namespace
+{
+	template<uint32 NumChannels>
+	void ImageResizeLinearDownByLessThanTwoSubImage(FVector2f ScalingFactor,
+			FImageSize DestSize, FImageSize DestSubSize, FImageSize SrcSize, FImageSize SrcSubSize,
+		  	uint8* RESTRICT SubDest, const uint8* RESTRICT SubSrc)
 	{
-		MUTABLE_CPUPROFILER_SCOPE(ImageResizeLinear);
-
-		check(!(InBase->m_flags & Image::IF_CANNOT_BE_SCALED));
+		struct alignas(uint64) FPixelData
+		{
+			uint16 Data[4] = {0};
+		};
 
-		ImagePtrConst Base = InBase;
-		ImagePtr BaseReformat;
+		static_assert(NumChannels <= 4);
+		
+		//TODO: Consider using fixed point arithmetic or vector registers as optimization.
+		
+		FVector2f InvScalingFactor = FVector2f(1.0f) / ScalingFactor;
+		check(InvScalingFactor.X <= 2.0f && InvScalingFactor.Y <= 2.0f && InvScalingFactor.X >= 1.0f && InvScalingFactor.Y >= 1.0f);
 
-		// Shouldn't happen! But if it does...
-		EImageFormat SourceFormat = Base->GetFormat();
-		EImageFormat UncompressedFormat = GetUncompressedFormat(SourceFormat);
-		if (SourceFormat != UncompressedFormat)
+		for (uint16 Y = 0; Y < DestSubSize.Y; ++Y)
 		{
-			BaseReformat = ImagePixelFormat(ImageCompressionQuality, InBase, UncompressedFormat);
-			Base = BaseReformat;
-		}
+			for (uint16 X = 0; X < DestSubSize.X; ++X)
+			{
+				auto LoadPixel = [](const uint8* Ptr) -> FPixelData
+				{
+					FPixelData Result;
+					if constexpr (NumChannels == 4)
+					{
+						uint32 PackedData;
+						FMemory::Memcpy(&PackedData, Ptr, sizeof(uint32));
 
-		FImageSize BaseSize = FImageSize(Base->GetSizeX(), Base->GetSizeY());
-		FImageSize DestSize = FImageSize(Dest->GetSizeX(), Dest->GetSizeY());
-		if (!DestSize[0] || !DestSize[1] || !BaseSize[0] || !BaseSize[1])
-		{
-			return;
+						Result.Data[0] = static_cast<uint16>((PackedData >> (8 * 0)) & 0xFF);
+						Result.Data[1] = static_cast<uint16>((PackedData >> (8 * 1)) & 0xFF);
+						Result.Data[2] = static_cast<uint16>((PackedData >> (8 * 2)) & 0xFF);
+						Result.Data[3] = static_cast<uint16>((PackedData >> (8 * 3)) & 0xFF);
+					}
+					else
+					{
+						for (int32 C = 0; C < NumChannels; ++C)
+						{
+							Result.Data[C] = static_cast<uint16>(Ptr[C]);
+						}
+					}
+
+					return Result;
+				};
+				
+				const FVector2f CoordsF = (FVector2f(X, Y) + 0.5f)*InvScalingFactor;
+				
+				using FUint16Vector2 = UE::Math::TIntVector2<uint16>;
+				const FUint16Vector2 FracINorm = FUint16Vector2(FMath::Frac(CoordsF.X) * 255.0f, FMath::Frac(CoordsF.Y) * 255.0f);
+				
+				const FIntVector2 Coords = FIntVector2(FMath::FloorToInt(CoordsF.X), FMath::FloorToInt(CoordsF.Y));
+				const FIntVector2 CoordsPlusOne = FIntVector2(
+						FMath::Min<int32>(Coords.X + 1, SrcSize.X - 1),
+						FMath::Min<int32>(Coords.Y + 1, SrcSize.Y - 1));
+
+				FPixelData PixelData00 = LoadPixel(SubSrc + (Coords.Y*SrcSize.X + Coords.X) * NumChannels);
+				FPixelData PixelData10 = LoadPixel(SubSrc + (Coords.Y*SrcSize.X + CoordsPlusOne.X) * NumChannels);
+				FPixelData PixelData01 = LoadPixel(SubSrc + (CoordsPlusOne.Y*SrcSize.X + Coords.X) * NumChannels);
+				FPixelData PixelData11 = LoadPixel(SubSrc + (CoordsPlusOne.Y*SrcSize.X + CoordsPlusOne.X) * NumChannels);
+				
+				FPixelData Result;
+				for (int32 C = 0; C < NumChannels; ++C)
+				{
+					const uint16 LerpY00 = ((PixelData10.Data[C] * FracINorm.X) + PixelData00.Data[C] * (255 - FracINorm.X)) / 255;
+					const uint16 LerpY01 = ((PixelData11.Data[C] * FracINorm.X) + PixelData01.Data[C] * (255 - FracINorm.X)) / 255;
+					Result.Data[C] = ((LerpY01 * FracINorm.Y) + LerpY00 * (255 - FracINorm.Y)) / 255;
+				}
+
+				uint8* Dest = SubDest + (Y*DestSize.X + X)*NumChannels;  
+
+				if constexpr (NumChannels == 4)
+				{
+					const uint32 PackedPixel = 
+						(static_cast<uint32>((Result.Data[0]) & 0xFF) << (8 * 0)) | 
+						(static_cast<uint32>((Result.Data[1]) & 0xFF) << (8 * 1)) | 
+						(static_cast<uint32>((Result.Data[2]) & 0xFF) << (8 * 2)) | 
+						(static_cast<uint32>((Result.Data[3]) & 0xFF) << (8 * 3));
+
+					FMemory::Memcpy(Dest, &PackedPixel, sizeof(uint32));
+				}
+				else
+				{
+					for (uint32 C = 0; C < NumChannels; ++C)
+					{
+						Dest[C] = static_cast<uint8>(Result.Data[C]);
+					}
+				}
+			}
 		}
+	}
+
+	void ImageResizeLinearDownByLessThanTwo(Image* Dest, const Image* InBase)
+	{
+
+		MUTABLE_CPUPROFILER_SCOPE(ImageResizeLinearDownByLessThanTwo);
+
+		const FImageSize DestSize = FImageSize(Dest->GetSizeX(), Dest->GetSizeY());
+		const FImageSize SrcSize = FImageSize(InBase->GetSizeX(), InBase->GetSizeY());
+		
+		check(SrcSize.X >= DestSize.X && SrcSize.Y >= DestSize.Y && SrcSize.X <= DestSize.X*2 && SrcSize.Y <= DestSize.Y*2);
+
+		const EImageFormat SrcFormat = InBase->GetFormat();
+		const EImageFormat DestFormat = Dest->GetFormat();
+	
+		uint8* DestData = Dest->GetData();
+		const uint8* SrcData = InBase->GetData();
+		
+		const FImageFormatData& DestFormatData = GetImageFormatData(Dest->GetFormat());
+		const FImageFormatData& SrcFormatData = GetImageFormatData(InBase->GetFormat());
+
+		const FVector2f ScalingFactor = FVector2f(DestSize.X, DestSize.Y) / FVector2f(SrcSize.X, SrcSize.Y);
 
-		// First resize X
-		ImagePtr Temp;
-		if (DestSize[0] > BaseSize[0])
+		if (SrcFormatData.PixelsPerBlockX == 1 && SrcFormatData.PixelsPerBlockY == 1)
 		{
-			Temp = CreateImage(DestSize[0], BaseSize[1], 1, Base->GetFormat(), EInitializationType::NotInitialized);
-			ImageMagnifyX(Temp.get(), Base.get());
+			MUTABLE_CPUPROFILER_SCOPE(ResizeUnCompressed);
+
+			// Uncompressed format, resize in one go directly to Dest.
+			switch(SrcFormat)
+			{
+
+			case EImageFormat::IF_L_UBYTE:
+			{
+				constexpr int32 NumChannels = 1;
+				ImageResizeLinearDownByLessThanTwoSubImage<NumChannels>(ScalingFactor,
+						DestSize, DestSize, SrcSize, SrcSize, DestData, SrcData);
+				break;
+			}
+
+			case EImageFormat::IF_RGB_UBYTE:
+			{
+				constexpr int32 NumChannels = 3;
+				ImageResizeLinearDownByLessThanTwoSubImage<NumChannels>(ScalingFactor,
+						DestSize, DestSize, SrcSize, SrcSize, DestData, SrcData);
+				break;
+			}
+
+			case EImageFormat::IF_RGBA_UBYTE:
+			case EImageFormat::IF_BGRA_UBYTE:
+			{
+				constexpr int32 NumChannels = 4;
+				ImageResizeLinearDownByLessThanTwoSubImage<NumChannels>(ScalingFactor,
+						DestSize, DestSize, SrcSize, SrcSize, DestData, SrcData);
+				break;
+			}
+			}
 		}
-		else if (DestSize[0] < BaseSize[0])
+		else if (SrcFormatData.PixelsPerBlockX > 0 && SrcFormatData.PixelsPerBlockY > 0)
 		{
-			Temp = CreateImage(DestSize[0], BaseSize[1], 1, Base->GetFormat(), EInitializationType::NotInitialized);
-			ImageMinifyX(Temp.get(), Base.get());
+			MUTABLE_CPUPROFILER_SCOPE(ResizeBlockCompressed);
+
+			check(DestFormatData.PixelsPerBlockX == 1 && DestFormatData.PixelsPerBlockY == 1);
+
+			const int32 DestNumChannels = DestFormatData.Channels; 
+			const FIntVector2 PixelsPerBlock = FIntVector2(SrcFormatData.PixelsPerBlockX, SrcFormatData.PixelsPerBlockY);
+			const int32 BlockSizeInBytes = SrcFormatData.BytesPerBlock;
+
+			const FIntVector2 NumBlocks = FIntVector2(
+					FMath::DivideAndRoundUp(SrcSize.X, uint16(PixelsPerBlock.X)),
+					FMath::DivideAndRoundUp(SrcSize.Y, uint16(PixelsPerBlock.Y)));
+
+
+			constexpr int32 BatchSizeInBlocksX = 1 << 5;
+			constexpr int32 BatchSizeInBlocksY = 1 << 4;
+
+			const FIntVector2 NumBatches = FIntVector2(
+					FMath::DivideAndRoundUp(NumBlocks.X, BatchSizeInBlocksX),
+					FMath::DivideAndRoundUp(NumBlocks.Y, BatchSizeInBlocksY));
+
+			// Limit the parallel job num based on actual num workers. Here we cannot rely on ParallelFor
+			// balancing the load as we need to allocate memory for every job. Make sure there is always 1 job.
+			// TODO: Consider balancing work on using a 2D division.
+			const int32 MaxParallelJobs = FMath::Max(1, FMath::Min(int32(LowLevelTasks::FScheduler::Get().GetNumWorkers()), 8));
+			
+			constexpr int32 MinRowBatchesPerJob = 2;
+
+			const int32 NumRowBatchesPerJob = 
+				 FMath::Min(NumBatches.Y, FMath::Max(MinRowBatchesPerJob, FMath::DivideAndRoundUp(NumBatches.Y, MaxParallelJobs)));
+
+			const int32 NumParallelJobs = FMath::DivideAndRoundUp(NumBatches.Y, NumRowBatchesPerJob); 
+
+			// Allocate memory for 1 extra block, this is needed so that the data to compute the batch edges
+			// if available. This implies extra work needs to be done, for simplicity re-decompress those blocks, if
+			// that becomes a problem more convoluted memory cacheing schemes could be implemented.
+			const miro::FImageSize StagingSize = miro::FImageSize(
+					uint16((BatchSizeInBlocksX + 1)*PixelsPerBlock.X), 
+					uint16((BatchSizeInBlocksY + 1)*PixelsPerBlock.Y));
+		
+			// Add some extra padding so diferent threads do not share cache lines.
+			const int32 PerJobStagingBytes = StagingSize.X*StagingSize.Y*DestNumChannels + 64;
+
+			TArray<uint8, FDefaultMemoryTrackingAllocator<MemoryCounters::FImageMemoryCounter>> StagingMemory;
+			StagingMemory.SetNum(PerJobStagingBytes * NumParallelJobs);
+			uint8 * const StagingMemoryData = StagingMemory.GetData();
+
+			miro::SubImageDecompression::FuncRefType DecompressionFunc = SelectDecompressionFunction(DestFormat, SrcFormat);
+
+			auto ProcessJob = 
+				[
+					NumParallelJobs, NumRowBatchesPerJob, 
+					StagingMemoryData, PerJobStagingBytes, 
+					NumBatches, NumBlocks, PixelsPerBlock, BlockSizeInBytes, DestNumChannels, DecompressionFunc,
+					SrcData, SrcSize, DestData, DestSize, ScalingFactor
+				](int32 JobId)
+			{
+				const int32 JobRowBegin = JobId*NumRowBatchesPerJob;
+				const int32 JobRowEnd   = FMath::Min(JobRowBegin + NumRowBatchesPerJob, NumBatches.Y);
+				uint8 * const JobStagingMemoryData = StagingMemoryData + JobId*PerJobStagingBytes;
+
+				for (int32 BatchY = JobRowBegin; BatchY < JobRowEnd; ++BatchY)
+				{
+					for (int32 BatchX = 0; BatchX < NumBatches.X; ++BatchX)
+					{
+						const FIntVector2 BatchBeginInBlocks = FIntVector2(
+								BatchX*BatchSizeInBlocksX, BatchY*BatchSizeInBlocksY);
+
+						const FIntVector2 BatchDecEndInBlocks = FIntVector2(
+								FMath::Min(BatchBeginInBlocks.X + BatchSizeInBlocksX + 1, NumBlocks.X),
+								FMath::Min(BatchBeginInBlocks.Y + BatchSizeInBlocksY + 1, NumBlocks.Y));
+
+						const FIntVector2 BatchEndInBlocks = FIntVector2(
+								FMath::Min(BatchBeginInBlocks.X + BatchSizeInBlocksX, NumBlocks.X),
+								FMath::Min(BatchBeginInBlocks.Y + BatchSizeInBlocksY, NumBlocks.Y));
+
+						const uint8* const SrcBatchData = 
+								SrcData + (BatchBeginInBlocks.Y*NumBlocks.X + BatchBeginInBlocks.X)*BlockSizeInBytes;
+
+						// Assume the decompressed size is always multiple of the block size. Trim unused bytes when copying to 
+						// the final destination.
+						const FIntVector2 BatchDecSizeInPixels = FIntVector2(
+								(BatchDecEndInBlocks.X - BatchBeginInBlocks.X)*PixelsPerBlock.X,
+								(BatchDecEndInBlocks.Y - BatchBeginInBlocks.Y)*PixelsPerBlock.Y);
+
+						const miro::FImageSize FromSize = miro::FImageSize(SrcSize.X, SrcSize.Y); 
+						const miro::FImageSize DecSize  = miro::FImageSize(uint16(BatchDecSizeInPixels.X), uint16(BatchDecSizeInPixels.Y));
+
+						DecompressionFunc(FromSize, DecSize, DecSize, SrcBatchData, JobStagingMemoryData);
+
+						const FIntVector2 BatchOutBeginInPixels = FIntVector2( 
+								FMath::DivideAndRoundUp(BatchBeginInBlocks.X*PixelsPerBlock.X*DestSize.X, int32(SrcSize.X)),
+								FMath::DivideAndRoundUp(BatchBeginInBlocks.Y*PixelsPerBlock.Y*DestSize.Y, int32(SrcSize.Y)));
+
+						const FIntVector2 BatchOutEndInPixels = FIntVector2(
+								FMath::Min(
+									FMath::DivideAndRoundUp(BatchEndInBlocks.X*PixelsPerBlock.X*DestSize.X, int32(SrcSize.X)), 
+									int32(DestSize.X)),
+								FMath::Min(
+									FMath::DivideAndRoundUp(BatchEndInBlocks.Y*PixelsPerBlock.Y*DestSize.Y, int32(SrcSize.Y)), 
+									int32(DestSize.Y))); 
+					
+						const FIntVector2 BatchSizeInPixels = (BatchEndInBlocks - BatchBeginInBlocks)*PixelsPerBlock;
+						
+						const FImageSize SrcSubSize  = FImageSize(BatchSizeInPixels.X, BatchSizeInPixels.Y);
+						const FImageSize DestSubSize = FImageSize(
+								BatchOutEndInPixels.X - BatchOutBeginInPixels.X, 
+								BatchOutEndInPixels.Y - BatchOutBeginInPixels.Y);
+		
+						const FIntVector2& DestSubOffset = FIntVector2(BatchOutBeginInPixels.X, BatchOutBeginInPixels.Y);
+					
+						const uint8* SrcSubData = JobStagingMemoryData;
+						switch(DestNumChannels)
+						{
+						case 1:
+						{
+							constexpr int32 NumChannels = 1;
+						
+							uint8* DestSubData = DestData + (DestSubOffset.Y*DestSize.X + DestSubOffset.X)*NumChannels;
+
+							ImageResizeLinearDownByLessThanTwoSubImage<NumChannels>(ScalingFactor, 
+									DestSize, DestSubSize, DecSize, SrcSubSize, DestSubData, SrcSubData);
+							break;
+						}
+
+						case 3:
+						{
+							constexpr int32 NumChannels = 3;
+
+							uint8* DestSubData = DestData + (DestSubOffset.Y*DestSize.X + DestSubOffset.X)*NumChannels;
+							
+							ImageResizeLinearDownByLessThanTwoSubImage<NumChannels>(ScalingFactor, 
+									DestSize, DestSubSize, DecSize, SrcSubSize, DestSubData, SrcSubData);
+							break;
+						}
+
+						case 4:
+						{
+							constexpr int32 NumChannels = 4;
+
+							uint8* DestSubData = DestData + (DestSubOffset.Y*DestSize.X + DestSubOffset.X)*NumChannels;
+							
+							ImageResizeLinearDownByLessThanTwoSubImage<NumChannels>(ScalingFactor,
+									DestSize, DestSubSize, DecSize, SrcSubSize, DestSubData, SrcSubData);
+							break;
+						}
+						}
+					}
+				}
+			};
+
+			if (NumParallelJobs == 1)
+			{
+				ProcessJob(0);
+			}
+			else if (NumParallelJobs > 1)
+			{	
+				ParallelFor(NumParallelJobs, ProcessJob);
+			}
 		}
 		else
 		{
-			Temp = CloneImage(Base.get());
+			checkf(false, TEXT("Format not supported"));
 		}
+	}
+}
+
+	void FImageOperator::ImageResizeLinear(Image* Dest, int32 ImageCompressionQuality, const Image* InBase)
+	{
+		MUTABLE_CPUPROFILER_SCOPE(ImageResizeLinear);
+
+		check(!(InBase->m_flags & Image::IF_CANNOT_BE_SCALED));
+
+		const FImageSize BaseSize = FImageSize(InBase->GetSizeX(), InBase->GetSizeY());
+		const FImageSize DestSize = FImageSize(Dest->GetSizeX(), Dest->GetSizeY());
+
+		// This case should be treated in the caller and skip the operation altogether.
+		check(BaseSize != DestSize);
 
-		// Now resize Y
-		ImagePtr Temp2;
-		if (DestSize[1] > BaseSize[1])
+		if (!DestSize[0] || !DestSize[1] || !BaseSize[0] || !BaseSize[1])
 		{
-			Temp2 = CreateImage(DestSize[0], DestSize[1], 1, Base->GetFormat(), EInitializationType::NotInitialized);
-			ImageMagnifyY(Temp2.get(), Temp.get());
-			ReleaseImage(Temp);
+			return;
 		}
-		else if (DestSize[1] < BaseSize[1])
+
+		const bool bIsDownSizeByLessThanTwo =
+			BaseSize.X >= DestSize.X && BaseSize.Y >= DestSize.Y && 
+			BaseSize.X < DestSize.X*2 && BaseSize.Y < DestSize.Y*2; 
+
+		const bool bIsUncompressedFormat = !mu::IsCompressedFormat(InBase->GetFormat());  
+		const bool bIsBlockCompressedFormat = mu::IsBlockCompressedFormat(InBase->GetFormat());
+
+		if (bIsDownSizeByLessThanTwo && (bIsUncompressedFormat || bIsBlockCompressedFormat) && bEnableImageResizeMemoryOptimizations)
 		{
-			Temp2 = CreateImage(DestSize[0], DestSize[1], 1, Base->GetFormat(), EInitializationType::NotInitialized);
-			ImageMinifyY(Temp2.get(), Temp.get());
-			ReleaseImage(Temp);
+			// Special case for small down resizes where no more than 2 pixels need to be sampled. 
+			// This could be extended to upsizes of any size.
+			
+			const EImageFormat BaseUncompressedFormat = GetUncompressedFormat(InBase->GetFormat());
+			const EImageFormat DestFormat = Dest->GetFormat();
+			
+			Ptr<Image> DecompressedDest = Dest;
+			if (DestFormat != BaseUncompressedFormat)
+			{
+				DecompressedDest = CreateImage(DestSize.X, DestSize.Y, 1, BaseUncompressedFormat, EInitializationType::NotInitialized);
+			}
+			
+			ImageResizeLinearDownByLessThanTwo(DecompressedDest.get(), InBase);
+
+			if (DestFormat != BaseUncompressedFormat)
+			{
+				bool bSuccess = false;
+				ImagePixelFormat(bSuccess, ImageCompressionQuality, Dest, DecompressedDest.get());
+				check(bSuccess);
+
+				ReleaseImage(DecompressedDest);
+			}
 		}
 		else
 		{
-			Temp2 = Temp;
-		}
-		Temp = nullptr;
+			ImagePtrConst Base = InBase;
+			ImagePtr BaseReformat;
+			
+			// Shouldn't happen! But if it does...
+			const EImageFormat SourceFormat = Base->GetFormat();
+			const EImageFormat UncompressedFormat = GetUncompressedFormat(SourceFormat);
+			if (SourceFormat != UncompressedFormat)
+			{
+				BaseReformat = ImagePixelFormat(ImageCompressionQuality, InBase, UncompressedFormat);
+				Base = BaseReformat;
+			}
 
+			const EImageFormat BaseFormat = Base->GetFormat();
 
-		// Reset format if it was changed to scale
-		if (SourceFormat != UncompressedFormat)
-		{
-			Ptr<Image> OldTemp = Temp2;
-			Temp2 = ImagePixelFormat(ImageCompressionQuality, Temp2.get(), SourceFormat);
-			ReleaseImage(OldTemp);
-		}
+			// First resize X
+			ImagePtr Temp;
+			if (DestSize[0] > BaseSize[0])
+			{
+				Temp = CreateImage(DestSize[0], BaseSize[1], 1, BaseFormat, EInitializationType::NotInitialized);
+				ImageMagnifyX(Temp.get(), Base.get());
+			}
+			else if (DestSize[0] < BaseSize[0])
+			{
+				Temp = CreateImage(DestSize[0], BaseSize[1], 1, BaseFormat, EInitializationType::NotInitialized);
+				ImageMinifyX(Temp.get(), Base.get());
+			}
+			else
+			{
+				Temp = CloneImage(Base.get());
+			}
 
-		Dest->CopyMove(Temp2.get());
-		ReleaseImage(Temp2);
-		Temp2 = nullptr;
+			// No lonegr needed. Should be only set if a reformat has been done. 
+			check(!(SourceFormat == UncompressedFormat) || !BaseReformat);
+			if (BaseReformat)
+			{
+				check(Base == BaseReformat);
+				ReleaseImage(BaseReformat);
+				BaseReformat = nullptr;
+				Base = nullptr;
+			}
 
+			// Now resize Y
+			ImagePtr Temp2;
+			if (DestSize[1] > BaseSize[1])
+			{
+				Temp2 = CreateImage(DestSize[0], DestSize[1], 1, BaseFormat, EInitializationType::NotInitialized);
+				ImageMagnifyY(Temp2.get(), Temp.get());
+				ReleaseImage(Temp);
+			}
+			else if (DestSize[1] < BaseSize[1])
+			{
+				Temp2 = CreateImage(DestSize[0], DestSize[1], 1, BaseFormat, EInitializationType::NotInitialized);
+				ImageMinifyY(Temp2.get(), Temp.get());
+				ReleaseImage(Temp);
+			}
+			else
+			{
+				Temp2 = Temp;
+			}
+			Temp = nullptr;
+
+
+			// Reset format if it was changed to scale
+			if (SourceFormat != UncompressedFormat)
+			{
+				Ptr<Image> OldTemp = Temp2;
+				Temp2 = ImagePixelFormat(ImageCompressionQuality, Temp2.get(), SourceFormat);
+				ReleaseImage(OldTemp);
+			}
+
+			Dest->CopyMove(Temp2.get());
+			ReleaseImage(Temp2);
+			Temp2 = nullptr;
+		}
+		
 		// Update the relevancy data of the image.
-		if (Base->m_flags & Image::EImageFlags::IF_HAS_RELEVANCY_MAP)
+		if (InBase->m_flags & Image::EImageFlags::IF_HAS_RELEVANCY_MAP)
 		{
 			Dest->m_flags |= Image::EImageFlags::IF_HAS_RELEVANCY_MAP;
 
 			float FactorY = float(DestSize[1]) / float(BaseSize[1]);
 
-			Dest->RelevancyMinY = uint16(FMath::FloorToFloat(Base->RelevancyMinY * FactorY));
-			Dest->RelevancyMaxY = uint16(FMath::Min((int32)FMath::CeilToFloat(Base->RelevancyMinY * FactorY), Dest->GetSizeY() - 1));
+			Dest->RelevancyMinY = uint16(FMath::FloorToFloat(InBase->RelevancyMinY * FactorY));
+			Dest->RelevancyMaxY = uint16(FMath::Min((int32)FMath::CeilToFloat(InBase->RelevancyMaxY * FactorY), Dest->GetSizeY() - 1));
 		}
-
-		ReleaseImage(BaseReformat);
 	}
 
-}
\ No newline at end of file
+}
