[Mutable] Hang when closing the Game/Editor while a Mutable update is running.

#rnx

diff --git a/Engine/Plugins/Experimental/Mutable/Source/CustomizableObject/Private/MuCO/FMutableTaskGraph.cpp b/Engine/Plugins/Experimental/Mutable/Source/CustomizableObject/Private/MuCO/FMutableTaskGraph.cpp
--- a/Engine/Plugins/Experimental/Mutable/Source/CustomizableObject/Private/MuCO/FMutableTaskGraph.cpp
+++ b/Engine/Plugins/Experimental/Mutable/Source/CustomizableObject/Private/MuCO/FMutableTaskGraph.cpp
@@ -119,9 +119,28 @@ void FMutableTaskGraph::WaitForMutableTasks()
 	if (LastMutableTask.IsValid())
 	{
 #ifdef MUTABLE_USE_NEW_TASKGRAPH
-		LastMutableTask.Wait();
+		if (CVarTaskGraphBusyWait->GetBool())
+		{
+			LastMutableTask.BusyWait();
+		}
+		else
+		{
+			LastMutableTask.Wait();
+		}
 #else
-		LastMutableTask->Wait();
+		if (CVarTaskGraphBusyWait->GetBool())
+		{
+			LowLevelTasks::BusyWaitUntil(
+				[Task = LastMutableTask]()
+				{
+					return Task->IsComplete();
+				}
+			);
+		}
+		else 
+		{
+			LastMutableTask->Wait();
+		}
 #endif
 		LastMutableTask = {};
 	}
diff --git a/Engine/Plugins/Experimental/Mutable/Source/MutableRuntime/Private/MuR/CodeRunnerTasked.cpp b/Engine/Plugins/Experimental/Mutable/Source/MutableRuntime/Private/MuR/CodeRunnerTasked.cpp
--- a/Engine/Plugins/Experimental/Mutable/Source/MutableRuntime/Private/MuR/CodeRunnerTasked.cpp
+++ b/Engine/Plugins/Experimental/Mutable/Source/MutableRuntime/Private/MuR/CodeRunnerTasked.cpp
@@ -54,16 +54,6 @@ static FAutoConsoleVariableRef CVarCoreRunnerTaskPriority(
 	TEXT("0 AnyThread. 1 to 6, from hight to low priority."),
 	ECVF_Default);
 
-
-bool bBusyWait = true;
-
-static FAutoConsoleVariableRef CVarCodeRunnerTaskGraphBusyWait(
-	TEXT("mutable.CodeRunnerTaskGraphBusyWait"),
-	bBusyWait,
-	TEXT("Use TaskGraph BusyWait instead of simple Wait. Required to avoid hangs on platforms with low number of cores."),
-	ECVF_Default);
-
-
 namespace mu
 {
 
@@ -558,7 +548,7 @@ namespace mu
 					if (IssuedTasks[IssuedIndex]->Event.IsValid())
 					{
 #ifdef MUTABLE_USE_NEW_TASKGRAPH
-						if (CVarCodeRunnerTaskGraphBusyWait->GetBool())
+						if (CVarTaskGraphBusyWait->GetBool())
 						{
 							IssuedTasks[IssuedIndex]->Event.BusyWait();							
 						}
@@ -567,12 +557,10 @@ namespace mu
 							IssuedTasks[IssuedIndex]->Event.Wait();
 						}
 #else
-						if (CVarCodeRunnerTaskGraphBusyWait->GetBool())
+						if (CVarTaskGraphBusyWait->GetBool())
 						{
-							FGraphEventRef Task = IssuedTasks[IssuedIndex]->Event;
-							
 							LowLevelTasks::BusyWaitUntil(
-								[Task]()
+								[Task = IssuedTasks[IssuedIndex]->Event]()
 								{
 									return Task->IsComplete();
 								}
diff --git a/Engine/Plugins/Experimental/Mutable/Source/MutableRuntime/Private/MuR/System.cpp b/Engine/Plugins/Experimental/Mutable/Source/MutableRuntime/Private/MuR/System.cpp
--- a/Engine/Plugins/Experimental/Mutable/Source/MutableRuntime/Private/MuR/System.cpp
+++ b/Engine/Plugins/Experimental/Mutable/Source/MutableRuntime/Private/MuR/System.cpp
@@ -35,6 +35,11 @@
 #include "ProfilingDebugging/CountersTrace.h"
 #include "PackedNormal.h"
 
+TAutoConsoleVariable<bool> CVarTaskGraphBusyWait(
+	TEXT("mutable.TaskGraphBusyWait"),
+	true,
+	TEXT("Use TaskGraph BusyWait instead of simple Wait. Required to avoid hangs on platforms with low number of cores."),
+	ECVF_Default);
 
 namespace mu
 {
diff --git a/Engine/Plugins/Experimental/Mutable/Source/MutableRuntime/Public/MuR/System.h b/Engine/Plugins/Experimental/Mutable/Source/MutableRuntime/Public/MuR/System.h
--- a/Engine/Plugins/Experimental/Mutable/Source/MutableRuntime/Public/MuR/System.h
+++ b/Engine/Plugins/Experimental/Mutable/Source/MutableRuntime/Public/MuR/System.h
@@ -46,6 +46,9 @@ enum class ETextureCompressionStrategy : uint8
 };
 
 
+MUTABLERUNTIME_API extern TAutoConsoleVariable<bool> CVarTaskGraphBusyWait;
+
+
 namespace mu
 {
 	// Forward references
