#rnx

#tests PIE

diff --git a/Engine/Plugins/Experimental/Mutable/Source/CustomizableObject/Classes/MuCO/CustomizableObjectInstanceUsage.h b/Engine/Plugins/Experimental/Mutable/Source/CustomizableObject/Classes/MuCO/CustomizableObjectInstanceUsage.h
--- a/Engine/Plugins/Experimental/Mutable/Source/CustomizableObject/Classes/MuCO/CustomizableObjectInstanceUsage.h
+++ b/Engine/Plugins/Experimental/Mutable/Source/CustomizableObject/Classes/MuCO/CustomizableObjectInstanceUsage.h
@@ -84,6 +84,7 @@ public:
 
 private:
 	// If this CustomizableSkeletalComponent is not null, it means this Usage was created by it, and all persistent properties should be obtained through it
+	UPROPERTY()
 	TObjectPtr<UCustomizableSkeletalComponent> CustomizableSkeletalComponent;
 
 	// If no CustomizableSkeletalComponent is associated, this SkeletalComponent will be used
