[mutable] Solution to L+A normalmaps losing the alpha channel in Android ASTC.
#rnx

diff --git a/Engine/Plugins/Experimental/Mutable/Source/CustomizableObjectEditor/Private/MuCOE/GenerateMutableSource/GenerateMutableSourceImage.cpp b/Engine/Plugins/Experimental/Mutable/Source/CustomizableObjectEditor/Private/MuCOE/GenerateMutableSource/GenerateMutableSourceImage.cpp
--- a/Engine/Plugins/Experimental/Mutable/Source/CustomizableObjectEditor/Private/MuCOE/GenerateMutableSource/GenerateMutableSourceImage.cpp
+++ b/Engine/Plugins/Experimental/Mutable/Source/CustomizableObjectEditor/Private/MuCOE/GenerateMutableSource/GenerateMutableSourceImage.cpp
@@ -1042,7 +1042,10 @@ mu::NodeImagePtr GenerateMutableSourceImage(const UEdGraphPin* Pin, FMutableGrap
 					bSuccess = false;
 				}
 
-				if (bSuccess && Pin->PinType.PinCategory != Schema->PC_MaterialAsset && !TypedNodeTable->GetColumnDefaultAssetByType<UTexture2D>(Pin) && !TypedNodeTable->GetColumnDefaultAssetByType<UTexture2DArray>(Pin))
+				UTexture2D* DefaultTexture2D = TypedNodeTable->GetColumnDefaultAssetByType<UTexture2D>(Pin);
+				UTexture2DArray* DefaultTextureArray = TypedNodeTable->GetColumnDefaultAssetByType<UTexture2DArray>(Pin);
+
+				if (bSuccess && Pin->PinType.PinCategory != Schema->PC_MaterialAsset && !DefaultTexture2D && !DefaultTextureArray)
 				{
 					FString Msg = FString::Printf(TEXT("Couldn't find a default value in the data table's struct for the column [%s]. The default value is null or not a supported Texture"), *ColumnName);
 					GenerationContext.Compiler->CompilerLog(FText::FromString(Msg), Node);
@@ -1115,8 +1118,14 @@ mu::NodeImagePtr GenerateMutableSourceImage(const UEdGraphPin* Pin, FMutableGrap
 							ImageTableNode->SetTable(Table);
 							ImageTableNode->SetColumn(ColumnName);
 							ImageTableNode->SetParameterName(TypedNodeTable->ParameterName);
-							ImageTableNode->SetMaxTextureSize(MaxTextureSize);
 
+							// TextureArrays are passthrough textures and do not need this step
+							if (DefaultTexture2D)
+							{
+								int32 DefaultMaxTextureSize = GetMaxTextureSize(DefaultTexture2D, GenerationContext);
+								ImageTableNode->SetMaxTextureSize(FMath::Max(DefaultTexture2D->Source.GetSizeX(), DefaultTexture2D->Source.GetSizeY()));
+							}
+							
 							GenerationContext.AddParameterNameUnique(Node, TypedNodeTable->ParameterName);
 						}
 					}
diff --git a/Engine/Plugins/Experimental/Mutable/Source/CustomizableObjectEditor/Private/MuCOE/GenerateMutableSource/GenerateMutableSourceSurface.cpp b/Engine/Plugins/Experimental/Mutable/Source/CustomizableObjectEditor/Private/MuCOE/GenerateMutableSource/GenerateMutableSourceSurface.cpp
--- a/Engine/Plugins/Experimental/Mutable/Source/CustomizableObjectEditor/Private/MuCOE/GenerateMutableSource/GenerateMutableSourceSurface.cpp
+++ b/Engine/Plugins/Experimental/Mutable/Source/CustomizableObjectEditor/Private/MuCOE/GenerateMutableSource/GenerateMutableSourceSurface.cpp
@@ -42,6 +42,7 @@
 #include "MuT/NodeImageMipmap.h"
 #include "MuT/NodeImageNormalComposite.h"
 #include "MuT/NodeImageResize.h"
+#include "MuT/NodeImageSwizzle.h"
 #include "MuT/NodeMeshConstant.h"
 #include "MuT/NodeMeshFormat.h"
 #include "MuT/NodeMeshFragment.h"
@@ -668,6 +669,7 @@ mu::NodeSurfacePtr GenerateMutableSourceSurface(const UEdGraphPin * Pin, FMutabl
 								LastImage = CompositedImage;
 							}
 
+							mu::Ptr<mu::NodeImage> FormatSource = LastImage;
 							mu::NodeImageFormatPtr FormatImage = new mu::NodeImageFormat();
 							FormatImage->SetSource(LastImage.get());
 							FormatImage->SetFormat(mu::EImageFormat::IF_RGBA_UBYTE);
@@ -718,6 +720,35 @@ mu::NodeSurfacePtr GenerateMutableSourceSurface(const UEdGraphPin * Pin, FMutabl
 									mu::EImageFormat MutableFormat = QualityAndPerformanceFix(UnrealToMutablePixelFormat(UnrealTargetPlatformFormat,false));
 									mu::EImageFormat MutableFormatIfAlpha = QualityAndPerformanceFix(UnrealToMutablePixelFormat(UnrealTargetPlatformFormatAlpha,true));
 
+									// Temp hack to enable RG->LA 
+									if (GenerationContext.Options.TargetPlatform)
+									{
+										bool bUseLA = GenerationContext.Options.TargetPlatform->SupportsFeature(ETargetPlatformFeatures::NormalmapLAEncodingMode);
+										if (bUseLA)
+										{
+											// See GetQualityFormat in TextureFormatASTC.cpp to understand why
+											if (UnrealTargetPlatformFormat == PF_ASTC_6x6 || UnrealTargetPlatformFormat == PF_ASTC_6x6_NORM_RG)
+											{
+												MutableFormat = mu::EImageFormat::IF_ASTC_4x4_RGBA_LDR;
+												MutableFormatIfAlpha = mu::EImageFormat::IF_ASTC_4x4_RGBA_LDR;
+
+												// Insert a channel swizzle
+												mu::Ptr<mu::NodeImageSwizzle> Swizzle = new mu::NodeImageSwizzle;
+												Swizzle->SetFormat( mu::EImageFormat::IF_RGBA_UBYTE );
+												Swizzle->SetSource(0, FormatSource);
+												Swizzle->SetSource(1, FormatSource);
+												Swizzle->SetSource(2, FormatSource);
+												Swizzle->SetSource(3, FormatSource);
+												Swizzle->SetSourceChannel(0, 0);
+												Swizzle->SetSourceChannel(1, 0);
+												Swizzle->SetSourceChannel(2, 0);
+												Swizzle->SetSourceChannel(3, 1);
+
+												FormatImage->SetSource( Swizzle.get() );
+											}
+										}
+									}
+
 									// Unsupported format: look for something generic
 									if (MutableFormat == mu::EImageFormat::IF_NONE)
 									{
diff --git a/Engine/Plugins/Experimental/Mutable/Source/MutableTools/Private/MuT/CodeGenerator.cpp b/Engine/Plugins/Experimental/Mutable/Source/MutableTools/Private/MuT/CodeGenerator.cpp
--- a/Engine/Plugins/Experimental/Mutable/Source/MutableTools/Private/MuT/CodeGenerator.cpp
+++ b/Engine/Plugins/Experimental/Mutable/Source/MutableTools/Private/MuT/CodeGenerator.cpp
@@ -21,6 +21,7 @@
 #include "MuT/ASTOpImageCompose.h"
 #include "MuT/ASTOpImageMipmap.h"
 #include "MuT/ASTOpImagePixelFormat.h"
+#include "MuT/ASTOpImageSwizzle.h"
 #include "MuT/ASTOpImageLayer.h"
 #include "MuT/ASTOpImageLayerColor.h"
 #include "MuT/ASTOpImagePatch.h"
@@ -53,6 +54,7 @@
 #include "MuT/NodeImageFormatPrivate.h"
 #include "MuT/NodeImageMipmap.h"
 #include "MuT/NodeImageMipmapPrivate.h"
+#include "MuT/NodeImageSwizzlePrivate.h"
 #include "MuT/NodeImageReference.h"
 #include "MuT/NodeLODPrivate.h"
 #include "MuT/NodeMesh.h"
@@ -1213,6 +1215,7 @@ namespace mu
 					// Any image-specific format or mipmapping needs to be applied at the end
 					NodeImageMipmapPtr mipmapNode;
 					NodeImageFormatPtr formatNode;
+					NodeImageSwizzlePtr swizzleNode;
 					bool found = false;
 					while (!found)
 					{
@@ -1226,6 +1229,15 @@ namespace mu
 							if (!formatNode) formatNode = tf;
 							pImageNode = tf->GetSource();
 						}
+						else if (NodeImageSwizzle* ts = dynamic_cast<NodeImageSwizzle*>(pImageNode.get()))
+						{
+							NodeImage* Source = ts->GetSource(0).get();
+							if (!swizzleNode && Source ==ts->GetSource(1) && Source==ts->GetSource(2) && Source==ts->GetSource(3))
+							{
+								swizzleNode = ts;
+								pImageNode = Source;
+							}
+						}
 						else
 						{
 							found = true;
@@ -1304,6 +1316,22 @@ namespace mu
 
 						check(imageAd);
 
+						if (swizzleNode)
+						{
+							Ptr<ASTOpImageSwizzle> fop = new ASTOpImageSwizzle();
+							fop->Format = swizzleNode->GetPrivate()->m_format;
+							fop->Sources[0] = imageAd;
+							fop->Sources[1] = imageAd;
+							fop->Sources[2] = imageAd;
+							fop->Sources[3] = imageAd;
+							fop->SourceChannels[0] = swizzleNode->GetPrivate()->m_sourceChannels[0];
+							fop->SourceChannels[1] = swizzleNode->GetPrivate()->m_sourceChannels[1];
+							fop->SourceChannels[2] = swizzleNode->GetPrivate()->m_sourceChannels[2];
+							fop->SourceChannels[3] = swizzleNode->GetPrivate()->m_sourceChannels[3];
+							check(fop->Format != EImageFormat::IF_NONE);
+							imageAd = fop;
+						}
+
 						if (mipmapNode)
 						{
 							Ptr<ASTOpImageMipmap> op = new ASTOpImageMipmap();
@@ -1584,6 +1612,21 @@ namespace mu
 								}
 							}
 
+							if (swizzleNode)
+							{
+								Ptr<ASTOpImageSwizzle> fop = new ASTOpImageSwizzle();
+								fop->Format = swizzleNode->GetPrivate()->m_format;
+								fop->Sources[0] = imageAd;
+								fop->Sources[1] = imageAd;
+								fop->Sources[2] = imageAd;
+								fop->Sources[3] = imageAd;
+								fop->SourceChannels[0] = swizzleNode->GetPrivate()->m_sourceChannels[0];
+								fop->SourceChannels[1] = swizzleNode->GetPrivate()->m_sourceChannels[1];
+								fop->SourceChannels[2] = swizzleNode->GetPrivate()->m_sourceChannels[2];
+								fop->SourceChannels[3] = swizzleNode->GetPrivate()->m_sourceChannels[3];
+								check(fop->Format != EImageFormat::IF_NONE);
+								imageAd = fop;
+							}
 
 							// Apply mipmap and format if necessary
 							if (mipmapNode)
