[Mutable] Fix static analyzer warning in OpImageProject.cpp.

#rnx

diff --git a/Engine/Plugins/Experimental/Mutable/Source/MutableRuntime/Private/MuR/OpImageProject.cpp b/Engine/Plugins/Experimental/Mutable/Source/MutableRuntime/Private/MuR/OpImageProject.cpp
--- a/Engine/Plugins/Experimental/Mutable/Source/MutableRuntime/Private/MuR/OpImageProject.cpp
+++ b/Engine/Plugins/Experimental/Mutable/Source/MutableRuntime/Private/MuR/OpImageProject.cpp
@@ -180,6 +180,7 @@ namespace Private
 		static_assert(CheckExactlyOneProjectionFlag(Features));
 
 		static constexpr int32 PIXEL_SIZE = GetFormatNumChannels(Features);
+		static_assert(PIXEL_SIZE != 0);
 
 	public:
 		static FProjectedPixelProcessorContext MakeContext(
