[mutable] Optimize mutable generated code involving swizzle+compose operations. Useful for platforms that have a swizzle at the top of the expression, like Android ASTC for LA formats.
#rnx

diff --git a/Engine/Plugins/Experimental/Mutable/Source/MutableTools/Private/MuT/ASTOpImageSwizzle.cpp b/Engine/Plugins/Experimental/Mutable/Source/MutableTools/Private/MuT/ASTOpImageSwizzle.cpp
--- a/Engine/Plugins/Experimental/Mutable/Source/MutableTools/Private/MuT/ASTOpImageSwizzle.cpp
+++ b/Engine/Plugins/Experimental/Mutable/Source/MutableTools/Private/MuT/ASTOpImageSwizzle.cpp
@@ -424,6 +424,32 @@ namespace mu
 				break;
 			}
 
+			case OP_TYPE::IM_COMPOSE:
+			{
+				// Move the swizzle down the base and the block, but not the mask.
+				Ptr<ASTOpImageCompose> nop = mu::Clone<ASTOpImageCompose>(channelSourceAt);
+
+				Ptr<ASTOpImageSwizzle> aOp = mu::Clone<ASTOpImageSwizzle>(this);
+				ReplaceAllSources(aOp, nop->Base.child());
+				nop->Base = aOp;
+
+				Ptr<ASTOpImageSwizzle> bOp = mu::Clone<ASTOpImageSwizzle>(this);
+				ReplaceAllSources(bOp, nop->BlockImage.child());
+				nop->BlockImage = bOp;
+
+				at = nop;
+				break;
+			}
+
+			case OP_TYPE::IM_BLANKLAYOUT:
+			{
+				// We can remove the swizzle entirely.
+				// It is not 100% equivalent, because blank layouts are initialized with 0,0,0,1 so the result could be
+				// different, but those pixels shouldn't be used anyway.
+				at = channelSourceAt;
+				break;
+			}
+
 			default:
 				bAllChannelsAreTheSame = false;
 				break;
