Properly disable 510 for UE_LOG
Fix/silence handful of additional warnings
#rnx

diff --git a/Engine/Plugins/Experimental/Mutable/Source/MutableRuntime/Private/MuR/OpImageProject.cpp b/Engine/Plugins/Experimental/Mutable/Source/MutableRuntime/Private/MuR/OpImageProject.cpp
--- a/Engine/Plugins/Experimental/Mutable/Source/MutableRuntime/Private/MuR/OpImageProject.cpp
+++ b/Engine/Plugins/Experimental/Mutable/Source/MutableRuntime/Private/MuR/OpImageProject.cpp
@@ -288,7 +288,7 @@ namespace Private
 				if (Factor > UE_SMALL_NUMBER)
 				{
 					MaskFactor = ((MaskFactor > 0.0f) && Context.MaskData)
-							? static_cast<float>(Context.MaskData[(BufferPos - Context.TargetData) / PIXEL_SIZE])
+							? static_cast<float>(Context.MaskData[(BufferPos - Context.TargetData) / PIXEL_SIZE]) //-V609
 							: MaskFactor;
 					Factor = (Factor * MaskFactor) * (1.0f/255.0f);
 				}
@@ -481,7 +481,7 @@ namespace Private
 				if (Factor > 0)
 				{
 					MaskFactor = ((MaskFactor > 0) && Context.MaskData) 
-							? Context.MaskData[(BufferPos - Context.TargetData) / PIXEL_SIZE] 
+							? Context.MaskData[(BufferPos - Context.TargetData) / PIXEL_SIZE]  //-V609
 							: MaskFactor;
 					Factor = (Factor * MaskFactor) / 255;
 				}
diff --git a/Engine/Source/Developer/Apple/MetalShaderFormat/Private/MetalShaderCompiler.cpp b/Engine/Source/Developer/Apple/MetalShaderFormat/Private/MetalShaderCompiler.cpp
--- a/Engine/Source/Developer/Apple/MetalShaderFormat/Private/MetalShaderCompiler.cpp
+++ b/Engine/Source/Developer/Apple/MetalShaderFormat/Private/MetalShaderCompiler.cpp
@@ -1060,7 +1060,7 @@ void CompileMetalShader(const FShaderCompilerInput& Input, const FString& InPrep
 
 	TCHAR const* StandardPlatform = bIsMobile ? TEXT("ios") : TEXT("macos");
 	FString Standard;
-	if (VersionEnum >= 8)
+	if (VersionEnum >= 8) //-V547
 	{
 		Standard = FString::Printf(TEXT("-std=metal%s"), *StandardVersion);
 	}
diff --git a/Engine/Source/Editor/KismetCompiler/Private/KismetCompiler.cpp b/Engine/Source/Editor/KismetCompiler/Private/KismetCompiler.cpp
--- a/Engine/Source/Editor/KismetCompiler/Private/KismetCompiler.cpp
+++ b/Engine/Source/Editor/KismetCompiler/Private/KismetCompiler.cpp
@@ -2530,7 +2530,7 @@ void FKismetCompilerContext::PrecompileFunction(FKismetFunctionContext& Context,
 
 		////////////////////////////////////////
 
-		if (Context.Function && Context.IsDelegateSignature())
+		if (Context.IsDelegateSignature())
 		{
 			Context.Function->FunctionFlags |= FUNC_Delegate;
 
@@ -2547,7 +2547,7 @@ void FKismetCompilerContext::PrecompileFunction(FKismetFunctionContext& Context,
 			}
 		}
 
-		if (Context.EntryPoint && Context.Function)
+		if (Context.EntryPoint)
 		{
 			if (Context.EntryPoint->MetaData.HasMetaData(FBlueprintMetadata::MD_FieldNotify) && NewClass->ImplementsInterface(UNotifyFieldValueChanged::StaticClass()))
 			{
diff --git a/Engine/Source/Runtime/Core/Public/Windows/WindowsPlatformCodeAnalysis.h b/Engine/Source/Runtime/Core/Public/Windows/WindowsPlatformCodeAnalysis.h
--- a/Engine/Source/Runtime/Core/Public/Windows/WindowsPlatformCodeAnalysis.h
+++ b/Engine/Source/Runtime/Core/Public/Windows/WindowsPlatformCodeAnalysis.h
@@ -189,7 +189,7 @@
 	//-V:DataStripped:581
 	//-V:FromInt:601
 	//-V:UE_CLOG(:501,560
-	//-V:UE_LOG(:501, 560
+	//-V:UE_LOG(:501,510,560
 	//-V:UGL_REQUIRED_VOID:501
 	//-V:AnimScriptInstance:595
 	//-V:Driver:595
@@ -229,8 +229,6 @@
 	//-V:Linker:678
 	//-V:self:678
 	//-V:AccumulateParentID:678
-	//-V:BasicLog:510
-	//-V:BasicFatalLog:510
 
 	// The following classes retain a reference to data supplied in the constructor by the derived class which can not yet be initialized.
 	//-V:FMemoryWriter(:1050
diff --git a/Engine/Source/Runtime/Renderer/Private/InstanceCulling/InstanceCullingContext.cpp b/Engine/Source/Runtime/Renderer/Private/InstanceCulling/InstanceCullingContext.cpp
--- a/Engine/Source/Runtime/Renderer/Private/InstanceCulling/InstanceCullingContext.cpp
+++ b/Engine/Source/Runtime/Renderer/Private/InstanceCulling/InstanceCullingContext.cpp
@@ -832,7 +832,7 @@ void FInstanceCullingContext::BuildRenderingCommandsInternal(
 
 	PassParametersTmp.ViewIds = GraphBuilder.CreateSRV(ViewIdsBuffer);
 	PassParametersTmp.NumCullingViews = 0;
-	if (bCullInstances || bAllowWPODisable) //-V595
+	if ((bCullInstances || bAllowWPODisable) && InstanceCullingManager)
 	{
 #if DO_CHECK
 		for (int32 ViewId : ViewIds)
diff --git a/Engine/Source/Runtime/Renderer/Private/SceneVisibility.cpp b/Engine/Source/Runtime/Renderer/Private/SceneVisibility.cpp
--- a/Engine/Source/Runtime/Renderer/Private/SceneVisibility.cpp
+++ b/Engine/Source/Runtime/Renderer/Private/SceneVisibility.cpp
@@ -4438,7 +4438,7 @@ void FSceneRenderer::PreVisibilityFrameSetup(FRDGBuilder& GraphBuilder)
 			{
 				// Prepare (skel.) data for guide/simulation update
 				// If we are rendering from scene capture we don't need to run another time the hair bookmarks.
-				if (IsHairStrandsEnabled(EHairStrandsShaderType::All, Scene->GetShaderPlatform()) && Views[0].AllowGPUParticleUpdate())
+				if (Scene && IsHairStrandsEnabled(EHairStrandsShaderType::All, Scene->GetShaderPlatform()) && Views[0].AllowGPUParticleUpdate())
 				{
 					Scene->WaitForGPUSkinCacheTask();
 
