Fixed a crash that occurs when geometry is missing for a LOD while generating a mesh from a Mutable instance

diff --git a/Engine/Plugins/Experimental/Mutable/Source/CustomizableObject/Private/MuCO/CustomizableObjectInstance.cpp b/Engine/Plugins/Experimental/Mutable/Source/CustomizableObject/Private/MuCO/CustomizableObjectInstance.cpp
--- a/Engine/Plugins/Experimental/Mutable/Source/CustomizableObject/Private/MuCO/CustomizableObjectInstance.cpp
+++ b/Engine/Plugins/Experimental/Mutable/Source/CustomizableObject/Private/MuCO/CustomizableObjectInstance.cpp
@@ -4632,7 +4632,7 @@ bool UCustomizableInstancePrivateData::BuildOrCopyRenderData(const TSharedRef<FU
 			Component = &OperationData->InstanceUpdateData.Components[LOD.FirstComponent + ComponentIndex];
 		}
 
-		if (!Component || !Component->bGenerated)
+		if (!Component || !Component->bGenerated || Component->SurfaceCount == 0)
 		{
 			TRACE_CPUPROFILER_EVENT_SCOPE_TEXT(*FString::Printf(TEXT("BuildOrCopyRenderData_CopyData: From LOD %d to LOD %d"), LastValidLODIndex, LODIndex));
 
