[Mutable] Fix Remove Selected Light not hidding gizmo.

#rnx

diff --git a/Engine/Plugins/Experimental/Mutable/Source/CustomizableObjectEditor/Private/MuCOE/CustomizableObjectEditor.cpp b/Engine/Plugins/Experimental/Mutable/Source/CustomizableObjectEditor/Private/MuCOE/CustomizableObjectEditor.cpp
--- a/Engine/Plugins/Experimental/Mutable/Source/CustomizableObjectEditor/Private/MuCOE/CustomizableObjectEditor.cpp
+++ b/Engine/Plugins/Experimental/Mutable/Source/CustomizableObjectEditor/Private/MuCOE/CustomizableObjectEditor.cpp
@@ -1225,6 +1225,9 @@ void FCustomizableObjectEditor::ShowGizmoLight(ULightComponent& InSelectedLight)
 void FCustomizableObjectEditor::HideGizmoLight()
 {
 	CustomSettings->SetSelectedLight(nullptr);
+
+	Viewport->HideGizmoLight();
+
 	CustomizableObjectEditorAdvancedPreviewSettings->Refresh();
 }
 
