#rnx
#tests PIE

diff --git a/Engine/Plugins/Experimental/Mutable/Source/CustomizableObject/Classes/MuCO/CustomizableObjectInstanceUsage.h b/Engine/Plugins/Experimental/Mutable/Source/CustomizableObject/Classes/MuCO/CustomizableObjectInstanceUsage.h
--- a/Engine/Plugins/Experimental/Mutable/Source/CustomizableObject/Classes/MuCO/CustomizableObjectInstanceUsage.h
+++ b/Engine/Plugins/Experimental/Mutable/Source/CustomizableObject/Classes/MuCO/CustomizableObjectInstanceUsage.h
@@ -89,7 +89,7 @@ private:
 
 	// If no CustomizableSkeletalComponent is associated, this SkeletalComponent will be used
 	UPROPERTY()
-	TObjectPtr<USkeletalMeshComponent> UsedSkeletalMeshComponent;
+	TWeakObjectPtr<USkeletalMeshComponent> UsedSkeletalMeshComponent;
 
 	// If no CustomizableSkeletalComponent is associated, this Instance will be used
 	UPROPERTY()
diff --git a/Engine/Plugins/Experimental/Mutable/Source/CustomizableObject/Private/MuCO/CustomizableObjectInstanceUsage.cpp b/Engine/Plugins/Experimental/Mutable/Source/CustomizableObject/Private/MuCO/CustomizableObjectInstanceUsage.cpp
--- a/Engine/Plugins/Experimental/Mutable/Source/CustomizableObject/Private/MuCO/CustomizableObjectInstanceUsage.cpp
+++ b/Engine/Plugins/Experimental/Mutable/Source/CustomizableObject/Private/MuCO/CustomizableObjectInstanceUsage.cpp
@@ -156,7 +156,14 @@ void UCustomizableObjectInstanceUsage::AttachTo(USkeletalMeshComponent* Skeletal
 	}
 	else
 	{
-		UsedSkeletalMeshComponent = SkeletalMeshComponent;
+		if (IsValid(SkeletalMeshComponent))
+		{
+			UsedSkeletalMeshComponent = SkeletalMeshComponent;
+		}
+		else
+		{
+			UsedSkeletalMeshComponent = nullptr;
+		}
 
 		// To mimic the behavior of UCustomizableSkeletalComponent::OnAttachmentChanged()
 		SetPendingSetSkeletalMesh(true);
@@ -170,10 +177,12 @@ USkeletalMeshComponent* UCustomizableObjectInstanceUsage::GetAttachParent() cons
 	{
 		return Cast<USkeletalMeshComponent>(CustomizableSkeletalComponent->GetAttachParent());
 	}
-	else
+	else if(UsedSkeletalMeshComponent.IsValid())
 	{
-		return UsedSkeletalMeshComponent;
+		return UsedSkeletalMeshComponent.Get();
 	}
+
+	return nullptr;
 }
 
 
@@ -580,7 +589,7 @@ bool UCustomizableObjectInstanceUsage::IsNetMode(ENetMode InNetMode) const
 	{
 		return CustomizableSkeletalComponent->IsNetMode(InNetMode);
 	}
-	else if(UsedSkeletalMeshComponent)
+	else if(UsedSkeletalMeshComponent.IsValid())
 	{
 		return UsedSkeletalMeshComponent->IsNetMode(InNetMode);
 	}
