[mutable] Optimize memory usage in CustomizableObject compilation. Don't preload unreal assets, and delay compilation-time access as much as possible.
#rnx

diff --git a/Engine/Plugins/Experimental/Mutable/Source/CustomizableObject/Classes/MuCO/CustomizableObject.h b/Engine/Plugins/Experimental/Mutable/Source/CustomizableObject/Classes/MuCO/CustomizableObject.h
--- a/Engine/Plugins/Experimental/Mutable/Source/CustomizableObject/Classes/MuCO/CustomizableObject.h
+++ b/Engine/Plugins/Experimental/Mutable/Source/CustomizableObject/Classes/MuCO/CustomizableObject.h
@@ -1490,7 +1490,7 @@ private:
 	// This is a manual version number for the binary blobs in this asset.
 	// Increasing it invalidates all the previously compiled models.
 	// Warning: If while merging code both versions have changed, take the highest+1.
-	static const int32 CurrentSupportedVersion = 412;
+	static const int32 CurrentSupportedVersion = 413;
 
 public:
 
diff --git a/Engine/Plugins/Experimental/Mutable/Source/CustomizableObjectEditor/Private/MuCOE/CustomizableObjectCompileRunnable.cpp b/Engine/Plugins/Experimental/Mutable/Source/CustomizableObjectEditor/Private/MuCOE/CustomizableObjectCompileRunnable.cpp
--- a/Engine/Plugins/Experimental/Mutable/Source/CustomizableObjectEditor/Private/MuCOE/CustomizableObjectCompileRunnable.cpp
+++ b/Engine/Plugins/Experimental/Mutable/Source/CustomizableObjectEditor/Private/MuCOE/CustomizableObjectCompileRunnable.cpp
@@ -4,11 +4,13 @@
 
 #include "HAL/FileManager.h"
 #include "MuCO/UnrealMutableModelDiskStreamer.h"
+#include "MuCO/UnrealToMutableTextureConversionUtils.h"
 #include "MuR/Model.h"
 #include "MuT/Compiler.h"
 #include "MuT/ErrorLog.h"
 #include "MuT/UnrealPixelFormatOverride.h"
 #include "Serialization/MemoryWriter.h"
+#include "Async/Async.h"
 #include "Trace/Trace.inl"
 
 class ITargetPlatform;
@@ -87,6 +89,62 @@ uint32 FCustomizableObjectCompileRunnable::Run()
 		CompilerOptions->SetImagePixelFormatOverride( UnrealPixelFormatFunc );
 	}
 
+	CompilerOptions->SetReferencedResourceCallback([this](int32 ID) 
+		{
+			mu::Ptr<mu::Image> Image;
+
+			auto LoadFunc = [this,ID]() 
+			{
+				check(IsInGameThread());
+
+				mu::Ptr<mu::Image> Image;
+				if (!ReferencedTextures.IsValidIndex(ID))
+				{
+					// The id is not valid for this CO
+					check(false);
+					return Image;
+				}
+
+				// Find the texture id
+				TSoftObjectPtr<UTexture> TexturePtr = ReferencedTextures[ID];
+
+				// This can cause a stall because of loading the asset.
+				UTexture2D* Texture = Cast<UTexture2D>(TexturePtr.LoadSynchronous());
+				if (!Texture)
+				{
+					// Failed to load the texture
+					check(false);
+					return Image;
+				}
+
+				// In the editor the src data can be directly accessed
+				Image = new mu::Image();
+				int32 MipmapsToSkip = 0;
+				bool bIsNormalComposite = false; // TODO?
+				EUnrealToMutableConversionError Error = ConvertTextureUnrealSourceToMutable(Image.get(), Texture, bIsNormalComposite, MipmapsToSkip);
+				check(Error == EUnrealToMutableConversionError::Success);
+				return Image;
+			};
+
+			// This runs in a random thread
+			if (IsInGameThread())
+			{
+				Image = LoadFunc();
+			}
+			else
+			{
+				UE::Tasks::FTaskEvent Completion(TEXT("SetReferencedResourceCallback"));
+
+				AsyncTask(ENamedThreads::GameThread, [&Image, &Completion, LoadFunc]() 
+					{
+						Image = LoadFunc();
+						Completion.Trigger();
+					});
+
+				Completion.Wait();
+			}
+			return Image;
+		});
 
 	// Minimum resident mip count.
 	const int MinResidentMips = UTexture::GetStaticMinTextureResidentMipCount();
@@ -292,3 +350,4 @@ const ITargetPlatform* FCustomizableObjectSaveDDRunnable::GetTargetPlatform() co
 }
 
 #undef LOCTEXT_NAMESPACE
+
diff --git a/Engine/Plugins/Experimental/Mutable/Source/CustomizableObjectEditor/Private/MuCOE/CustomizableObjectCompileRunnable.h b/Engine/Plugins/Experimental/Mutable/Source/CustomizableObjectEditor/Private/MuCOE/CustomizableObjectCompileRunnable.h
--- a/Engine/Plugins/Experimental/Mutable/Source/CustomizableObjectEditor/Private/MuCOE/CustomizableObjectCompileRunnable.h
+++ b/Engine/Plugins/Experimental/Mutable/Source/CustomizableObjectEditor/Private/MuCOE/CustomizableObjectCompileRunnable.h
@@ -58,6 +58,8 @@ public:
 	TSharedPtr<mu::Model, ESPMode::ThreadSafe> Model;
 
 	FCompilationOptions Options;
+	
+	TArray<TSoftObjectPtr<UTexture>> ReferencedTextures;
 
 	FString ErrorMsg;
 
diff --git a/Engine/Plugins/Experimental/Mutable/Source/CustomizableObjectEditor/Private/MuCOE/CustomizableObjectCompiler.cpp b/Engine/Plugins/Experimental/Mutable/Source/CustomizableObjectEditor/Private/MuCOE/CustomizableObjectCompiler.cpp
--- a/Engine/Plugins/Experimental/Mutable/Source/CustomizableObjectEditor/Private/MuCOE/CustomizableObjectCompiler.cpp
+++ b/Engine/Plugins/Experimental/Mutable/Source/CustomizableObjectEditor/Private/MuCOE/CustomizableObjectCompiler.cpp
@@ -32,14 +32,12 @@ class UTexture2D;
 #define UE_MUTABLE_COMPILE_REGION		TEXT("Mutable Compile")
 #define UE_MUTABLE_PRELOAD_REGION		TEXT("Mutable Preload")
 #define UE_MUTABLE_SAVEDD_REGION		TEXT("Mutable SaveDD")
-#define UE_MUTABLE_CONVERTING_REGION	TEXT("Mutable Converting")
 
 UCustomizableObjectNodeObject* GetRootNode(UCustomizableObject* Object, bool &bOutMultipleBaseObjectsFound);
 
 
 
 FCustomizableObjectCompiler::FCustomizableObjectCompiler() : FCustomizableObjectCompilerBase()
-	, PendingTexturesToLoad(false)
 	, CompilationLaunchPending(false)
 	, PreloadingReferencerAssets(false)
 	, CurrentGAsyncLoadingTimeLimit(-1.0f)
@@ -103,10 +101,7 @@ bool FCustomizableObjectCompiler::Tick()
 		TRACE_END_REGION(UE_MUTABLE_SAVEDD_REGION);
 	}
 
-	// In editor, when compiling a CO, referencer assets and Unreal to Mutable texture conversion are performed asynchronously
-	UpdatePendingTextureConversion(true);
-
-	if (!PendingTexturesToLoad && CompilationLaunchPending)
+	if (CompilationLaunchPending)
 	{
 		CompilationLaunchPending = false;
 		LaunchMutableCompile(true);
@@ -309,12 +304,6 @@ void FCustomizableObjectCompiler::AddReferencedObjects(FReferenceCollector& Coll
 		Collector.AddReferencedObject(ArrayGCProtect[i]);
 	}
 
-	// Protect images that are pending to be loaded as well
-	for (FTextureUnrealToMutableTask& Image: ArrayTextureUnrealToMutableTask )
-	{
-		Collector.AddReferencedObject(Image.Texture);
-	}
-
 	if (CurrentObject)
 	{
 		Collector.AddReferencedObject(CurrentObject);
@@ -828,86 +817,6 @@ FAssetData* FCustomizableObjectCompiler::GetCachedAssetData(const FString& Packa
 }
 
 
-void FCustomizableObjectCompiler::UpdatePendingTextureConversion(bool UseTimeLimit)
-{
-	// In editor, when compiling a CO, referencer assets and Unreal to Mutable texture conversion are performed asynchronously
-	if (PendingTexturesToLoad)
-	{
-		double InitialTimeSeconds = FPlatformTime::Seconds();
-
-		float CurrentMaxConvertToMutableTextureTimeSeconds = MaxConvertToMutableTextureTime;
-		if (!FApp::HasFocus())
-		{
-			// If the app is not focused, it will reduce the framerate to 10fps. Do 2 more seconds of work to compensate for that.
-			CurrentMaxConvertToMutableTextureTimeSeconds += 2.0f;
-		}
-
-		while ( CompletedUnrealToMutableTask < ArrayTextureUnrealToMutableTask.Num() )
-		{
-			UTexture2D* Texture = ArrayTextureUnrealToMutableTask[CompletedUnrealToMutableTask].Texture;
-
-			// If the texture has been set to null, it means it was a duplicate that we already processed.
-			if (!Texture)
-			{ 
-				CompletedUnrealToMutableTask++;
-				continue;
-			}
-
-			// Convert the texture
-			mu::ImagePtr Image = ConvertTextureUnrealToMutable( 
-				Texture, 
-				ArrayTextureUnrealToMutableTask[CompletedUnrealToMutableTask].Node,
-				this,
-				ArrayTextureUnrealToMutableTask[CompletedUnrealToMutableTask].bIsNormalComposite );
-
-			// Assign to all tasks referring to the same texture
-			for (int32 j = CompletedUnrealToMutableTask; j < ArrayTextureUnrealToMutableTask.Num(); ++j)
-			{
-				if (ArrayTextureUnrealToMutableTask[j].Texture == Texture)
-				{
-					if (ArrayTextureUnrealToMutableTask[j].ImageNode.get())
-					{
-						ArrayTextureUnrealToMutableTask[j].ImageNode->SetValue(Image.get());
-						ArrayTextureUnrealToMutableTask[j].Texture = nullptr;
-					}
-
-					else if (ArrayTextureUnrealToMutableTask[j].TableNode.get())
-					{
-						int32 ColumnIndx = ArrayTextureUnrealToMutableTask[j].TableColumn;
-						int32 RowIndx = ArrayTextureUnrealToMutableTask[j].TableRow;
-
-						mu::Ptr<mu::ResourceProxyMemory<mu::Image>> ImageProxy = new mu::ResourceProxyMemory<mu::Image>(Image);
-						ArrayTextureUnrealToMutableTask[j].TableNode->SetCell(ColumnIndx, RowIndx, ImageProxy.get());
-
-						ArrayTextureUnrealToMutableTask[j].Texture = nullptr;
-					}
-				}
-			}
-
-			CompletedUnrealToMutableTask++;
-
-			if (UseTimeLimit)
-			{
-				double UsedTimeSeconds = (FPlatformTime::Seconds() - InitialTimeSeconds);
-				if (UsedTimeSeconds > CurrentMaxConvertToMutableTextureTimeSeconds)
-				{
-					break;
-				}
-			}
-		}
-
-		if (CompletedUnrealToMutableTask >= ArrayTextureUnrealToMutableTask.Num())
-		{
-			PendingTexturesToLoad = false;
-			ArrayTextureUnrealToMutableTask.Empty();
-			CompletedUnrealToMutableTask = 0;
-
-			TRACE_END_REGION(UE_MUTABLE_CONVERTING_REGION);
-		}
-	}
-}
-
-
 float FCustomizableObjectCompiler::ComputeAsyncLoadingTimeLimit()
 {
 	float DeltaTimeSeconds = FApp::GetDeltaTime();
@@ -1276,19 +1185,24 @@ void FCustomizableObjectCompiler::CompileInternal(UCustomizableObject* Object, c
 		Object->SetBoneNamesArray(GenerationContext.BoneNames);
 
 		// Pass-through textures
-		TArray<TSoftObjectPtr<UTexture>> NewReferencedPassThroughTextures;
+		TArray<TSoftObjectPtr<UTexture>> NewCompileTimeReferencedTextures;
+		for (const TPair<TSoftObjectPtr<UTexture>, FMutableGraphGenerationContext::FGeneratedReferencedTexture>& Pair : GenerationContext.CompileTimeTextureMap)
+		{
+			check(Pair.Value.ID == NewCompileTimeReferencedTextures.Num());
+			NewCompileTimeReferencedTextures.Add(Pair.Key);
+		}
 
-		for (const TPair<TSoftObjectPtr<UTexture>, FMutableGraphGenerationContext::FGeneratedPassThroughTexture>& Pair : GenerationContext.PassThroughTextureMap)
+		TArray<TSoftObjectPtr<UTexture>> NewRuntimeReferencedTextures;
+		for (const TPair<TSoftObjectPtr<UTexture>, FMutableGraphGenerationContext::FGeneratedReferencedTexture>& Pair : GenerationContext.RuntimeReferencedTextureMap)
 		{
-			check(Pair.Value.ID == NewReferencedPassThroughTextures.Num());
-			NewReferencedPassThroughTextures.Add(Pair.Key);
+			check(Pair.Value.ID == NewRuntimeReferencedTextures.Num());
+			NewRuntimeReferencedTextures.Add(Pair.Key);
 		}
 
 		// Mark the object as modified, used to avoid missing assets in packages.
-		if (Object->ReferencedPassThroughTextures != NewReferencedPassThroughTextures)
+		if (Object->ReferencedPassThroughTextures != NewRuntimeReferencedTextures)
 		{
-			Object->ReferencedPassThroughTextures.Empty(NewReferencedPassThroughTextures.Num());
-			Object->ReferencedPassThroughTextures = NewReferencedPassThroughTextures;
+			Object->ReferencedPassThroughTextures = NewRuntimeReferencedTextures;
 
 			if (!ParamNamesToSelectedOptions.Num()) // Don't mark the objects as modified because of a partial compilation
 			{
@@ -1382,22 +1296,13 @@ void FCustomizableObjectCompiler::CompileInternal(UCustomizableObject* Object, c
 
 		CompileTask = MakeShareable(new FCustomizableObjectCompileRunnable(MutableRoot));
 		CompileTask->Options = Options;
+		CompileTask->ReferencedTextures = NewCompileTimeReferencedTextures;
 		SetCompilationState(ECustomizableObjectCompilationState::InProgress);
 
-		if (GenerationContext.ArrayTextureUnrealToMutableTask.Num() > 0)
-		{
-			ArrayTextureUnrealToMutableTask.Insert(GenerationContext.ArrayTextureUnrealToMutableTask, ArrayTextureUnrealToMutableTask.Num());
-			PendingTexturesToLoad = true;
-			TRACE_BEGIN_REGION(UE_MUTABLE_CONVERTING_REGION);
-		}
-
-		// If synchronous compilation is requested, proceed the same way if in editor and if packaging
 		if (!bAsync)
 		{
-			UpdatePendingTextureConversion(false);
-			LaunchMutableCompile(false);
-			MUTABLE_CPUPROFILER_SCOPE(WaitForCompletion);
-			CompileThread->WaitForCompletion();
+			CompileTask->Init();
+			CompileTask->Run();
 			FinishCompilation();
 
 			if (SaveDDTask.IsValid())
@@ -1419,7 +1324,6 @@ void FCustomizableObjectCompiler::CompileInternal(UCustomizableObject* Object, c
 			// If packaging, convert textures and launch Mutable compile thread
 			if (IsRunningCommandlet())
 			{
-				UpdatePendingTextureConversion(false);
 				LaunchMutableCompile(false);
 			}
 			else
@@ -1427,11 +1331,6 @@ void FCustomizableObjectCompiler::CompileInternal(UCustomizableObject* Object, c
 				// If asynchronous compilation and not packaging, set CompilationLaunchPending to true and FCustomizableObjectCompiler::Tick()
 				// will do the remaining steps (convert textures asynchronously and launch the Mutable compile thread)
 				CompilationLaunchPending = true;
-
-				if (PendingTexturesToLoad)
-				{
-					AddCompileNotification(LOCTEXT("ConvertingToMutableTexture", "Converting textures"));
-				}
 			}
 		}
 	}
@@ -1455,7 +1354,7 @@ void FCustomizableObjectCompiler::CompileInternal(UCustomizableObject* Object, c
 }
 
 
-mu::NodePtr FCustomizableObjectCompiler::Export(UCustomizableObject* Object, const FCompilationOptions& InCompilerOptions, TArray<TSoftObjectPtr<UTexture>>& OutReferencedTextures)
+mu::NodePtr FCustomizableObjectCompiler::Export(UCustomizableObject* Object, const FCompilationOptions& InCompilerOptions, TArray<TSoftObjectPtr<UTexture>>& OutRuntimeReferencedTextures)
 {
 	UE_LOG(LogMutable, Log, TEXT("Started Customizable Object Export %s."), *Object->GetName());
 
@@ -1490,20 +1389,12 @@ mu::NodePtr FCustomizableObjectCompiler::Export(UCustomizableObject* Object, con
 		return nullptr;
 	}
 
-	// Ensure the images are converted
-	if (GenerationContext.ArrayTextureUnrealToMutableTask.Num() > 0)
-	{
-		ArrayTextureUnrealToMutableTask.Insert(GenerationContext.ArrayTextureUnrealToMutableTask, ArrayTextureUnrealToMutableTask.Num());
-		PendingTexturesToLoad = true;
-		UpdatePendingTextureConversion(false);
-	}
-
 	// Pass out the references textures
-	OutReferencedTextures.Empty();
-	for (const TPair<TSoftObjectPtr<UTexture>, FMutableGraphGenerationContext::FGeneratedPassThroughTexture>& Pair : GenerationContext.PassThroughTextureMap)
+	OutRuntimeReferencedTextures.Empty();
+	for (const TPair<TSoftObjectPtr<UTexture>, FMutableGraphGenerationContext::FGeneratedReferencedTexture>& Pair : GenerationContext.RuntimeReferencedTextureMap)
 	{
-		check(Pair.Value.ID == OutReferencedTextures.Num());
-		OutReferencedTextures.Add(Pair.Key);
+		check(Pair.Value.ID == OutRuntimeReferencedTextures.Num());
+		OutRuntimeReferencedTextures.Add(Pair.Key);
 	}
 
 	return MutableRoot;
diff --git a/Engine/Plugins/Experimental/Mutable/Source/CustomizableObjectEditor/Private/MuCOE/CustomizableObjectCompiler.h b/Engine/Plugins/Experimental/Mutable/Source/CustomizableObjectEditor/Private/MuCOE/CustomizableObjectCompiler.h
--- a/Engine/Plugins/Experimental/Mutable/Source/CustomizableObjectEditor/Private/MuCOE/CustomizableObjectCompiler.h
+++ b/Engine/Plugins/Experimental/Mutable/Source/CustomizableObjectEditor/Private/MuCOE/CustomizableObjectCompiler.h
@@ -178,9 +178,6 @@ private:
 	* @return nothing */
 	void SaveCODerivedData(bool bShowNotification);
 
-	/** When compiling a CO in the editor, flag to know when the Unreal textures have been converted to Mutable textures */
-	bool PendingTexturesToLoad;
-
 	/** When compiling a CO in the editor, flag to know if there's a mutable task pending to be launched through LaunchMutableCompile */
 	bool CompilationLaunchPending;
 
@@ -198,13 +195,6 @@ private:
 	* @return pointer to element if any found, nullptr otherwise */
 	FAssetData* GetCachedAssetData(const FString& PackageName);
 
-	/** When compiling a CO, performs the texture conversion from Unreal to Mutable. If in editor, this method is called in
-	* FCustomizableObjectCompiler::Tick with time limit, otherwise it is called in FCustomizableObjectCompiler::Compile with
-	* no time limit to convert all textures in a row
-	* @param UseTimeLimit [in] if true, the method will return after TimeLimit seconds have been used to convert textures, if false, all textures will be processed in a row
-	* @return nothing */
-	void UpdatePendingTextureConversion(bool UseTimeLimit);
-
 	/** Helper function to compute the value for Unreal Engine variable s.AsyncLoadingTimeLimit while asynchronous loading is used.
 	* Also assigned to MaxConvertToMutableTextureTime
 	* @return value in milliseconds to use for AsyncLoadingTimeLimit and MaxConvertToMutableTextureTime */
@@ -227,13 +217,10 @@ private:
 	/** Copy of GAsyncLoadingTimeLimit while assets are loaded, previous value is restored after asset load */
 	float CurrentGAsyncLoadingTimeLimit;
 
-	/** Array with the textures that need to be asynchronously converted from Unreal to Mutable */
-	TArray<FTextureUnrealToMutableTask> ArrayTextureUnrealToMutableTask;
-
 	/** Counter to know how many of the textures in ArrayTextureUnrealToMutableTask have been converted from Unreal to Mutable */
 	int32 CompletedUnrealToMutableTask;
 
-	/** Time threshold used in UpdatePendingTextureConversion to stop converting textures from Unreal to Mutable until next tick */
+	/** Time threshold used to prevent doing too much work in one tick */
 	float MaxConvertToMutableTextureTime;
 
 	// Stores the only option of an Int Param that should be compiled
diff --git a/Engine/Plugins/Experimental/Mutable/Source/CustomizableObjectEditor/Private/MuCOE/GenerateMutableSource/GenerateMutableSource.cpp b/Engine/Plugins/Experimental/Mutable/Source/CustomizableObjectEditor/Private/MuCOE/GenerateMutableSource/GenerateMutableSource.cpp
--- a/Engine/Plugins/Experimental/Mutable/Source/CustomizableObjectEditor/Private/MuCOE/GenerateMutableSource/GenerateMutableSource.cpp
+++ b/Engine/Plugins/Experimental/Mutable/Source/CustomizableObjectEditor/Private/MuCOE/GenerateMutableSource/GenerateMutableSource.cpp
@@ -50,6 +50,7 @@
 #include "PhysicsEngine/PhysicsAsset.h"
 #include "PlatformInfo.h"
 #include "MuCO/CustomizableObjectPrivate.h"
+#include "Math/NumericLimits.h"
 
 #define LOCTEXT_NAMESPACE "CustomizableObjectEditor"
 
@@ -1687,6 +1688,73 @@ int32 GetMaxTextureSize(const UTexture2D* ReferenceTexture, const FMutableGraphG
 }
 
 
+mu::Ptr<mu::Image> GenerateImageConstant(UTexture* Texture, FMutableGraphGenerationContext& GenerationContext, bool bIsReference)
+{
+	if (!Texture)
+	{
+		return nullptr;
+	}
+
+	bool bForceLoad = false;
+	bool bIsCompileTime = false;
+	if (!bIsReference)
+	{
+		bForceLoad = true;
+		if (GenerationContext.Options.OptimizationLevel == 0)
+		{
+			bIsCompileTime = false;
+		}
+		else
+		{
+			bIsCompileTime = true;
+		}
+	}
+
+	TMap<TSoftObjectPtr<UTexture>, FMutableGraphGenerationContext::FGeneratedReferencedTexture>& TextureMap = bIsCompileTime
+		? GenerationContext.CompileTimeTextureMap
+		: GenerationContext.RuntimeReferencedTextureMap;
+
+	FMutableGraphGenerationContext::FGeneratedReferencedTexture InvalidEntry;
+	InvalidEntry.ID = TNumericLimits<uint32>::Max();
+	FMutableGraphGenerationContext::FGeneratedReferencedTexture& Entry = TextureMap.FindOrAdd(Texture,InvalidEntry);
+
+	if (Entry.ID == TNumericLimits<uint32>::Max())
+	{
+		Entry.ID = TextureMap.Num()-1;
+	}
+
+	// Create a descriptor for the image but fill it only if it is not a true passthrough image. Otherwise we want it empty.
+	mu::FImageDesc ImageDesc;
+	if (bForceLoad)
+	{
+		ImageDesc.m_size[0] = Texture->Source.GetSizeX();
+		ImageDesc.m_size[1] = Texture->Source.GetSizeY();
+		ImageDesc.m_lods = Texture->Source.GetNumMips();
+
+		mu::EImageFormat MutableFormat = mu::EImageFormat::IF_RGBA_UBYTE;
+		ETextureSourceFormat SourceFormat = Texture->Source.GetFormat();
+		switch (SourceFormat)
+		{
+		case ETextureSourceFormat::TSF_G8:
+		case ETextureSourceFormat::TSF_G16:
+		case ETextureSourceFormat::TSF_R16F:
+		case ETextureSourceFormat::TSF_R32F:
+			MutableFormat = mu::EImageFormat::IF_L_UBYTE;
+			break;
+
+		default:
+			break;
+		}
+
+		ImageDesc.m_format = MutableFormat;
+	}
+
+	// Compile-time references that are left should be resolved immediately (should only happen in editor).
+	mu::Ptr<mu::Image> Result = mu::Image::CreateAsReference(Entry.ID, ImageDesc, bForceLoad);
+	return Result;
+}
+
+
 void AddSocketTagsToMesh(const USkeletalMesh* SourceMesh, mu::MeshPtr MutableMesh, FMutableGraphGenerationContext& GenerationContext)
 {
 	for (int32 SocketIndex = 0; SocketIndex < SourceMesh->NumSockets(); ++SocketIndex)
diff --git a/Engine/Plugins/Experimental/Mutable/Source/CustomizableObjectEditor/Private/MuCOE/GenerateMutableSource/GenerateMutableSource.h b/Engine/Plugins/Experimental/Mutable/Source/CustomizableObjectEditor/Private/MuCOE/GenerateMutableSource/GenerateMutableSource.h
--- a/Engine/Plugins/Experimental/Mutable/Source/CustomizableObjectEditor/Private/MuCOE/GenerateMutableSource/GenerateMutableSource.h
+++ b/Engine/Plugins/Experimental/Mutable/Source/CustomizableObjectEditor/Private/MuCOE/GenerateMutableSource/GenerateMutableSource.h
@@ -238,37 +238,6 @@ inline uint32 GetTypeHash(const FGeneratedImageKey& Key)
 }
 
 
-/** Struct describing the conversion task for each Unreal texture that needs to be converted to Mutable format */
-struct FTextureUnrealToMutableTask
-{
-	/** Default constructor */
-	FTextureUnrealToMutableTask(mu::NodeImageConstantPtr ImageNodeParameter, UTexture2D* TextureParameter, const UCustomizableObjectNode* NodeParameter, bool bIsNormalCompositeParameter = false)
-		: ImageNode(ImageNodeParameter)
-		, Texture(TextureParameter)
-		, Node(NodeParameter)
-		, bIsNormalComposite(bIsNormalCompositeParameter)
-	{}
-
-	/** Table Node constructor */
-	FTextureUnrealToMutableTask(mu::TablePtr TableParameter, UTexture2D* TextureParameter, const UCustomizableObjectNode* NodeParameter, int32 ColumnIndex, int32 RowIndex, bool bIsNormalCompositeParameter = false)
-		: TableNode(TableParameter)
-		, Texture(TextureParameter)
-		, Node(NodeParameter)
-		, TableColumn(ColumnIndex)
-		, TableRow(RowIndex)
-		, bIsNormalComposite(bIsNormalCompositeParameter)
-	{}
-
-	mu::NodeImageConstantPtr ImageNode;
-	mu::TablePtr TableNode;
-	TObjectPtr<UTexture2D> Texture;
-	const UCustomizableObjectNode* Node;
-	int32 TableColumn;
-	int32 TableRow;
-	bool bIsNormalComposite;	
-};
-
-
 struct FPoseBoneData
 {
 	TArray<FName> ArrayBoneName;
@@ -560,13 +529,19 @@ struct FMutableGraphGenerationContext
 	TMap<FGeneratedImageKey, mu::NodeImagePtr> GeneratedImages;
 
 	/** Data stored per-generated passthrough texture. */
-	struct FGeneratedPassThroughTexture
+	struct FGeneratedReferencedTexture
 	{
 		uint32 ID;
-		mu::FImageDesc ImageDesc;
+		//mu::FImageDesc ImageDesc;
 	};
-	// Cache of pass-through images and their IDs used in the core to indentify them
-	TMap<TSoftObjectPtr<UTexture>, FGeneratedPassThroughTexture> PassThroughTextureMap;
+
+	// Cache of runtime pass-through images and their IDs used in the core to indentify them.
+	// These textures will remain as external references even in optimized models.
+	TMap<TSoftObjectPtr<UTexture>, FGeneratedReferencedTexture> RuntimeReferencedTextureMap;
+
+	// Cache of runtime pass-through images and their IDs used in the core to indentify them
+	// These textures will become mutable images in the compiled model.
+	TMap<TSoftObjectPtr<UTexture>, FGeneratedReferencedTexture> CompileTimeTextureMap;
 
     // Global morph selection overrides.
     TArray<FRealTimeMorphSelectionOverride> RealTimeMorphTargetsOverrides;
@@ -747,12 +722,6 @@ struct FMutableGraphGenerationContext
 	// Set of all the guids of all the child CustomizableObjects in the compilation
 	TSet<FGuid> CustomizableObjectGuidsInCompilation;
 
-	/** Array with the conversion tasks for each Unreal texture that needs to be converted to Mutable format for a particular Customizable Object.
-	* Multiple tasks for the same texture may be added here when referenced from different LODs or nodes, however when precessing the array it'll
-	* make sure to avoid converting the same texture more than once.
-	*/
-	TArray<FTextureUnrealToMutableTask> ArrayTextureUnrealToMutableTask;
-
 	/** Stores the physics assets gathered from the SkeletalMesh nodes during compilation, to be used in mesh generation in-game */
 	TMap<FString, TSoftObjectPtr<UPhysicsAsset>> PhysicsAssetMap;
 
@@ -880,3 +849,5 @@ int32 ComputeLODBias(const FMutableGraphGenerationContext& GenerationContext, co
 
 
 int32 GetMaxTextureSize(const UTexture2D* ReferenceTexture, const FMutableGraphGenerationContext& GenerationContext);
+
+mu::Ptr<mu::Image> GenerateImageConstant( UTexture*, FMutableGraphGenerationContext&, bool bIsReference);
diff --git a/Engine/Plugins/Experimental/Mutable/Source/CustomizableObjectEditor/Private/MuCOE/GenerateMutableSource/GenerateMutableSourceGroupProjector.cpp b/Engine/Plugins/Experimental/Mutable/Source/CustomizableObjectEditor/Private/MuCOE/GenerateMutableSource/GenerateMutableSourceGroupProjector.cpp
--- a/Engine/Plugins/Experimental/Mutable/Source/CustomizableObjectEditor/Private/MuCOE/GenerateMutableSource/GenerateMutableSourceGroupProjector.cpp
+++ b/Engine/Plugins/Experimental/Mutable/Source/CustomizableObjectEditor/Private/MuCOE/GenerateMutableSource/GenerateMutableSourceGroupProjector.cpp
@@ -552,8 +552,11 @@ bool GenerateMutableSourceGroupProjector(const UEdGraphPin* Pin, FMutableGraphGe
 			{
 				if (ArrayOptionImage[SelectorIndex].OptionImage)
 				{
+					mu::Ptr<mu::Image> ImageConstant = GenerateImageConstant(ArrayOptionImage[SelectorIndex].OptionImage, GenerationContext, false);
+
 					mu::NodeImageConstantPtr ImageNode = new mu::NodeImageConstant();
-					GenerationContext.ArrayTextureUnrealToMutableTask.Add(FTextureUnrealToMutableTask(ImageNode, ArrayOptionImage[SelectorIndex].OptionImage, ProjParamNode));
+					ImageNode->SetValue(ImageConstant.get());
+
 					SwitchNode->SetOption(SelectorIndex, ImageNode);
 				}
 				else
diff --git a/Engine/Plugins/Experimental/Mutable/Source/CustomizableObjectEditor/Private/MuCOE/GenerateMutableSource/GenerateMutableSourceImage.cpp b/Engine/Plugins/Experimental/Mutable/Source/CustomizableObjectEditor/Private/MuCOE/GenerateMutableSource/GenerateMutableSourceImage.cpp
--- a/Engine/Plugins/Experimental/Mutable/Source/CustomizableObjectEditor/Private/MuCOE/GenerateMutableSource/GenerateMutableSourceImage.cpp
+++ b/Engine/Plugins/Experimental/Mutable/Source/CustomizableObjectEditor/Private/MuCOE/GenerateMutableSource/GenerateMutableSourceImage.cpp
@@ -48,7 +48,6 @@
 #include "MuT/NodeImageTable.h"
 #include "MuT/NodeImageTransform.h"
 #include "MuT/NodeImageVariation.h"
-#include "MuT/NodeImageReference.h"
 #include "MuT/NodeImageConstant.h"
 #include "MuT/NodeScalarConstant.h"
 
@@ -165,35 +164,11 @@ mu::NodeImagePtr GenerateMutableSourceImage(const UEdGraphPin* Pin, FMutableGrap
 			}
 			else
 			{
-				if (GenerationContext.Options.OptimizationLevel==0)
-				{
-					// In the "None" optimization level, we don't process any image, and we set them all as image references
-					// that will be loaded at instance generation time: fast compilation x slow generation.
-					FMutableGraphGenerationContext::FGeneratedPassThroughTexture* FoundIndex = GenerationContext.PassThroughTextureMap.Find(BaseTexture);
-					FMutableGraphGenerationContext::FGeneratedPassThroughTexture NewEntry;
-
-					if (!FoundIndex)
-					{
-						FoundIndex = &NewEntry;
-						NewEntry.ID = GenerationContext.PassThroughTextureMap.Num();
-						NewEntry.ImageDesc.m_size[0] = BaseTexture->Source.GetSizeX();
-						NewEntry.ImageDesc.m_size[1] = BaseTexture->Source.GetSizeY();
-						NewEntry.ImageDesc.m_lods = BaseTexture->Source.GetNumMips();
-						NewEntry.ImageDesc.m_format = mu::EImageFormat::IF_RGBA_UBYTE; //TODO: it cannot be known without actually loading the image, which we don't want to do here.
-						GenerationContext.PassThroughTextureMap.Add(BaseTexture, NewEntry);
-					}
+				mu::Ptr<mu::Image> ImageConstant = GenerateImageConstant(BaseTexture, GenerationContext, false);
 
-					mu::Ptr<mu::NodeImageReference> ReferenceImageNode = new mu::NodeImageReference;
-					ReferenceImageNode->SetImageReference(FoundIndex->ID,FoundIndex->ImageDesc);
-					ReferenceImageNode->SetForceLoad(true);
-					ImageNode = ReferenceImageNode;
-				}
-				else
-				{
-					mu::Ptr<mu::NodeImageConstant> ConstantImageNode = new mu::NodeImageConstant();
-					ImageNode = ConstantImageNode;
-					GenerationContext.ArrayTextureUnrealToMutableTask.Add(FTextureUnrealToMutableTask(ConstantImageNode, BaseTexture, Node));
-				}
+				mu::Ptr<mu::NodeImageConstant> ReferenceImageNode = new mu::NodeImageConstant;
+				ReferenceImageNode->SetValue( ImageConstant.get() );
+				ImageNode = ReferenceImageNode;
 
 				GenerationContext.GeneratedImages.Add(imageKey, ImageNode);
 			}
@@ -270,11 +245,7 @@ mu::NodeImagePtr GenerateMutableSourceImage(const UEdGraphPin* Pin, FMutableGrap
 		Result = ImageNode;
 
 		UTexture2D* Texture = TypedNodeMesh->FindTextureForPin(Pin);
-
-		if (Texture)
-		{
-			GenerationContext.ArrayTextureUnrealToMutableTask.Add(FTextureUnrealToMutableTask(ImageNode, Texture, Node));
-		}
+		ImageNode->SetValue( GenerateImageConstant(Texture, GenerationContext, false).get() );
 
 		Result = ResizeToMaxTextureSize(MaxTextureSize, Texture, ImageNode);
 	}
@@ -898,24 +869,10 @@ mu::NodeImagePtr GenerateMutableSourceImage(const UEdGraphPin* Pin, FMutableGrap
 		UTexture2D* BaseTexture = TypedNodePassThroughTex->Texture;
 		if (BaseTexture)
 		{
-			FMutableGraphGenerationContext::FGeneratedPassThroughTexture* FoundIndex = GenerationContext.PassThroughTextureMap.Find(BaseTexture);
-			FMutableGraphGenerationContext::FGeneratedPassThroughTexture NewEntry;
-
-			if (!FoundIndex)
-			{
-				FoundIndex = &NewEntry;
-				NewEntry.ID = GenerationContext.PassThroughTextureMap.Num();
-				NewEntry.ImageDesc.m_size[0] = BaseTexture->Source.GetSizeX();
-				NewEntry.ImageDesc.m_size[1] = BaseTexture->Source.GetSizeY();
-				NewEntry.ImageDesc.m_lods = BaseTexture->Source.GetNumMips();
-				NewEntry.ImageDesc.m_format = mu::EImageFormat::IF_RGBA_UBYTE; //TODO: it cannot be known without actually loading the image, which we don't want to do here.
-				GenerationContext.PassThroughTextureMap.Add(BaseTexture, NewEntry);
-			}
-
-			mu::Ptr<mu::NodeImageReference> ImageNode = new mu::NodeImageReference;
-			ImageNode->SetImageReference(FoundIndex->ID, FoundIndex->ImageDesc);
-
+			mu::NodeImageConstantPtr ImageNode = new mu::NodeImageConstant();
 			Result = ImageNode;
+
+			ImageNode->SetValue(GenerateImageConstant(BaseTexture, GenerationContext, true).get());
 		}
 		else
 		{
diff --git a/Engine/Plugins/Experimental/Mutable/Source/CustomizableObjectEditor/Private/MuCOE/GenerateMutableSource/GenerateMutableSourceSurface.cpp b/Engine/Plugins/Experimental/Mutable/Source/CustomizableObjectEditor/Private/MuCOE/GenerateMutableSource/GenerateMutableSourceSurface.cpp
--- a/Engine/Plugins/Experimental/Mutable/Source/CustomizableObjectEditor/Private/MuCOE/GenerateMutableSource/GenerateMutableSourceSurface.cpp
+++ b/Engine/Plugins/Experimental/Mutable/Source/CustomizableObjectEditor/Private/MuCOE/GenerateMutableSource/GenerateMutableSourceSurface.cpp
@@ -589,7 +589,8 @@ mu::NodeSurfacePtr GenerateMutableSourceSurface(const UEdGraphPin * Pin, FMutabl
 									UTexture2D* Texture2D = TypedNodeMat->GetImageValue(ImageIndex);
 
 									const mu::NodeImageConstantPtr ConstImageNode = new mu::NodeImageConstant();
-									GenerationContext.ArrayTextureUnrealToMutableTask.Add(FTextureUnrealToMutableTask(ConstImageNode, Texture2D, Node));
+									mu::Ptr<mu::Image> ImageConstant = GenerateImageConstant(Texture2D, GenerationContext, false);
+									ConstImageNode->SetValue(ImageConstant.get());
 
 									return ResizeToMaxTextureSize(Props.MaxTextureSize, Texture2D, ConstImageNode);
 								}
@@ -657,7 +658,11 @@ mu::NodeSurfacePtr GenerateMutableSourceSurface(const UEdGraphPin * Pin, FMutabl
 								UTexture2D* ReferenceCompositeNormalTexture = Cast<UTexture2D>(ReferenceTexture->GetCompositeTexture());
 								if (ReferenceCompositeNormalTexture)
 								{
-									GenerationContext.ArrayTextureUnrealToMutableTask.Add(FTextureUnrealToMutableTask(CompositeNormalImage, ReferenceCompositeNormalTexture, Node, true));
+									// GenerationContext.ArrayTextureUnrealToMutableTask.Add(FTextureUnrealToMutableTask(CompositeNormalImage, ReferenceCompositeNormalTexture, Node, true));
+									// TODO: The normal composite part is not propagated, so it will be unsupported. Create a task that performs the required transforms at mutable image level, and add the right operations here
+									// instead of propagating the flag and doing them on unreal-convert.
+									mu::Ptr<mu::Image> ImageConstant = GenerateImageConstant(ReferenceCompositeNormalTexture, GenerationContext, false);
+									CompositeNormalImage->SetValue(ImageConstant.get());
 
 									mu::NodeImageMipmapPtr NormalCompositeMipmapImage = new mu::NodeImageMipmap();
 									NormalCompositeMipmapImage->SetSource(CompositeNormalImage);
diff --git a/Engine/Plugins/Experimental/Mutable/Source/CustomizableObjectEditor/Private/MuCOE/GenerateMutableSource/GenerateMutableSourceTable.cpp b/Engine/Plugins/Experimental/Mutable/Source/CustomizableObjectEditor/Private/MuCOE/GenerateMutableSource/GenerateMutableSourceTable.cpp
--- a/Engine/Plugins/Experimental/Mutable/Source/CustomizableObjectEditor/Private/MuCOE/GenerateMutableSource/GenerateMutableSourceTable.cpp
+++ b/Engine/Plugins/Experimental/Mutable/Source/CustomizableObjectEditor/Private/MuCOE/GenerateMutableSource/GenerateMutableSourceTable.cpp
@@ -419,59 +419,17 @@ bool FillTableColumn(const UCustomizableObjectNodeTable* TableNode,	mu::TablePtr
 				CurrentColumn = MutableTable->AddColumn(ColumnName, mu::ETableColumnType::Image);
 			}
 
-			if (TableNode->GetColumnImageMode(PinName) == ETableTextureType::PASSTHROUGH_TEXTURE)
+			bool bIsPassthroughTexture = TableNode->GetColumnImageMode(PinName) == ETableTextureType::PASSTHROUGH_TEXTURE;
+			if (!bIsPassthroughTexture)
 			{
-				FMutableGraphGenerationContext::FGeneratedPassThroughTexture* FoundIndex = GenerationContext.PassThroughTextureMap.Find(Texture);
-				FMutableGraphGenerationContext::FGeneratedPassThroughTexture NewEntry;
-
-				if (!FoundIndex)
+				if (Texture)
 				{
-					FoundIndex = &NewEntry;
-					NewEntry.ID = GenerationContext.PassThroughTextureMap.Num();
-					NewEntry.ImageDesc.m_size[0] = Texture->Source.GetSizeX();
-					NewEntry.ImageDesc.m_size[1] = Texture->Source.GetSizeY();
-					NewEntry.ImageDesc.m_lods = Texture->Source.GetNumMips();
-					NewEntry.ImageDesc.m_format = mu::EImageFormat::IF_RGBA_UBYTE; //TODO: it cannot be known without actually loading the image, which we don't want to do here.
-					GenerationContext.PassThroughTextureMap.Add(Texture, NewEntry);
+					GenerationContext.AddParticipatingObject(*Texture);
 				}
-
-				mu::Ptr<mu::ResourceProxyMemory<mu::Image>> Proxy = new mu::ResourceProxyMemory<mu::Image>(mu::Image::CreateAsReference(FoundIndex->ID,FoundIndex->ImageDesc, false));
-				MutableTable->SetCell(CurrentColumn, RowIdx, Proxy.get());
 			}
-			else 
-			{
-				if (Texture2D)
-				{
-					GenerationContext.AddParticipatingObject(*Texture2D);
-				}
-
-				if (GenerationContext.Options.OptimizationLevel == 0)
-				{
-					// In the "None" optimization level, we don't process any image, and we set them all as image references
-					// that will be loaded at instance generation time: fast compilation x slow generation.
-					FMutableGraphGenerationContext::FGeneratedPassThroughTexture* FoundIndex = GenerationContext.PassThroughTextureMap.Find(Texture);
-					FMutableGraphGenerationContext::FGeneratedPassThroughTexture NewEntry;
-
-					if (!FoundIndex)
-					{
-						FoundIndex = &NewEntry;
-						NewEntry.ID = GenerationContext.PassThroughTextureMap.Num();
-						NewEntry.ImageDesc.m_size[0] = Texture->Source.GetSizeX();
-						NewEntry.ImageDesc.m_size[1] = Texture->Source.GetSizeY();
-						NewEntry.ImageDesc.m_lods = Texture->Source.GetNumMips();
-						NewEntry.ImageDesc.m_format = mu::EImageFormat::IF_RGBA_UBYTE; //TODO: it cannot be known without actually loading the image, which we don't want to do here.
-						GenerationContext.PassThroughTextureMap.Add(Texture, NewEntry);
-					}
-
-					mu::Ptr<mu::ResourceProxyMemory<mu::Image>> Proxy = new mu::ResourceProxyMemory<mu::Image>(mu::Image::CreateAsReference(FoundIndex->ID, FoundIndex->ImageDesc, true));
-					MutableTable->SetCell(CurrentColumn, RowIdx, Proxy.get());
-				}
-				else
-				{
-					GenerationContext.ArrayTextureUnrealToMutableTask.Add(FTextureUnrealToMutableTask(MutableTable, Texture2D, TableNode, CurrentColumn, RowIdx));
-				}
 
-			}
+			mu::Ptr<mu::ResourceProxyMemory<mu::Image>> Proxy = new mu::ResourceProxyMemory<mu::Image>(GenerateImageConstant(Texture, GenerationContext, bIsPassthroughTexture));
+			MutableTable->SetCell(CurrentColumn, RowIdx, Proxy.get());
 		}
 
 		else if (SoftObjectProperty->PropertyClass->IsChildOf(UMaterialInstance::StaticClass()))
@@ -594,7 +552,9 @@ bool FillTableColumn(const UCustomizableObjectNodeTable* TableNode,	mu::TablePtr
 					LogRowGenerationMessage(TableNode, DataTablePtr, GenerationContext, msg, RowName);
 				}
 
-				GenerationContext.ArrayTextureUnrealToMutableTask.Add(FTextureUnrealToMutableTask(MutableTable, ParameterTexture, TableNode, ColumnIndex, RowIdx));
+				bool bIsPassthroughTexture = false;
+				mu::Ptr<mu::ResourceProxyMemory<mu::Image>> Proxy = new mu::ResourceProxyMemory<mu::Image>(GenerateImageConstant(ParameterTexture, GenerationContext, bIsPassthroughTexture));
+				MutableTable->SetCell(ColumnIndex, RowIdx, Proxy.get());
 
 				return true;
 			}
diff --git a/Engine/Plugins/Experimental/Mutable/Source/CustomizableObjectEditor/Private/MuCOE/SMutableObjectViewer.cpp b/Engine/Plugins/Experimental/Mutable/Source/CustomizableObjectEditor/Private/MuCOE/SMutableObjectViewer.cpp
--- a/Engine/Plugins/Experimental/Mutable/Source/CustomizableObjectEditor/Private/MuCOE/SMutableObjectViewer.cpp
+++ b/Engine/Plugins/Experimental/Mutable/Source/CustomizableObjectEditor/Private/MuCOE/SMutableObjectViewer.cpp
@@ -278,6 +278,7 @@ void SMutableObjectViewer::CompileMutableCodePressed()
 	// Do the compilation to Mutable Code synchronously.
 	TSharedPtr<FCustomizableObjectCompileRunnable> CompileTask = MakeShareable(new FCustomizableObjectCompileRunnable(RootNode));
 	CompileTask->Options = CompileOptions;
+	CompileTask->ReferencedTextures = ReferencedTextures;
 	CompileTask->Init();
 	CompileTask->Run();
 
diff --git a/Engine/Plugins/Experimental/Mutable/Source/MutableRuntime/Private/MuR/CodeRunnerTasked.cpp b/Engine/Plugins/Experimental/Mutable/Source/MutableRuntime/Private/MuR/CodeRunnerTasked.cpp
--- a/Engine/Plugins/Experimental/Mutable/Source/MutableRuntime/Private/MuR/CodeRunnerTasked.cpp
+++ b/Engine/Plugins/Experimental/Mutable/Source/MutableRuntime/Private/MuR/CodeRunnerTasked.cpp
@@ -696,7 +696,8 @@ namespace mu
 			Base = Formatted;
 		}
 
-		if (Blended && InitialFormat != Blended->GetFormat() && !Args.flags)
+		bool bMustHaveSameFormat = !(Args.flags & (OP::ImageLayerArgs::F_BASE_RGB_FROM_ALPHA | OP::ImageLayerArgs::F_BLENDED_RGB_FROM_ALPHA));
+		if (Blended && InitialFormat != Blended->GetFormat() && bMustHaveSameFormat)
 		{
 			Ptr<Image> Formatted = Runner->CreateImage(Blended->GetSizeX(), Blended->GetSizeY(), Blended->GetLODCount(), Base->GetFormat(), EInitializationType::NotInitialized);
 			bool bSuccess = false;
@@ -1616,6 +1617,13 @@ namespace mu
 			Lods = Image::GetMipmapCount(destSize[0], destSize[1]);
 		}
 
+		if (Base->IsReference())
+		{
+			// We are trying to resize an external reference. This shouldn't happen, but be deffensive.
+			Runner->StoreImage(Op, Base);
+			return false;
+		}
+
 		Result = Runner->CreateImage( destSize[0], destSize[1], Lods, Base->GetFormat(), EInitializationType::NotInitialized );
 
 		ImagePixelFormatFunc = Runner->m_pSystem->ImagePixelFormatOverride;
diff --git a/Engine/Plugins/Experimental/Mutable/Source/MutableRuntime/Private/MuR/Image.cpp b/Engine/Plugins/Experimental/Mutable/Source/MutableRuntime/Private/MuR/Image.cpp
--- a/Engine/Plugins/Experimental/Mutable/Source/MutableRuntime/Private/MuR/Image.cpp
+++ b/Engine/Plugins/Experimental/Mutable/Source/MutableRuntime/Private/MuR/Image.cpp
@@ -296,6 +296,12 @@ namespace mu
 		return m_flags & EImageFlags::IF_IS_REFERENCE;
 	}
 
+	//---------------------------------------------------------------------------------------------
+	bool Image::IsForceLoad() const
+	{
+		return m_flags & EImageFlags::IF_IS_FORCELOAD;
+	}
+
 	//---------------------------------------------------------------------------------------------
 	uint32 Image::GetReferencedTexture() const
 	{
diff --git a/Engine/Plugins/Experimental/Mutable/Source/MutableRuntime/Public/MuR/Image.h b/Engine/Plugins/Experimental/Mutable/Source/MutableRuntime/Public/MuR/Image.h
--- a/Engine/Plugins/Experimental/Mutable/Source/MutableRuntime/Public/MuR/Image.h
+++ b/Engine/Plugins/Experimental/Mutable/Source/MutableRuntime/Public/MuR/Image.h
@@ -96,6 +96,9 @@ namespace mu
 		/** Return true if this is a reference to an engine image. */
 		bool IsReference() const;
 
+		/** If true, this is a reference that must be resolved at compile time. */
+		bool IsForceLoad() const;
+
 		/** Return the id of the engine referenced texture. Only valid if IsReference. */
 		uint32 GetReferencedTexture() const;
 
diff --git a/Engine/Plugins/Experimental/Mutable/Source/MutableTools/Private/MuT/AST.h b/Engine/Plugins/Experimental/Mutable/Source/MutableTools/Private/MuT/AST.h
--- a/Engine/Plugins/Experimental/Mutable/Source/MutableTools/Private/MuT/AST.h
+++ b/Engine/Plugins/Experimental/Mutable/Source/MutableTools/Private/MuT/AST.h
@@ -178,6 +178,8 @@ namespace mu
     typedef TSet<Ptr<ASTOp>> ASTOpSet;
 
 
+	typedef TFunction<Ptr<Image>(int32)> FReferencedResourceFunc;
+
     //! Detailed optimization flags
     struct FModelOptimizationOptions
     {
@@ -206,6 +208,9 @@ namespace mu
 
         //! Ratio used to decide if it is worth to generate a crop operation
         float MinRLECompressionGain = 1.2f;
+
+		//
+		FReferencedResourceFunc ReferencedResourceProvider;
     };
 
 
diff --git a/Engine/Plugins/Experimental/Mutable/Source/MutableTools/Private/MuT/ASTOpReferenceResource.cpp b/Engine/Plugins/Experimental/Mutable/Source/MutableTools/Private/MuT/ASTOpReferenceResource.cpp
--- a/Engine/Plugins/Experimental/Mutable/Source/MutableTools/Private/MuT/ASTOpReferenceResource.cpp
+++ b/Engine/Plugins/Experimental/Mutable/Source/MutableTools/Private/MuT/ASTOpReferenceResource.cpp
@@ -2,9 +2,9 @@
 
 #include "MuT/ASTOpReferenceResource.h"
 
+#include "MuT/ASTOpConstantResource.h"
 #include "Containers/Array.h"
 #include "HAL/PlatformMath.h"
-#include "Hash/CityHash.h"
 #include "Misc/AssertionMacros.h"
 #include "MuR/Types.h"
 #include "MuT/StreamsPrivate.h"
@@ -123,12 +123,6 @@ namespace mu
 	}
 
 
-	//-------------------------------------------------------------------------------------------------
-	ASTOpReferenceResource::~ASTOpReferenceResource()
-	{
-	}
-
-
 	//-------------------------------------------------------------------------------------------------
 	mu::Ptr<ImageSizeExpression> ASTOpReferenceResource::GetImageSizeExpression() const
 	{
@@ -138,4 +132,36 @@ namespace mu
 		return pRes;
 	}
 
+
+	Ptr<ASTOp> ASTOpReferenceResource::OptimiseSemantic(const FModelOptimizationOptions& options, int32 Pass) const
+	{
+		mu::Ptr<ASTOp> NewOp;
+
+		switch (type)
+		{
+
+		case OP_TYPE::IM_REFERENCE:
+		{
+			// If we are in reference resolution stage
+			if (Pass>=2 && bForceLoad)
+			{
+				check(options.ReferencedResourceProvider);
+
+				Ptr<Image> ConstantImage = options.ReferencedResourceProvider(ID);
+
+				Ptr<ASTOpConstantResource> ConstantOp = new ASTOpConstantResource;
+				ConstantOp->type = OP_TYPE::IM_CONSTANT;
+				ConstantOp->SetValue( ConstantImage.get(), options.bUseDiskCache );
+				NewOp = ConstantOp;
+			}
+			break;
+		}
+
+		default:
+			checkf(false, TEXT("Instruction not supported"));
+		}
+
+		return NewOp;
+	}
+
 }
\ No newline at end of file
diff --git a/Engine/Plugins/Experimental/Mutable/Source/MutableTools/Private/MuT/ASTOpReferenceResource.h b/Engine/Plugins/Experimental/Mutable/Source/MutableTools/Private/MuT/ASTOpReferenceResource.h
--- a/Engine/Plugins/Experimental/Mutable/Source/MutableTools/Private/MuT/ASTOpReferenceResource.h
+++ b/Engine/Plugins/Experimental/Mutable/Source/MutableTools/Private/MuT/ASTOpReferenceResource.h
@@ -36,13 +36,8 @@ namespace mu
 
 	public:
 
-		~ASTOpReferenceResource() override;
-
 		// Own interface
 
-		//! Get a hash of the stored value.
-		uint64 GetValueHash() const;
-
 		// ASTOp interface
 		OP_TYPE GetOpType() const override { return type; }
 		void ForEachChild(const TFunctionRef<void(ASTChild&)>) override;
@@ -55,6 +50,7 @@ namespace mu
 		void GetLayoutBlockSize(int* pBlockX, int* pBlockY) override;
 		bool GetNonBlackRect(FImageRect& maskUsage) const override;
 		Ptr<ImageSizeExpression> GetImageSizeExpression() const override;
+		Ptr<ASTOp> OptimiseSemantic(const FModelOptimizationOptions& options, int32 Pass) const override;
 	};
 
 
diff --git a/Engine/Plugins/Experimental/Mutable/Source/MutableTools/Private/MuT/CodeGenerator.cpp b/Engine/Plugins/Experimental/Mutable/Source/MutableTools/Private/MuT/CodeGenerator.cpp
--- a/Engine/Plugins/Experimental/Mutable/Source/MutableTools/Private/MuT/CodeGenerator.cpp
+++ b/Engine/Plugins/Experimental/Mutable/Source/MutableTools/Private/MuT/CodeGenerator.cpp
@@ -55,7 +55,6 @@
 #include "MuT/NodeImageMipmap.h"
 #include "MuT/NodeImageMipmapPrivate.h"
 #include "MuT/NodeImageSwizzlePrivate.h"
-#include "MuT/NodeImageReference.h"
 #include "MuT/NodeLODPrivate.h"
 #include "MuT/NodeMesh.h"
 #include "MuT/NodeMeshClipMorphPlane.h"
diff --git a/Engine/Plugins/Experimental/Mutable/Source/MutableTools/Private/MuT/CodeGenerator_Image.cpp b/Engine/Plugins/Experimental/Mutable/Source/MutableTools/Private/MuT/CodeGenerator_Image.cpp
--- a/Engine/Plugins/Experimental/Mutable/Source/MutableTools/Private/MuT/CodeGenerator_Image.cpp
+++ b/Engine/Plugins/Experimental/Mutable/Source/MutableTools/Private/MuT/CodeGenerator_Image.cpp
@@ -57,7 +57,6 @@
 #include "MuT/NodeImageConditionalPrivate.h"
 #include "MuT/NodeImageConstant.h"
 #include "MuT/NodeImageConstantPrivate.h"
-#include "MuT/NodeImageReferencePrivate.h"
 #include "MuT/NodeImageFormat.h"
 #include "MuT/NodeImageFormatPrivate.h"
 #include "MuT/NodeImageGradient.h"
@@ -162,7 +161,6 @@ namespace mu
 			case NodeImage::EType::Variation: GenerateImage_Variation(Options, Result, static_cast<const NodeImageVariation*>(Node)); break;
 			case NodeImage::EType::NormalComposite: GenerateImage_NormalComposite(Options, Result, static_cast<const NodeImageNormalComposite*>(Node)); break;
 			case NodeImage::EType::Transform: GenerateImage_Transform(Options, Result, static_cast<const NodeImageTransform*>(Node)); break;
-			case NodeImage::EType::Reference: GenerateImage_Reference(Options, Result, static_cast<const NodeImageReference*>(Node)); break;
 			case NodeImage::EType::None: check(false);
 			}
 
@@ -193,86 +191,31 @@ namespace mu
             m_pErrorLog->GetPrivate()->Add( "Constant image not set.", ELMT_WARNING, node.m_errorContext );
         }
 
-		if (Options.ImageLayoutStrategy!= CompilerOptions::TextureLayoutStrategy::None && Options.LayoutToApply)
-		{
-			// We want to generate only a block from the image.
-
-			FIntVector2 SourceImageSize(pImage->GetSizeX(), pImage->GetSizeY());
-
-			int32 BlockIndex = Options.LayoutToApply->FindBlock(Options.LayoutBlockId);
-			check( BlockIndex>=0 );
-
-			// Block in layout grid units
-			box< UE::Math::TIntVector2<uint16> > RectInCells;
-			Options.LayoutToApply->GetBlock
-			(
-				BlockIndex,
-				&RectInCells.min[0], &RectInCells.min[1],
-				&RectInCells.size[0], &RectInCells.size[1]
-			);
-
-			FIntPoint grid = Options.LayoutToApply->GetGridSize();
-			grid[0] = FMath::Max(1, grid[0]);
-			grid[1] = FMath::Max(1, grid[1]);
-
-			// Transform to pixels
-			box< UE::Math::TIntVector2<int32> > rect;
-			rect.min[0]  = (RectInCells.min[0]  * SourceImageSize[0]) / grid[0];
-			rect.min[1]  = (RectInCells.min[1]  * SourceImageSize[1]) / grid[1];
-			rect.size[0] = (RectInCells.size[0] * SourceImageSize[0]) / grid[0];
-			rect.size[1] = (RectInCells.size[1] * SourceImageSize[1]) / grid[1];
-
-			// Do we need to crop?
-			if (rect.min[0]!=0 || rect.min[1]!=0 || pImage->GetSizeX() != rect.size[0] || pImage->GetSizeY() != rect.size[1])
-			{
-				// Crop now
-				FImageOperator ImOp = FImageOperator::GetDefault(m_compilerOptions->ImageFormatFunc);
 
-				Ptr<Image> pCropped = new Image(rect.size[0], rect.size[1], 1, pImage->GetFormat(), EInitializationType::NotInitialized);
-				ImOp.ImageCrop(pCropped.get(), m_compilerOptions->ImageCompressionQuality, pImage.get(), rect);
-
-				Ptr<ASTOpConstantResource> op = new ASTOpConstantResource();
-				op->type = OP_TYPE::IM_CONSTANT;
-				op->SetValue(pCropped, m_compilerOptions->OptimisationOptions.bUseDiskCache);
-				Result.op = op;
-			}
-			else
-			{
-				// No need to crop.
-				Ptr<ASTOpConstantResource> op = new ASTOpConstantResource();
-				op->type = OP_TYPE::IM_CONSTANT;
-				op->SetValue(pImage, m_compilerOptions->OptimisationOptions.bUseDiskCache);
-				Result.op = op;
-			}
+		if (pImage->IsReference())
+		{
+			Ptr<ASTOpReferenceResource> ReferenceOp = new ASTOpReferenceResource();
+			ReferenceOp->type = OP_TYPE::IM_REFERENCE;
+			ReferenceOp->ID = pImage->GetReferencedTexture();
+			ReferenceOp->bForceLoad = pImage->IsForceLoad();
+			Result.op = ReferenceOp;
 		}
-        else
-        {
+		else
+		{
 			Ptr<ASTOpConstantResource> op = new ASTOpConstantResource();
 			op->type = OP_TYPE::IM_CONSTANT;
-			op->SetValue( pImage, m_compilerOptions->OptimisationOptions.bUseDiskCache );
+			op->SetValue(pImage, m_compilerOptions->OptimisationOptions.bUseDiskCache);
 			Result.op = op;
 		}
-    }
-
 
-	//---------------------------------------------------------------------------------------------
-	void CodeGenerator::GenerateImage_Reference(const FImageGenerationOptions& Options, FImageGenerationResult& Result, const NodeImageReference* InNode)
-	{
-		const NodeImageReference::Private& node = *InNode->GetPrivate();
-
-		Ptr<ASTOpReferenceResource> ReferenceOp = new ASTOpReferenceResource();
-		ReferenceOp->type = OP_TYPE::IM_REFERENCE;
-		ReferenceOp->ID = node.ImageReferenceID;
-		ReferenceOp->bForceLoad = node.bForceLoad;
-
-		if (Options.ImageLayoutStrategy != CompilerOptions::TextureLayoutStrategy::None && Options.LayoutToApply)
+		if (Options.ImageLayoutStrategy!= CompilerOptions::TextureLayoutStrategy::None && Options.LayoutToApply)
 		{
 			// We want to generate only a block from the image.
 
-			FIntVector2 SourceImageSize(node.ImageDesc.m_size[0], node.ImageDesc.m_size[1]);
+			FIntVector2 SourceImageSize(pImage->GetSizeX(), pImage->GetSizeY());
 
 			int32 BlockIndex = Options.LayoutToApply->FindBlock(Options.LayoutBlockId);
-			check(BlockIndex >= 0);
+			check( BlockIndex>=0 );
 
 			// Block in layout grid units
 			box< UE::Math::TIntVector2<uint16> > RectInCells;
@@ -289,36 +232,24 @@ namespace mu
 
 			// Transform to pixels
 			box< UE::Math::TIntVector2<int32> > rect;
-			rect.min[0] = (RectInCells.min[0] * SourceImageSize[0]) / grid[0];
-			rect.min[1] = (RectInCells.min[1] * SourceImageSize[1]) / grid[1];
+			rect.min[0]  = (RectInCells.min[0]  * SourceImageSize[0]) / grid[0];
+			rect.min[1]  = (RectInCells.min[1]  * SourceImageSize[1]) / grid[1];
 			rect.size[0] = (RectInCells.size[0] * SourceImageSize[0]) / grid[0];
 			rect.size[1] = (RectInCells.size[1] * SourceImageSize[1]) / grid[1];
 
 			// Do we need to crop?
-			if (rect.min[0] != 0 || rect.min[1] != 0 || node.ImageDesc.m_size[0] != rect.size[0] || node.ImageDesc.m_size[1] != rect.size[1])
+			if (rect.min[0]!=0 || rect.min[1]!=0 || pImage->GetSizeX() != rect.size[0] || pImage->GetSizeY() != rect.size[1])
 			{
-				// If need to crop it has to be the case of an image we will force to load.
-				check(node.bForceLoad);
-
 				Ptr<ASTOpImageCrop> CropOp = new ASTOpImageCrop();
-				CropOp->Source = ReferenceOp;
+				CropOp->Source = Result.op;
 				CropOp->Min[0] = rect.min[0];
 				CropOp->Min[1] = rect.min[1];
 				CropOp->Size[0] = rect.size[0];
 				CropOp->Size[1] = rect.size[1];
-
 				Result.op = CropOp;
 			}
-			else
-			{
-				Result.op = ReferenceOp;
-			}
 		}
-		else
-		{
-			Result.op = ReferenceOp;
-		}
-	}
+    }
 
 
     //---------------------------------------------------------------------------------------------
@@ -1764,13 +1695,14 @@ namespace mu
 			{
 				const FTableValue& CellData = node.m_pTable->GetPrivate()->Rows[row].Values[colIndex];
 				ImagePtrConst pImage = nullptr;
-				NodeImagePtr CellImage = nullptr;
 
 				if (Ptr<ResourceProxy<Image>> pProxyImage = CellData.ProxyImage)
 				{
 					pImage = pProxyImage->Get();
 				}
 
+				Ptr<ASTOp> ImageOp;
+
 				if (!pImage)
 				{
 					FString Msg = FString::Printf(TEXT("Table has a missing image in column %d, row %d."), colIndex, row);
@@ -1778,42 +1710,26 @@ namespace mu
 				}
 				else
 				{
-					if (pImage->IsReference())
-					{
-						FImageDesc ImageDesc;
-						ImageDesc.m_format = pImage->GetFormat();
-						ImageDesc.m_lods = pImage->GetLODCount();
-						ImageDesc.m_size[0] = pImage->GetSizeX();
-						ImageDesc.m_size[1] = pImage->GetSizeY();
-
-						Ptr<NodeImageReference> ImageRef = new NodeImageReference();
-						ImageRef->SetImageReference(pImage->GetReferencedTexture(), ImageDesc);
-						ImageRef->SetForceLoad( pImage->m_flags & Image::IF_IS_FORCELOAD );
-
-						CellImage = ImageRef;
-					}
-					else
-					{
-						if (node.MaxTextureSize > 0 && (node.MaxTextureSize < pImage->GetSizeX() || node.MaxTextureSize < pImage->GetSizeY()))
-						{
-							float Factor = FMath::Min(node.MaxTextureSize / (float)(pImage->GetSizeX()), node.MaxTextureSize / (float)(pImage->GetSizeY()));
-
-							FImageOperator ImOp = FImageOperator::GetDefault(m_compilerOptions->ImageFormatFunc);
-							Ptr<Image> ResizedImage = new Image(pImage->GetSizeX() * Factor, pImage->GetSizeY() * Factor, pImage->GetLODCount(), pImage->GetFormat(), EInitializationType::NotInitialized);
-							ImOp.ImageResizeLinear(ResizedImage.get(), m_compilerOptions->ImageCompressionQuality, pImage.get());
-							pImage = ResizedImage;
-						}
-
-						NodeImageConstantPtr ImageConst = new NodeImageConstant();
-						ImageConst->SetValue(pImage);
-
-						CellImage = ImageConst;
-					}
+					NodeImageConstantPtr ImageConst = new NodeImageConstant();
+					ImageConst->SetValue(pImage.get());
+
+					FImageGenerationResult Result;
+					GenerateImage(Options, Result, ImageConst);
+					ImageOp = Result.op;
+				}
+
+				if (node.MaxTextureSize > 0 && (node.MaxTextureSize < pImage->GetSizeX() || node.MaxTextureSize < pImage->GetSizeY()))
+				{
+					float Factor = FMath::Min(node.MaxTextureSize / (float)(pImage->GetSizeX()), node.MaxTextureSize / (float)(pImage->GetSizeY()));
+					Ptr<ASTOpFixed> op = new ASTOpFixed();
+					op->op.type = OP_TYPE::IM_RESIZE;
+					op->op.args.ImageResize.size[0] = (uint16)pImage->GetSizeX() * Factor;
+					op->op.args.ImageResize.size[1] = (uint16)pImage->GetSizeY() * Factor;
+					op->SetChild(op->op.args.ImageResize.source, ImageOp);
+					ImageOp = op;
 				}
 
-				FImageGenerationResult Result;
-				GenerateImage(Options, Result, CellImage);
-				return Result.op;
+				return ImageOp;
 			});
 	}
 
diff --git a/Engine/Plugins/Experimental/Mutable/Source/MutableTools/Private/MuT/CodeOptimiser.cpp b/Engine/Plugins/Experimental/Mutable/Source/MutableTools/Private/MuT/CodeOptimiser.cpp
--- a/Engine/Plugins/Experimental/Mutable/Source/MutableTools/Private/MuT/CodeOptimiser.cpp
+++ b/Engine/Plugins/Experimental/Mutable/Source/MutableTools/Private/MuT/CodeOptimiser.cpp
@@ -1266,6 +1266,24 @@ namespace mu
 				UE_LOG(LogMutableCore, Verbose, TEXT("(int) %s : %ld"), TEXT("ast size"), int64(ASTOp::CountNodes(roots)));
 				ASTOp::LogHistogram(roots);
 
+				FullOptimiseAST( roots, 1 );
+				UE_LOG(LogMutableCore, Verbose, TEXT("(int) %s : %ld"), TEXT("ast size"), int64(ASTOp::CountNodes(roots)));
+				ASTOp::LogHistogram(roots);
+			}
+
+			// Constant resolution stage: resolve referenced assets.
+			{
+				MUTABLE_CPUPROFILER_SCOPE(ReferenceResolution);
+				FullOptimiseAST(roots, 2);
+			}
+
+			// Main optimisation stage again for data-aware optimizations
+			{
+				MUTABLE_CPUPROFILER_SCOPE(FinalStage);
+				FullOptimiseAST(roots, 0);
+				UE_LOG(LogMutableCore, Verbose, TEXT("(int) %s : %ld"), TEXT("ast size"), int64(ASTOp::CountNodes(roots)));
+				ASTOp::LogHistogram(roots);
+
 				FullOptimiseAST(roots, 1);
 				UE_LOG(LogMutableCore, Verbose, TEXT("(int) %s : %ld"), TEXT("ast size"), int64(ASTOp::CountNodes(roots)));
 				ASTOp::LogHistogram(roots);
@@ -1320,6 +1338,12 @@ namespace mu
 			DuplicatedCodeRemoverAST( roots );
 			UE_LOG(LogMutableCore, Verbose, TEXT("(int) %s : %ld"), TEXT("ast size"), int64(ASTOp::CountNodes(roots)));
 
+			// Constant resolution stage: resolve referenced assets.
+			{
+				MUTABLE_CPUPROFILER_SCOPE(ReferenceResolution);
+				FullOptimiseAST(roots, 2);
+			}
+
 			for ( size_t s=0;  s<m_states.size(); ++s )
 			{
 				UE_LOG(LogMutableCore, Verbose, TEXT(" - constant generator"));
diff --git a/Engine/Plugins/Experimental/Mutable/Source/MutableTools/Private/MuT/Compiler.cpp b/Engine/Plugins/Experimental/Mutable/Source/MutableTools/Private/MuT/Compiler.cpp
--- a/Engine/Plugins/Experimental/Mutable/Source/MutableTools/Private/MuT/Compiler.cpp
+++ b/Engine/Plugins/Experimental/Mutable/Source/MutableTools/Private/MuT/Compiler.cpp
@@ -160,6 +160,13 @@ namespace mu
 	}
 
 
+	//---------------------------------------------------------------------------------------------
+	void CompilerOptions::SetReferencedResourceCallback(const FReferencedResourceFunc& Provider)
+	{
+		m_pD->OptimisationOptions.ReferencedResourceProvider = Provider;
+	}
+
+
     //---------------------------------------------------------------------------------------------
     //---------------------------------------------------------------------------------------------
     //---------------------------------------------------------------------------------------------
diff --git a/Engine/Plugins/Experimental/Mutable/Source/MutableTools/Private/MuT/ImageDescGenerator.cpp b/Engine/Plugins/Experimental/Mutable/Source/MutableTools/Private/MuT/ImageDescGenerator.cpp
--- a/Engine/Plugins/Experimental/Mutable/Source/MutableTools/Private/MuT/ImageDescGenerator.cpp
+++ b/Engine/Plugins/Experimental/Mutable/Source/MutableTools/Private/MuT/ImageDescGenerator.cpp
@@ -16,7 +16,6 @@
 #include "MuT/NodeImageColourMapPrivate.h"
 #include "MuT/NodeImageConditionalPrivate.h"
 #include "MuT/NodeImageConstantPrivate.h"
-#include "MuT/NodeImageReferencePrivate.h"
 #include "MuT/NodeImageFormatPrivate.h"
 #include "MuT/NodeImageGradientPrivate.h"
 #include "MuT/NodeImageInterpolatePrivate.h"
@@ -97,14 +96,6 @@ namespace mu
     }
 
 
-	//---------------------------------------------------------------------------------------------
-	Ptr<ASTOp> ImageDescGenerator::Visit( const NodeImageReference::Private& node )
-	{
-		m_desc = node.ImageDesc;
-		return 0;
-	}
-
-
     //---------------------------------------------------------------------------------------------
     Ptr<ASTOp> ImageDescGenerator::Visit( const NodeImageParameter::Private& )
     {
diff --git a/Engine/Plugins/Experimental/Mutable/Source/MutableTools/Private/MuT/ImageDescGenerator.h b/Engine/Plugins/Experimental/Mutable/Source/MutableTools/Private/MuT/ImageDescGenerator.h
--- a/Engine/Plugins/Experimental/Mutable/Source/MutableTools/Private/MuT/ImageDescGenerator.h
+++ b/Engine/Plugins/Experimental/Mutable/Source/MutableTools/Private/MuT/ImageDescGenerator.h
@@ -30,7 +30,6 @@
 #include "MuT/NodeImageTable.h"
 #include "MuT/NodeImageTransform.h"
 #include "MuT/NodeImageVariation.h"
-#include "MuT/NodeImageReference.h"
 #include "MuT/Visitor.h"
 
 namespace mu { class ASTOp; }
@@ -72,9 +71,7 @@ namespace mu
                                public Visitor<NodeImageMultiLayer::Private, Ptr<ASTOp>, true>,
                                public Visitor<NodeImageInvert::Private, Ptr<ASTOp>, true>,
                                public Visitor<NodeImageNormalComposite::Private, Ptr<ASTOp>, true>,
-                               public Visitor<NodeImageTransform::Private, Ptr<ASTOp>, true>, 
-							   public Visitor<NodeImageReference::Private, Ptr<ASTOp>, true>
-
+                               public Visitor<NodeImageTransform::Private, Ptr<ASTOp>, true>
     {
     public:
 
@@ -104,7 +101,6 @@ namespace mu
         Ptr<ASTOp> Visit( const NodeImageInvert::Private& );
         Ptr<ASTOp> Visit( const NodeImageNormalComposite::Private& );
 		Ptr<ASTOp> Visit( const NodeImageTransform::Private& );
-		Ptr<ASTOp> Visit( const NodeImageReference::Private& );
 
     public:
 
diff --git a/Engine/Plugins/Experimental/Mutable/Source/MutableTools/Private/MuT/NodeImage.cpp b/Engine/Plugins/Experimental/Mutable/Source/MutableTools/Private/MuT/NodeImage.cpp
--- a/Engine/Plugins/Experimental/Mutable/Source/MutableTools/Private/MuT/NodeImage.cpp
+++ b/Engine/Plugins/Experimental/Mutable/Source/MutableTools/Private/MuT/NodeImage.cpp
@@ -30,7 +30,6 @@
 #include "MuT/NodeImageTable.h"
 #include "MuT/NodeImageTransform.h"
 #include "MuT/NodeImageVariation.h"
-#include "MuT/NodeImageReference.h"
 
 
 namespace mu
@@ -95,7 +94,7 @@ namespace mu
         case 32 :  return NodeImageVariation::StaticUnserialise( arch ); break;
         case 33 :  return NodeImageNormalComposite::StaticUnserialise( arch ); break;
 		case 34:  return NodeImageTransform::StaticUnserialise(arch); break;
-		case 35:  return NodeImageReference::StaticUnserialise(arch); break;
+		//case 35:  return NodeImageReference::StaticUnserialise(arch); break;
 		default : check(false);
 		}
 
diff --git a/Engine/Plugins/Experimental/Mutable/Source/MutableTools/Private/MuT/NodeImageConstant.cpp b/Engine/Plugins/Experimental/Mutable/Source/MutableTools/Private/MuT/NodeImageConstant.cpp
--- a/Engine/Plugins/Experimental/Mutable/Source/MutableTools/Private/MuT/NodeImageConstant.cpp
+++ b/Engine/Plugins/Experimental/Mutable/Source/MutableTools/Private/MuT/NodeImageConstant.cpp
@@ -56,9 +56,9 @@ namespace mu
 
 
 	//---------------------------------------------------------------------------------------------
-	void NodeImageConstant::SetValue(Ptr<const Image> pValue)
+	void NodeImageConstant::SetValue(const Image* Value)
 	{
-		m_pD->m_pProxy = new ResourceProxyMemory<Image>(pValue);
+		m_pD->m_pProxy = new ResourceProxyMemory<Image>(Value);
 	}
 
 
diff --git a/Engine/Plugins/Experimental/Mutable/Source/MutableTools/Private/MuT/NodeImageReference.cpp b/Engine/Plugins/Experimental/Mutable/Source/MutableTools/Private/MuT/NodeImageReference.cpp
deleted file mode 100644
diff --git a/Engine/Plugins/Experimental/Mutable/Source/MutableTools/Private/MuT/NodeImageReferencePrivate.h b/Engine/Plugins/Experimental/Mutable/Source/MutableTools/Private/MuT/NodeImageReferencePrivate.h
deleted file mode 100644
diff --git a/Engine/Plugins/Experimental/Mutable/Source/MutableTools/Public/MuT/Compiler.h b/Engine/Plugins/Experimental/Mutable/Source/MutableTools/Public/MuT/Compiler.h
--- a/Engine/Plugins/Experimental/Mutable/Source/MutableTools/Public/MuT/Compiler.h
+++ b/Engine/Plugins/Experimental/Mutable/Source/MutableTools/Public/MuT/Compiler.h
@@ -2,11 +2,12 @@
 
 #pragma once
 
-#include "HAL/PlatformMath.h"
+#include "MuT/AST.h"
 #include "MuR/Ptr.h"
 #include "MuR/RefCounted.h"
 #include "MuR/Image.h"
 #include "Templates/SharedPointer.h"
+#include "HAL/PlatformMath.h"
 
 namespace mu
 {
@@ -75,6 +76,9 @@ namespace mu
 		/** Set an optional pixel conversion function that will be called before any pixel format conversion. */
 		void SetImagePixelFormatOverride(const FImageOperator::FImagePixelFormatFunc&);
 
+		/** */
+		void SetReferencedResourceCallback(const FReferencedResourceFunc&);
+
         //! Different data packing strategies
         enum class TextureLayoutStrategy : uint8
         {
diff --git a/Engine/Plugins/Experimental/Mutable/Source/MutableTools/Public/MuT/NodeImageConstant.h b/Engine/Plugins/Experimental/Mutable/Source/MutableTools/Public/MuT/NodeImageConstant.h
--- a/Engine/Plugins/Experimental/Mutable/Source/MutableTools/Public/MuT/NodeImageConstant.h
+++ b/Engine/Plugins/Experimental/Mutable/Source/MutableTools/Public/MuT/NodeImageConstant.h
@@ -58,7 +58,7 @@ namespace mu
         Ptr<const Image> GetValue() const;
 
         //! Set the image to be output by this node
-        void SetValue( Ptr<const Image> pImage );
+        void SetValue( const Image* );
 
         //! Set the image proxy that will provide the image for this node when necessary
         void SetValue( Ptr<ResourceProxy<Image>> pImage );
diff --git a/Engine/Plugins/Experimental/Mutable/Source/MutableTools/Public/MuT/NodeImageReference.h b/Engine/Plugins/Experimental/Mutable/Source/MutableTools/Public/MuT/NodeImageReference.h
deleted file mode 100644
