[Mutable] Node Variation removed missing default value warnings.

- Unified behaviour with all Variation Nodes.

#rnx

diff --git a/Engine/Plugins/Experimental/Mutable/Source/CustomizableObjectEditor/Private/MuCOE/GenerateMutableSource/GenerateMutableSourceColor.cpp b/Engine/Plugins/Experimental/Mutable/Source/CustomizableObjectEditor/Private/MuCOE/GenerateMutableSource/GenerateMutableSourceColor.cpp
--- a/Engine/Plugins/Experimental/Mutable/Source/CustomizableObjectEditor/Private/MuCOE/GenerateMutableSource/GenerateMutableSourceColor.cpp
+++ b/Engine/Plugins/Experimental/Mutable/Source/CustomizableObjectEditor/Private/MuCOE/GenerateMutableSource/GenerateMutableSourceColor.cpp
@@ -247,10 +247,6 @@ mu::NodeColourPtr GenerateMutableSourceColor(const UEdGraphPin* Pin, FMutableGra
 				GenerationContext.Compiler->CompilerLog(LOCTEXT("ColorFailed", "Color generation failed."), Node);
 			}
 		}
-		else
-		{
-			GenerationContext.Compiler->CompilerLog(LOCTEXT("ColorVarMissingDef", "Color variation node requires a default value."), Node);
-		}
 
 		const int32 NumVariations = TypedNodeColorVar->GetNumVariations();
 		ColorNode->SetVariationCount(NumVariations);
diff --git a/Engine/Plugins/Experimental/Mutable/Source/CustomizableObjectEditor/Private/MuCOE/GenerateMutableSource/GenerateMutableSourceFloat.cpp b/Engine/Plugins/Experimental/Mutable/Source/CustomizableObjectEditor/Private/MuCOE/GenerateMutableSource/GenerateMutableSourceFloat.cpp
--- a/Engine/Plugins/Experimental/Mutable/Source/CustomizableObjectEditor/Private/MuCOE/GenerateMutableSource/GenerateMutableSourceFloat.cpp
+++ b/Engine/Plugins/Experimental/Mutable/Source/CustomizableObjectEditor/Private/MuCOE/GenerateMutableSource/GenerateMutableSourceFloat.cpp
@@ -275,10 +275,6 @@ mu::NodeScalarPtr GenerateMutableSourceFloat(const UEdGraphPin* Pin, FMutableGra
 				GenerationContext.Compiler->CompilerLog(LOCTEXT("FloatFailed", "Float generation failed."), Node);
 			}
 		}
-		else
-		{
-			GenerationContext.Compiler->CompilerLog(LOCTEXT("FloatVarMissingDef", "Float variation node requires a default value."), Node);
-		}
 
 		const int32 NumVariations = TypedNodeFloatVar->GetNumVariations();
 		FloatNode->SetVariationCount(NumVariations);
diff --git a/Engine/Plugins/Experimental/Mutable/Source/CustomizableObjectEditor/Private/MuCOE/GenerateMutableSource/GenerateMutableSourceImage.cpp b/Engine/Plugins/Experimental/Mutable/Source/CustomizableObjectEditor/Private/MuCOE/GenerateMutableSource/GenerateMutableSourceImage.cpp
--- a/Engine/Plugins/Experimental/Mutable/Source/CustomizableObjectEditor/Private/MuCOE/GenerateMutableSource/GenerateMutableSourceImage.cpp
+++ b/Engine/Plugins/Experimental/Mutable/Source/CustomizableObjectEditor/Private/MuCOE/GenerateMutableSource/GenerateMutableSourceImage.cpp
@@ -431,10 +431,6 @@ mu::NodeImagePtr GenerateMutableSourceImage(const UEdGraphPin* Pin, FMutableGrap
 				GenerationContext.Compiler->CompilerLog(LOCTEXT("TextureFailed", "Texture generation failed."), Node);
 			}
 		}
-		else
-		{
-			GenerationContext.Compiler->CompilerLog(LOCTEXT("TextureVarMissingDef", "Texture variation node requires a default value."), Node);
-		}
 
 		TextureNode->SetVariationCount(TypedNodeImageVar->Variations.Num());
 		for (int32 VariationIndex = 0; VariationIndex < TypedNodeImageVar->Variations.Num(); ++VariationIndex)
diff --git a/Engine/Plugins/Experimental/Mutable/Source/CustomizableObjectEditor/Private/MuCOE/GenerateMutableSource/GenerateMutableSourceMesh.cpp b/Engine/Plugins/Experimental/Mutable/Source/CustomizableObjectEditor/Private/MuCOE/GenerateMutableSource/GenerateMutableSourceMesh.cpp
--- a/Engine/Plugins/Experimental/Mutable/Source/CustomizableObjectEditor/Private/MuCOE/GenerateMutableSource/GenerateMutableSourceMesh.cpp
+++ b/Engine/Plugins/Experimental/Mutable/Source/CustomizableObjectEditor/Private/MuCOE/GenerateMutableSource/GenerateMutableSourceMesh.cpp
@@ -3301,10 +3301,6 @@ mu::NodeMeshPtr GenerateMutableSourceMesh(const UEdGraphPin* Pin,
 				GenerationContext.Compiler->CompilerLog(LOCTEXT("MeshFailed", "Mesh generation failed."), Node);
 			}
 		}
-		else
-		{
-			GenerationContext.Compiler->CompilerLog(LOCTEXT("MeshVarMissingDef", "Mesh variation node requires a default value."), Node);
-		}
 
 		const int32 NumVariations = TypedNodeMeshVar->GetNumVariations();
 		MeshNode->SetVariationCount(NumVariations);
@@ -3345,7 +3341,8 @@ mu::NodeMeshPtr GenerateMutableSourceMesh(const UEdGraphPin* Pin,
 		}
 		else
 		{
-			GenerationContext.Compiler->CompilerLog(LOCTEXT("MeshGeometryMissingDef", "Mesh variation node requires a default value."), Node);
+			FText Text = FText::Format(LOCTEXT("MeshGeometryMissingDef", "Geometry Operation node requires the {0} value."), TypedNodeGeometry->MeshAPin()->PinFriendlyName);
+			GenerationContext.Compiler->CompilerLog(Text, Node);
 		}
 
 		if (const UEdGraphPin* ConnectedPin = FollowInputPin(*TypedNodeGeometry->MeshBPin()))
diff --git a/Engine/Plugins/Experimental/Mutable/Source/CustomizableObjectEditor/Private/MuCOE/Nodes/CustomizableObjectNodeExtensionDataVariation.cpp b/Engine/Plugins/Experimental/Mutable/Source/CustomizableObjectEditor/Private/MuCOE/Nodes/CustomizableObjectNodeExtensionDataVariation.cpp
--- a/Engine/Plugins/Experimental/Mutable/Source/CustomizableObjectEditor/Private/MuCOE/Nodes/CustomizableObjectNodeExtensionDataVariation.cpp
+++ b/Engine/Plugins/Experimental/Mutable/Source/CustomizableObjectEditor/Private/MuCOE/Nodes/CustomizableObjectNodeExtensionDataVariation.cpp
@@ -75,10 +75,6 @@ mu::NodeExtensionDataPtr UCustomizableObjectNodeExtensionDataVariation::Generate
 			}
 		}
 	}
-	else
-	{
-		InCompilerInterface.CompilerLog(LOCTEXT("ExtensionDataVariationMissinfDef", "Extension Data Variation node requires a default value"), this);
-	}
 
 	VariationNode->SetVariationCount(Variations.Num());
 	for (int32 VariationIndex = 0; VariationIndex < Variations.Num(); ++VariationIndex)
diff --git a/Engine/Plugins/Experimental/Mutable/Source/MutableTools/Private/MuT/CodeGenerator_Color.cpp b/Engine/Plugins/Experimental/Mutable/Source/MutableTools/Private/MuT/CodeGenerator_Color.cpp
--- a/Engine/Plugins/Experimental/Mutable/Source/MutableTools/Private/MuT/CodeGenerator_Color.cpp
+++ b/Engine/Plugins/Experimental/Mutable/Source/MutableTools/Private/MuT/CodeGenerator_Color.cpp
@@ -241,11 +241,6 @@ namespace mu
 			FMeshGenerationResult branchResults;
 			currentOp = Generate(node.m_defaultColour);
 		}
-		else
-		{
-			// This argument is required
-			currentOp = GenerateMissingColourCode(TEXT("Variation default"), node.m_errorContext);
-		}
 
 		// Process variations in reverse order, since conditionals are built bottom-up.
 		for (int t = node.m_variations.Num() - 1; t >= 0; --t)
diff --git a/Engine/Plugins/Experimental/Mutable/Source/MutableTools/Private/MuT/CodeGenerator_ExtensionData.cpp b/Engine/Plugins/Experimental/Mutable/Source/MutableTools/Private/MuT/CodeGenerator_ExtensionData.cpp
--- a/Engine/Plugins/Experimental/Mutable/Source/MutableTools/Private/MuT/CodeGenerator_ExtensionData.cpp
+++ b/Engine/Plugins/Experimental/Mutable/Source/MutableTools/Private/MuT/CodeGenerator_ExtensionData.cpp
@@ -156,11 +156,6 @@ class Node;
 
 			CurrentOp = DefaultResult.Op;
 		}
-		else
-		{
-			// This argument is required
-			CurrentOp = GenerateMissingExtensionDataCode(TEXT("Variation default"), Node.m_errorContext);
-		}
 
 		// Process variations in reverse order, since conditionals are built bottom-up.
 		for (int32 VariationIndex = Node.Variations.Num() - 1; VariationIndex >= 0; --VariationIndex)
diff --git a/Engine/Plugins/Experimental/Mutable/Source/MutableTools/Private/MuT/CodeGenerator_Image.cpp b/Engine/Plugins/Experimental/Mutable/Source/MutableTools/Private/MuT/CodeGenerator_Image.cpp
--- a/Engine/Plugins/Experimental/Mutable/Source/MutableTools/Private/MuT/CodeGenerator_Image.cpp
+++ b/Engine/Plugins/Experimental/Mutable/Source/MutableTools/Private/MuT/CodeGenerator_Image.cpp
@@ -1605,12 +1605,7 @@ namespace mu
 			GenerateImage(Options, BranchResults, node.m_defaultImage);
 			currentOp = BranchResults.op;
         }
-        else
-        {
-            // This argument is required
-            currentOp = GenerateMissingImageCode(TEXT("Variation default"), EImageFormat::IF_RGBA_UBYTE, node.m_errorContext, Options);
-        }
-
+        
         // Process variations in reverse order, since conditionals are built bottom-up.
         for ( int t = int( node.m_variations.Num() ) - 1; t >= 0; --t )
         {
diff --git a/Engine/Plugins/Experimental/Mutable/Source/MutableTools/Private/MuT/CodeGenerator_Scalar.cpp b/Engine/Plugins/Experimental/Mutable/Source/MutableTools/Private/MuT/CodeGenerator_Scalar.cpp
--- a/Engine/Plugins/Experimental/Mutable/Source/MutableTools/Private/MuT/CodeGenerator_Scalar.cpp
+++ b/Engine/Plugins/Experimental/Mutable/Source/MutableTools/Private/MuT/CodeGenerator_Scalar.cpp
@@ -284,11 +284,6 @@ namespace mu
 			GenerateScalar(ChildResult, node.m_defaultScalar);
 			op = ChildResult.op;
 		}
-		else
-		{
-			// This argument is required
-			op = GenerateMissingScalarCode(TEXT("Variation default"), 0.0f, node.m_errorContext);
-		}
 
 		// Process variations in reverse order, since conditionals are built bottom-up.
 		for (int t = node.m_variations.Num() - 1; t >= 0; --t)
