[mutable] Ensure the plugin is tagged as experimental in the UI.
#rnx

diff --git a/Engine/Plugins/Experimental/Mutable/Mutable.uplugin b/Engine/Plugins/Experimental/Mutable/Mutable.uplugin
--- a/Engine/Plugins/Experimental/Mutable/Mutable.uplugin
+++ b/Engine/Plugins/Experimental/Mutable/Mutable.uplugin
@@ -3,9 +3,8 @@
 
 	"FriendlyName": "Mutable",
 	"Version": 1,
-	"VersionName": "1.7.0",
-	"CreatedBy": "Anticto Estudi Binari, SL",
-	"CreatedByURL": "https://anticto.com",
+	"VersionName": "1.8.0",
+	"CreatedBy": "Epic Games, Inc.",
 	"Description": "Mutable adds the tools and runtime to create customizable objects for your games.",
 	"Category": "CustomizableObjects",
 	"Modules": [
@@ -61,5 +60,6 @@
 	],
 
 	"CanContainContent": true,
+	"IsExperimentalVersion": true,
 	"EnabledByDefault": false
 }
