[Mutable] Fix assert when failing to rename a streamed resource container at cook time by reusing existing containers.

#rnx

diff --git a/Engine/Plugins/Experimental/Mutable/Source/CustomizableObjectEditor/Private/MuCOE/GenerateMutableSource/GenerateMutableSource.cpp b/Engine/Plugins/Experimental/Mutable/Source/CustomizableObjectEditor/Private/MuCOE/GenerateMutableSource/GenerateMutableSource.cpp
--- a/Engine/Plugins/Experimental/Mutable/Source/CustomizableObjectEditor/Private/MuCOE/GenerateMutableSource/GenerateMutableSource.cpp
+++ b/Engine/Plugins/Experimental/Mutable/Source/CustomizableObjectEditor/Private/MuCOE/GenerateMutableSource/GenerateMutableSource.cpp
@@ -167,21 +167,10 @@ int32 FMutableGraphGenerationContext::AddStreamedResource(uint32 InResourceHash,
 	}
 
 	int32 NewResourceIndex = StreamedResourceData.Num();
-	const FString ContainerName = FString::Printf(TEXT("SR_%d"), NewResourceIndex);
+	const FString ContainerName = GetNameSafe(Object) + FString::Printf(TEXT("_SR_%d"), NewResourceIndex);
 
-	UObject* ExistingObject = FindObject<UObject>(Object, *ContainerName);
-	if (ExistingObject)
-	{
-		// This must have been left behind from a previous compilation and hasn't been deleted by 
-		// GC yet.
-		//
-		// Move it into the transient package to get it out of the way.
-		ExistingObject->Rename(nullptr, GetTransientPackage(), REN_DontCreateRedirectors);
-
-		check(!FindObject<UObject>(Object, *ContainerName));
-	}
-
-	OutNewResource = NewObject<UCustomizableObjectResourceDataContainer>(
+	UCustomizableObjectResourceDataContainer* ExistingContainer = Cast<UCustomizableObjectResourceDataContainer>(FindObject<UObject>(Object, *ContainerName));
+	OutNewResource = ExistingContainer ? ExistingContainer : NewObject<UCustomizableObjectResourceDataContainer>(
 		Object,
 		FName(*ContainerName),
 		RF_Public);
