[Mutable] Fix FCustomizableObjectLayoutBlock::Id is not initialized properly error.

#rnx

diff --git a/Engine/Plugins/Experimental/Mutable/Source/CustomizableObjectEditor/Classes/MuCOE/CustomizableObjectLayout.h b/Engine/Plugins/Experimental/Mutable/Source/CustomizableObjectEditor/Classes/MuCOE/CustomizableObjectLayout.h
--- a/Engine/Plugins/Experimental/Mutable/Source/CustomizableObjectEditor/Classes/MuCOE/CustomizableObjectLayout.h
+++ b/Engine/Plugins/Experimental/Mutable/Source/CustomizableObjectEditor/Classes/MuCOE/CustomizableObjectLayout.h
@@ -52,7 +52,7 @@ struct CUSTOMIZABLEOBJECTEDITOR_API FCustomizableObjectLayoutBlock
 	uint32 Priority;
 
 	/** Unique unchangeable id used to reference this block from other nodes. */
-	UPROPERTY()
+	UPROPERTY(meta = (IgnoreForMemberInitializationTest))
 	FGuid Id;
 
 	/** Block will be reduced on both sizes at the same time on each reduction. */
