[mutable] Runtime texture decompression optimizations.
Mostly tweaks to help code generation with inline and templates. Also some optimization to functions known to be called only with specific parameter value ranges.
#rnx

diff --git a/Engine/Plugins/Experimental/Mutable/Source/MutableRuntime/Private/MuR/BlockCompression/Miro/Miro.cpp b/Engine/Plugins/Experimental/Mutable/Source/MutableRuntime/Private/MuR/BlockCompression/Miro/Miro.cpp
--- a/Engine/Plugins/Experimental/Mutable/Source/MutableRuntime/Private/MuR/BlockCompression/Miro/Miro.cpp
+++ b/Engine/Plugins/Experimental/Mutable/Source/MutableRuntime/Private/MuR/BlockCompression/Miro/Miro.cpp
@@ -19,7 +19,6 @@ static int32 s_CurrentBlock = -1;
 
 inline constexpr void miro_check(bool x)
 {
-	(void)x;
 	check(x);
 }
 
@@ -322,8 +321,8 @@ namespace impl
 				refColors[3] = 0xFF000000 | (b << 16) | (g << 8) | r;
 			}
 			else {                            /* Quite rare BC1A mode */
-			 /* color_2 = 1/2*color_0 + 1/2*color_1;
-				color_3 = 0;                         */
+				/* color_2 = 1/2*color_0 + 1/2*color_1;
+				   color_3 = 0;                         */
 				r = (r0 + r1 + 1) >> 1;
 				g = (g0 + g1 + 1) >> 1;
 				b = (b0 + b1 + 1) >> 1;
@@ -2151,21 +2150,16 @@ namespace impl
 		constexpr size_t MAX_BLOCK_HEIGHT = 12;
 
 
-		inline bool getbit(size_t number, size_t n) {
+		FORCEINLINE bool getbit(size_t number, size_t n) {
 			return (number >> n) & 1;
 		}
 
-		inline uint8 getbits(uint8 number, uint8 msb, uint8 lsb) {
+		FORCEINLINE uint8 getbits(uint8 number, uint8 msb, uint8 lsb) {
 			int count = msb - lsb + 1;
 			return static_cast<uint8>((number >> lsb) & ((1 << count) - 1));
 		}
 
-		inline size_t getbits(size_t number, size_t msb, size_t lsb) {
-			size_t count = msb - lsb + 1;
-			return (number >> lsb) & ((size_t(1) << count) - 1);
-		}
-
-		inline void orbits8_ptr(uint8* ptr,
+		FORCEINLINE void orbits8_ptr(uint8* ptr,
 			size_t bitoffset,
 			size_t number,
 			size_t bitcount) {
@@ -2346,7 +2340,7 @@ namespace impl
 		typedef vec4_t<int> vec4i_t;
 
 		template <typename T>
-		vec3_t<T> operator+(vec3_t<T> a, vec3_t<T> b) 
+		vec3_t<T> operator+(vec3_t<T> a, vec3_t<T> b)
 		{
 			vec3_t<T> result;
 			result.r = a.r + b.r;
@@ -2356,7 +2350,7 @@ namespace impl
 		}
 
 		template <typename T>
-		vec4_t<T> operator+(vec4_t<T> a, vec4_t<T> b) 
+		vec4_t<T> operator+(vec4_t<T> a, vec4_t<T> b)
 		{
 			vec4_t<T> result;
 			result.r = a.r + b.r;
@@ -2367,7 +2361,7 @@ namespace impl
 		}
 
 		template <typename T>
-		vec3_t<T> operator-(vec3_t<T> a, vec3_t<T> b) 
+		vec3_t<T> operator-(vec3_t<T> a, vec3_t<T> b)
 		{
 			vec3_t<T> result;
 			result.r = a.r - b.r;
@@ -2377,7 +2371,7 @@ namespace impl
 		}
 
 		template <typename T>
-		vec4_t<T> operator-(vec4_t<T> a, vec4_t<T> b) 
+		vec4_t<T> operator-(vec4_t<T> a, vec4_t<T> b)
 		{
 			vec4_t<T> result;
 			result.r = a.r - b.r;
@@ -2388,7 +2382,7 @@ namespace impl
 		}
 
 		template <typename T>
-		vec3_t<T> operator*(vec3_t<T> a, vec3_t<T> b) 
+		vec3_t<T> operator*(vec3_t<T> a, vec3_t<T> b)
 		{
 			vec3_t<T> result;
 			result.r = a.r * b.r;
@@ -2398,7 +2392,7 @@ namespace impl
 		}
 
 		template <typename T>
-		vec4_t<T> operator*(vec4_t<T> a, vec4_t<T> b) 
+		vec4_t<T> operator*(vec4_t<T> a, vec4_t<T> b)
 		{
 			vec4_t<T> result;
 			result.r = a.r * b.r;
@@ -2409,7 +2403,7 @@ namespace impl
 		}
 
 		template <typename T>
-		vec4_t<T> operator/(vec4_t<T> a, vec4_t<T> b) 
+		vec4_t<T> operator/(vec4_t<T> a, vec4_t<T> b)
 		{
 			vec4_t<T> result;
 			result.r = a.r / b.r;
@@ -2462,7 +2456,7 @@ namespace impl
 		}
 
 		template <typename T>
-		vec3_t<T> operator/(vec3_t<T> a, vec3_t<T> b) 
+		vec3_t<T> operator/(vec3_t<T> a, vec3_t<T> b)
 		{
 			vec3_t<T> result;
 			result.x = a.x / b.x;
@@ -2513,13 +2507,13 @@ namespace impl
 		}
 
 		template <typename T>
-		T norm(vec3_t<T> a) 
+		T norm(vec3_t<T> a)
 		{
 			return static_cast<T>(sqrt(quadrance(a)));
 		}
 
 		template <typename T>
-		T norm(vec4_t<T> a) 
+		T norm(vec4_t<T> a)
 		{
 			return static_cast<T>(sqrt(quadrance(a)));
 		}
@@ -2531,7 +2525,7 @@ namespace impl
 		}
 
 		template <typename T>
-		T qd(vec3_t<T> a, vec3_t<T> b) 
+		T qd(vec3_t<T> a, vec3_t<T> b)
 		{
 			return quadrance(a - b);
 		}
@@ -2559,13 +2553,13 @@ namespace impl
 			// \todo This should be intercepted earlier.
 			// miro_check(x != 0.0);
 			if (x == 0.0)
-				return vec4_t<T>(0, 1, 0,0);
+				return vec4_t<T>(0, 1, 0, 0);
 
 			return a / x;
 		}
 
 		template <typename T>
-		vec3_t<T> vecmin(vec3_t<T> a, vec3_t<T> b) 
+		vec3_t<T> vecmin(vec3_t<T> a, vec3_t<T> b)
 		{
 			vec3_t<T> result;
 			result.x = FMath::Min(a.x, b.x);
@@ -2575,7 +2569,7 @@ namespace impl
 		}
 
 		template <typename T>
-		vec3_t<T> vecmax(vec3_t<T> a, vec3_t<T> b) 
+		vec3_t<T> vecmax(vec3_t<T> a, vec3_t<T> b)
 		{
 			vec3_t<T> result;
 			result.x = FMath::Max(a.x, b.x);
@@ -2585,7 +2579,7 @@ namespace impl
 		}
 
 		template <typename T>
-		T qd_to_line(vec3_t<T> m, vec3_t<T> k, T kk, vec3_t<T> p) 
+		T qd_to_line(vec3_t<T> m, vec3_t<T> k, T kk, vec3_t<T> p)
 		{
 			T t = dot(p - m, k) / kk;
 			vec3_t<T> q = k * t + m;
@@ -2593,28 +2587,28 @@ namespace impl
 		}
 
 
-		inline bool is_greyscale(vec3i_t color) 
+		inline bool is_greyscale(vec3i_t color)
 		{
 			// integer equality is transitive
 			return color.r == color.g && color.g == color.b;
 		}
 
-		inline int luminance(vec3i_t color) 
+		inline int luminance(vec3i_t color)
 		{
 			return (color.r + color.g + color.b) / 3;
 		}
 
-		inline bool approx_equal(vec3i_t a, vec3i_t b, int32 Epsilon ) 
+		inline bool approx_equal(vec3i_t a, vec3i_t b, int32 Epsilon)
 		{
 			return quadrance(a - b) <= Epsilon;
 		}
 
-		inline bool approx_equal(vec4i_t a, vec4i_t b, int32 Epsilon) 
+		inline bool approx_equal(vec4i_t a, vec4i_t b, int32 Epsilon)
 		{
 			return quadrance(a - b) <= Epsilon;
 		}
 
-		inline vec3i_t clamp_rgb(vec3i_t color) 
+		inline vec3i_t clamp_rgb(vec3i_t color)
 		{
 			vec3i_t result;
 			result.r = FMath::Clamp(color.r, 0, 255);
@@ -2623,7 +2617,7 @@ namespace impl
 			return result;
 		}
 
-		inline vec3f_t clamp_rgb(vec3f_t color) 
+		inline vec3f_t clamp_rgb(vec3f_t color)
 		{
 			vec3f_t result;
 			result.r = FMath::Clamp(color.r, 0.0f, 255.0f);
@@ -2632,7 +2626,7 @@ namespace impl
 			return result;
 		}
 
-		inline vec4f_t clamp_rgba(vec4f_t color) 
+		inline vec4f_t clamp_rgba(vec4f_t color)
 		{
 			vec4f_t result;
 			result.r = FMath::Clamp(color.r, 0.0f, 255.0f);
@@ -2642,17 +2636,17 @@ namespace impl
 			return result;
 		}
 
-		inline bool is_rgb(float color) 
+		inline bool is_rgb(float color)
 		{
 			return color >= 0.0f && color <= 255.0f;
 		}
 
-		inline bool is_rgb(vec3f_t color) 
+		inline bool is_rgb(vec3f_t color)
 		{
 			return is_rgb(color.r) && is_rgb(color.g) && is_rgb(color.b);
 		}
 
-		inline vec3i_t floor(vec3f_t color) 
+		inline vec3i_t floor(vec3f_t color)
 		{
 			vec3i_t result;
 			result.r = static_cast<int>(FMath::Floor(color.r));
@@ -2661,7 +2655,7 @@ namespace impl
 			return result;
 		}
 
-		inline vec3i_t round(vec3f_t color) 
+		inline vec3i_t round(vec3f_t color)
 		{
 			vec3i_t result;
 			result.r = static_cast<int>(FMath::RoundToInt32(color.r));
@@ -2670,7 +2664,7 @@ namespace impl
 			return result;
 		}
 
-		inline vec4i_t round(vec4f_t color) 
+		inline vec4i_t round(vec4f_t color)
 		{
 			vec4i_t result;
 			result.r = static_cast<int>(FMath::RoundToInt32(color.r));
@@ -2680,7 +2674,7 @@ namespace impl
 			return result;
 		}
 
-		inline vec3i_t to_vec3i(unorm8_t color) 
+		inline vec3i_t to_vec3i(unorm8_t color)
 		{
 			vec3i_t result;
 			result.r = color.channels.r;
@@ -2689,7 +2683,7 @@ namespace impl
 			return result;
 		}
 
-		inline vec4i_t to_vec4i(unorm8_t color) 
+		inline vec4i_t to_vec4i(unorm8_t color)
 		{
 			vec4i_t result;
 			result.r = color.channels.r;
@@ -2699,7 +2693,7 @@ namespace impl
 			return result;
 		}
 
-		inline vec3i_t to_vec3i(vec3f_t color) 
+		inline vec3i_t to_vec3i(vec3f_t color)
 		{
 			vec3i_t result;
 			result.r = static_cast<int>(color.r);
@@ -2708,7 +2702,7 @@ namespace impl
 			return result;
 		}
 
-		inline vec3f_t to_vec3f(unorm8_t color) 
+		inline vec3f_t to_vec3f(unorm8_t color)
 		{
 			vec3f_t result;
 			result.r = color.channels.r;
@@ -2717,7 +2711,7 @@ namespace impl
 			return result;
 		}
 
-		inline vec4f_t to_vec4f(unorm8_t color) 
+		inline vec4f_t to_vec4f(unorm8_t color)
 		{
 			vec4f_t result;
 			result.r = color.channels.r;
@@ -2727,7 +2721,7 @@ namespace impl
 			return result;
 		}
 
-		inline vec3f_t to_vec3f(vec3i_t color) 
+		inline vec3f_t to_vec3f(vec3i_t color)
 		{
 			vec3f_t result;
 			result.r = static_cast<float>(color.r);
@@ -2736,7 +2730,7 @@ namespace impl
 			return result;
 		}
 
-		inline vec4f_t to_vec4f(vec4i_t color) 
+		inline vec4f_t to_vec4f(vec4i_t color)
 		{
 			vec4f_t result;
 			result.r = static_cast<float>(color.r);
@@ -2746,7 +2740,7 @@ namespace impl
 			return result;
 		}
 
-		inline unorm8_t to_unorm8(vec3i_t color) 
+		inline unorm8_t to_unorm8(vec3i_t color)
 		{
 			unorm8_t result;
 			result.channels.r = static_cast<uint8>(color.r);
@@ -2756,7 +2750,7 @@ namespace impl
 			return result;
 		}
 
-		inline unorm8_t to_unorm8(vec4i_t color) 
+		inline unorm8_t to_unorm8(vec4i_t color)
 		{
 			unorm8_t result;
 			result.channels.r = static_cast<uint8>(color.r);
@@ -2766,7 +2760,7 @@ namespace impl
 			return result;
 		}
 
-		inline unorm16_t unorm8_to_unorm16(unorm8_t c8) 
+		inline unorm16_t unorm8_to_unorm16(unorm8_t c8)
 		{
 			// (x / 255) * (2^16-1) = x * 65535 / 255 = x * 257
 			unorm16_t result;
@@ -3689,7 +3683,7 @@ namespace impl
 			}
 		}
 
-		inline void integer_sequence_encode(const uint8* numbers, size_t count, range_t range, uint8* output) 
+		inline void integer_sequence_encode(const uint8* numbers, size_t count, range_t range, uint8* output)
 		{
 			integer_sequence_encode(numbers, count, range, bitwriter(output));
 		}
@@ -3698,7 +3692,7 @@ namespace impl
 		 * Compute the number of bits required to store a number of items in a specific
 		 * range using the binary integer sequence encoding.
 		 */
-		inline size_t compute_ise_bitcount(size_t items, range_t range) 
+		inline size_t compute_ise_bitcount(size_t items, range_t range)
 		{
 			size_t bits = bits_trits_quints_table[range][0];
 			size_t trits = bits_trits_quints_table[range][1];
@@ -3786,8 +3780,8 @@ namespace impl
 		}
 
 		void encode_luminance_direct(range_t endpoint_quant,
-			int v0,
-			int v1,
+			int32 v0,
+			int32 v1,
 			uint8 endpoint_unquantized[2],
 			uint8 endpoint_quantized[2])
 		{
@@ -3933,7 +3927,7 @@ namespace impl
 
 				endpoint_unquantized[0] = e1u;
 				endpoint_unquantized[1] = e0u;
-				
+
 				Swap(e0, e1);
 			}
 			else
@@ -3971,6 +3965,36 @@ namespace impl
 			}
 		}
 
+		void encode_la_direct(range_t endpoint_quant,
+			vec4i_t e0,
+			vec4i_t e1,
+			uint8 endpoint_quantized[4],
+			vec4i_t endpoint_unquantized[2])
+		{
+			int32 l0 = (e0.r+e0.g+e0.b)/3;
+			int32 a0 = e0.a;
+			int32 l1 = (e1.r+e1.r+e1.r)/3;
+			int32 a1 = e1.a;
+
+			int32 l0q = quantize_color(endpoint_quant, l0);
+			int32 a0q = quantize_color(endpoint_quant, a0);
+			int32 l1q = quantize_color(endpoint_quant, l1);
+			int32 a1q = quantize_color(endpoint_quant, a1);
+
+			int32 l0u = unquantize_color(endpoint_quant, l0q);
+			int32 a0u = unquantize_color(endpoint_quant, a0q);
+			int32 l1u = unquantize_color(endpoint_quant, l1q);
+			int32 a1u = unquantize_color(endpoint_quant, a1q);
+
+			endpoint_quantized[0] = static_cast<uint8>(l0q);
+			endpoint_quantized[1] = static_cast<uint8>(l1q);
+			endpoint_quantized[2] = static_cast<uint8>(a0q);
+			endpoint_quantized[3] = static_cast<uint8>(a1q);
+
+			endpoint_unquantized[0] = vec4i_t(l0u, l0u, l0u, a0u);
+			endpoint_unquantized[1] = vec4i_t(l1u, l1u, l1u, a1u);
+		}
+
 		template<int32 BLOCK_WIDTH, int32 BLOCK_HEIGHT>
 		inline void symbolic_to_physical
 		(
@@ -3982,7 +4006,7 @@ namespace impl
 
 			const uint8 endpoint_ise[MAXIMUM_ENCODED_COLOR_ENDPOINT_BYTES],
 
-			// FIXME: +1 needed here because orbits_8ptr breaks when the offset reaches
+			// +1 needed here because orbits_8ptr breaks when the offset reaches
 			// the last byte which always happens if the weight mode is RANGE_32.
 			const uint8 weights_ise[MAXIMUM_ENCODED_WEIGHT_BYTES + 1],
 
@@ -3994,7 +4018,7 @@ namespace impl
 			miro_check(endpoint_quant < RANGE_MAX);
 			miro_check(color_endpoint_mode < CEM_MAX);
 			miro_check(partition_count == 1);
-			miro_check(compute_ise_bitcount(BLOCK_WIDTH*BLOCK_HEIGHT, weight_quant) < MAXIMUM_ENCODED_WEIGHT_BITS);
+			miro_check(compute_ise_bitcount(BLOCK_WIDTH * BLOCK_HEIGHT, weight_quant) < MAXIMUM_ENCODED_WEIGHT_BITS);
 
 			size_t n = BLOCK_WIDTH;
 			size_t m = BLOCK_HEIGHT;
@@ -4019,7 +4043,7 @@ namespace impl
 			size_t cem = color_endpoint_mode;
 
 			// Block mode
-			
+
 			// Actually weight-grid and block sizes don't need to match. We are defining the weight-grid size here but for
 			// now we make them match.
 			if (BLOCK_WIDTH == 4 && BLOCK_HEIGHT == 4)
@@ -4037,6 +4061,10 @@ namespace impl
 				orbits8_ptr(pb->data, 7, b, 2);
 				orbits8_ptr(pb->data, 9, h, 1);
 				orbits8_ptr(pb->data, 10, d, 1);
+				
+				// This would be for rgba
+				//check(pb->data[0] == 0b01010011);
+				//check(pb->data[1] == 0b000);
 			}
 			else if (BLOCK_WIDTH == 6 && BLOCK_HEIGHT == 6)
 			{
@@ -4055,11 +4083,15 @@ namespace impl
 				orbits8_ptr(pb->data, 7, 0, 1);
 				orbits8_ptr(pb->data, 8, 1, 1);
 				orbits8_ptr(pb->data, 9, b, 2);
+
+				// This would be for rgba
+				//check(pb->data[0] == 0b00001000);
+				//check(pb->data[1] == 0b001);
 			}
 			else if (BLOCK_WIDTH == 8 && BLOCK_HEIGHT == 8)
 			{
 				// This encoder doesn't support 8x8 and dual plane: that requires mismatching and interpolating the pixel grid and weight grid
-				miro_check(dual_plane==0);
+				miro_check(dual_plane == 0);
 
 				// Use the tenth row of Table C.2.8 in the ASTC specification.
 				size_t a = m - 6;
@@ -4224,7 +4256,7 @@ namespace impl
 					check(pb->data[1] == 0b011);
 #else
 					pb->data[0] = 0b00110100;
-					pb->data[1] = 0b011;					
+					pb->data[1] = 0b011;
 #endif
 				}
 				else
@@ -4456,8 +4488,8 @@ namespace impl
 		{
 			// Brute force: shouldn't happen at runtime
 			int8 ce_range = RANGE_MAX;
-			
-			miro_check( !bDualPlane || Partitions<4);
+
+			miro_check(!bDualPlane || Partitions < 4);
 			int32 ConfigBits = 17;
 			if (Partitions > 1)
 			{
@@ -4480,11 +4512,11 @@ namespace impl
 			int32 RemainingBits = 128 - ConfigBits - WeightBits;
 
 			int32 CEMValueCount = cem_value_count[endpoint_mode] * Partitions;
-			miro_check(CEMValueCount<=18);
+			miro_check(CEMValueCount <= 18);
 
 			for (ce_range = RANGE_MAX - 1; ce_range >= 0; --ce_range)
 			{
-				int32 RangeBits = compute_ise_bitcount(CEMValueCount, static_cast<range_t>(ce_range) );
+				int32 RangeBits = compute_ise_bitcount(CEMValueCount, static_cast<range_t>(ce_range));
 				if (RangeBits <= RemainingBits)
 				{
 					break;
@@ -4511,7 +4543,7 @@ namespace impl
 
 
 		template<int32 BLOCK_WIDTH, int32 BLOCK_HEIGHT>
-		void encode_luminance(const uint8 texels[BLOCK_WIDTH*BLOCK_HEIGHT], PhysicalBlock* physical_block)
+		void encode_luminance(const uint8 texels[BLOCK_WIDTH * BLOCK_HEIGHT], PhysicalBlock* physical_block)
 		{
 			size_t partition_count = 1;
 			bool bDualPlane = false;
@@ -4519,14 +4551,29 @@ namespace impl
 			color_endpoint_mode_t color_endpoint_mode = CEM_LDR_LUMINANCE_DIRECT;
 			range_t weight_quant = RANGE_MAX;
 			range_t endpoint_quant = RANGE_MAX;
+
+			uint8 l0 = 255;
+			uint8 l1 = 0;
+			for (size_t i = 0; i < BLOCK_WIDTH * BLOCK_HEIGHT; ++i)
+			{
+				l0 = FMath::Min(l0, texels[i]);
+				l1 = FMath::Max(l1, texels[i]);
+			}
+
+			uint8 endpoint_unquantized[2];
+			uint8 endpoint_quantized[2];
+			constexpr uint8 NumEndpointValues = 2;
+
 			if (BLOCK_WIDTH == 4 && BLOCK_HEIGHT == 4)
 			{
+				// Hack because we have excess bits with the encoding that we support. It should be RANGE_48, but we only support up to RANGE_32.
+				check(false);
 				weight_quant = RANGE_32;
 				endpoint_quant = RANGE_256;
 			}
 			else if (BLOCK_WIDTH == 6 && BLOCK_HEIGHT == 6)
 			{
-				weight_quant = RANGE_4;
+				weight_quant = RANGE_6;
 				endpoint_quant = RANGE_256;
 			}
 			else if (BLOCK_WIDTH == 8 && BLOCK_HEIGHT == 8)
@@ -4544,29 +4591,19 @@ namespace impl
 			miro_check(endpoint_quant == ExpectedEndpointQuant);
 #endif
 
-			uint8 l0 = 255;
-			uint8 l1 = 0;
-			for (size_t i = 0; i < BLOCK_WIDTH * BLOCK_HEIGHT; ++i)
-			{
-				l0 = FMath::Min(l0, texels[i]);
-				l1 = FMath::Max(l1, texels[i]);
-			}
-
-			uint8 endpoint_unquantized[2];
-			uint8 endpoint_quantized[2];
 			encode_luminance_direct(endpoint_quant, l0, l1, endpoint_quantized, endpoint_unquantized);
 
 			uint8 weights_quantized[BLOCK_WIDTH * BLOCK_HEIGHT];
-			calculate_quantized_weights_luminance<BLOCK_WIDTH*BLOCK_HEIGHT>(
+			calculate_quantized_weights_luminance<BLOCK_WIDTH* BLOCK_HEIGHT>(
 				texels, weight_quant, endpoint_unquantized[0], endpoint_unquantized[1], weights_quantized);
 
 			uint8 endpoint_ise[MAXIMUM_ENCODED_COLOR_ENDPOINT_BYTES] = { 0 };
-			integer_sequence_encode(endpoint_quantized, 2, endpoint_quant, endpoint_ise);
+			integer_sequence_encode(endpoint_quantized, NumEndpointValues, endpoint_quant, endpoint_ise);
 
 			uint8 weights_ise[MAXIMUM_ENCODED_WEIGHT_BYTES + 1] = { 0 };
 			integer_sequence_encode(weights_quantized, BLOCK_WIDTH * BLOCK_HEIGHT, weight_quant, weights_ise);
 
-			symbolic_to_physical<BLOCK_WIDTH,BLOCK_HEIGHT>(color_endpoint_mode, endpoint_quant, weight_quant,
+			symbolic_to_physical<BLOCK_WIDTH, BLOCK_HEIGHT>(color_endpoint_mode, endpoint_quant, weight_quant,
 				partition_count, endpoint_ise,
 				weights_ise, physical_block);
 		}
@@ -4587,18 +4624,18 @@ namespace impl
 
 			if (BLOCK_WIDTH == 4 && BLOCK_HEIGHT == 4)
 			{
-				weight_quant = RANGE_12;
 				endpoint_quant = RANGE_256;
+				weight_quant = RANGE_12;
 			}
 			else if (BLOCK_WIDTH == 6 && BLOCK_HEIGHT == 6)
 			{
-				weight_quant = RANGE_3;
 				endpoint_quant = RANGE_256;
+				weight_quant = RANGE_3;
 			}
 			else if (BLOCK_WIDTH == 8 && BLOCK_HEIGHT == 8)
 			{
-				weight_quant = RANGE_2;
 				endpoint_quant = RANGE_192;
+				weight_quant = RANGE_2;
 			}
 			else
 			{
@@ -4618,7 +4655,7 @@ namespace impl
 			encode_rgb_direct(endpoint_quant, endpoint0, endpoint1, endpoint_quantized, endpoint_unquantized);
 
 			uint8 weights_quantized[BLOCK_WIDTH * BLOCK_HEIGHT];
-			calculate_quantized_weights_rgb<BLOCK_WIDTH*BLOCK_HEIGHT>(texels, weight_quant, endpoint_unquantized[0], endpoint_unquantized[1], weights_quantized);
+			calculate_quantized_weights_rgb<BLOCK_WIDTH* BLOCK_HEIGHT>(texels, weight_quant, endpoint_unquantized[0], endpoint_unquantized[1], weights_quantized);
 
 			uint8 endpoint_ise[MAXIMUM_ENCODED_COLOR_ENDPOINT_BYTES] = { 0 };
 			integer_sequence_encode(endpoint_quantized, 6, endpoint_quant, endpoint_ise);
@@ -4651,7 +4688,7 @@ namespace impl
 				// weight_quant = RANGE_4;
 				// endpoint_quant = RANGE_48;
 			}
-			else if(BLOCK_WIDTH == 6 && BLOCK_HEIGHT == 6)
+			else if (BLOCK_WIDTH == 6 && BLOCK_HEIGHT == 6)
 			{
 				weight_quant = RANGE_2;
 				endpoint_quant = RANGE_24;
@@ -4664,7 +4701,7 @@ namespace impl
 			// There is only one valid option for endpoint_quant
 #ifdef UE_MIRO_DEBUG
 			range_t ExpectedEndpointQuant = endpoint_quantization(BLOCK_WIDTH, BLOCK_HEIGHT, partition_count, weight_quant, color_endpoint_mode, bDualPlane);
-			miro_check( endpoint_quant == ExpectedEndpointQuant );
+			miro_check(endpoint_quant == ExpectedEndpointQuant);
 #endif
 
 			vec4i_t endpoint_unquantized[2];
@@ -4674,14 +4711,14 @@ namespace impl
 			encode_rgba_direct(endpoint_quant, endpoint0, endpoint1, endpoint_quantized, endpoint_unquantized, true);
 
 			uint8 weights_quantized[BLOCK_WIDTH * BLOCK_HEIGHT];
-			calculate_quantized_weights_rgb<BLOCK_WIDTH*BLOCK_HEIGHT>(texels, weight_quant,
+			calculate_quantized_weights_rgb<BLOCK_WIDTH* BLOCK_HEIGHT>(texels, weight_quant,
 				endpoint_unquantized[0].rgb(),
 				endpoint_unquantized[1].rgb(),
 				weights_quantized);
 
 
 			uint8 alpha_weights_quantized[BLOCK_WIDTH * BLOCK_HEIGHT];
-			calculate_quantized_weights_channel<BLOCK_WIDTH*BLOCK_HEIGHT>(texels, weight_quant,
+			calculate_quantized_weights_channel<BLOCK_WIDTH* BLOCK_HEIGHT>(texels, weight_quant,
 				(uint8)endpoint_unquantized[0].a,
 				(uint8)endpoint_unquantized[1].a,
 				alpha_weights_quantized,
@@ -4752,7 +4789,7 @@ namespace impl
 			encode_rgba_direct(endpoint_quant, endpoint0, endpoint1, endpoint_quantized, endpoint_unquantized, false);
 
 			uint8 weights_quantized[BLOCK_WIDTH * BLOCK_HEIGHT];
-			calculate_quantized_weights_rgba<BLOCK_WIDTH*BLOCK_HEIGHT>(texels, weight_quant,
+			calculate_quantized_weights_rgba<BLOCK_WIDTH* BLOCK_HEIGHT>(texels, weight_quant,
 				endpoint_unquantized[0],
 				endpoint_unquantized[1],
 				weights_quantized);
@@ -4763,13 +4800,13 @@ namespace impl
 			uint8 weights_ise[MAXIMUM_ENCODED_WEIGHT_BYTES + 1] = { 0 };
 			integer_sequence_encode(weights_quantized, BLOCK_WIDTH * BLOCK_HEIGHT, weight_quant, weights_ise);
 
-			symbolic_to_physical<BLOCK_WIDTH, BLOCK_HEIGHT>(color_endpoint_mode, endpoint_quant, weight_quant, partition_count, 
-				endpoint_ise, weights_ise, 
+			symbolic_to_physical<BLOCK_WIDTH, BLOCK_HEIGHT>(color_endpoint_mode, endpoint_quant, weight_quant, partition_count,
+				endpoint_ise, weights_ise,
 				physical_block, bDualPlane);
 		}
 
 
-		void encode_rgba_single_partition_8x8_infill_7x7(const unorm8_t texels[8*8],
+		void encode_rgba_single_partition_8x8_infill_7x7(const unorm8_t texels[8 * 8],
 			vec4f_t e0, vec4f_t e1,
 			PhysicalBlock* physical_block,
 			bool bDualPlane)
@@ -4869,6 +4906,7 @@ namespace impl
 
 				// another option
 				{
+					check(LastOption<MaxOptionCount);
 					FEncodingOption& Option = Options[LastOption++];
 
 					Option.color_endpoint_mode = CEM_LDR_LUMINANCE_DIRECT;
@@ -4894,6 +4932,7 @@ namespace impl
 				// another option
 				//if (bRGBLinear)
 				{
+					check(LastOption < MaxOptionCount);
 					FEncodingOption& Option = Options[LastOption++];
 
 					Option.color_endpoint_mode = CEM_LDR_RGB_BASE_SCALE;
@@ -4903,8 +4942,8 @@ namespace impl
 					Option.endpoint0 = round(e0);
 					Option.endpoint1 = round(e1);
 
-					int32 MaxComponent = (ve1[0] > ve1[1]) 
-						? ((ve1[0] > ve1[2]) ? 0 : ((ve1[1] > ve1[2]) ? 1 : 2) ) 
+					int32 MaxComponent = (ve1[0] > ve1[1])
+						? ((ve1[0] > ve1[2]) ? 0 : ((ve1[1] > ve1[2]) ? 1 : 2))
 						: ((ve1[1] > ve1[2]) ? 1 : 2);
 					float Ratio = 1.0;
 					if (ve1[MaxComponent] > UE_SMALL_NUMBER)
@@ -4937,6 +4976,7 @@ namespace impl
 
 				// another option
 				{
+					check(LastOption < MaxOptionCount);
 					FEncodingOption& Option = Options[LastOption++];
 
 					Option.color_endpoint_mode = CEM_LDR_RGB_DIRECT;
@@ -4962,24 +5002,42 @@ namespace impl
 			}
 			else
 			{
-				// most generic option.
-				FEncodingOption& Option = Options[LastOption++];
+				// luminance and alpha.
+				{
+					check(LastOption < MaxOptionCount);
+					FEncodingOption& Option = Options[LastOption++];
 
-				Option.color_endpoint_mode = CEM_LDR_RGBA_DIRECT;
-				Option.EndpointValueCount = 8;
-				Option.weight_quant = RANGE_2;
-				Option.endpoint_quant = RANGE_192;
-				Option.endpoint0 = round(e0);
-				Option.endpoint1 = round(e1);
+					Option.color_endpoint_mode = CEM_LDR_LUMINANCE_ALPHA_DIRECT;
+					Option.EndpointValueCount = 4;
+					Option.weight_quant = RANGE_3;
+					Option.endpoint_quant = RANGE_256;
+					Option.endpoint0 = round(e0);
+					Option.endpoint1 = round(e1);
+
+					encode_la_direct(Option.endpoint_quant, Option.endpoint0, Option.endpoint1, Option.endpoint_quantized, Option.endpoint_unquantized);
+				}
 
-				encode_rgba_direct(Option.endpoint_quant, Option.endpoint0, Option.endpoint1, Option.endpoint_quantized, Option.endpoint_unquantized, false);
+				// most generic option.
+				{
+					check(LastOption < MaxOptionCount);
+					FEncodingOption& Option = Options[LastOption++];
+	
+					Option.color_endpoint_mode = CEM_LDR_RGBA_DIRECT;
+					Option.EndpointValueCount = 8;
+					Option.weight_quant = RANGE_2;
+					Option.endpoint_quant = RANGE_192;
+					Option.endpoint0 = round(e0);
+					Option.endpoint1 = round(e1);
+	
+					encode_rgba_direct(Option.endpoint_quant, Option.endpoint0, Option.endpoint1, Option.endpoint_quantized, Option.endpoint_unquantized, false);
+				}
 			}
-			check(LastOption<=MaxOptionCount);
+			check(LastOption <= MaxOptionCount);
 
 			int32 MinError = TNumericLimits<int32>::Max();
 			int32 MinErrorOptionIndex = -1;
 
-			for (int32 o=0; o<LastOption; ++o)
+			for (int32 o = 0; o < LastOption; ++o)
 			{
 				FEncodingOption& Option = Options[o];
 
@@ -4991,7 +5049,7 @@ namespace impl
 				int32 EndpointError = quadrance(Option.endpoint0 - Option.endpoint_unquantized[0])
 					+ quadrance(Option.endpoint1 - Option.endpoint_unquantized[1]);
 
-				// Heuristic to compensate for the bigger precision in some options.
+				// Heuristic to compensate for the bigger weight precision in some options.
 				//const int32 RangeErrorFactor[RANGE_MAX] = { 128, 85, 64, 51, 42, 32, 25, 20, 16, 12, 10, 8, 6, 5, 4, 1, 1, 1, 1, 1, 1 };
 				//OptionError *= RangeErrorFactor[Option.weight_quant];
 
@@ -5011,7 +5069,7 @@ namespace impl
 			const FEncodingOption& Option = Options[MinErrorOptionIndex];
 
 			uint8 grid_weights_quantized[MAX_BLOCK_WIDTH * MAX_BLOCK_HEIGHT];
-			calculate_quantized_weights_rgba<GRID_X * GRID_Y>(grid_texels, Option.weight_quant,
+			calculate_quantized_weights_rgba<GRID_X* GRID_Y>(grid_texels, Option.weight_quant,
 				Option.endpoint_unquantized[0], Option.endpoint_unquantized[1],
 				grid_weights_quantized);
 
@@ -5082,7 +5140,7 @@ namespace impl
 
 
 		template<int32 BLOCK_WIDTH, int32 BLOCK_HEIGHT>
-		void encode_rg_single_partition(const unorm8_t texels[BLOCK_WIDTH*BLOCK_HEIGHT],
+		void encode_rg_single_partition(const unorm8_t texels[BLOCK_WIDTH * BLOCK_HEIGHT],
 			vec3f_t e0,
 			vec3f_t e1,
 			PhysicalBlock* physical_block)
@@ -5127,7 +5185,7 @@ namespace impl
 
 			uint8 weights_quantized[BLOCK_WIDTH * BLOCK_HEIGHT];
 			// TODO: optimize with rg skipping b
-			calculate_quantized_weights_rgb<BLOCK_WIDTH*BLOCK_HEIGHT>(texels, weight_quant, endpoint_unquantized[0], endpoint_unquantized[1], weights_quantized);
+			calculate_quantized_weights_rgb<BLOCK_WIDTH* BLOCK_HEIGHT>(texels, weight_quant, endpoint_unquantized[0], endpoint_unquantized[1], weights_quantized);
 
 			uint8 endpoint_ise[MAXIMUM_ENCODED_COLOR_ENDPOINT_BYTES] = { 0 };
 			integer_sequence_encode(endpoint_quantized, 4, endpoint_quant, endpoint_ise);
@@ -5135,7 +5193,7 @@ namespace impl
 			uint8 weights_ise[MAXIMUM_ENCODED_WEIGHT_BYTES + 1] = { 0 };
 			integer_sequence_encode(weights_quantized, BLOCK_WIDTH * BLOCK_HEIGHT, weight_quant, weights_ise);
 
-			symbolic_to_physical<BLOCK_WIDTH,BLOCK_HEIGHT>(color_endpoint_mode, endpoint_quant, weight_quant,
+			symbolic_to_physical<BLOCK_WIDTH, BLOCK_HEIGHT>(color_endpoint_mode, endpoint_quant, weight_quant,
 				partition_count, endpoint_ise,
 				weights_ise, physical_block, bDualPlane);
 		}
@@ -5519,8 +5577,9 @@ namespace impl
 
 
 		template<int32 BLOCK_WIDTH, int32 BLOCK_HEIGHT>
-		void compress_block(const unorm8_t texels[BLOCK_WIDTH*BLOCK_HEIGHT], PhysicalBlock* physical_block)
+		void compress_block(const unorm8_t texels[BLOCK_WIDTH * BLOCK_HEIGHT], PhysicalBlock* physical_block)
 		{
+			// Option 1: block mode 10111111100 : void extent block.
 			{
 				unorm8_t color;
 				int32 Epsilon = APPROX_COLOR_EPSILON;
@@ -5531,21 +5590,25 @@ namespace impl
 				}
 			}
 
+			// Option 2: luminance
+			// We skip this option for 4x4 beacuse of a bug.
+			if (BLOCK_WIDTH>4)
 			{
 				uint8 luminances[BLOCK_WIDTH * BLOCK_HEIGHT];
 				int32 Epsilon = APPROX_COLOR_EPSILON;
 				if (is_greyscale(texels, BLOCK_WIDTH * BLOCK_HEIGHT, luminances, Epsilon))
 				{
-					encode_luminance<BLOCK_WIDTH,BLOCK_HEIGHT>(luminances, physical_block);
+					encode_luminance<BLOCK_WIDTH, BLOCK_HEIGHT>(luminances, physical_block);
 					return;
 				}
 			}
 
+			// Option 3: generic
 			vec3f_t k, m;
-			principal_component_analysis<BLOCK_WIDTH*BLOCK_HEIGHT>(texels, BLOCK_WIDTH * BLOCK_HEIGHT, k, m);
+			principal_component_analysis<BLOCK_WIDTH* BLOCK_HEIGHT>(texels, BLOCK_WIDTH * BLOCK_HEIGHT, k, m);
 			vec3f_t e0, e1;
 			find_min_max(texels, BLOCK_WIDTH * BLOCK_HEIGHT, k, m, e0, e1);
-			encode_rgb_single_partition<BLOCK_WIDTH,BLOCK_HEIGHT>(texels, e0, e1, physical_block);
+			encode_rgb_single_partition<BLOCK_WIDTH, BLOCK_HEIGHT>(texels, e0, e1, physical_block);
 		}
 
 
@@ -5555,10 +5618,10 @@ namespace impl
 			unorm8_t color;
 
 			// Option 1: block mode 10111111100 : void extent block.
-			
+
 			// \TODO: This actually helps reduce the error due to the bad quality of the non-void extent blocks
 			// When that quality improves, reduce this epsilon.
-			int32 Epsilon = APPROX_COLOR_EPSILON; 
+			int32 Epsilon = APPROX_COLOR_EPSILON;
 
 			bool isSolidRGBA = is_solid_rgba(texels, BLOCK_WIDTH * BLOCK_HEIGHT, &color, Epsilon);
 			if (isSolidRGBA)
@@ -5582,7 +5645,7 @@ namespace impl
 				else
 				{
 					vec3f_t k, m;
-					principal_component_analysis<BLOCK_WIDTH*BLOCK_HEIGHT>(texels, BLOCK_WIDTH * BLOCK_HEIGHT, k, m);
+					principal_component_analysis<BLOCK_WIDTH* BLOCK_HEIGHT>(texels, BLOCK_WIDTH * BLOCK_HEIGHT, k, m);
 					find_min_max(texels, BLOCK_WIDTH * BLOCK_HEIGHT, k, m, e0, e1);
 				}
 
@@ -5604,18 +5667,10 @@ namespace impl
 			{
 				// Analysing as 4D
 				vec4f_t re0, re1;
-				bool isSolid = is_solid_rgba(texels, BLOCK_WIDTH * BLOCK_HEIGHT, &color, Epsilon);
-				if (isSolid)
-				{
-					re0 = re1 = to_vec4f(texels[0]);
-				}
-				else
-				{
-					vec4f_t k, m;
-					principal_component_analysis<BLOCK_WIDTH* BLOCK_HEIGHT>(texels, BLOCK_WIDTH * BLOCK_HEIGHT, k, m);
-					find_min_max(texels, BLOCK_WIDTH * BLOCK_HEIGHT, k, m, re0, re1);
-				}
-				
+				vec4f_t k, m;
+				principal_component_analysis<BLOCK_WIDTH* BLOCK_HEIGHT>(texels, BLOCK_WIDTH * BLOCK_HEIGHT, k, m);
+				find_min_max(texels, BLOCK_WIDTH * BLOCK_HEIGHT, k, m, re0, re1);
+
 				// This would be an RGB analysis with a hacky patch for alpha limits.
 				// {
 				//vec3f_t e0, e1;
@@ -5702,136 +5757,30 @@ namespace impl
 
 	// Based on astc_dec released under Apache license
 	// See: [REDACTED_PROFILE_URL]/astc_dec
-	// Just modified for integration.
+	// Modified for integration.
 	namespace astcdec
 	{
 
 #define DE_LENGTH_OF_ARRAY(x) (sizeof(x)/sizeof(x[0]))
 #define DE_UNREF(x) (void)x
 
-		typedef uint8_t deUint8;
-		typedef int8_t deInt8;
-		typedef uint32_t deUint32;
-		typedef int32_t deInt32;
-		typedef uint16_t deUint16;
-		typedef int16_t deInt16;
-		typedef int64_t deInt64;
-		typedef uint64_t deUint64;
-
-#define DE_ASSERT check
-
-		static bool inBounds(int v, int l, int h)
+		static bool inBounds(int32 v, int32 l, int32 h)
 		{
 			return (v >= l) && (v < h);
 		}
 
-		static bool inRange(int v, int l, int h)
+		static bool inRange(int32 v, int32 l, int32 h)
 		{
 			return (v >= l) && (v <= h);
 		}
 
-		struct UVec4
-		{
-			uint32_t m_c[4];
-
-			UVec4()
-			{
-				m_c[0] = 0;
-				m_c[1] = 0;
-				m_c[2] = 0;
-				m_c[3] = 0;
-			}
-
-			UVec4(uint32_t x, uint32_t y, uint32_t z, uint32_t w)
-			{
-				m_c[0] = x;
-				m_c[1] = y;
-				m_c[2] = z;
-				m_c[3] = w;
-			}
-
-			uint32_t x() const { return m_c[0]; }
-			uint32_t y() const { return m_c[1]; }
-			uint32_t z() const { return m_c[2]; }
-			uint32_t w() const { return m_c[3]; }
-
-			uint32_t& x() { return m_c[0]; }
-			uint32_t& y() { return m_c[1]; }
-			uint32_t& z() { return m_c[2]; }
-			uint32_t& w() { return m_c[3]; }
+		using UVec4 = FUint32Vector4;
+		using IVec4 = FInt32Vector4;
 
-			inline uint32_t operator[] (uint32_t idx) const { check(idx < 4); if (idx < 4) return m_c[idx]; return 0; }
-			inline uint32_t& operator[] (uint32_t idx) { check(idx < 4); if (idx < 4) return m_c[idx]; return m_c[0]; }
-		};
-
-		struct IVec4
+		inline FUint32Vector4 asUint(FInt32Vector4 v)
 		{
-			int32_t m_c[4];
-
-			IVec4()
-			{
-				m_c[0] = 0;
-				m_c[1] = 0;
-				m_c[2] = 0;
-				m_c[3] = 0;
-			}
-
-			IVec4(int32_t x, int32_t y, int32_t z, int32_t w)
-			{
-				m_c[0] = x;
-				m_c[1] = y;
-				m_c[2] = z;
-				m_c[3] = w;
-			}
-
-			int32_t x() const { return m_c[0]; }
-			int32_t y() const { return m_c[1]; }
-			int32_t z() const { return m_c[2]; }
-			int32_t w() const { return m_c[3]; }
-
-			int32_t& x() { return m_c[0]; }
-			int32_t& y() { return m_c[1]; }
-			int32_t& z() { return m_c[2]; }
-			int32_t& w() { return m_c[3]; }
-
-			UVec4 asUint() const
-			{
-				return UVec4(FMath::Max(0, m_c[0]), FMath::Max(0, m_c[1]), FMath::Max(0, m_c[2]), FMath::Max(0, m_c[3]));
-			}
-
-			inline int32_t operator[] (uint32_t idx) const { check(idx < 4); if (idx < 4) return m_c[idx]; return 0; }
-			inline int32_t& operator[] (uint32_t idx) { check(idx < 4); if (idx < 4) return m_c[idx]; return m_c[0]; }
-		};
-
-		struct IVec3
-		{
-			int32_t m_c[3];
-
-			IVec3()
-			{
-				m_c[0] = 0;
-				m_c[1] = 0;
-				m_c[2] = 0;
-			}
-
-			IVec3(int32_t x, int32_t y, int32_t z)
-			{
-				m_c[0] = x;
-				m_c[1] = y;
-				m_c[2] = z;
-			}
-
-			int32_t x() const { return m_c[0]; }
-			int32_t y() const { return m_c[1]; }
-			int32_t z() const { return m_c[2]; }
-
-			int32_t& x() { return m_c[0]; }
-			int32_t& y() { return m_c[1]; }
-			int32_t& z() { return m_c[2]; }
-
-			inline int32_t operator[] (uint32_t idx) const { check(idx < 3); if (idx < 3) return m_c[idx]; return 0; }
-			inline int32_t& operator[] (uint32_t idx) { check(idx < 3); if (idx < 3) return m_c[idx]; return m_c[0]; }
-		};
+			return FUint32Vector4(FMath::Max(0, v[0]), FMath::Max(0, v[1]), FMath::Max(0, v[2]), FMath::Max(0, v[3]));
+		}
 
 		static uint32_t deDivRoundUp32(uint32_t a, uint32_t b)
 		{
@@ -5850,54 +5799,65 @@ namespace impl
 			MAX_BLOCK_WIDTH = 12,
 			MAX_BLOCK_HEIGHT = 12
 		};
-		inline deUint32 getBit(deUint32 src, int ndx)
+		inline uint32 getBit(uint32 src, int ndx)
 		{
-			DE_ASSERT(inBounds(ndx, 0, 32));
+			check(inBounds(ndx, 0, 32));
 			return (src >> ndx) & 1;
 		}
-		inline deUint32 getBits(deUint32 src, int low, int high)
+		inline uint32 getBits(uint32 src, int low, int high)
 		{
 			const int numBits = (high - low) + 1;
-			DE_ASSERT(inRange(numBits, 1, 32));
+			check(inRange(numBits, 1, 32));
 			if (numBits < 32)
-				return (deUint32)((src >> low) & ((1u << numBits) - 1));
+				return (uint32)((src >> low) & ((1u << numBits) - 1));
 			else
-				return (deUint32)((src >> low) & 0xFFFFFFFFu);
+				return (uint32)((src >> low) & 0xFFFFFFFFu);
 		}
-		inline bool isBitSet(deUint32 src, int ndx)
+		inline bool isBitSet(uint32 src, int ndx)
 		{
 			return getBit(src, ndx) != 0;
 		}
-		inline deUint32 reverseBits(deUint32 src, int numBits)
+
+		FORCEINLINE uint32 NewReverseBits(uint32 src, int numBits)
 		{
-			DE_ASSERT(inRange(numBits, 0, 32));
-			deUint32 result = 0;
-			for (int i = 0; i < numBits; i++)
-				result |= ((src >> i) & 1) << (numBits - 1 - i);
-			return result;
+			miro_check(numBits <= 8);
+
+			// If we were to really support 32 bits
+			//uint32 ReversedSrc =
+			//	(astcrt::bit_reverse_table[(src >>  0) & 0xff] << 24) |
+			//	(astcrt::bit_reverse_table[(src >>  8) & 0xff] << 16) |
+			//	(astcrt::bit_reverse_table[(src >> 16) & 0xff] <<  8) |
+			//	(astcrt::bit_reverse_table[(src >> 24) & 0xff] <<  0);
+			//return ReversedSrc >> (32 - numBits);
+
+			uint8 r = astcrt::bit_reverse_table[src];
+
+			// Alernative without table
+			//r = src;
+			//r = (r & 0xF0) >> 4 | (r & 0x0F) << 4;
+			//r = (r & 0xCC) >> 2 | (r & 0x33) << 2;
+			//r = (r & 0xAA) >> 1 | (r & 0x55) << 1;
+
+			return r >> (8 - numBits);
 		}
-		inline deUint32 bitReplicationScale(deUint32 src, int numSrcBits, int numDstBits)
+
+		inline uint32 bitReplicationScale(uint32 src, int numSrcBits, int numDstBits)
 		{
-			DE_ASSERT(numSrcBits <= numDstBits);
-			DE_ASSERT((src & ((1 << numSrcBits) - 1)) == src);
-			deUint32 dst = 0;
+			check(numSrcBits <= numDstBits);
+			check((src & ((1 << numSrcBits) - 1)) == src);
+			uint32 dst = 0;
 			for (int shift = numDstBits - numSrcBits; shift > -numSrcBits; shift -= numSrcBits)
 				dst |= shift >= 0 ? src << shift : src >> -shift;
 			return dst;
 		}
 
-		inline deInt32 signExtend(deInt32 src, int numSrcBits)
+		inline int32 signExtend(int32 src, int numSrcBits)
 		{
-			DE_ASSERT(inRange(numSrcBits, 2, 31));
+			check(inRange(numSrcBits, 2, 31));
 			const bool negative = (src & (1 << (numSrcBits - 1))) != 0;
 			return src | (negative ? ~((1 << numSrcBits) - 1) : 0);
 		}
 
-		//inline bool isFloat16InfOrNan (deFloat16 v)
-		//{
-		//	return getBits(v, 10, 14) == 31;
-		//}
-
 		enum ISEMode
 		{
 			ISEMODE_TRIT = 0,
@@ -5905,12 +5865,14 @@ namespace impl
 			ISEMODE_PLAIN_BIT,
 			ISEMODE_LAST
 		};
+
 		struct ISEParams
 		{
 			ISEMode		mode;
 			int			numBits;
 			ISEParams(ISEMode mode_, int numBits_) : mode(mode_), numBits(numBits_) {}
 		};
+
 		inline int computeNumRequiredBits(const ISEParams& iseParams, int numValues)
 		{
 			switch (iseParams.mode)
@@ -5919,10 +5881,11 @@ namespace impl
 			case ISEMODE_QUINT:			return deDivRoundUp32(numValues * 7, 3) + numValues * iseParams.numBits;
 			case ISEMODE_PLAIN_BIT:		return numValues * iseParams.numBits;
 			default:
-				DE_ASSERT(false);
+				check(false);
 				return -1;
 			}
 		}
+
 		ISEParams computeMaximumRangeISEParams(int numAvailableBits, int numValuesInSequence)
 		{
 			int curBitsForTritMode = 6;
@@ -5930,7 +5893,7 @@ namespace impl
 			int curBitsForPlainBitMode = 8;
 			while (true)
 			{
-				DE_ASSERT(curBitsForTritMode > 0 || curBitsForQuintMode > 0 || curBitsForPlainBitMode > 0);
+				check(curBitsForTritMode > 0 || curBitsForQuintMode > 0 || curBitsForPlainBitMode > 0);
 				const int tritRange = curBitsForTritMode > 0 ? (3 << curBitsForTritMode) - 1 : -1;
 				const int quintRange = curBitsForQuintMode > 0 ? (5 << curBitsForQuintMode) - 1 : -1;
 				const int plainBitRange = curBitsForPlainBitMode > 0 ? (1 << curBitsForPlainBitMode) - 1 : -1;
@@ -5952,18 +5915,20 @@ namespace impl
 				else
 				{
 					const ISEParams params(ISEMODE_PLAIN_BIT, curBitsForPlainBitMode);
-					DE_ASSERT(maxRange == plainBitRange);
+					check(maxRange == plainBitRange);
 					if (computeNumRequiredBits(params, numValuesInSequence) <= numAvailableBits)
 						return ISEParams(ISEMODE_PLAIN_BIT, curBitsForPlainBitMode);
 					curBitsForPlainBitMode--;
 				}
 			}
 		}
-		inline int computeNumColorEndpointValues(deUint32 endpointMode)
+
+		inline int computeNumColorEndpointValues(uint32 endpointMode)
 		{
-			DE_ASSERT(endpointMode < 16);
+			check(endpointMode < 16);
 			return (endpointMode / 4 + 1) * 2;
 		}
+
 		// Decompression utilities
 		enum DecompressResult
 		{
@@ -5971,11 +5936,12 @@ namespace impl
 			DECOMPRESS_RESULT_ERROR,				//!< Encountered error while decompressing, error color written
 			DECOMPRESS_RESULT_LAST
 		};
+
 		// A helper for getting bits from a 128-bit block.
 		class Block128
 		{
 		private:
-			typedef deUint64 Word;
+			typedef uint64 Word;
 			enum
 			{
 				WORD_BYTES = sizeof(Word),
@@ -5984,7 +5950,7 @@ namespace impl
 			};
 			//DE_STATIC_ASSERT(128 % WORD_BITS == 0);
 		public:
-			Block128(const deUint8* src)
+			Block128(const uint8* src)
 			{
 				for (int wordNdx = 0; wordNdx < NUM_WORDS; wordNdx++)
 				{
@@ -5993,37 +5959,38 @@ namespace impl
 						m_words[wordNdx] |= (Word)src[wordNdx * WORD_BYTES + byteNdx] << (8 * byteNdx);
 				}
 			}
-			deUint32 getBit(int ndx) const
+			FORCEINLINE uint32 getBit(int ndx) const
 			{
-				DE_ASSERT(inBounds(ndx, 0, 128));
+				miro_check(inBounds(ndx, 0, 128));
 				return (m_words[ndx / WORD_BITS] >> (ndx % WORD_BITS)) & 1;
 			}
-			deUint32 getBits(int low, int high) const
+
+			FORCEINLINE uint32 getBits(int low, int high) const
 			{
-				DE_ASSERT(inBounds(low, 0, 128));
-				DE_ASSERT(inBounds(high, 0, 128));
-				DE_ASSERT(inRange(high - low + 1, 0, 32));
+				miro_check(inBounds(low, 0, 128));
+				miro_check(inBounds(high, 0, 128));
+				miro_check(inRange(high - low + 1, 0, 32));
 				if (high - low + 1 == 0)
 					return 0;
 				const int word0Ndx = low / WORD_BITS;
 				const int word1Ndx = high / WORD_BITS;
 				// \note "foo << bar << 1" done instead of "foo << (bar+1)" to avoid overflow, i.e. shift amount being too big.
 				if (word0Ndx == word1Ndx)
-					return (deUint32)((m_words[word0Ndx] & ((((Word)1 << high % WORD_BITS << 1) - 1))) >> ((Word)low % WORD_BITS));
+					return (uint32)((m_words[word0Ndx] & ((((Word)1 << high % WORD_BITS << 1) - 1))) >> ((Word)low % WORD_BITS));
 				else
 				{
-					DE_ASSERT(word1Ndx == word0Ndx + 1);
-					return (deUint32)(m_words[word0Ndx] >> (low % WORD_BITS)) |
-						(deUint32)((m_words[word1Ndx] & (((Word)1 << high % WORD_BITS << 1) - 1)) << (high - low - high % WORD_BITS));
+					miro_check(word1Ndx == word0Ndx + 1);
+					return (uint32)(m_words[word0Ndx] >> (low % WORD_BITS)) |
+						(uint32)((m_words[word1Ndx] & (((Word)1 << high % WORD_BITS << 1) - 1)) << (high - low - high % WORD_BITS));
 				}
 			}
-			bool isBitSet(int ndx) const
+			FORCEINLINE bool isBitSet(int ndx) const
 			{
-				DE_ASSERT(inBounds(ndx, 0, 128));
+				miro_check(inBounds(ndx, 0, 128));
 				return getBit(ndx) != 0;
 			}
 
-			bool isZero() const
+			FORCEINLINE bool isZero() const
 			{
 				for (int wordNdx = 0; wordNdx < NUM_WORDS; wordNdx++)
 				{
@@ -6031,23 +5998,28 @@ namespace impl
 				}
 				return true;
 			}
+
+			FORCEINLINE uint64 GetWord(int i) const { return m_words[i]; }
+
 		private:
 			Word m_words[NUM_WORDS];
 		};
+
 		// A helper for sequential access into a Block128.
+		template<bool Forward>
 		class BitAccessStream
 		{
 		public:
-			BitAccessStream(const Block128& src, int startNdxInSrc, int length, bool forward)
+			BitAccessStream(const Block128& src, int startNdxInSrc, int length)
 				: m_src(src)
 				, m_startNdxInSrc(startNdxInSrc)
 				, m_length(length)
-				, m_forward(forward)
 				, m_ndx(0)
 			{
 			}
+
 			// Get the next num bits. Bits at positions greater than or equal to m_length are zeros.
-			deUint32 getNext(int num)
+			uint32 getNext(int num)
 			{
 				if (num == 0 || m_ndx >= m_length)
 					return 0;
@@ -6056,22 +6028,32 @@ namespace impl
 				const int low = m_ndx;
 				const int high = m_ndx + numBitsFromSrc - 1;
 				m_ndx += num;
-				return m_forward ? m_src.getBits(m_startNdxInSrc + low, m_startNdxInSrc + high)
-					: reverseBits(m_src.getBits(m_startNdxInSrc - high, m_startNdxInSrc - low), numBitsFromSrc);
+				if (Forward)
+				{
+					return m_src.getBits(m_startNdxInSrc + low, m_startNdxInSrc + high);
+				}
+				else
+				{
+					uint32 PreReverse = m_src.getBits(m_startNdxInSrc - high, m_startNdxInSrc - low);
+					uint32 NewReverse = NewReverseBits(PreReverse, numBitsFromSrc);
+					return NewReverse;
+				}
 			}
+
 		private:
 			const Block128& m_src;
 			const int			m_startNdxInSrc;
 			const int			m_length;
-			const bool			m_forward;
 			int					m_ndx;
 		};
+
 		struct ISEDecodedResult
 		{
-			deUint32 m;
-			deUint32 tq; //!< Trit or quint value, depending on ISE mode.
-			deUint32 v;
+			uint32 m;
+			uint32 tq; //!< Trit or quint value, depending on ISE mode.
+			uint32 v;
 		};
+
 		// Data from an ASTC block's "block mode" part (i.e. bits [0,10]).
 		struct ASTCBlockMode
 		{
@@ -6093,20 +6075,24 @@ namespace impl
 			{
 			}
 		};
+
 		inline int computeNumWeights(const ASTCBlockMode& mode)
 		{
 			return mode.weightGridWidth * mode.weightGridHeight * (mode.isDualPlane ? 2 : 1);
 		}
+
 		struct ColorEndpointPair
 		{
 			UVec4 e0;
 			UVec4 e1;
 		};
+
 		struct TexelWeightPair
 		{
-			deUint32 w[2];
+			uint32 w[2];
 		};
-		ASTCBlockMode getASTCBlockMode(deUint32 blockModeData)
+
+		ASTCBlockMode getASTCBlockMode(uint32 blockModeData)
 		{
 			ASTCBlockMode blockMode;
 			blockMode.isError = true; // \note Set to false later, if not error.
@@ -6115,13 +6101,13 @@ namespace impl
 			{
 				if ((getBits(blockModeData, 0, 1) == 0 && getBits(blockModeData, 6, 8) == 7) || getBits(blockModeData, 0, 3) == 0)
 					return blockMode; // Invalid ("reserved").
-				deUint32 r = (deUint32)-1; // \note Set in the following branches.
+				uint32 r = (uint32)-1; // \note Set in the following branches.
 				if (getBits(blockModeData, 0, 1) == 0)
 				{
-					const deUint32 r0 = getBit(blockModeData, 4);
-					const deUint32 r1 = getBit(blockModeData, 2);
-					const deUint32 r2 = getBit(blockModeData, 3);
-					const deUint32 i78 = getBits(blockModeData, 7, 8);
+					const uint32 r0 = getBit(blockModeData, 4);
+					const uint32 r1 = getBit(blockModeData, 2);
+					const uint32 r2 = getBit(blockModeData, 3);
+					const uint32 i78 = getBits(blockModeData, 7, 8);
 					r = (r2 << 2) | (r1 << 1) | (r0 << 0);
 					if (i78 == 3)
 					{
@@ -6131,40 +6117,40 @@ namespace impl
 					}
 					else
 					{
-						const deUint32 a = getBits(blockModeData, 5, 6);
+						const uint32 a = getBits(blockModeData, 5, 6);
 						switch (i78)
 						{
 						case 0:		blockMode.weightGridWidth = 12;		blockMode.weightGridHeight = a + 2;									break;
 						case 1:		blockMode.weightGridWidth = a + 2;	blockMode.weightGridHeight = 12;									break;
 						case 2:		blockMode.weightGridWidth = a + 6;	blockMode.weightGridHeight = getBits(blockModeData, 9, 10) + 6;		break;
-						default: DE_ASSERT(false);
+						default: check(false);
 						}
 					}
 				}
 				else
 				{
-					const deUint32 r0 = getBit(blockModeData, 4);
-					const deUint32 r1 = getBit(blockModeData, 0);
-					const deUint32 r2 = getBit(blockModeData, 1);
-					const deUint32 i23 = getBits(blockModeData, 2, 3);
-					const deUint32 a = getBits(blockModeData, 5, 6);
+					const uint32 r0 = getBit(blockModeData, 4);
+					const uint32 r1 = getBit(blockModeData, 0);
+					const uint32 r2 = getBit(blockModeData, 1);
+					const uint32 i23 = getBits(blockModeData, 2, 3);
+					const uint32 a = getBits(blockModeData, 5, 6);
 					r = (r2 << 2) | (r1 << 1) | (r0 << 0);
 					if (i23 == 3)
 					{
-						const deUint32	b = getBit(blockModeData, 7);
+						const uint32	b = getBit(blockModeData, 7);
 						const bool		i8 = isBitSet(blockModeData, 8);
 						blockMode.weightGridWidth = i8 ? b + 2 : a + 2;
 						blockMode.weightGridHeight = i8 ? a + 2 : b + 6;
 					}
 					else
 					{
-						const deUint32 b = getBits(blockModeData, 7, 8);
+						const uint32 b = getBits(blockModeData, 7, 8);
 						switch (i23)
 						{
 						case 0:		blockMode.weightGridWidth = b + 4;	blockMode.weightGridHeight = a + 2;	break;
 						case 1:		blockMode.weightGridWidth = b + 8;	blockMode.weightGridHeight = a + 2;	break;
 						case 2:		blockMode.weightGridWidth = a + 2;	blockMode.weightGridHeight = b + 8;	break;
-						default: DE_ASSERT(false);
+						default: check(false);
 						}
 					}
 				}
@@ -6186,7 +6172,7 @@ namespace impl
 						case 5:							m = ISEMODE_QUINT;	b = 2;	break;
 						case 6:		m = ISEMODE_TRIT;						b = 3;	break;
 						case 7:												b = 5;	break;
-						default:	DE_ASSERT(false);
+						default:	check(false);
 						}
 					}
 					else
@@ -6199,7 +6185,7 @@ namespace impl
 						case 5:							m = ISEMODE_QUINT;			break;
 						case 6:		m = ISEMODE_TRIT;						b = 1;	break;
 						case 7:												b = 3;	break;
-						default:	DE_ASSERT(false);
+						default:	check(false);
 						}
 					}
 				}
@@ -6207,11 +6193,12 @@ namespace impl
 			blockMode.isError = false;
 			return blockMode;
 		}
+
 		inline void setASTCErrorColorBlock(void* dst, int blockWidth, int blockHeight, bool isSRGB)
 		{
 			if (isSRGB)
 			{
-				deUint8* const dstU = (deUint8*)dst;
+				uint8* const dstU = (uint8*)dst;
 				for (int i = 0; i < blockWidth * blockHeight; i++)
 				{
 					dstU[4 * i + 0] = 0xff;
@@ -6232,69 +6219,35 @@ namespace impl
 				}
 			}
 		}
-		DecompressResult decodeVoidExtentBlock(void* dst, const Block128& blockData, int blockWidth, int blockHeight, bool isSRGB, bool isLDRMode)
+
+		template<int32 BlockSize>
+		void decodeVoidExtentBlockSRGB(void* dst, const Block128& blockData)
 		{
-			const deUint32	minSExtent = blockData.getBits(12, 24);
-			const deUint32	maxSExtent = blockData.getBits(25, 37);
-			const deUint32	minTExtent = blockData.getBits(38, 50);
-			const deUint32	maxTExtent = blockData.getBits(51, 63);
-			const bool		allExtentsAllOnes = minSExtent == 0x1fff && maxSExtent == 0x1fff && minTExtent == 0x1fff && maxTExtent == 0x1fff;
-			const bool		isHDRBlock = blockData.isBitSet(9);
-			if ((isLDRMode && isHDRBlock) || (!allExtentsAllOnes && (minSExtent >= maxSExtent || minTExtent >= maxTExtent)))
-			{
-				setASTCErrorColorBlock(dst, blockWidth, blockHeight, isSRGB);
-				return DECOMPRESS_RESULT_ERROR;
-			}
-			const deUint32 rgba[4] =
-			{
-				blockData.getBits(64,  79),
-				blockData.getBits(80,  95),
-				blockData.getBits(96,  111),
-				blockData.getBits(112, 127)
-			};
-			if (isSRGB)
-			{
-				deUint8* const dstU = (deUint8*)dst;
-				for (int i = 0; i < blockWidth * blockHeight; i++)
-					for (int c = 0; c < 4; c++)
-						dstU[i * 4 + c] = (deUint8)((rgba[c] & 0xff00) >> 8);
-			}
-			else
+			uint64 HighWord = blockData.GetWord(1);
+
+			uint8 rgba[4];
+			rgba[0] = uint8(HighWord >> 8);
+			rgba[1] = uint8(HighWord >> 24);
+			rgba[2] = uint8(HighWord >> 40);
+			rgba[3] = uint8(HighWord >> 56);
+
+			uint32* const dstU = (uint32*)dst;
+			for (int i = 0; i < BlockSize * BlockSize; i++)
 			{
-				float* const dstF = (float*)dst;
-				if (isHDRBlock)
-				{
-					// rg - REMOVING HDR SUPPORT FOR NOW
-#if 0
-					for (int c = 0; c < 4; c++)
-					{
-						if (isFloat16InfOrNan((deFloat16)rgba[c]))
-							throw InternalError("Infinity or NaN color component in HDR void extent block in ASTC texture (behavior undefined by ASTC specification)");
-					}
-					for (int i = 0; i < blockWidth * blockHeight; i++)
-						for (int c = 0; c < 4; c++)
-							dstF[i * 4 + c] = deFloat16To32((deFloat16)rgba[c]);
-#endif
-				}
-				else
-				{
-					for (int i = 0; i < blockWidth * blockHeight; i++)
-						for (int c = 0; c < 4; c++)
-							dstF[i * 4 + c] = rgba[c] == 65535 ? 1.0f : (float)rgba[c] / 65536.0f;
-				}
+				FMemory::Memcpy(&dstU[i], rgba, 4);
 			}
-			return DECOMPRESS_RESULT_VALID_BLOCK;
 		}
-		void decodeColorEndpointModes(deUint32* endpointModesDst, const Block128& blockData, int numPartitions, int extraCemBitsStart)
+
+		void decodeColorEndpointModes(uint32* endpointModesDst, const Block128& blockData, int numPartitions, int extraCemBitsStart)
 		{
 			if (numPartitions == 1)
 				endpointModesDst[0] = blockData.getBits(13, 16);
 			else
 			{
-				const deUint32 highLevelSelector = blockData.getBits(23, 24);
+				const uint32 highLevelSelector = blockData.getBits(23, 24);
 				if (highLevelSelector == 0)
 				{
-					const deUint32 mode = blockData.getBits(25, 28);
+					const uint32 mode = blockData.getBits(25, 28);
 					for (int i = 0; i < numPartitions; i++)
 						endpointModesDst[i] = mode;
 				}
@@ -6302,37 +6255,40 @@ namespace impl
 				{
 					for (int partNdx = 0; partNdx < numPartitions; partNdx++)
 					{
-						const deUint32 cemClass = highLevelSelector - (blockData.isBitSet(25 + partNdx) ? 0 : 1);
-						const deUint32 lowBit0Ndx = numPartitions + 2 * partNdx;
-						const deUint32 lowBit1Ndx = numPartitions + 2 * partNdx + 1;
-						const deUint32 lowBit0 = blockData.getBit(lowBit0Ndx < 4 ? 25 + lowBit0Ndx : extraCemBitsStart + lowBit0Ndx - 4);
-						const deUint32 lowBit1 = blockData.getBit(lowBit1Ndx < 4 ? 25 + lowBit1Ndx : extraCemBitsStart + lowBit1Ndx - 4);
+						const uint32 cemClass = highLevelSelector - (blockData.isBitSet(25 + partNdx) ? 0 : 1);
+						const uint32 lowBit0Ndx = numPartitions + 2 * partNdx;
+						const uint32 lowBit1Ndx = numPartitions + 2 * partNdx + 1;
+						const uint32 lowBit0 = blockData.getBit(lowBit0Ndx < 4 ? 25 + lowBit0Ndx : extraCemBitsStart + lowBit0Ndx - 4);
+						const uint32 lowBit1 = blockData.getBit(lowBit1Ndx < 4 ? 25 + lowBit1Ndx : extraCemBitsStart + lowBit1Ndx - 4);
 						endpointModesDst[partNdx] = (cemClass << 2) | (lowBit1 << 1) | lowBit0;
 					}
 				}
 			}
 		}
-		int computeNumColorEndpointValues(const deUint32* endpointModes, int numPartitions)
+
+		int computeNumColorEndpointValues(const uint32* endpointModes, int numPartitions)
 		{
 			int result = 0;
 			for (int i = 0; i < numPartitions; i++)
 				result += computeNumColorEndpointValues(endpointModes[i]);
 			return result;
 		}
-		void decodeISETritBlock(ISEDecodedResult* dst, int numValues, BitAccessStream& data, int numBits)
+
+		template<bool Forward>
+		void decodeISETritBlock(ISEDecodedResult* dst, int numValues, BitAccessStream<Forward>& data, int numBits)
 		{
-			DE_ASSERT(inRange(numValues, 1, 5));
-			deUint32 m[5];
+			check(inRange(numValues, 1, 5));
+			uint32 m[5];
 			m[0] = data.getNext(numBits);
-			deUint32 T01 = data.getNext(2);
+			uint32 T01 = data.getNext(2);
 			m[1] = data.getNext(numBits);
-			deUint32 T23 = data.getNext(2);
+			uint32 T23 = data.getNext(2);
 			m[2] = data.getNext(numBits);
-			deUint32 T4 = data.getNext(1);
+			uint32 T4 = data.getNext(1);
 			m[3] = data.getNext(numBits);
-			deUint32 T56 = data.getNext(2);
+			uint32 T56 = data.getNext(2);
 			m[4] = data.getNext(numBits);
-			deUint32 T7 = data.getNext(1);
+			uint32 T7 = data.getNext(1);
 			switch (numValues)
 			{
 				// \note Fall-throughs.
@@ -6342,10 +6298,10 @@ namespace impl
 			case 4: T7 = 0;
 			case 5: break;
 			default:
-				DE_ASSERT(false);
+				check(false);
 			}
-			const deUint32 T = (T7 << 7) | (T56 << 5) | (T4 << 4) | (T23 << 2) | (T01 << 0);
-			static const deUint32 tritsFromT[256][5] =
+			const uint32 T = (T7 << 7) | (T56 << 5) | (T4 << 4) | (T23 << 2) | (T01 << 0);
+			static const uint32 tritsFromT[256][5] =
 			{
 				{ 0,0,0,0,0 }, { 1,0,0,0,0 }, { 2,0,0,0,0 }, { 0,0,2,0,0 }, { 0,1,0,0,0 }, { 1,1,0,0,0 }, { 2,1,0,0,0 }, { 1,0,2,0,0 }, { 0,2,0,0,0 }, { 1,2,0,0,0 }, { 2,2,0,0,0 }, { 2,0,2,0,0 }, { 0,2,2,0,0 }, { 1,2,2,0,0 }, { 2,2,2,0,0 }, { 2,0,2,0,0 },
 				{ 0,0,1,0,0 }, { 1,0,1,0,0 }, { 2,0,1,0,0 }, { 0,1,2,0,0 }, { 0,1,1,0,0 }, { 1,1,1,0,0 }, { 2,1,1,0,0 }, { 1,1,2,0,0 }, { 0,2,1,0,0 }, { 1,2,1,0,0 }, { 2,2,1,0,0 }, { 2,1,2,0,0 }, { 0,0,0,2,2 }, { 1,0,0,2,2 }, { 2,0,0,2,2 }, { 0,0,2,2,2 },
@@ -6364,7 +6320,7 @@ namespace impl
 				{ 0,0,0,1,2 }, { 1,0,0,1,2 }, { 2,0,0,1,2 }, { 0,0,2,1,2 }, { 0,1,0,1,2 }, { 1,1,0,1,2 }, { 2,1,0,1,2 }, { 1,0,2,1,2 }, { 0,2,0,1,2 }, { 1,2,0,1,2 }, { 2,2,0,1,2 }, { 2,0,2,1,2 }, { 0,2,2,1,2 }, { 1,2,2,1,2 }, { 2,2,2,1,2 }, { 2,0,2,1,2 },
 				{ 0,0,1,1,2 }, { 1,0,1,1,2 }, { 2,0,1,1,2 }, { 0,1,2,1,2 }, { 0,1,1,1,2 }, { 1,1,1,1,2 }, { 2,1,1,1,2 }, { 1,1,2,1,2 }, { 0,2,1,1,2 }, { 1,2,1,1,2 }, { 2,2,1,1,2 }, { 2,1,2,1,2 }, { 0,2,2,2,2 }, { 1,2,2,2,2 }, { 2,2,2,2,2 }, { 2,1,2,2,2 }
 			};
-			const deUint32(&trits)[5] = tritsFromT[T];
+			const uint32(&trits)[5] = tritsFromT[T];
 			for (int i = 0; i < numValues; i++)
 			{
 				dst[i].m = m[i];
@@ -6372,16 +6328,18 @@ namespace impl
 				dst[i].v = (trits[i] << numBits) + m[i];
 			}
 		}
-		void decodeISEQuintBlock(ISEDecodedResult* dst, int numValues, BitAccessStream& data, int numBits)
+
+		template<bool Forward>
+		void decodeISEQuintBlock(ISEDecodedResult* dst, int numValues, BitAccessStream<Forward>& data, int numBits)
 		{
-			DE_ASSERT(inRange(numValues, 1, 3));
-			deUint32 m[3];
+			check(inRange(numValues, 1, 3));
+			uint32 m[3];
 			m[0] = data.getNext(numBits);
-			deUint32 Q012 = data.getNext(3);
+			uint32 Q012 = data.getNext(3);
 			m[1] = data.getNext(numBits);
-			deUint32 Q34 = data.getNext(2);
+			uint32 Q34 = data.getNext(2);
 			m[2] = data.getNext(numBits);
-			deUint32 Q56 = data.getNext(2);
+			uint32 Q56 = data.getNext(2);
 			switch (numValues)
 			{
 				// \note Fall-throughs.
@@ -6389,10 +6347,10 @@ namespace impl
 			case 2: Q56 = 0;
 			case 3: break;
 			default:
-				DE_ASSERT(false);
+				check(false);
 			}
-			const deUint32 Q = (Q56 << 5) | (Q34 << 3) | (Q012 << 0);
-			static const deUint32 quintsFromQ[256][3] =
+			const uint32 Q = (Q56 << 5) | (Q34 << 3) | (Q012 << 0);
+			static const uint32 quintsFromQ[128][3] =
 			{
 				{ 0,0,0 }, { 1,0,0 }, { 2,0,0 }, { 3,0,0 }, { 4,0,0 }, { 0,4,0 }, { 4,4,0 }, { 4,4,4 }, { 0,1,0 }, { 1,1,0 }, { 2,1,0 }, { 3,1,0 }, { 4,1,0 }, { 1,4,0 }, { 4,4,1 }, { 4,4,4 },
 				{ 0,2,0 }, { 1,2,0 }, { 2,2,0 }, { 3,2,0 }, { 4,2,0 }, { 2,4,0 }, { 4,4,2 }, { 4,4,4 }, { 0,3,0 }, { 1,3,0 }, { 2,3,0 }, { 3,3,0 }, { 4,3,0 }, { 3,4,0 }, { 4,4,3 }, { 4,4,4 },
@@ -6403,20 +6361,36 @@ namespace impl
 				{ 0,0,3 }, { 1,0,3 }, { 2,0,3 }, { 3,0,3 }, { 4,0,3 }, { 0,4,3 }, { 0,0,4 }, { 1,0,4 }, { 0,1,3 }, { 1,1,3 }, { 2,1,3 }, { 3,1,3 }, { 4,1,3 }, { 1,4,3 }, { 0,1,4 }, { 1,1,4 },
 				{ 0,2,3 }, { 1,2,3 }, { 2,2,3 }, { 3,2,3 }, { 4,2,3 }, { 2,4,3 }, { 0,2,4 }, { 1,2,4 }, { 0,3,3 }, { 1,3,3 }, { 2,3,3 }, { 3,3,3 }, { 4,3,3 }, { 3,4,3 }, { 0,3,4 }, { 1,3,4 }
 			};
-			const deUint32(&quints)[3] = quintsFromQ[Q];
-			for (int i = 0; i < numValues; i++)
+			if (Q < 128)
 			{
-				dst[i].m = m[i];
-				dst[i].tq = quints[i];
-				dst[i].v = (quints[i] << numBits) + m[i];
+				const uint32(&quints)[3] = quintsFromQ[Q];
+				for (int i = 0; i < numValues; i++)
+				{
+					dst[i].m = m[i];
+					dst[i].tq = quints[i];
+					dst[i].v = (quints[i] << numBits) + m[i];
+				}
+			}
+			else
+			{
+				for (int i = 0; i < numValues; i++)
+				{
+					dst[i].m = 0;
+					dst[i].tq = 0;
+					dst[i].v = 0;
+				}
 			}
 		}
-		inline void decodeISEBitBlock(ISEDecodedResult* dst, BitAccessStream& data, int numBits)
+
+		template<bool Forward>
+		inline void decodeISEBitBlock(ISEDecodedResult* dst, BitAccessStream<Forward>& data, int numBits)
 		{
 			dst[0].m = data.getNext(numBits);
 			dst[0].v = dst[0].m;
 		}
-		void decodeISE(ISEDecodedResult* dst, int numValues, BitAccessStream& data, const ISEParams& params)
+
+		template<bool Forward>
+		void decodeISE(ISEDecodedResult* dst, int numValues, BitAccessStream<Forward>& data, const ISEParams& params)
 		{
 			if (params.mode == ISEMODE_TRIT)
 			{
@@ -6438,29 +6412,30 @@ namespace impl
 			}
 			else
 			{
-				DE_ASSERT(params.mode == ISEMODE_PLAIN_BIT);
+				check(params.mode == ISEMODE_PLAIN_BIT);
 				for (int i = 0; i < numValues; i++)
 					decodeISEBitBlock(&dst[i], data, params.numBits);
 			}
 		}
-		void unquantizeColorEndpoints(deUint32* dst, const ISEDecodedResult* iseResults, int numEndpoints, const ISEParams& iseParams)
+
+		void unquantizeColorEndpoints(uint32* dst, const ISEDecodedResult* iseResults, int numEndpoints, const ISEParams& iseParams)
 		{
 			if (iseParams.mode == ISEMODE_TRIT || iseParams.mode == ISEMODE_QUINT)
 			{
 				const int rangeCase = iseParams.numBits * 2 - (iseParams.mode == ISEMODE_TRIT ? 2 : 1);
-				DE_ASSERT(inRange(rangeCase, 0, 10));
-				static const deUint32	Ca[11] = { 204, 113, 93, 54, 44, 26, 22, 13, 11, 6, 5 };
-				const deUint32			C = Ca[rangeCase];
+				check(inRange(rangeCase, 0, 10));
+				static const uint32	Ca[11] = { 204, 113, 93, 54, 44, 26, 22, 13, 11, 6, 5 };
+				const uint32			C = Ca[rangeCase];
 				for (int endpointNdx = 0; endpointNdx < numEndpoints; endpointNdx++)
 				{
-					const deUint32 a = getBit(iseResults[endpointNdx].m, 0);
-					const deUint32 b = getBit(iseResults[endpointNdx].m, 1);
-					const deUint32 c = getBit(iseResults[endpointNdx].m, 2);
-					const deUint32 d = getBit(iseResults[endpointNdx].m, 3);
-					const deUint32 e = getBit(iseResults[endpointNdx].m, 4);
-					const deUint32 f = getBit(iseResults[endpointNdx].m, 5);
-					const deUint32 A = a == 0 ? 0 : (1 << 9) - 1;
-					const deUint32 B = rangeCase == 0 ? 0
+					const uint32 a = getBit(iseResults[endpointNdx].m, 0);
+					const uint32 b = getBit(iseResults[endpointNdx].m, 1);
+					const uint32 c = getBit(iseResults[endpointNdx].m, 2);
+					const uint32 d = getBit(iseResults[endpointNdx].m, 3);
+					const uint32 e = getBit(iseResults[endpointNdx].m, 4);
+					const uint32 f = getBit(iseResults[endpointNdx].m, 5);
+					const uint32 A = a == 0 ? 0 : (1 << 9) - 1;
+					const uint32 B = rangeCase == 0 ? 0
 						: rangeCase == 1 ? 0
 						: rangeCase == 2 ? (b << 8) | (b << 4) | (b << 2) | (b << 1)
 						: rangeCase == 3 ? (b << 8) | (b << 3) | (b << 2)
@@ -6471,19 +6446,20 @@ namespace impl
 						: rangeCase == 8 ? (e << 8) | (d << 7) | (c << 6) | (b << 5) | (e << 1) | (d << 0)
 						: rangeCase == 9 ? (e << 8) | (d << 7) | (c << 6) | (b << 5) | (e << 0)
 						: rangeCase == 10 ? (f << 8) | (e << 7) | (d << 6) | (c << 5) | (b << 4) | (f << 0)
-						: (deUint32)-1;
-					DE_ASSERT(B != (deUint32)-1);
+						: (uint32)-1;
+					check(B != (uint32)-1);
 					dst[endpointNdx] = (((iseResults[endpointNdx].tq * C + B) ^ A) >> 2) | (A & 0x80);
 				}
 			}
 			else
 			{
-				DE_ASSERT(iseParams.mode == ISEMODE_PLAIN_BIT);
+				check(iseParams.mode == ISEMODE_PLAIN_BIT);
 				for (int endpointNdx = 0; endpointNdx < numEndpoints; endpointNdx++)
 					dst[endpointNdx] = bitReplicationScale(iseResults[endpointNdx].v, iseParams.numBits, 8);
 			}
 		}
-		inline void bitTransferSigned(deInt32& a, deInt32& b)
+
+		inline void bitTransferSigned(int32& a, int32& b)
 		{
 			b >>= 1;
 			b |= a & 0x80;
@@ -6492,18 +6468,21 @@ namespace impl
 			if (isBitSet(a, 5))
 				a -= 0x40;
 		}
+
 		inline UVec4 clampedRGBA(const IVec4& rgba)
 		{
-			return UVec4(FMath::Clamp(rgba.x(), 0, 0xff),
-				FMath::Clamp(rgba.y(), 0, 0xff),
-				FMath::Clamp(rgba.z(), 0, 0xff),
-				FMath::Clamp(rgba.w(), 0, 0xff));
+			return UVec4(FMath::Clamp(rgba[0], 0, 0xff),
+				FMath::Clamp(rgba[1], 0, 0xff),
+				FMath::Clamp(rgba[2], 0, 0xff),
+				FMath::Clamp(rgba[3], 0, 0xff));
 		}
+
 		inline IVec4 blueContract(int r, int g, int b, int a)
 		{
 			return IVec4((r + b) >> 1, (g + b) >> 1, b, a);
 		}
-		inline bool isColorEndpointModeHDR(deUint32 mode)
+
+		inline bool isColorEndpointModeHDR(uint32 mode)
 		{
 			return mode == 2 ||
 				mode == 3 ||
@@ -6512,34 +6491,35 @@ namespace impl
 				mode == 14 ||
 				mode == 15;
 		}
-		void decodeHDREndpointMode7(UVec4& e0, UVec4& e1, deUint32 v0, deUint32 v1, deUint32 v2, deUint32 v3)
+
+		void decodeHDREndpointMode7(UVec4& e0, UVec4& e1, uint32 v0, uint32 v1, uint32 v2, uint32 v3)
 		{
-			const deUint32 m10 = getBit(v1, 7) | (getBit(v2, 7) << 1);
-			const deUint32 m23 = getBits(v0, 6, 7);
-			const deUint32 majComp = m10 != 3 ? m10
+			const uint32 m10 = getBit(v1, 7) | (getBit(v2, 7) << 1);
+			const uint32 m23 = getBits(v0, 6, 7);
+			const uint32 majComp = m10 != 3 ? m10
 				: m23 != 3 ? m23
 				: 0;
-			const deUint32 mode = m10 != 3 ? m23
+			const uint32 mode = m10 != 3 ? m23
 				: m23 != 3 ? 4
 				: 5;
-			deInt32			red = (deInt32)getBits(v0, 0, 5);
-			deInt32			green = (deInt32)getBits(v1, 0, 4);
-			deInt32			blue = (deInt32)getBits(v2, 0, 4);
-			deInt32			scale = (deInt32)getBits(v3, 0, 4);
+			int32			red = (int32)getBits(v0, 0, 5);
+			int32			green = (int32)getBits(v1, 0, 4);
+			int32			blue = (int32)getBits(v2, 0, 4);
+			int32			scale = (int32)getBits(v3, 0, 4);
 			{
 #define SHOR(DST_VAR, SHIFT, BIT_VAR) (DST_VAR) |= (BIT_VAR) << (SHIFT)
 #define ASSIGN_X_BITS(V0,S0, V1,S1, V2,S2, V3,S3, V4,S4, V5,S5, V6,S6) do { SHOR(V0,S0,x0); SHOR(V1,S1,x1); SHOR(V2,S2,x2); SHOR(V3,S3,x3); SHOR(V4,S4,x4); SHOR(V5,S5,x5); SHOR(V6,S6,x6); } while (false)
-				const deUint32	x0 = getBit(v1, 6);
-				const deUint32	x1 = getBit(v1, 5);
-				const deUint32	x2 = getBit(v2, 6);
-				const deUint32	x3 = getBit(v2, 5);
-				const deUint32	x4 = getBit(v3, 7);
-				const deUint32	x5 = getBit(v3, 6);
-				const deUint32	x6 = getBit(v3, 5);
-				deInt32& R = red;
-				deInt32& G = green;
-				deInt32& B = blue;
-				deInt32& S = scale;
+				const uint32	x0 = getBit(v1, 6);
+				const uint32	x1 = getBit(v1, 5);
+				const uint32	x2 = getBit(v2, 6);
+				const uint32	x3 = getBit(v2, 5);
+				const uint32	x4 = getBit(v3, 7);
+				const uint32	x5 = getBit(v3, 6);
+				const uint32	x6 = getBit(v3, 5);
+				int32& R = red;
+				int32& G = green;
+				int32& B = blue;
+				int32& S = scale;
 				switch (mode)
 				{
 				case 0: ASSIGN_X_BITS(R, 9, R, 8, R, 7, R, 10, R, 6, S, 6, S, 5); break;
@@ -6549,13 +6529,13 @@ namespace impl
 				case 4: ASSIGN_X_BITS(G, 6, G, 5, B, 6, B, 5, R, 6, R, 7, S, 5); break;
 				case 5: ASSIGN_X_BITS(G, 6, G, 5, B, 6, B, 5, R, 6, S, 6, S, 5); break;
 				default:
-					DE_ASSERT(false);
+					check(false);
 				}
 #undef ASSIGN_X_BITS
 #undef SHOR
 			}
 			static const int shiftAmounts[] = { 1, 1, 2, 3, 4, 5 };
-			DE_ASSERT(mode < DE_LENGTH_OF_ARRAY(shiftAmounts));
+			check(mode < DE_LENGTH_OF_ARRAY(shiftAmounts));
 			red <<= shiftAmounts[mode];
 			green <<= shiftAmounts[mode];
 			blue <<= shiftAmounts[mode];
@@ -6578,9 +6558,10 @@ namespace impl
 				FMath::Clamp(blue, 0, 0xfff),
 				0x780);
 		}
-		void decodeHDREndpointMode11(UVec4& e0, UVec4& e1, deUint32 v0, deUint32 v1, deUint32 v2, deUint32 v3, deUint32 v4, deUint32 v5)
+
+		void decodeHDREndpointMode11(UVec4& e0, UVec4& e1, uint32 v0, uint32 v1, uint32 v2, uint32 v3, uint32 v4, uint32 v5)
 		{
-			const deUint32 major = (getBit(v5, 7) << 1) | getBit(v4, 7);
+			const uint32 major = (getBit(v5, 7) << 1) | getBit(v4, 7);
 			if (major == 3)
 			{
 				e0 = UVec4(v0 << 4, v2 << 4, getBits(v4, 0, 6) << 5, 0x780);
@@ -6588,22 +6569,22 @@ namespace impl
 			}
 			else
 			{
-				const deUint32 mode = (getBit(v3, 7) << 2) | (getBit(v2, 7) << 1) | getBit(v1, 7);
-				deInt32 a = (deInt32)((getBit(v1, 6) << 8) | v0);
-				deInt32 c = (deInt32)(getBits(v1, 0, 5));
-				deInt32 b0 = (deInt32)(getBits(v2, 0, 5));
-				deInt32 b1 = (deInt32)(getBits(v3, 0, 5));
-				deInt32 d0 = (deInt32)(getBits(v4, 0, 4));
-				deInt32 d1 = (deInt32)(getBits(v5, 0, 4));
+				const uint32 mode = (getBit(v3, 7) << 2) | (getBit(v2, 7) << 1) | getBit(v1, 7);
+				int32 a = (int32)((getBit(v1, 6) << 8) | v0);
+				int32 c = (int32)(getBits(v1, 0, 5));
+				int32 b0 = (int32)(getBits(v2, 0, 5));
+				int32 b1 = (int32)(getBits(v3, 0, 5));
+				int32 d0 = (int32)(getBits(v4, 0, 4));
+				int32 d1 = (int32)(getBits(v5, 0, 4));
 				{
 #define SHOR(DST_VAR, SHIFT, BIT_VAR) (DST_VAR) |= (BIT_VAR) << (SHIFT)
 #define ASSIGN_X_BITS(V0,S0, V1,S1, V2,S2, V3,S3, V4,S4, V5,S5) do { SHOR(V0,S0,x0); SHOR(V1,S1,x1); SHOR(V2,S2,x2); SHOR(V3,S3,x3); SHOR(V4,S4,x4); SHOR(V5,S5,x5); } while (false)
-					const deUint32 x0 = getBit(v2, 6);
-					const deUint32 x1 = getBit(v3, 6);
-					const deUint32 x2 = getBit(v4, 6);
-					const deUint32 x3 = getBit(v5, 6);
-					const deUint32 x4 = getBit(v4, 5);
-					const deUint32 x5 = getBit(v5, 5);
+					const uint32 x0 = getBit(v2, 6);
+					const uint32 x1 = getBit(v3, 6);
+					const uint32 x2 = getBit(v4, 6);
+					const uint32 x3 = getBit(v5, 6);
+					const uint32 x4 = getBit(v4, 5);
+					const uint32 x5 = getBit(v5, 5);
 					switch (mode)
 					{
 					case 0: ASSIGN_X_BITS(b0, 6, b1, 6, d0, 6, d1, 6, d0, 5, d1, 5); break;
@@ -6615,13 +6596,13 @@ namespace impl
 					case 6: ASSIGN_X_BITS(b0, 6, b1, 6, a, 11, c, 6, a, 9, a, 10); break;
 					case 7: ASSIGN_X_BITS(a, 9, a, 10, a, 11, c, 6, d0, 5, d1, 5); break;
 					default:
-						DE_ASSERT(false);
+						check(false);
 					}
 #undef ASSIGN_X_BITS
 #undef SHOR
 				}
 				static const int numDBits[] = { 7, 6, 7, 6, 5, 6, 5, 6 };
-				DE_ASSERT(mode < DE_LENGTH_OF_ARRAY(numDBits));
+				check(mode < DE_LENGTH_OF_ARRAY(numDBits));
 				d0 = signExtend(d0, numDBits[mode]);
 				d1 = signExtend(d1, numDBits[mode]);
 				const int shiftAmount = (mode >> 1) ^ 3;
@@ -6641,26 +6622,27 @@ namespace impl
 					0x780);
 				if (major == 1)
 				{
-					std::swap(e0.x(), e0.y());
-					std::swap(e1.x(), e1.y());
+					Swap(e0[0], e0[1]);
+					Swap(e1[0], e1[1]);
 				}
 				else if (major == 2)
 				{
-					std::swap(e0.x(), e0.z());
-					std::swap(e1.x(), e1.z());
+					Swap(e0[0], e0[2]);
+					Swap(e1[0], e1[2]);
 				}
 			}
 		}
-		void decodeHDREndpointMode15(UVec4& e0, UVec4& e1, deUint32 v0, deUint32 v1, deUint32 v2, deUint32 v3, deUint32 v4, deUint32 v5, deUint32 v6In, deUint32 v7In)
+
+		void decodeHDREndpointMode15(UVec4& e0, UVec4& e1, uint32 v0, uint32 v1, uint32 v2, uint32 v3, uint32 v4, uint32 v5, uint32 v6In, uint32 v7In)
 		{
 			decodeHDREndpointMode11(e0, e1, v0, v1, v2, v3, v4, v5);
-			const deUint32	mode = (getBit(v7In, 7) << 1) | getBit(v6In, 7);
-			deInt32			v6 = (deInt32)getBits(v6In, 0, 6);
-			deInt32			v7 = (deInt32)getBits(v7In, 0, 6);
+			const uint32	mode = (getBit(v7In, 7) << 1) | getBit(v6In, 7);
+			int32			v6 = (int32)getBits(v6In, 0, 6);
+			int32			v7 = (int32)getBits(v7In, 0, 6);
 			if (mode == 3)
 			{
-				e0.w() = v6 << 5;
-				e1.w() = v7 << 5;
+				e0[3] = v6 << 5;
+				e1[3] = v7 << 5;
 			}
 			else
 			{
@@ -6672,17 +6654,18 @@ namespace impl
 				v7 <<= 4 - mode;
 				v7 += v6;
 				v7 = FMath::Clamp(v7, 0, 0xfff);
-				e0.w() = v6;
-				e1.w() = v7;
+				e0[3] = v6;
+				e1[3] = v7;
 			}
 		}
-		void decodeColorEndpoints(ColorEndpointPair* dst, const deUint32* unquantizedEndpoints, const deUint32* endpointModes, int numPartitions)
+
+		void decodeColorEndpoints(ColorEndpointPair* dst, const uint32* unquantizedEndpoints, const uint32* endpointModes, int numPartitions)
 		{
 			int unquantizedNdx = 0;
 			for (int partitionNdx = 0; partitionNdx < numPartitions; partitionNdx++)
 			{
-				const deUint32		endpointMode = endpointModes[partitionNdx];
-				const deUint32* v = &unquantizedEndpoints[unquantizedNdx];
+				const uint32		endpointMode = endpointModes[partitionNdx];
+				const uint32* v = &unquantizedEndpoints[unquantizedNdx];
 				UVec4& e0 = dst[partitionNdx].e0;
 				UVec4& e1 = dst[partitionNdx].e1;
 				unquantizedNdx += computeNumColorEndpointValues(endpointMode);
@@ -6694,17 +6677,17 @@ namespace impl
 					break;
 				case 1:
 				{
-					const deUint32 L0 = (v[0] >> 2) | (getBits(v[1], 6, 7) << 6);
-					const deUint32 L1 = FMath::Min(0xffu, L0 + getBits(v[1], 0, 5));
+					const uint32 L0 = (v[0] >> 2) | (getBits(v[1], 6, 7) << 6);
+					const uint32 L1 = FMath::Min(0xffu, L0 + getBits(v[1], 0, 5));
 					e0 = UVec4(L0, L0, L0, 0xff);
 					e1 = UVec4(L1, L1, L1, 0xff);
 					break;
 				}
 				case 2:
 				{
-					const deUint32 v1Gr = v[1] >= v[0];
-					const deUint32 y0 = v1Gr ? v[0] << 4 : (v[1] << 4) + 8;
-					const deUint32 y1 = v1Gr ? v[1] << 4 : (v[0] << 4) - 8;
+					const uint32 v1Gr = v[1] >= v[0];
+					const uint32 y0 = v1Gr ? v[0] << 4 : (v[1] << 4) + 8;
+					const uint32 y1 = v1Gr ? v[1] << 4 : (v[0] << 4) - 8;
 					e0 = UVec4(y0, y0, y0, 0x780);
 					e1 = UVec4(y1, y1, y1, 0x780);
 					break;
@@ -6712,11 +6695,11 @@ namespace impl
 				case 3:
 				{
 					const bool		m = isBitSet(v[0], 7);
-					const deUint32	y0 = m ? (getBits(v[1], 5, 7) << 9) | (getBits(v[0], 0, 6) << 2)
+					const uint32	y0 = m ? (getBits(v[1], 5, 7) << 9) | (getBits(v[0], 0, 6) << 2)
 						: (getBits(v[1], 4, 7) << 8) | (getBits(v[0], 0, 6) << 1);
-					const deUint32	d = m ? getBits(v[1], 0, 4) << 2
+					const uint32	d = m ? getBits(v[1], 0, 4) << 2
 						: getBits(v[1], 0, 3) << 1;
-					const deUint32	y1 = FMath::Min(0xfffu, y0 + d);
+					const uint32	y1 = FMath::Min(0xfffu, y0 + d);
 					e0 = UVec4(y0, y0, y0, 0x780);
 					e1 = UVec4(y1, y1, y1, 0x780);
 					break;
@@ -6727,10 +6710,10 @@ namespace impl
 					break;
 				case 5:
 				{
-					deInt32 v0 = (deInt32)v[0];
-					deInt32 v1 = (deInt32)v[1];
-					deInt32 v2 = (deInt32)v[2];
-					deInt32 v3 = (deInt32)v[3];
+					int32 v0 = (int32)v[0];
+					int32 v1 = (int32)v[1];
+					int32 v2 = (int32)v[2];
+					int32 v3 = (int32)v[3];
 					bitTransferSigned(v1, v0);
 					bitTransferSigned(v3, v2);
 					e0 = clampedRGBA(IVec4(v0, v0, v0, v2));
@@ -6752,18 +6735,18 @@ namespace impl
 					}
 					else
 					{
-						e0 = blueContract(v[1], v[3], v[5], 0xff).asUint();
-						e1 = blueContract(v[0], v[2], v[4], 0xff).asUint();
+						e0 = asUint(blueContract(v[1], v[3], v[5], 0xff));
+						e1 = asUint(blueContract(v[0], v[2], v[4], 0xff));
 					}
 					break;
 				case 9:
 				{
-					deInt32 v0 = (deInt32)v[0];
-					deInt32 v1 = (deInt32)v[1];
-					deInt32 v2 = (deInt32)v[2];
-					deInt32 v3 = (deInt32)v[3];
-					deInt32 v4 = (deInt32)v[4];
-					deInt32 v5 = (deInt32)v[5];
+					int32 v0 = (int32)v[0];
+					int32 v1 = (int32)v[1];
+					int32 v2 = (int32)v[2];
+					int32 v3 = (int32)v[3];
+					int32 v4 = (int32)v[4];
+					int32 v5 = (int32)v[5];
 					bitTransferSigned(v1, v0);
 					bitTransferSigned(v3, v2);
 					bitTransferSigned(v5, v4);
@@ -6800,14 +6783,14 @@ namespace impl
 					break;
 				case 13:
 				{
-					deInt32 v0 = (deInt32)v[0];
-					deInt32 v1 = (deInt32)v[1];
-					deInt32 v2 = (deInt32)v[2];
-					deInt32 v3 = (deInt32)v[3];
-					deInt32 v4 = (deInt32)v[4];
-					deInt32 v5 = (deInt32)v[5];
-					deInt32 v6 = (deInt32)v[6];
-					deInt32 v7 = (deInt32)v[7];
+					int32 v0 = (int32)v[0];
+					int32 v1 = (int32)v[1];
+					int32 v2 = (int32)v[2];
+					int32 v3 = (int32)v[3];
+					int32 v4 = (int32)v[4];
+					int32 v5 = (int32)v[5];
+					int32 v6 = (int32)v[6];
+					int32 v7 = (int32)v[7];
 					bitTransferSigned(v1, v0);
 					bitTransferSigned(v3, v2);
 					bitTransferSigned(v5, v4);
@@ -6826,32 +6809,34 @@ namespace impl
 				}
 				case 14:
 					decodeHDREndpointMode11(e0, e1, v[0], v[1], v[2], v[3], v[4], v[5]);
-					e0.w() = v[6];
-					e1.w() = v[7];
+					e0[3] = v[6];
+					e1[3] = v[7];
 					break;
 				case 15:
 					decodeHDREndpointMode15(e0, e1, v[0], v[1], v[2], v[3], v[4], v[5], v[6], v[7]);
 					break;
 				default:
-					DE_ASSERT(false);
+					check(false);
 				}
 			}
 		}
-		void computeColorEndpoints(ColorEndpointPair* dst, const Block128& blockData, const deUint32* endpointModes, int numPartitions, int numColorEndpointValues, const ISEParams& iseParams, int numBitsAvailable)
+
+		void computeColorEndpoints(ColorEndpointPair* dst, const Block128& blockData, const uint32* endpointModes, int numPartitions, int numColorEndpointValues, const ISEParams& iseParams, int numBitsAvailable)
 		{
 			const int			colorEndpointDataStart = numPartitions == 1 ? 17 : 29;
 			ISEDecodedResult	colorEndpointData[18];
 			{
-				BitAccessStream dataStream(blockData, colorEndpointDataStart, numBitsAvailable, true);
+				BitAccessStream<true> dataStream(blockData, colorEndpointDataStart, numBitsAvailable);
 				decodeISE(&colorEndpointData[0], numColorEndpointValues, dataStream, iseParams);
 			}
 			{
-				deUint32 unquantizedEndpoints[18];
+				uint32 unquantizedEndpoints[18];
 				unquantizeColorEndpoints(&unquantizedEndpoints[0], &colorEndpointData[0], numColorEndpointValues, iseParams);
 				decodeColorEndpoints(dst, &unquantizedEndpoints[0], &endpointModes[0], numPartitions);
 			}
 		}
-		void unquantizeWeights(deUint32 dst[64], const ISEDecodedResult* weightGrid, const ASTCBlockMode& blockMode)
+
+		void unquantizeWeights(uint32 dst[64], const ISEDecodedResult* weightGrid, const ASTCBlockMode& blockMode)
 		{
 			const int			numWeights = computeNumWeights(blockMode);
 			const ISEParams& iseParams = blockMode.weightISEParams;
@@ -6860,39 +6845,39 @@ namespace impl
 				const int rangeCase = iseParams.numBits * 2 + (iseParams.mode == ISEMODE_QUINT ? 1 : 0);
 				if (rangeCase == 0 || rangeCase == 1)
 				{
-					static const deUint32 map0[3] = { 0, 32, 63 };
-					static const deUint32 map1[5] = { 0, 16, 32, 47, 63 };
-					const deUint32* const map = rangeCase == 0 ? &map0[0] : &map1[0];
+					static const uint32 map0[3] = { 0, 32, 63 };
+					static const uint32 map1[5] = { 0, 16, 32, 47, 63 };
+					const uint32* const map = rangeCase == 0 ? &map0[0] : &map1[0];
 					for (int i = 0; i < numWeights; i++)
 					{
-						DE_ASSERT(weightGrid[i].v < (rangeCase == 0 ? 3u : 5u));
+						check(weightGrid[i].v < (rangeCase == 0 ? 3u : 5u));
 						dst[i] = map[weightGrid[i].v];
 					}
 				}
 				else
 				{
-					DE_ASSERT(rangeCase <= 6);
-					static const deUint32	Ca[5] = { 50, 28, 23, 13, 11 };
-					const deUint32			C = Ca[rangeCase - 2];
+					check(rangeCase <= 6);
+					static const uint32	Ca[5] = { 50, 28, 23, 13, 11 };
+					const uint32			C = Ca[rangeCase - 2];
 					for (int weightNdx = 0; weightNdx < numWeights; weightNdx++)
 					{
-						const deUint32 a = getBit(weightGrid[weightNdx].m, 0);
-						const deUint32 b = getBit(weightGrid[weightNdx].m, 1);
-						const deUint32 c = getBit(weightGrid[weightNdx].m, 2);
-						const deUint32 A = a == 0 ? 0 : (1 << 7) - 1;
-						const deUint32 B = rangeCase == 2 ? 0
+						const uint32 a = getBit(weightGrid[weightNdx].m, 0);
+						const uint32 b = getBit(weightGrid[weightNdx].m, 1);
+						const uint32 c = getBit(weightGrid[weightNdx].m, 2);
+						const uint32 A = a == 0 ? 0 : (1 << 7) - 1;
+						const uint32 B = rangeCase == 2 ? 0
 							: rangeCase == 3 ? 0
 							: rangeCase == 4 ? (b << 6) | (b << 2) | (b << 0)
 							: rangeCase == 5 ? (b << 6) | (b << 1)
 							: rangeCase == 6 ? (c << 6) | (b << 5) | (c << 1) | (b << 0)
-							: (deUint32)-1;
+							: (uint32)-1;
 						dst[weightNdx] = (((weightGrid[weightNdx].tq * C + B) ^ A) >> 2) | (A & 0x20);
 					}
 				}
 			}
 			else
 			{
-				DE_ASSERT(iseParams.mode == ISEMODE_PLAIN_BIT);
+				check(iseParams.mode == ISEMODE_PLAIN_BIT);
 				for (int weightNdx = 0; weightNdx < numWeights; weightNdx++)
 					dst[weightNdx] = bitReplicationScale(weightGrid[weightNdx].v, iseParams.numBits, 6);
 			}
@@ -6902,117 +6887,121 @@ namespace impl
 			for (int weightNdx = numWeights; weightNdx < 64; weightNdx++)
 				dst[weightNdx] = ~0u;
 		}
-		void interpolateWeights(TexelWeightPair* dst, const deUint32(&unquantizedWeights)[64], int blockWidth, int blockHeight, const ASTCBlockMode& blockMode)
+
+		void interpolateWeights(TexelWeightPair* dst, const uint32(&unquantizedWeights)[64], int blockWidth, int blockHeight, const ASTCBlockMode& blockMode)
 		{
 			const int		numWeightsPerTexel = blockMode.isDualPlane ? 2 : 1;
-			const deUint32	scaleX = (1024 + blockWidth / 2) / (blockWidth - 1);
-			const deUint32	scaleY = (1024 + blockHeight / 2) / (blockHeight - 1);
-			DE_ASSERT(blockMode.weightGridWidth * blockMode.weightGridHeight * numWeightsPerTexel <= DE_LENGTH_OF_ARRAY(unquantizedWeights));
+			const uint32	scaleX = (1024 + blockWidth / 2) / (blockWidth - 1);
+			const uint32	scaleY = (1024 + blockHeight / 2) / (blockHeight - 1);
+			miro_check(blockMode.weightGridWidth * blockMode.weightGridHeight * numWeightsPerTexel <= DE_LENGTH_OF_ARRAY(unquantizedWeights));
 			for (int texelY = 0; texelY < blockHeight; texelY++)
 			{
 				for (int texelX = 0; texelX < blockWidth; texelX++)
 				{
-					const deUint32 gX = (scaleX * texelX * (blockMode.weightGridWidth - 1) + 32) >> 6;
-					const deUint32 gY = (scaleY * texelY * (blockMode.weightGridHeight - 1) + 32) >> 6;
-					const deUint32 jX = gX >> 4;
-					const deUint32 jY = gY >> 4;
-					const deUint32 fX = gX & 0xf;
-					const deUint32 fY = gY & 0xf;
-					const deUint32 w11 = (fX * fY + 8) >> 4;
-					const deUint32 w10 = fY - w11;
-					const deUint32 w01 = fX - w11;
-					const deUint32 w00 = 16 - fX - fY + w11;
-					const deUint32 i00 = jY * blockMode.weightGridWidth + jX;
-					const deUint32 i01 = i00 + 1;
-					const deUint32 i10 = i00 + blockMode.weightGridWidth;
-					const deUint32 i11 = i00 + blockMode.weightGridWidth + 1;
+					const uint32 gX = (scaleX * texelX * (blockMode.weightGridWidth - 1) + 32) >> 6;
+					const uint32 gY = (scaleY * texelY * (blockMode.weightGridHeight - 1) + 32) >> 6;
+					const uint32 jX = gX >> 4;
+					const uint32 jY = gY >> 4;
+					const uint32 fX = gX & 0xf;
+					const uint32 fY = gY & 0xf;
+					const uint32 w11 = (fX * fY + 8) >> 4;
+					const uint32 w10 = fY - w11;
+					const uint32 w01 = fX - w11;
+					const uint32 w00 = 16 - fX - fY + w11;
+					const uint32 i00 = jY * blockMode.weightGridWidth + jX;
+					const uint32 i01 = i00 + 1;
+					const uint32 i10 = i00 + blockMode.weightGridWidth;
+					const uint32 i11 = i00 + blockMode.weightGridWidth + 1;
 					// These addresses can be out of bounds, but respective weights will be 0 then.
-					DE_ASSERT(deInBounds32(i00, 0, blockMode.weightGridWidth * blockMode.weightGridHeight) || w00 == 0);
-					DE_ASSERT(deInBounds32(i01, 0, blockMode.weightGridWidth * blockMode.weightGridHeight) || w01 == 0);
-					DE_ASSERT(deInBounds32(i10, 0, blockMode.weightGridWidth * blockMode.weightGridHeight) || w10 == 0);
-					DE_ASSERT(deInBounds32(i11, 0, blockMode.weightGridWidth * blockMode.weightGridHeight) || w11 == 0);
+					miro_check(deInBounds32(i00, 0, blockMode.weightGridWidth * blockMode.weightGridHeight) || w00 == 0);
+					miro_check(deInBounds32(i01, 0, blockMode.weightGridWidth * blockMode.weightGridHeight) || w01 == 0);
+					miro_check(deInBounds32(i10, 0, blockMode.weightGridWidth * blockMode.weightGridHeight) || w10 == 0);
+					miro_check(deInBounds32(i11, 0, blockMode.weightGridWidth * blockMode.weightGridHeight) || w11 == 0);
 					for (int texelWeightNdx = 0; texelWeightNdx < numWeightsPerTexel; texelWeightNdx++)
 					{
 						// & 0x3f clamps address to bounds of unquantizedWeights
-						const deUint32 p00 = unquantizedWeights[(i00 * numWeightsPerTexel + texelWeightNdx) & 0x3f];
-						const deUint32 p01 = unquantizedWeights[(i01 * numWeightsPerTexel + texelWeightNdx) & 0x3f];
-						const deUint32 p10 = unquantizedWeights[(i10 * numWeightsPerTexel + texelWeightNdx) & 0x3f];
-						const deUint32 p11 = unquantizedWeights[(i11 * numWeightsPerTexel + texelWeightNdx) & 0x3f];
+						const uint32 p00 = unquantizedWeights[(i00 * numWeightsPerTexel + texelWeightNdx) & 0x3f];
+						const uint32 p01 = unquantizedWeights[(i01 * numWeightsPerTexel + texelWeightNdx) & 0x3f];
+						const uint32 p10 = unquantizedWeights[(i10 * numWeightsPerTexel + texelWeightNdx) & 0x3f];
+						const uint32 p11 = unquantizedWeights[(i11 * numWeightsPerTexel + texelWeightNdx) & 0x3f];
 						dst[texelY * blockWidth + texelX].w[texelWeightNdx] = (p00 * w00 + p01 * w01 + p10 * w10 + p11 * w11 + 8) >> 4;
 					}
 				}
 			}
 		}
+
 		void computeTexelWeights(TexelWeightPair* dst, const Block128& blockData, int blockWidth, int blockHeight, const ASTCBlockMode& blockMode)
 		{
 			ISEDecodedResult weightGrid[64];
 			{
-				BitAccessStream dataStream(blockData, 127, computeNumRequiredBits(blockMode.weightISEParams, computeNumWeights(blockMode)), false);
+				BitAccessStream<false> dataStream(blockData, 127, computeNumRequiredBits(blockMode.weightISEParams, computeNumWeights(blockMode)));
 				decodeISE(&weightGrid[0], computeNumWeights(blockMode), dataStream, blockMode.weightISEParams);
 			}
 			{
-				deUint32 unquantizedWeights[64];
+				uint32 unquantizedWeights[64];
 				unquantizeWeights(&unquantizedWeights[0], &weightGrid[0], blockMode);
 				interpolateWeights(dst, unquantizedWeights, blockWidth, blockHeight, blockMode);
 			}
 		}
-		inline deUint32 hash52(deUint32 v)
+
+		inline uint32 hash52(uint32 v)
 		{
-			deUint32 p = v;
+			uint32 p = v;
 			p ^= p >> 15;	p -= p << 17;	p += p << 7;	p += p << 4;
 			p ^= p >> 5;	p += p << 16;	p ^= p >> 7;	p ^= p >> 3;
 			p ^= p << 6;	p ^= p >> 17;
 			return p;
 		}
-		int computeTexelPartition(deUint32 seedIn, deUint32 xIn, deUint32 yIn, deUint32 zIn, int numPartitions, bool smallBlock)
-		{
-			DE_ASSERT(zIn == 0);
-			const deUint32	x = smallBlock ? xIn << 1 : xIn;
-			const deUint32	y = smallBlock ? yIn << 1 : yIn;
-			const deUint32	z = smallBlock ? zIn << 1 : zIn;
-			const deUint32	seed = seedIn + 1024 * (numPartitions - 1);
-			const deUint32	rnum = hash52(seed);
-			deUint8			seed1 = (deUint8)(rnum & 0xf);
-			deUint8			seed2 = (deUint8)((rnum >> 4) & 0xf);
-			deUint8			seed3 = (deUint8)((rnum >> 8) & 0xf);
-			deUint8			seed4 = (deUint8)((rnum >> 12) & 0xf);
-			deUint8			seed5 = (deUint8)((rnum >> 16) & 0xf);
-			deUint8			seed6 = (deUint8)((rnum >> 20) & 0xf);
-			deUint8			seed7 = (deUint8)((rnum >> 24) & 0xf);
-			deUint8			seed8 = (deUint8)((rnum >> 28) & 0xf);
-			deUint8			seed9 = (deUint8)((rnum >> 18) & 0xf);
-			deUint8			seed10 = (deUint8)((rnum >> 22) & 0xf);
-			deUint8			seed11 = (deUint8)((rnum >> 26) & 0xf);
-			deUint8			seed12 = (deUint8)(((rnum >> 30) | (rnum << 2)) & 0xf);
-			seed1 = (deUint8)(seed1 * seed1);
-			seed2 = (deUint8)(seed2 * seed2);
-			seed3 = (deUint8)(seed3 * seed3);
-			seed4 = (deUint8)(seed4 * seed4);
-			seed5 = (deUint8)(seed5 * seed5);
-			seed6 = (deUint8)(seed6 * seed6);
-			seed7 = (deUint8)(seed7 * seed7);
-			seed8 = (deUint8)(seed8 * seed8);
-			seed9 = (deUint8)(seed9 * seed9);
-			seed10 = (deUint8)(seed10 * seed10);
-			seed11 = (deUint8)(seed11 * seed11);
-			seed12 = (deUint8)(seed12 * seed12);
+
+		int computeTexelPartition(uint32 seedIn, uint32 xIn, uint32 yIn, uint32 zIn, int numPartitions, bool smallBlock)
+		{
+			check(zIn == 0);
+			const uint32	x = smallBlock ? xIn << 1 : xIn;
+			const uint32	y = smallBlock ? yIn << 1 : yIn;
+			const uint32	z = smallBlock ? zIn << 1 : zIn;
+			const uint32	seed = seedIn + 1024 * (numPartitions - 1);
+			const uint32	rnum = hash52(seed);
+			uint8			seed1 = (uint8)(rnum & 0xf);
+			uint8			seed2 = (uint8)((rnum >> 4) & 0xf);
+			uint8			seed3 = (uint8)((rnum >> 8) & 0xf);
+			uint8			seed4 = (uint8)((rnum >> 12) & 0xf);
+			uint8			seed5 = (uint8)((rnum >> 16) & 0xf);
+			uint8			seed6 = (uint8)((rnum >> 20) & 0xf);
+			uint8			seed7 = (uint8)((rnum >> 24) & 0xf);
+			uint8			seed8 = (uint8)((rnum >> 28) & 0xf);
+			uint8			seed9 = (uint8)((rnum >> 18) & 0xf);
+			uint8			seed10 = (uint8)((rnum >> 22) & 0xf);
+			uint8			seed11 = (uint8)((rnum >> 26) & 0xf);
+			uint8			seed12 = (uint8)(((rnum >> 30) | (rnum << 2)) & 0xf);
+			seed1 = (uint8)(seed1 * seed1);
+			seed2 = (uint8)(seed2 * seed2);
+			seed3 = (uint8)(seed3 * seed3);
+			seed4 = (uint8)(seed4 * seed4);
+			seed5 = (uint8)(seed5 * seed5);
+			seed6 = (uint8)(seed6 * seed6);
+			seed7 = (uint8)(seed7 * seed7);
+			seed8 = (uint8)(seed8 * seed8);
+			seed9 = (uint8)(seed9 * seed9);
+			seed10 = (uint8)(seed10 * seed10);
+			seed11 = (uint8)(seed11 * seed11);
+			seed12 = (uint8)(seed12 * seed12);
 			const int shA = (seed & 2) != 0 ? 4 : 5;
 			const int shB = numPartitions == 3 ? 6 : 5;
 			const int sh1 = (seed & 1) != 0 ? shA : shB;
 			const int sh2 = (seed & 1) != 0 ? shB : shA;
 			const int sh3 = (seed & 0x10) != 0 ? sh1 : sh2;
-			seed1 = (deUint8)(seed1 >> sh1);
-			seed2 = (deUint8)(seed2 >> sh2);
-			seed3 = (deUint8)(seed3 >> sh1);
-			seed4 = (deUint8)(seed4 >> sh2);
-			seed5 = (deUint8)(seed5 >> sh1);
-			seed6 = (deUint8)(seed6 >> sh2);
-			seed7 = (deUint8)(seed7 >> sh1);
-			seed8 = (deUint8)(seed8 >> sh2);
-			seed9 = (deUint8)(seed9 >> sh3);
-			seed10 = (deUint8)(seed10 >> sh3);
-			seed11 = (deUint8)(seed11 >> sh3);
-			seed12 = (deUint8)(seed12 >> sh3);
+			seed1 = (uint8)(seed1 >> sh1);
+			seed2 = (uint8)(seed2 >> sh2);
+			seed3 = (uint8)(seed3 >> sh1);
+			seed4 = (uint8)(seed4 >> sh2);
+			seed5 = (uint8)(seed5 >> sh1);
+			seed6 = (uint8)(seed6 >> sh2);
+			seed7 = (uint8)(seed7 >> sh1);
+			seed8 = (uint8)(seed8 >> sh2);
+			seed9 = (uint8)(seed9 >> sh3);
+			seed10 = (uint8)(seed10 >> sh3);
+			seed11 = (uint8)(seed11 >> sh3);
+			seed12 = (uint8)(seed12 >> sh3);
 			const int a = 0x3f & (seed1 * x + seed2 * y + seed11 * z + (rnum >> 14));
 			const int b = 0x3f & (seed3 * x + seed4 * y + seed12 * z + (rnum >> 10));
 			const int c = numPartitions >= 3 ? 0x3f & (seed5 * x + seed6 * y + seed9 * z + (rnum >> 6)) : 0;
@@ -7022,12 +7011,13 @@ namespace impl
 				: c >= d ? 2
 				: 3;
 		}
-		DecompressResult setTexelColors(void* dst, ColorEndpointPair* colorEndpoints, TexelWeightPair* texelWeights, int ccs, deUint32 partitionIndexSeed,
-			int numPartitions, int blockWidth, int blockHeight, bool isSRGB, bool isLDRMode, const deUint32* colorEndpointModes)
+
+		template<int32 BlockSize>
+		DecompressResult setTexelColorsSRGB(void* dst, ColorEndpointPair* colorEndpoints, TexelWeightPair* texelWeights, int ccs, uint32 partitionIndexSeed, int numPartitions, const uint32* colorEndpointModes)
 		{
-			const bool			smallBlock = blockWidth * blockHeight < 31;
+			const bool			smallBlock = BlockSize * BlockSize < 31;
 			DecompressResult	result = DECOMPRESS_RESULT_VALID_BLOCK;
-			bool				isHDREndpoint[4] = {false,false,false,false};
+			bool				isHDREndpoint[4] = { false,false,false,false };
 			for (int i = 0; i < numPartitions; i++)
 			{
 				isHDREndpoint[i] = isColorEndpointModeHDR(colorEndpointModes[i]);
@@ -7037,89 +7027,50 @@ namespace impl
 					return DECOMPRESS_RESULT_ERROR;
 			}
 
-			for (int texelY = 0; texelY < blockHeight; texelY++)
-				for (int texelX = 0; texelX < blockWidth; texelX++)
+			for (int texelY = 0; texelY < BlockSize; texelY++)
+				for (int texelX = 0; texelX < BlockSize; texelX++)
 				{
-					const int				texelNdx = texelY * blockWidth + texelX;
-					int						colorEndpointNdx = numPartitions == 1 ? 0 : computeTexelPartition(partitionIndexSeed, texelX, texelY, 0, numPartitions, smallBlock);
-					DE_ASSERT(colorEndpointNdx < numPartitions);
-					colorEndpointNdx = FMath::Clamp<int>(colorEndpointNdx,0, FMath::Min(3,numPartitions));
+					const int texelNdx = texelY * BlockSize + texelX;
+					int	 colorEndpointNdx = numPartitions == 1 ? 0 : computeTexelPartition(partitionIndexSeed, texelX, texelY, 0, numPartitions, smallBlock);
+					check(colorEndpointNdx < numPartitions);
+					colorEndpointNdx = FMath::Clamp<int>(colorEndpointNdx, 0, FMath::Min(3, numPartitions));
 					const UVec4& e0 = colorEndpoints[colorEndpointNdx].e0;
 					const UVec4& e1 = colorEndpoints[colorEndpointNdx].e1;
 					const TexelWeightPair& weight = texelWeights[texelNdx];
-					if (isLDRMode && isHDREndpoint[colorEndpointNdx] && false) // REMOVING HDR SUPPORT FOR NOW
-					{
-						if (isSRGB)
-						{
-							((deUint8*)dst)[texelNdx * 4 + 0] = 0xff;
-							((deUint8*)dst)[texelNdx * 4 + 1] = 0;
-							((deUint8*)dst)[texelNdx * 4 + 2] = 0xff;
-							((deUint8*)dst)[texelNdx * 4 + 3] = 0xff;
-						}
-						else
-						{
-							((float*)dst)[texelNdx * 4 + 0] = 1.0f;
-							((float*)dst)[texelNdx * 4 + 1] = 0;
-							((float*)dst)[texelNdx * 4 + 2] = 1.0f;
-							((float*)dst)[texelNdx * 4 + 3] = 1.0f;
-						}
-						result = DECOMPRESS_RESULT_ERROR;
-					}
-					else 
+
+					for (int channelNdx = 0; channelNdx < 4; channelNdx++)
 					{
-						for (int channelNdx = 0; channelNdx < 4; channelNdx++)
-						{
-							if (!isHDREndpoint[colorEndpointNdx] || (channelNdx == 3 && colorEndpointModes[colorEndpointNdx] == 14) || true) // \note Alpha for mode 14 is treated the same as LDR. // always true because there HDR is currently not supported.
-							{
-								const deUint32 c0 = (e0[channelNdx] << 8) | (isSRGB ? 0x80 : e0[channelNdx]);
-								const deUint32 c1 = (e1[channelNdx] << 8) | (isSRGB ? 0x80 : e1[channelNdx]);
-								const deUint32 w = weight.w[ccs == channelNdx ? 1 : 0];
-								const deUint32 c = (c0 * (64 - w) + c1 * w + 32) / 64;
-								if (isSRGB)
-									((deUint8*)dst)[texelNdx * 4 + channelNdx] = (deUint8)((c & 0xff00) >> 8);
-								else
-									((float*)dst)[texelNdx * 4 + channelNdx] = c == 65535 ? 1.0f : (float)c / 65536.0f;
-							}
-							else
-							{
-								//DE_STATIC_ASSERT((meta::TypesSame<deFloat16, deUint16>::Value));
-								// rg - REMOVING HDR SUPPORT FOR NOW
-#if 0
-								const deUint32		c0 = e0[channelNdx] << 4;
-								const deUint32		c1 = e1[channelNdx] << 4;
-								const deUint32		w = weight.w[ccs == channelNdx ? 1 : 0];
-								const deUint32		c = (c0 * (64 - w) + c1 * w + 32) / 64;
-								const deUint32		e = getBits(c, 11, 15);
-								const deUint32		m = getBits(c, 0, 10);
-								const deUint32		mt = m < 512 ? 3 * m
-									: m >= 1536 ? 5 * m - 2048
-									: 4 * m - 512;
-								const deFloat16		cf = (deFloat16)((e << 10) + (mt >> 3));
-								((float*)dst)[texelNdx * 4 + channelNdx] = deFloat16To32(isFloat16InfOrNan(cf) ? 0x7bff : cf);
-#endif
-							}
-						}
+						const uint32 c0 = (e0[channelNdx] << 8) | 0x80;
+						const uint32 c1 = (e1[channelNdx] << 8) | 0x80;
+						const uint32 w = weight.w[ccs == channelNdx ? 1 : 0];
+						const uint32 c = (c0 * (64 - w) + c1 * w + 32) / 64;
+						((uint8*)dst)[texelNdx * 4 + channelNdx] = (uint8)((c & 0xff00) >> 8);
 					}
 				}
 			return result;
 		}
 
-
-		DecompressResult decompressBlock(void* dst, const Block128& blockData, int blockWidth, int blockHeight, bool isSRGB, bool isLDR)
+		template<int32 BlockSize>
+		void decompressBlockSRGB(void* dst, const Block128& blockData)
 		{
-			DE_ASSERT(isLDR || !isSRGB);
+			bool isLDR = true;
+			bool isSRGB = true;
 
 			// Decode block mode.
 			const ASTCBlockMode blockMode = getASTCBlockMode(blockData.getBits(0, 10));
+
 			// Check for block mode errors.
 			if (blockMode.isError)
 			{
-				setASTCErrorColorBlock(dst, blockWidth, blockHeight, isSRGB);
-				return DECOMPRESS_RESULT_ERROR;
+				setASTCErrorColorBlock(dst, BlockSize, BlockSize, isSRGB);
+				return;
 			}
 			// Separate path for void-extent.
 			if (blockMode.isVoidExtent)
-				return decodeVoidExtentBlock(dst, blockData, blockWidth, blockHeight, isSRGB, isLDR);
+			{
+				return decodeVoidExtentBlockSRGB<BlockSize>(dst, blockData);
+			}
+
 			// Compute weight grid values.
 			const int numWeights = computeNumWeights(blockMode);
 			const int numWeightDataBits = computeNumRequiredBits(blockMode.weightISEParams, numWeights);
@@ -7128,12 +7079,12 @@ namespace impl
 			if (numWeights > 64 ||
 				numWeightDataBits > 96 ||
 				numWeightDataBits < 24 ||
-				blockMode.weightGridWidth > blockWidth ||
-				blockMode.weightGridHeight > blockHeight ||
+				blockMode.weightGridWidth > BlockSize ||
+				blockMode.weightGridHeight > BlockSize ||
 				(numPartitions == 4 && blockMode.isDualPlane))
 			{
-				setASTCErrorColorBlock(dst, blockWidth, blockHeight, isSRGB);
-				return DECOMPRESS_RESULT_ERROR;
+				setASTCErrorColorBlock(dst, BlockSize, BlockSize, isSRGB);
+				return;
 			}
 			// Compute number of bits available for color endpoint data.
 			const bool	isSingleUniqueCem = numPartitions == 1 || blockData.getBits(23, 24) == 0;
@@ -7146,14 +7097,14 @@ namespace impl
 				: numPartitions == 2 ? 1
 				: 0);
 			// Decode color endpoint modes.
-			deUint32 colorEndpointModes[4];
+			uint32 colorEndpointModes[4];
 			decodeColorEndpointModes(&colorEndpointModes[0], blockData, numPartitions, extraCemBitsStart);
 			const int numColorEndpointValues = computeNumColorEndpointValues(colorEndpointModes, numPartitions);
 			// Check for errors in color endpoint value count.
 			if (numColorEndpointValues > 18 || numBitsForColorEndpoints < (int)deDivRoundUp32(13 * numColorEndpointValues, 5))
 			{
-				setASTCErrorColorBlock(dst, blockWidth, blockHeight, isSRGB);
-				return DECOMPRESS_RESULT_ERROR;
+				setASTCErrorColorBlock(dst, BlockSize, BlockSize, isSRGB);
+				return;
 			}
 			// Compute color endpoints.
 			ColorEndpointPair colorEndpoints[4];
@@ -7161,19 +7112,18 @@ namespace impl
 				computeMaximumRangeISEParams(numBitsForColorEndpoints, numColorEndpointValues), numBitsForColorEndpoints);
 			// Compute texel weights.
 			TexelWeightPair texelWeights[MAX_BLOCK_WIDTH * MAX_BLOCK_HEIGHT];
-			computeTexelWeights(&texelWeights[0], blockData, blockWidth, blockHeight, blockMode);
+			computeTexelWeights(&texelWeights[0], blockData, BlockSize, BlockSize, blockMode);
 			// Set texel colors.
 			const int		ccs = blockMode.isDualPlane ? (int)blockData.getBits(extraCemBitsStart - 2, extraCemBitsStart - 1) : -1;
-			const deUint32	partitionIndexSeed = numPartitions > 1 ? blockData.getBits(13, 22) : (deUint32)-1;
-			return setTexelColors(dst, &colorEndpoints[0], &texelWeights[0], ccs, partitionIndexSeed, numPartitions, blockWidth, blockHeight, isSRGB, isLDR, &colorEndpointModes[0]);
+			const uint32	partitionIndexSeed = numPartitions > 1 ? blockData.getBits(13, 22) : (uint32)-1;
+			setTexelColorsSRGB<BlockSize>(dst, &colorEndpoints[0], &texelWeights[0], ccs, partitionIndexSeed, numPartitions, &colorEndpointModes[0]);
 		}
 
 
 #ifdef UE_MIRO_DEBUG
-		void LogBlock(const Block128& blockData, int blockWidth, int blockHeight, bool isSRGB, bool isLDR)
+		template<int32 BlockSize>
+		void LogBlock(const Block128& blockData)
 		{
-			DE_ASSERT(isLDR || !isSRGB);
-
 			// Decode block mode.
 			uint32 BlockMode = blockData.getBits(0, 10);
 			const ASTCBlockMode blockMode = getASTCBlockMode(BlockMode);
@@ -7191,7 +7141,7 @@ namespace impl
   ((byte) & 0x002 ? L'1' : L'0'), \
   ((byte) & 0x001 ? L'1' : L'0') 
 
-			UE_LOG(LogMutableCore, Log, TEXT("block_mode : %c%c %c%c%c%c %c%c%c%c%c"), MODE_TO_BINARY(BlockMode) );
+			UE_LOG(LogMutableCore, Log, TEXT("block_mode : %c%c %c%c%c%c %c%c%c%c%c"), MODE_TO_BINARY(BlockMode));
 
 
 			// Check for block mode errors.
@@ -7220,8 +7170,8 @@ namespace impl
 			if (numWeights > 64 ||
 				numWeightDataBits > 96 ||
 				numWeightDataBits < 24 ||
-				blockMode.weightGridWidth > blockWidth ||
-				blockMode.weightGridHeight > blockHeight ||
+				blockMode.weightGridWidth > BlockSize ||
+				blockMode.weightGridHeight > BlockSize ||
 				(numPartitions == 4 && blockMode.isDualPlane))
 			{
 				miro_check(false);
@@ -7239,7 +7189,7 @@ namespace impl
 				: 0);
 
 			// Decode color endpoint modes.
-			deUint32 colorEndpointModes[4];
+			uint32 colorEndpointModes[4];
 			decodeColorEndpointModes(&colorEndpointModes[0], blockData, numPartitions, extraCemBitsStart);
 			const int numColorEndpointValues = computeNumColorEndpointValues(colorEndpointModes, numPartitions);
 
@@ -7257,9 +7207,9 @@ namespace impl
 			computeColorEndpoints(&colorEndpoints[0], blockData, &colorEndpointModes[0], numPartitions, numColorEndpointValues,
 				computeMaximumRangeISEParams(numBitsForColorEndpoints, numColorEndpointValues), numBitsForColorEndpoints);
 
-			UE_LOG(LogMutableCore, Log, TEXT("endpoints : ( %3d, %3d, %3d, %3d ), ( %3d, %3d, %3d, %3d )"), 
+			UE_LOG(LogMutableCore, Log, TEXT("endpoints : ( %3d, %3d, %3d, %3d ), ( %3d, %3d, %3d, %3d )"),
 				colorEndpoints[0].e0[0], colorEndpoints[0].e0[1], colorEndpoints[0].e0[2], colorEndpoints[0].e0[3],
-				colorEndpoints[0].e1[0], colorEndpoints[0].e1[1], colorEndpoints[0].e1[2], colorEndpoints[0].e1[3] );
+				colorEndpoints[0].e1[0], colorEndpoints[0].e1[1], colorEndpoints[0].e1[2], colorEndpoints[0].e1[3]);
 
 			// Compute texel weights.
 			TexelWeightPair texelWeights[MAX_BLOCK_WIDTH * MAX_BLOCK_HEIGHT];
@@ -7267,13 +7217,13 @@ namespace impl
 
 			ISEDecodedResult weightGrid[64];
 			{
-				BitAccessStream dataStream(blockData, 127, numWeightDataBits, false);
+				BitAccessStream<false> dataStream(blockData, 127, numWeightDataBits);
 				decodeISE(&weightGrid[0], computeNumWeights(blockMode), dataStream, blockMode.weightISEParams);
 			}
-			deUint32 unquantizedWeights[64];
-			{				
+			uint32 unquantizedWeights[64];
+			{
 				unquantizeWeights(&unquantizedWeights[0], &weightGrid[0], blockMode);
-				interpolateWeights(&texelWeights[0], unquantizedWeights, blockWidth, blockHeight, blockMode);
+				interpolateWeights(&texelWeights[0], unquantizedWeights, BlockSize, BlockSize, blockMode);
 			}
 
 			int i = 0;
@@ -7290,17 +7240,17 @@ namespace impl
 
 			// Set texel colors.
 			const int		ccs = blockMode.isDualPlane ? (int)blockData.getBits(extraCemBitsStart - 2, extraCemBitsStart - 1) : -1;
-			const deUint32	partitionIndexSeed = numPartitions > 1 ? blockData.getBits(13, 22) : (deUint32)-1;
+			const uint32	partitionIndexSeed = numPartitions > 1 ? blockData.getBits(13, 22) : (uint32)-1;
 			astcrt::vec4f_t dst[MAX_BLOCK_WIDTH * MAX_BLOCK_HEIGHT];
-			DecompressResult Result = setTexelColors(dst, &colorEndpoints[0], &texelWeights[0], ccs, partitionIndexSeed, numPartitions, blockWidth, blockHeight, isSRGB, isLDR, &colorEndpointModes[0]);
-			check(Result==DecompressResult::DECOMPRESS_RESULT_VALID_BLOCK);
+			DecompressResult Result = setTexelColorsSRGB<BlockSize>(dst, &colorEndpoints[0], &texelWeights[0], ccs, partitionIndexSeed, numPartitions, &colorEndpointModes[0]);
+			check(Result == DecompressResult::DECOMPRESS_RESULT_VALID_BLOCK);
 
 			UE_LOG(LogMutableCore, Log, TEXT("decoded texels:"));
 			i = 0;
-			for (int32 Row = 0; Row < blockHeight; ++Row)
+			for (int32 Row = 0; Row < BlockSize; ++Row)
 			{
 				FString Line;
-				for (uint64 X = 0; X < blockWidth; ++X)
+				for (uint64 X = 0; X < BlockSize; ++X)
 				{
 					astcrt::vec4i_t texel;
 					texel.components(0) = FMath::Clamp<int>((int)(dst[i].components(0) * 65536.0f + .5f), 0, 65535) >> 8;
@@ -7319,55 +7269,21 @@ namespace impl
 		}
 #endif //UE_MIRO_DEBUG
 
-		bool decompress(uint8_t* pDst, const uint8_t* data, bool isSRGB, int blockWidth, int blockHeight)
+		template<int32 BlockSize>
+		void decompressSRGB(uint8* Dest, const uint8* Block)
 		{
-			// rg - We only support LDR here, although adding back in HDR would be easy.
-			const bool isLDR = true;
-			DE_ASSERT(isLDR || !isSRGB);
-
-			float linear[MAX_BLOCK_WIDTH * MAX_BLOCK_HEIGHT * 4];
-
-			const Block128 blockData(data);
+			const Block128 blockData(Block);
 
 			// (anticto) shortcut for blank blocks. This is not correct really, an all-zero block is invalid but we support it 
 			// so that zero-initialized ASTC textures are actually black.
 			if (blockData.isZero())
 			{
-				int32 pix = 0;
-				for (int32 i = 0; i < blockHeight* blockWidth; i++)
-				{
-					pDst[4 * i + 0] = 0;
-					pDst[4 * i + 1] = 0;
-					pDst[4 * i + 2] = 0;
-					pDst[4 * i + 3] = 0;
-				}
-				return true;
-			}
-
-			if (decompressBlock(isSRGB ? (void*)pDst : (void*)&linear[0],
-				blockData, blockWidth, blockHeight, isSRGB, isLDR) != DECOMPRESS_RESULT_VALID_BLOCK)
-			{
-				// Invalid ASTC block.
-				miro_check(false);
-				return false;
+				FMemory::Memzero(Dest, BlockSize * BlockSize * 4);
 			}
-
-			if (!isSRGB)
+			else
 			{
-				int pix = 0;
-				for (int i = 0; i < blockHeight; i++)
-				{
-					for (int j = 0; j < blockWidth; j++, pix++)
-					{
-						pDst[4 * pix + 0] = (uint8_t)(FMath::Clamp<int>((int)(linear[pix * 4 + 0] * 65536.0f + .5f), 0, 65535) >> 8);
-						pDst[4 * pix + 1] = (uint8_t)(FMath::Clamp<int>((int)(linear[pix * 4 + 1] * 65536.0f + .5f), 0, 65535) >> 8);
-						pDst[4 * pix + 2] = (uint8_t)(FMath::Clamp<int>((int)(linear[pix * 4 + 2] * 65536.0f + .5f), 0, 65535) >> 8);
-						pDst[4 * pix + 3] = (uint8_t)(FMath::Clamp<int>((int)(linear[pix * 4 + 3] * 65536.0f + .5f), 0, 65535) >> 8);
-					}
-				}
+				decompressBlockSRGB<BlockSize>(Dest, blockData);
 			}
-
-			return true;
 		}
 
 	}
@@ -8507,24 +8423,21 @@ namespace miro
 
 
 	//---------------------------------------------------------------------------------------------
-	template<uint32 BLOCK_SIZE>
+	template<uint32 BlockSize>
 	void Generic_ASTCRGBL_to_RGB(uint32 sx, uint32 sy, const uint8* from, uint8* to)
 	{
 		init_astc_decompress();
 
-		for (uint32 y = 0; y < sy; y += BLOCK_SIZE)
+		for (uint32 y = 0; y < sy; y += BlockSize)
 		{
-			for (uint32 x = 0; x < sx; x += BLOCK_SIZE)
+			for (uint32 x = 0; x < sx; x += BlockSize)
 			{
-				bool bIsSRGB = false;
-				uint8 Block[BLOCK_SIZE * BLOCK_SIZE * 4];
-				bool bSuccess = astcdec::decompress(Block, from, bIsSRGB, BLOCK_SIZE, BLOCK_SIZE);
-				miro_check(bSuccess);
-
+				uint8 Block[BlockSize * BlockSize * 4];
+				astcdec::decompressSRGB<BlockSize>(Block, from);
 
-				for (uint32 py = 0; py < BLOCK_SIZE; py++)
+				for (uint32 py = 0; py < BlockSize; py++)
 				{
-					for (uint32 px = 0; px < BLOCK_SIZE; px++)
+					for (uint32 px = 0; px < BlockSize; px++)
 					{
 						uint32 xi = x + px;
 						uint32 yi = y + py;
@@ -8532,9 +8445,9 @@ namespace miro
 						if (xi < sx && yi < sy)
 						{
 							uint8* toPixel = to + (yi * sx + xi) * 3;
-							toPixel[0] = Block[py * BLOCK_SIZE * 4 + px * 4 + 0];
-							toPixel[1] = Block[py * BLOCK_SIZE * 4 + px * 4 + 1];
-							toPixel[2] = Block[py * BLOCK_SIZE * 4 + px * 4 + 2];
+							toPixel[0] = Block[py * BlockSize * 4 + px * 4 + 0];
+							toPixel[1] = Block[py * BlockSize * 4 + px * 4 + 1];
+							toPixel[2] = Block[py * BlockSize * 4 + px * 4 + 2];
 						}
 					}
 				}
@@ -8547,23 +8460,21 @@ namespace miro
 
 
 	//---------------------------------------------------------------------------------------------
-	template<uint32 BLOCK_SIZE>
+	template<uint32 BlockSize>
 	void Generic_ASTCRGBL_to_RGBA(uint32 sx, uint32 sy, const uint8* from, uint8* to)
 	{
 		init_astc_decompress();
 
-		for (uint32 y = 0; y < sy; y += BLOCK_SIZE)
+		for (uint32 y = 0; y < sy; y += BlockSize)
 		{
-			for (uint32 x = 0; x < sx; x += BLOCK_SIZE)
+			for (uint32 x = 0; x < sx; x += BlockSize)
 			{
-				bool bIsSRGB = false;
-				uint8 Block[BLOCK_SIZE * BLOCK_SIZE * 4];
-				bool bSuccess = astcdec::decompress(Block, from, bIsSRGB, BLOCK_SIZE, BLOCK_SIZE);
-				miro_check(bSuccess);
+				uint8 Block[BlockSize * BlockSize * 4];
+				astcdec::decompressSRGB<BlockSize>(Block, from);
 
-				for (uint32 py = 0; py < BLOCK_SIZE; py++)
+				for (uint32 py = 0; py < BlockSize; py++)
 				{
-					for (uint32 px = 0; px < BLOCK_SIZE; px++)
+					for (uint32 px = 0; px < BlockSize; px++)
 					{
 						uint32 xi = x + px;
 						uint32 yi = y + py;
@@ -8571,9 +8482,9 @@ namespace miro
 						if (xi < sx && yi < sy)
 						{
 							uint8* toPixel = to + (yi * sx + xi) * 4;
-							toPixel[0] = Block[py * BLOCK_SIZE * 4 + px * 4 + 0];
-							toPixel[1] = Block[py * BLOCK_SIZE * 4 + px * 4 + 1];
-							toPixel[2] = Block[py * BLOCK_SIZE * 4 + px * 4 + 2];
+							toPixel[0] = Block[py * BlockSize * 4 + px * 4 + 0];
+							toPixel[1] = Block[py * BlockSize * 4 + px * 4 + 1];
+							toPixel[2] = Block[py * BlockSize * 4 + px * 4 + 2];
 							toPixel[3] = 255;
 						}
 					}
@@ -8596,15 +8507,15 @@ namespace miro
 
 		uint32 bx = FMath::DivideAndRoundUp(sx, BLOCK_SIZE);
 		uint32 by = FMath::DivideAndRoundUp(sy, BLOCK_SIZE);
-		
+
 #ifdef UE_MIRO_DEBUG
 		s_CurrentBlock = 0;
-		for ( uint32 y=0; y<by; ++y )
+		for (uint32 y = 0; y < by; ++y)
 #else
-		ParallelFor(by, [ bx, sx, sy, from, to, &physical_block_zero ] (uint32 y)
+		ParallelFor(by, [bx, sx, sy, from, to, &physical_block_zero](uint32 y)
 #endif
 			{
-				astcrt::PhysicalBlock * rowTo = reinterpret_cast<astcrt::PhysicalBlock*>(to + sizeof(astcrt::PhysicalBlock) * bx * y);
+				astcrt::PhysicalBlock* rowTo = reinterpret_cast<astcrt::PhysicalBlock*>(to + sizeof(astcrt::PhysicalBlock) * bx * y);
 
 				for (uint32 x = 0; x < bx; ++x)
 				{
@@ -8640,8 +8551,8 @@ namespace miro
 				}
 			}
 #else
-				}
-			});
+	}
+});
 #endif
 	}
 
@@ -8700,13 +8611,13 @@ namespace miro
 
 
 	//---------------------------------------------------------------------------------------------
-	template<uint32 BLOCK_SIZE>
+	template<uint32 BlockSize>
 	void Generic_ASTCRGBAL_to_RGBA(uint32 sx, uint32 sy, const uint8* from, uint8* to)
 	{
 		init_astc_decompress();
-	
-		const int32 NumBlocksY = FMath::DivideAndRoundUp(sy, BLOCK_SIZE);
-		const int32 NumBlocksX = FMath::DivideAndRoundUp(sx, BLOCK_SIZE);
+
+		const int32 NumBlocksY = FMath::DivideAndRoundUp(sy, BlockSize);
+		const int32 NumBlocksX = FMath::DivideAndRoundUp(sx, BlockSize);
 
 		if (sx <= 0 || sy <= 0)
 		{
@@ -8719,92 +8630,88 @@ namespace miro
 #else
 		ParallelFor(NumBlocksY, [NumBlocksX, sx, sy, from, to](int32 BlockY)
 #endif
-		{
-			for (int32 BlockX = 0; BlockX < NumBlocksX; ++BlockX)
 			{
-				constexpr bool bIsSRGB = false;
-				uint8 Block[BLOCK_SIZE * BLOCK_SIZE * 4];
+				for (int32 BlockX = 0; BlockX < NumBlocksX; ++BlockX)
+				{
+					uint8 Block[BlockSize * BlockSize * 4];
 
-				constexpr int32 CompressedBlockSize = 16;
-				const uint8* SrcBlockPtr = from + (BlockY * NumBlocksX + BlockX) * CompressedBlockSize;
+					constexpr int32 CompressedBlockSize = 16;
+					const uint8* SrcBlockPtr = from + (BlockY * NumBlocksX + BlockX) * CompressedBlockSize;
 
 #ifdef UE_MIRO_DEBUG
-				if (s_CurrentBlock==s_DebugBlock)
-				{
-					astcdec::LogBlock(SrcBlockPtr, BLOCK_SIZE, BLOCK_SIZE, bIsSRGB, true );
-				}
+					if (s_CurrentBlock == s_DebugBlock)
+					{
+						astcdec::LogBlock<BlockSize>(SrcBlockPtr);
+					}
 #endif
 
-				bool bSuccess = astcdec::decompress(Block, SrcBlockPtr, bIsSRGB, BLOCK_SIZE, BLOCK_SIZE);
-				miro_check(bSuccess);
+					astcdec::decompressSRGB<BlockSize>(Block, SrcBlockPtr);
 
-				for (uint32 py = 0; py < BLOCK_SIZE; py++)
-				{
-					for (uint32 px = 0; px < BLOCK_SIZE; px++)
+					for (uint32 py = 0; py < BlockSize; py++)
 					{
-						uint32 xi = FMath::Min(BlockX * BLOCK_SIZE + px, sx - 1);
-						uint32 yi = FMath::Min(BlockY * BLOCK_SIZE + py, sy - 1);
-
-						uint8* toPixel = to + (yi * sx + xi) * 4;
-						toPixel[0] = Block[py * BLOCK_SIZE * 4 + px * 4 + 0];
-						toPixel[1] = Block[py * BLOCK_SIZE * 4 + px * 4 + 1];
-						toPixel[2] = Block[py * BLOCK_SIZE * 4 + px * 4 + 2];
-						toPixel[3] = Block[py * BLOCK_SIZE * 4 + px * 4 + 3];
+						for (uint32 px = 0; px < BlockSize; px++)
+						{
+							uint32 xi = FMath::Min(BlockX * BlockSize + px, sx - 1);
+							uint32 yi = FMath::Min(BlockY * BlockSize + py, sy - 1);
+
+							uint8* toPixel = to + (yi * sx + xi) * 4;
+							toPixel[0] = Block[py * BlockSize * 4 + px * 4 + 0];
+							toPixel[1] = Block[py * BlockSize * 4 + px * 4 + 1];
+							toPixel[2] = Block[py * BlockSize * 4 + px * 4 + 2];
+							toPixel[3] = Block[py * BlockSize * 4 + px * 4 + 3];
+						}
 					}
-				}
 
 #ifdef UE_MIRO_DEBUG
-				++s_CurrentBlock;
+					++s_CurrentBlock;
+				}
 			}
-		}
 #else
-			}
+	}
 		});
 #endif
 	}
 
 
 	//---------------------------------------------------------------------------------------------
-	template<uint32 BLOCK_SIZE>
+	template<uint32 BlockSize>
 	void Generic_ASTCRGBAL_to_RGB(uint32 sx, uint32 sy, const uint8* from, uint8* to)
 	{
 		init_astc_decompress();
 
-		const int32 NumBlocksY = FMath::DivideAndRoundUp(sy, BLOCK_SIZE);
-		const int32 NumBlocksX = FMath::DivideAndRoundUp(sx, BLOCK_SIZE);
+		const int32 NumBlocksY = FMath::DivideAndRoundUp(sy, BlockSize);
+		const int32 NumBlocksX = FMath::DivideAndRoundUp(sx, BlockSize);
 
 		//for (uint32 BlockY = 0; BlockY < NumBlocksY; ++BlockY)
 		ParallelFor(NumBlocksY, [NumBlocksX, sx, sy, from, to](int32 BlockY)
-		{
-			for (int32 BlockX = 0; BlockX < NumBlocksX; ++BlockX)
 			{
-				constexpr bool bIsSRGB = false;
-				uint8 Block[BLOCK_SIZE * BLOCK_SIZE * 4];
+				for (int32 BlockX = 0; BlockX < NumBlocksX; ++BlockX)
+				{
+					uint8 Block[BlockSize * BlockSize * 4];
 
-				constexpr int32 CompressedBlockSize = 16;
-				const uint8* SrcBlockPtr = from + (BlockY * NumBlocksX + BlockX) * CompressedBlockSize;
+					constexpr int32 CompressedBlockSize = 16;
+					const uint8* SrcBlockPtr = from + (BlockY * NumBlocksX + BlockX) * CompressedBlockSize;
 
-				bool bSuccess = astcdec::decompress(Block, SrcBlockPtr, bIsSRGB, BLOCK_SIZE, BLOCK_SIZE);
-				miro_check(bSuccess);
+					astcdec::decompressSRGB<BlockSize>(Block, SrcBlockPtr);
 
-				for (uint32 py = 0; py < BLOCK_SIZE; py++)
-				{
-					for (uint32 px = 0; px < BLOCK_SIZE; px++)
+					for (uint32 py = 0; py < BlockSize; py++)
 					{
-						uint32 xi = BlockX * BLOCK_SIZE + px;
-						uint32 yi = BlockY * BLOCK_SIZE + py;
-
-						if (xi < sx && yi < sy)
+						for (uint32 px = 0; px < BlockSize; px++)
 						{
-							uint8* toPixel = to + (yi * sx + xi) * 3;
-							toPixel[0] = Block[py * BLOCK_SIZE * 4 + px * 4 + 0];
-							toPixel[1] = Block[py * BLOCK_SIZE * 4 + px * 4 + 1];
-							toPixel[2] = Block[py * BLOCK_SIZE * 4 + px * 4 + 2];
+							uint32 xi = BlockX * BlockSize + px;
+							uint32 yi = BlockY * BlockSize + py;
+
+							if (xi < sx && yi < sy)
+							{
+								uint8* toPixel = to + (yi * sx + xi) * 3;
+								toPixel[0] = Block[py * BlockSize * 4 + px * 4 + 0];
+								toPixel[1] = Block[py * BlockSize * 4 + px * 4 + 1];
+								toPixel[2] = Block[py * BlockSize * 4 + px * 4 + 2];
+							}
 						}
 					}
 				}
-			}
-		});
+			});
 	}
 
 
@@ -8963,13 +8870,13 @@ namespace miro
 
 
 	//---------------------------------------------------------------------------------------------
-	template<uint32 BLOCK_SIZE>
+	template<uint32 BlockSize>
 	void Generic_ASTCRGL_to_RGB(uint32 sx, uint32 sy, const uint8* from, uint8* to)
 	{
 		init_astc_decompress();
 
-		const int32 NumBlocksY = FMath::DivideAndRoundUp(sy, BLOCK_SIZE);
-		const int32 NumBlocksX = FMath::DivideAndRoundUp(sx, BLOCK_SIZE);
+		const int32 NumBlocksY = FMath::DivideAndRoundUp(sy, BlockSize);
+		const int32 NumBlocksX = FMath::DivideAndRoundUp(sx, BlockSize);
 
 		//for (uint32 by = 0; by < NumBlocksY; ++by)
 		ParallelFor(NumBlocksY, [NumBlocksX, sx, sy, from, to](int32 by)
@@ -8979,23 +8886,21 @@ namespace miro
 					constexpr int32 CompressedBlockSize = 16;
 					const uint8* SrcBlockPtr = from + (by * NumBlocksX + bx) * CompressedBlockSize;
 
-					bool bIsSRGB = false;
-					uint8 Block[BLOCK_SIZE * BLOCK_SIZE * 4];
-					bool bSuccess = astcdec::decompress(Block, SrcBlockPtr, bIsSRGB, BLOCK_SIZE, BLOCK_SIZE);
-					miro_check(bSuccess);
+					uint8 Block[BlockSize * BlockSize * 4];
+					astcdec::decompressSRGB<BlockSize>(Block, SrcBlockPtr);
 
-					for (uint32 BlockY = 0; BlockY < BLOCK_SIZE; BlockY++)
+					for (uint32 BlockY = 0; BlockY < BlockSize; BlockY++)
 					{
-						for (uint32 BlockX = 0; BlockX < BLOCK_SIZE; BlockX++)
+						for (uint32 BlockX = 0; BlockX < BlockSize; BlockX++)
 						{
-							uint32 xi = bx * BLOCK_SIZE + BlockX;
-							uint32 yi = by * BLOCK_SIZE + BlockY;
+							uint32 xi = bx * BlockSize + BlockX;
+							uint32 yi = by * BlockSize + BlockY;
 
 							if (xi < sx && yi < sy)
 							{
 								uint8* ToPixel = to + (yi * sx + xi) * 3;
-								ToPixel[0] = Block[BlockY * BLOCK_SIZE * 4 + BlockX * 4 + 0];
-								ToPixel[1] = Block[BlockY * BLOCK_SIZE * 4 + BlockX * 4 + 3];
+								ToPixel[0] = Block[BlockY * BlockSize * 4 + BlockX * 4 + 0];
+								ToPixel[1] = Block[BlockY * BlockSize * 4 + BlockX * 4 + 3];
 								ToPixel[2] = 255;
 							}
 						}
@@ -9006,13 +8911,13 @@ namespace miro
 
 
 	//---------------------------------------------------------------------------------------------
-	template<uint32 BLOCK_SIZE>
+	template<uint32 BlockSize>
 	void Generic_ASTCRGL_to_RGBA(uint32 sx, uint32 sy, const uint8* from, uint8* to)
 	{
 		init_astc_decompress();
 
-		const int32 NumBlocksY = FMath::DivideAndRoundUp(sy, BLOCK_SIZE);
-		const int32 NumBlocksX = FMath::DivideAndRoundUp(sx, BLOCK_SIZE);
+		const int32 NumBlocksY = FMath::DivideAndRoundUp(sy, BlockSize);
+		const int32 NumBlocksX = FMath::DivideAndRoundUp(sx, BlockSize);
 
 		//for (uint32 by = 0; by < NumBlocksY; ++by)
 		ParallelFor(NumBlocksY, [NumBlocksX, sx, sy, from, to](int32 by)
@@ -9022,23 +8927,21 @@ namespace miro
 					constexpr int32 CompressedBlockSize = 16;
 					const uint8* SrcBlockPtr = from + (by * NumBlocksX + bx) * CompressedBlockSize;
 
-					bool bIsSRGB = false;
-					uint8 Block[BLOCK_SIZE * BLOCK_SIZE * 4];
-					bool bSuccess = astcdec::decompress(Block, SrcBlockPtr, bIsSRGB, BLOCK_SIZE, BLOCK_SIZE);
-					miro_check(bSuccess);
+					uint8 Block[BlockSize * BlockSize * 4];
+					astcdec::decompressSRGB<BlockSize>(Block, SrcBlockPtr);
 
-					for (uint32 BlockY = 0; BlockY < BLOCK_SIZE; BlockY++)
+					for (uint32 BlockY = 0; BlockY < BlockSize; BlockY++)
 					{
-						for (uint32 BlockX = 0; BlockX < BLOCK_SIZE; BlockX++)
+						for (uint32 BlockX = 0; BlockX < BlockSize; BlockX++)
 						{
-							uint32 xi = bx * BLOCK_SIZE + BlockX;
-							uint32 yi = by * BLOCK_SIZE + BlockY;
+							uint32 xi = bx * BlockSize + BlockX;
+							uint32 yi = by * BlockSize + BlockY;
 
 							if (xi < sx && yi < sy)
 							{
 								uint8* ToPixel = to + (yi * sx + xi) * 4;
-								ToPixel[0] = Block[BlockY * BLOCK_SIZE * 4 + BlockX * 4 + 0];
-								ToPixel[1] = Block[BlockY * BLOCK_SIZE * 4 + BlockX * 4 + 3];
+								ToPixel[0] = Block[BlockY * BlockSize * 4 + BlockX * 4 + 0];
+								ToPixel[1] = Block[BlockY * BlockSize * 4 + BlockX * 4 + 3];
 								ToPixel[2] = 255;
 								ToPixel[3] = 255;
 							}
@@ -9103,7 +9006,7 @@ namespace miro
 	{
 		Generic_ASTCRGBL_to_RGBA<4>(sx, sy, from, to);
 	}
-	
+
 	void L_to_ASTC4x4RGBL(uint32 sx, uint32 sy, const uint8* from, uint8* to, int32 Quality)
 	{
 		Generic_L_to_ASTCRGBL<4>(sx, sy, from, to, Quality);
@@ -9113,17 +9016,17 @@ namespace miro
 	{
 		Generic_RGB_to_ASTCRGL<4>(sx, sy, from, to, Quality);
 	}
-	
+
 	void RGBA_to_ASTC4x4RGL(uint32 sx, uint32 sy, const uint8* from, uint8* to, int32 Quality)
 	{
 		Generic_RGBA_to_ASTCRGL<4>(sx, sy, from, to, Quality);
 	}
-	
+
 	void ASTC4x4RGL_to_RGB(uint32 sx, uint32 sy, const uint8* from, uint8* to)
 	{
 		Generic_ASTCRGL_to_RGB<4>(sx, sy, from, to);
 	}
-	
+
 	void ASTC4x4RGL_to_RGBA(uint32 sx, uint32 sy, const uint8* from, uint8* to)
 	{
 		Generic_ASTCRGL_to_RGBA<4>(sx, sy, from, to);
@@ -9341,10 +9244,8 @@ namespace miro::SubImageDecompression
 		{
 			for (uint32 x = 0; x < SubSize.X; x += BlockSize)
 			{
-				constexpr bool bIsSRGB = false;
 				uint8 Block[BlockSize * BlockSize * 4];
-				const bool bSuccess = astcdec::decompress(Block, From, bIsSRGB, BlockSize, BlockSize);
-				miro_check(bSuccess);
+				astcdec::decompressSRGB<BlockSize>(Block, From);
 
 				for (uint32 py = 0; py < BlockSize; py++)
 				{
@@ -9378,10 +9279,8 @@ namespace miro::SubImageDecompression
 		{
 			for (uint32 x = 0; x < SubSize.X; x += BlockSize)
 			{
-				bool bIsSRGB = false;
 				uint8 Block[BlockSize * BlockSize * 4];
-				bool bSuccess = astcdec::decompress(Block, From, bIsSRGB, BlockSize, BlockSize);
-				miro_check(bSuccess);
+				astcdec::decompressSRGB<BlockSize>(Block, From);
 
 				for (uint32 py = 0; py < BlockSize; py++)
 				{
@@ -9410,7 +9309,7 @@ namespace miro::SubImageDecompression
 	void GenericSubImage_ASTCRGBAL_To_RGBA(FImageSize FromSize, FImageSize ToSize, FImageSize SubSize, const uint8* From, uint8* To)
 	{
 		init_astc_decompress();
-	
+
 		const uint32 NumBlocksX = FMath::DivideAndRoundUp(uint32(FromSize.X), BlockSize);
 		const uint32 NumSubBlocksX = FMath::DivideAndRoundUp(uint32(SubSize.X), BlockSize);
 		const uint32 NumSubBlocksY = FMath::DivideAndRoundUp(uint32(SubSize.Y), BlockSize);
@@ -9424,14 +9323,12 @@ namespace miro::SubImageDecompression
 		{
 			for (uint32 BlockX = 0; BlockX < NumSubBlocksX; ++BlockX)
 			{
-				constexpr bool bIsSRGB = false;
 				uint8 Block[BlockSize * BlockSize * 4];
 
 				constexpr int32 CompressedBlockSize = 16;
 				const uint8* SrcBlockPtr = From + (BlockY * NumBlocksX + BlockX) * CompressedBlockSize;
 
-				bool bSuccess = astcdec::decompress(Block, SrcBlockPtr, bIsSRGB, BlockSize, BlockSize);
-				miro_check(bSuccess);
+				astcdec::decompressSRGB<BlockSize>(Block, SrcBlockPtr);
 
 				for (uint32 py = 0; py < BlockSize; py++)
 				{
@@ -9465,14 +9362,12 @@ namespace miro::SubImageDecompression
 		{
 			for (uint32 BlockX = 0; BlockX < NumSubBlocksX; ++BlockX)
 			{
-				constexpr bool bIsSRGB = false;
 				uint8 Block[BlockSize * BlockSize * 4];
 
 				constexpr int32 CompressedBlockSize = 16;
 				const uint8* SrcBlockPtr = From + (BlockY * NumBlocksX + BlockX) * CompressedBlockSize;
 
-				bool bSuccess = astcdec::decompress(Block, SrcBlockPtr, bIsSRGB, BlockSize, BlockSize);
-				miro_check(bSuccess);
+				astcdec::decompressSRGB<BlockSize>(Block, SrcBlockPtr);
 
 				for (uint32 py = 0; py < BlockSize; py++)
 				{
@@ -9510,10 +9405,8 @@ namespace miro::SubImageDecompression
 				constexpr int32 CompressedBlockSize = 16;
 				const uint8* SrcBlockPtr = From + (by * NumBlocksX + bx) * CompressedBlockSize;
 
-				constexpr bool bIsSRGB = false;
 				uint8 Block[BlockSize * BlockSize * 4];
-				const bool bSuccess = astcdec::decompress(Block, SrcBlockPtr, bIsSRGB, BlockSize, BlockSize);
-				miro_check(bSuccess);
+				astcdec::decompressSRGB<BlockSize>(Block, SrcBlockPtr);
 
 				for (uint32 BlockY = 0; BlockY < BlockSize; BlockY++)
 				{
@@ -9539,7 +9432,7 @@ namespace miro::SubImageDecompression
 	void GenericSubImage_ASTCRGL_To_RGBA(FImageSize FromSize, FImageSize ToSize, FImageSize SubSize, const uint8* From, uint8* To)
 	{
 		init_astc_decompress();
-		
+
 		const uint32 NumBlocksX = FMath::DivideAndRoundUp(uint32(FromSize.X), BlockSize);
 		const uint32 NumSubBlocksX = FMath::DivideAndRoundUp(uint32(SubSize.X), BlockSize);
 		const uint32 NumSubBlocksY = FMath::DivideAndRoundUp(uint32(SubSize.Y), BlockSize);
@@ -9551,10 +9444,8 @@ namespace miro::SubImageDecompression
 				constexpr int32 CompressedBlockSize = 16;
 				const uint8* SrcBlockPtr = From + (by * NumBlocksX + bx) * CompressedBlockSize;
 
-				constexpr bool bIsSRGB = false;
 				uint8 Block[BlockSize * BlockSize * 4];
-				bool bSuccess = astcdec::decompress(Block, SrcBlockPtr, bIsSRGB, BlockSize, BlockSize);
-				miro_check(bSuccess);
+				astcdec::decompressSRGB<BlockSize>(Block, SrcBlockPtr);
 
 				for (uint32 BlockY = 0; BlockY < BlockSize; BlockY++)
 				{
@@ -9597,12 +9488,12 @@ namespace miro::SubImageDecompression
 	{
 		GenericSubImage_ASTCRGBL_To_RGBA<4>(FromSize, ToSize, SubSize, From, To);
 	}
-	
+
 	void ASTC4x4RGL_To_RGBSubImage(FImageSize FromSize, FImageSize ToSize, FImageSize SubSize, const uint8* From, uint8* To)
 	{
 		GenericSubImage_ASTCRGL_To_RGB<4>(FromSize, ToSize, SubSize, From, To);
 	}
-	
+
 	void ASTC4x4RGL_To_RGBASubImage(FImageSize FromSize, FImageSize ToSize, FImageSize SubSize, const uint8* From, uint8* To)
 	{
 		GenericSubImage_ASTCRGL_To_RGBA<4>(FromSize, ToSize, SubSize, From, To);
@@ -10013,7 +9904,7 @@ namespace miro::SubImageDecompression
 		const uint32 NumBlocksX = FMath::DivideAndRoundUp(uint32(FromSize.X), BlockSize);
 		const uint32 NumSubBlocksX = FMath::DivideAndRoundUp(uint32(SubSize.X), BlockSize);
 		const uint32 NumSubBlocksY = FMath::DivideAndRoundUp(uint32(SubSize.Y), BlockSize);
-		
+
 		for (uint32 by = 0; by < NumSubBlocksY; ++by)
 		{
 			const uint8* RowFrom = From + 8 * NumBlocksX * by;
