[Mutable] Fix OpImageResize sampling coordinates issue when memory optimizations enabled. Add vector implemementation for small down resizes.

- Vector implementation is disabled by default. Cvar mutable.EnableVectorImplementationForSmallResizes added to enable the optimization.

#rnx

diff --git a/Engine/Plugins/Experimental/Mutable/Source/MutableRuntime/Private/MuR/OpImageResize.cpp b/Engine/Plugins/Experimental/Mutable/Source/MutableRuntime/Private/MuR/OpImageResize.cpp
--- a/Engine/Plugins/Experimental/Mutable/Source/MutableRuntime/Private/MuR/OpImageResize.cpp
+++ b/Engine/Plugins/Experimental/Mutable/Source/MutableRuntime/Private/MuR/OpImageResize.cpp
@@ -13,6 +13,15 @@ static FAutoConsoleVariableRef CVarEnableImageResizeMemoryOptimizations (
 	bEnableImageResizeMemoryOptimizations,
 	TEXT("If set to true, enables image resize optimizations form small resizes and block compressed images."),
 	ECVF_Default);
+
+
+bool bEnableVectorImplementationForSmallResizes = false;
+static FAutoConsoleVariableRef CVarEnableVectorImplementationForSmallResizes (
+	TEXT("mutable.EnableVectorImplementationForSmallResizes"),
+	bEnableVectorImplementationForSmallResizes,
+	TEXT("If set to true, enables image resize vector optimizations form small resizes. Will only have effect if mutable.EnableImageResizeMemoryOptimizations is enabled"),
+	ECVF_Default);
+
 }
 
 namespace mu
@@ -713,8 +722,92 @@ namespace mu
 
 namespace
 {
+
 	template<uint32 NumChannels>
-	void ImageResizeLinearDownByLessThanTwoSubImage(FVector2f ScalingFactor,
+	FORCENOINLINE void ImageResizeLinearDownByLessThanTwoSubImageVectorImpl(FVector2f InOneOverScalingFactor,
+			FImageSize DestSize, FImageSize DestSubSize, FImageSize SrcSize, FImageSize SrcSubSize,
+		  	uint8* RESTRICT SubDest, const uint8* RESTRICT SubSrc)
+	{	
+		static_assert(NumChannels > 0 && NumChannels <= 4);
+	
+		// For some reason (probably cache related), doing this copy yields marginally but consistently 
+		// better performance than using the function argument. 
+		const FVector2f OneOverScalingFactor = InOneOverScalingFactor;
+		
+		check(OneOverScalingFactor.X <= 2.0f && OneOverScalingFactor.Y <= 2.0f && 
+			  OneOverScalingFactor.X >= 1.0f && OneOverScalingFactor.Y >= 1.0f);
+
+		for (uint16 Y = 0; Y < DestSubSize.Y; ++Y)
+		{
+			for (uint16 X = 0; X < DestSubSize.X; ++X)
+			{
+				auto LoadPixel = [](const uint8* Ptr) -> VectorRegister4Float
+				{	
+					alignas(4) uint8 PixelData[4] = {0};
+					if constexpr (NumChannels == 4)
+					{
+						FMemory::Memcpy(&(PixelData[0]), Ptr, 4);
+					}
+					else
+					{
+						for (uint32 C = 0; C < NumChannels; ++C)
+						{
+							PixelData[C] = Ptr[C];
+						}
+					}
+
+					return VectorIntToFloat(MakeVectorRegisterInt(PixelData[0], PixelData[1], PixelData[2], PixelData[3]));
+				};
+				
+				const FVector2f CoordsF = FVector2f(X, Y)*OneOverScalingFactor;
+	
+				const FIntVector2 Coords = FIntVector2(FMath::FloorToInt(CoordsF.X), FMath::FloorToInt(CoordsF.Y));
+				const FIntVector2 CoordsPlusOne = FIntVector2(
+						FMath::Min<int32>(Coords.X + 1, SrcSize.X - 1),
+						FMath::Min<int32>(Coords.Y + 1, SrcSize.Y - 1));
+
+				const VectorRegister4Float Pixel00 = LoadPixel(SubSrc + (Coords.Y*SrcSize.X + Coords.X) * NumChannels);
+				const VectorRegister4Float Pixel10 = LoadPixel(SubSrc + (Coords.Y*SrcSize.X + CoordsPlusOne.X) * NumChannels);
+				const VectorRegister4Float Pixel01 = LoadPixel(SubSrc + (CoordsPlusOne.Y*SrcSize.X + Coords.X) * NumChannels);
+				const VectorRegister4Float Pixel11 = LoadPixel(SubSrc + (CoordsPlusOne.Y*SrcSize.X + CoordsPlusOne.X) * NumChannels);
+
+				const VectorRegister4Float CoordsFVector = MakeVectorRegister(CoordsF.X, CoordsF.Y, 0.0f, 0.0f);
+				const VectorRegister4Float Frac = VectorSubtract(CoordsFVector, VectorFloor(CoordsFVector));
+
+				const VectorRegister4Float FracX = VectorReplicate(Frac, 0);
+				const VectorRegister4Float FracY = VectorReplicate(Frac, 1);
+				
+				VectorRegister4Float Result = VectorMultiplyAdd(FracX, VectorSubtract(Pixel10, Pixel00), Pixel00);
+				Result = VectorMultiplyAdd(
+						FracY,
+						VectorSubtract(VectorMultiplyAdd(FracX, VectorSubtract(Pixel11, Pixel01), Pixel01), Result),
+						Result);
+
+				uint8* Dest = SubDest + (Y*DestSize.X + X)*NumChannels;  
+
+				if constexpr (NumChannels == 4)
+				{
+					VectorStoreByte4(Result, Dest);
+					VectorResetFloatRegisters();
+				}
+				else
+				{
+					alignas(4) uint8 ResultData[4];
+					VectorStoreByte4(Result, ResultData);
+					VectorResetFloatRegisters();
+				
+					for (uint32 C = 0; C < NumChannels; ++C)
+					{
+						Dest[C] = ResultData[C];
+					}
+				}
+			}
+		}
+	}
+
+
+	template<uint32 NumChannels>
+	FORCENOINLINE void ImageResizeLinearDownByLessThanTwoSubImageNonVectorImpl(FVector2f InOneOverScalingFactor,
 			FImageSize DestSize, FImageSize DestSubSize, FImageSize SrcSize, FImageSize SrcSubSize,
 		  	uint8* RESTRICT SubDest, const uint8* RESTRICT SubSrc)
 	{
@@ -725,10 +818,12 @@ namespace
 
 		static_assert(NumChannels <= 4);
 		
-		//TODO: Consider using fixed point arithmetic or vector registers as optimization.
+		// For some reason (probably cache related), doing this copy yields marginally but consistently 
+		// better performance than using the function argument. 
+		const FVector2f OneOverScalingFactor = InOneOverScalingFactor;
 		
-		FVector2f InvScalingFactor = FVector2f(1.0f) / ScalingFactor;
-		check(InvScalingFactor.X <= 2.0f && InvScalingFactor.Y <= 2.0f && InvScalingFactor.X >= 1.0f && InvScalingFactor.Y >= 1.0f);
+		check(OneOverScalingFactor.X <= 2.0f && OneOverScalingFactor.Y <= 2.0f && 
+			  OneOverScalingFactor.X >= 1.0f && OneOverScalingFactor.Y >= 1.0f);
 
 		for (uint16 Y = 0; Y < DestSubSize.Y; ++Y)
 		{
@@ -758,7 +853,7 @@ namespace
 					return Result;
 				};
 				
-				const FVector2f CoordsF = (FVector2f(X, Y) + 0.5f)*InvScalingFactor;
+				const FVector2f CoordsF = FVector2f(X, Y)*OneOverScalingFactor;
 				
 				using FUint16Vector2 = UE::Math::TIntVector2<uint16>;
 				const FUint16Vector2 FracINorm = FUint16Vector2(FMath::Frac(CoordsF.X) * 255.0f, FMath::Frac(CoordsF.Y) * 255.0f);
@@ -804,9 +899,26 @@ namespace
 		}
 	}
 
-	void ImageResizeLinearDownByLessThanTwo(Image* Dest, const Image* InBase)
+	template<int32 NumChannels, bool bUseVectorImpl>
+	FORCEINLINE void ImageResizeLinearDownByLessThanTwoSubImage(FVector2f OneOverScalingFactor,
+			FImageSize DestSize, FImageSize DestSubSize, FImageSize SrcSize, FImageSize SrcSubSize,
+		  	uint8* RESTRICT SubDest, const uint8* RESTRICT SubSrc)
 	{
+		if constexpr (bUseVectorImpl)
+		{
+			ImageResizeLinearDownByLessThanTwoSubImageVectorImpl<NumChannels>(
+					OneOverScalingFactor, DestSize, DestSubSize, SrcSize, SrcSubSize, SubDest, SubSrc);
+		}
+		else
+		{
+			ImageResizeLinearDownByLessThanTwoSubImageNonVectorImpl<NumChannels>(
+					OneOverScalingFactor, DestSize, DestSubSize, SrcSize, SrcSubSize, SubDest, SubSrc);
+		}
+	}
 
+	template<bool bUseVectorImpl>
+	void ImageResizeLinearDownByLessThanTwo(Image* Dest, const Image* InBase)
+	{
 		MUTABLE_CPUPROFILER_SCOPE(ImageResizeLinearDownByLessThanTwo);
 
 		const FImageSize DestSize = FImageSize(Dest->GetSizeX(), Dest->GetSizeY());
@@ -823,7 +935,7 @@ namespace
 		const FImageFormatData& DestFormatData = GetImageFormatData(Dest->GetFormat());
 		const FImageFormatData& SrcFormatData = GetImageFormatData(InBase->GetFormat());
 
-		const FVector2f ScalingFactor = FVector2f(DestSize.X, DestSize.Y) / FVector2f(SrcSize.X, SrcSize.Y);
+		const FVector2f OneOverScalingFactor = FVector2f(SrcSize.X, SrcSize.Y) / FVector2f(DestSize.X, DestSize.Y);
 
 		if (SrcFormatData.PixelsPerBlockX == 1 && SrcFormatData.PixelsPerBlockY == 1)
 		{
@@ -836,7 +948,7 @@ namespace
 			case EImageFormat::IF_L_UBYTE:
 			{
 				constexpr int32 NumChannels = 1;
-				ImageResizeLinearDownByLessThanTwoSubImage<NumChannels>(ScalingFactor,
+				ImageResizeLinearDownByLessThanTwoSubImage<NumChannels, bUseVectorImpl>(OneOverScalingFactor,
 						DestSize, DestSize, SrcSize, SrcSize, DestData, SrcData);
 				break;
 			}
@@ -844,7 +956,7 @@ namespace
 			case EImageFormat::IF_RGB_UBYTE:
 			{
 				constexpr int32 NumChannels = 3;
-				ImageResizeLinearDownByLessThanTwoSubImage<NumChannels>(ScalingFactor,
+				ImageResizeLinearDownByLessThanTwoSubImage<NumChannels, bUseVectorImpl>(OneOverScalingFactor,
 						DestSize, DestSize, SrcSize, SrcSize, DestData, SrcData);
 				break;
 			}
@@ -853,7 +965,7 @@ namespace
 			case EImageFormat::IF_BGRA_UBYTE:
 			{
 				constexpr int32 NumChannels = 4;
-				ImageResizeLinearDownByLessThanTwoSubImage<NumChannels>(ScalingFactor,
+				ImageResizeLinearDownByLessThanTwoSubImage<NumChannels, bUseVectorImpl>(OneOverScalingFactor,
 						DestSize, DestSize, SrcSize, SrcSize, DestData, SrcData);
 				break;
 			}
@@ -894,13 +1006,13 @@ namespace
 			const int32 NumParallelJobs = FMath::DivideAndRoundUp(NumBatches.Y, NumRowBatchesPerJob); 
 
 			// Allocate memory for 1 extra block, this is needed so that the data to compute the batch edges
-			// if available. This implies extra work needs to be done, for simplicity re-decompress those blocks, if
+			// is available. This implies extra work needs to be done, for simplicity re-decompress those blocks, if
 			// that becomes a problem more convoluted memory cacheing schemes could be implemented.
 			const miro::FImageSize StagingSize = miro::FImageSize(
 					uint16((BatchSizeInBlocksX + 1)*PixelsPerBlock.X), 
 					uint16((BatchSizeInBlocksY + 1)*PixelsPerBlock.Y));
 		
-			// Add some extra padding so diferent threads do not share cache lines.
+			// Add some extra padding so different threads do not share cache lines.
 			const int32 PerJobStagingBytes = StagingSize.X*StagingSize.Y*DestNumChannels + 64;
 
 			TArray<uint8, FDefaultMemoryTrackingAllocator<MemoryCounters::FImageMemoryCounter>> StagingMemory;
@@ -914,7 +1026,7 @@ namespace
 					NumParallelJobs, NumRowBatchesPerJob, 
 					StagingMemoryData, PerJobStagingBytes, 
 					NumBatches, NumBlocks, PixelsPerBlock, BlockSizeInBytes, DestNumChannels, DecompressionFunc,
-					SrcData, SrcSize, DestData, DestSize, ScalingFactor
+					SrcData, SrcSize, DestData, DestSize, OneOverScalingFactor
 				](int32 JobId)
 			{
 				const int32 JobRowBegin = JobId*NumRowBatchesPerJob;
@@ -980,7 +1092,7 @@ namespace
 						
 							uint8* DestSubData = DestData + (DestSubOffset.Y*DestSize.X + DestSubOffset.X)*NumChannels;
 
-							ImageResizeLinearDownByLessThanTwoSubImage<NumChannels>(ScalingFactor, 
+							ImageResizeLinearDownByLessThanTwoSubImage<NumChannels, bUseVectorImpl>(OneOverScalingFactor, 
 									DestSize, DestSubSize, DecSize, SrcSubSize, DestSubData, SrcSubData);
 							break;
 						}
@@ -991,7 +1103,7 @@ namespace
 
 							uint8* DestSubData = DestData + (DestSubOffset.Y*DestSize.X + DestSubOffset.X)*NumChannels;
 							
-							ImageResizeLinearDownByLessThanTwoSubImage<NumChannels>(ScalingFactor, 
+							ImageResizeLinearDownByLessThanTwoSubImage<NumChannels, bUseVectorImpl>(OneOverScalingFactor, 
 									DestSize, DestSubSize, DecSize, SrcSubSize, DestSubData, SrcSubData);
 							break;
 						}
@@ -1002,7 +1114,7 @@ namespace
 
 							uint8* DestSubData = DestData + (DestSubOffset.Y*DestSize.X + DestSubOffset.X)*NumChannels;
 							
-							ImageResizeLinearDownByLessThanTwoSubImage<NumChannels>(ScalingFactor,
+							ImageResizeLinearDownByLessThanTwoSubImage<NumChannels, bUseVectorImpl>(OneOverScalingFactor,
 									DestSize, DestSubSize, DecSize, SrcSubSize, DestSubData, SrcSubData);
 							break;
 						}
@@ -1064,8 +1176,17 @@ namespace
 			{
 				DecompressedDest = CreateImage(DestSize.X, DestSize.Y, 1, BaseUncompressedFormat, EInitializationType::NotInitialized);
 			}
-			
-			ImageResizeLinearDownByLessThanTwo(DecompressedDest.get(), InBase);
+		
+			if (bEnableVectorImplementationForSmallResizes)
+			{
+				constexpr bool bUseVectorImpl = true;
+				ImageResizeLinearDownByLessThanTwo<bUseVectorImpl>(DecompressedDest.get(), InBase);
+			}
+			else
+			{
+				constexpr bool bUseVectorImpl = false;
+				ImageResizeLinearDownByLessThanTwo<bUseVectorImpl>(DecompressedDest.get(), InBase);
+			}
 
 			if (DestFormat != BaseUncompressedFormat)
 			{
