[mutable] Added log message to display total model roms size and count

diff --git a/Engine/Plugins/Experimental/Mutable/Source/MutableValidation/MutableValidation.Build.cs b/Engine/Plugins/Experimental/Mutable/Source/MutableValidation/MutableValidation.Build.cs
--- a/Engine/Plugins/Experimental/Mutable/Source/MutableValidation/MutableValidation.Build.cs
+++ b/Engine/Plugins/Experimental/Mutable/Source/MutableValidation/MutableValidation.Build.cs
@@ -16,7 +16,8 @@ namespace UnrealBuildTool.Rules
 			PublicDependencyModuleNames.AddRange(
 				new string[] {
 					"Settings", 
-					"Engine"
+					"Engine", 
+					"MutableRuntime"
 				});
 			
 			PrivateDependencyModuleNames.AddRange(
@@ -27,7 +28,7 @@ namespace UnrealBuildTool.Rules
 					"UnrealEd",
 					
 					"DataValidation",
-					"CustomizableObject", 
+					"CustomizableObject"
 				}
 			);
 		}
diff --git a/Engine/Plugins/Experimental/Mutable/Source/MutableValidation/Private/MuV/CustomizableObjectValidationCommandlet.cpp b/Engine/Plugins/Experimental/Mutable/Source/MutableValidation/Private/MuV/CustomizableObjectValidationCommandlet.cpp
--- a/Engine/Plugins/Experimental/Mutable/Source/MutableValidation/Private/MuV/CustomizableObjectValidationCommandlet.cpp
+++ b/Engine/Plugins/Experimental/Mutable/Source/MutableValidation/Private/MuV/CustomizableObjectValidationCommandlet.cpp
@@ -6,6 +6,7 @@
 #include "Containers/Array.h"
 #include "MuCO/CustomizableObject.h"
 #include "MuCO/CustomizableObjectSystem.h"
+#include "MuR/Model.h"
 
 
 int32 UCustomizableObjectValidationCommandlet::Main(const FString& Params)
@@ -64,7 +65,7 @@ int32 UCustomizableObjectValidationCommandlet::Main(const FString& Params)
     	// Run Sync compilation -> Warning : Potentially long operation -------------
     	Compiler->Compile(*ToTestCustomizableObject, ToTestCustomizableObject->CompileOptions, false);
     	// --------------------------------------------------------------------------
-    		
+		
     	// Get the compilation result
     	const ECustomizableObjectCompilationState CompilationEndResult = Compiler->GetCompilationState();
     	check(CompilationEndResult != ECustomizableObjectCompilationState::None);
@@ -77,7 +78,22 @@ int32 UCustomizableObjectValidationCommandlet::Main(const FString& Params)
 	if (bWasCoCompilationSuccesfull)
 	{
 		UE_LOG(LogMutable,Display,TEXT("Customizable Object was compiled succesfully."));
-				
+		
+		// GHet the total size of the streaming data of the model ---------------------------------------------- //
+		if (const TSharedPtr<const mu::Model> MutableModel = ToTestCustomizableObject->GetModel())
+		{
+			const int32 RomCount =  MutableModel->GetRomCount();
+			uint32 TotalRomSizeBits = 0;
+			for (int32 RomIndex = 0; RomIndex < RomCount; RomIndex++)
+			{
+				const uint32 RomBitSize = MutableModel->GetRomSize(RomIndex);
+				TotalRomSizeBits += RomBitSize;
+			}
+			
+			UE_LOG(LogMutable, Log,TEXT("(int) model_rom_count : %d "), RomCount);
+			UE_LOG(LogMutable, Log,TEXT("(uint32) model_roms_size : %u "), TotalRomSizeBits);
+		}
+		
 		// Generate target random instances to be tested ------------------------------------------------------------ //
 		bool bWasInstancesCreationSuccessful = true;
 		{
