[Mutable] Mutable Benchmark improvements.

- Fixed Task_Game_Callbacks timer scope not reporting correctly.
- Added Queue time.
- Added Task_Mutable_UpdateImage time.
- Added Update result.
- Removed unused CVarBenchmarkFilePath.
- Moved incorrect Enqueued UpdateSkeletalMesh Async log.

#rnx

diff --git a/Engine/Plugins/Experimental/Mutable/Source/CustomizableObject/Private/MuCO/CustomizableObjectInstance.cpp b/Engine/Plugins/Experimental/Mutable/Source/CustomizableObject/Private/MuCO/CustomizableObjectInstance.cpp
--- a/Engine/Plugins/Experimental/Mutable/Source/CustomizableObject/Private/MuCO/CustomizableObjectInstance.cpp
+++ b/Engine/Plugins/Experimental/Mutable/Source/CustomizableObject/Private/MuCO/CustomizableObjectInstance.cpp
@@ -5353,7 +5353,10 @@ void UCustomizableInstancePrivateData::BuildMaterials(const TSharedRef<FUpdateCo
 	GeneratedMaterials.Reset();
 
 	// Prepare the data to store in order to regenerate resources for this instance (usually texture mips).
-	TSharedPtr<FMutableUpdateContext> UpdateContext = MakeShared<FMutableUpdateContext>(UCustomizableObjectSystem::GetInstance()->GetPrivate()->MutableSystem,
+	TSharedPtr<FMutableUpdateContext> UpdateContext = MakeShared<FMutableUpdateContext>(
+		CustomizableObject->GetPathName(),
+		Public->GetPathName(),
+		UCustomizableObjectSystem::GetInstance()->GetPrivate()->MutableSystem,
 		CustomizableObject->GetModel(),
 		OperationData->Parameters,
 	    OperationData->State);
diff --git a/Engine/Plugins/Experimental/Mutable/Source/CustomizableObject/Private/MuCO/CustomizableObjectMipDataProvider.cpp b/Engine/Plugins/Experimental/Mutable/Source/CustomizableObject/Private/MuCO/CustomizableObjectMipDataProvider.cpp
--- a/Engine/Plugins/Experimental/Mutable/Source/CustomizableObject/Private/MuCO/CustomizableObjectMipDataProvider.cpp
+++ b/Engine/Plugins/Experimental/Mutable/Source/CustomizableObject/Private/MuCO/CustomizableObjectMipDataProvider.cpp
@@ -16,8 +16,15 @@ UMutableTextureMipDataProviderFactory::UMutableTextureMipDataProviderFactory(con
 }
 
 
-FMutableUpdateContext::FMutableUpdateContext(mu::Ptr<mu::System> InSystem,
-	TSharedPtr<mu::Model, ESPMode::ThreadSafe> InModel, mu::Ptr<const mu::Parameters> InParameters, int32 InState):
+FMutableUpdateContext::FMutableUpdateContext(
+	const FString& InCustomizableObjectPathName,
+	const FString& InInstancePathName,
+	mu::Ptr<mu::System> InSystem,
+	TSharedPtr<mu::Model, ESPMode::ThreadSafe> InModel,
+	mu::Ptr<const mu::Parameters> InParameters,
+	int32 InState) :
+	CustomizableObjectPathName(InCustomizableObjectPathName),
+	InstancePathName(InInstancePathName),
 	System(InSystem),
 	Model(InModel),
 	Parameters(InParameters),
@@ -42,6 +49,18 @@ FMutableUpdateContext::~FMutableUpdateContext()
 }
 
 
+const FString& FMutableUpdateContext::GetCustomizableObjectPathName() const
+{
+	return CustomizableObjectPathName;
+}
+
+
+const FString& FMutableUpdateContext::GetInstancePathName() const
+{
+	return InstancePathName;
+}
+
+
 mu::Ptr<mu::System> FMutableUpdateContext::GetSystem() const
 {
 	return System;
@@ -109,7 +128,8 @@ namespace impl
 	void Task_Mutable_UpdateImage(TSharedPtr<FMutableImageOperationData> OperationData)
 	{
 		MUTABLE_CPUPROFILER_SCOPE(Task_Mutable_UpdateImage);
-
+		const double StartTime = FPlatformTime::Seconds();
+		
 		// Any external texture that may be needed for this update will be requested from Mutable Core's GetImage
 		// which will safely access the GlobalExternalImages map, and then just get the cached image or issue a disk read
 
@@ -209,6 +229,29 @@ namespace impl
 			}
 		}
 
+		if (CVarEnableBenchmark.GetValueOnAnyThread())
+		{
+			double Time = FPlatformTime::Seconds() - StartTime;
+
+			const FString& CustomizableObjectPathName = OperationData->UpdateContext->GetCustomizableObjectPathName();
+			const FString& InstancePathName = OperationData->UpdateContext->GetInstancePathName();
+			
+			FFunctionGraphTask::CreateAndDispatchWhenReady(
+			[CustomizableObjectPathName, InstancePathName, Time]()
+			{
+				UCustomizableObjectSystem* System = UCustomizableObjectSystem::GetInstance();
+				if (!System)
+				{
+					return;
+				}
+
+				System->GetPrivateChecked()->LogBenchmarkUtil.FinishUpdateImage(CustomizableObjectPathName, InstancePathName, Time);
+			},
+			TStatId{},
+			nullptr,
+			ENamedThreads::GameThread);
+		}
+		
 		{
 			// The request could be cancelled in parallel from CancelCounterSafely and its value be changed
 			// between reading it and actually running Decrement() and RescheduleCallback(), so lock
diff --git a/Engine/Plugins/Experimental/Mutable/Source/CustomizableObject/Private/MuCO/CustomizableObjectSystem.cpp b/Engine/Plugins/Experimental/Mutable/Source/CustomizableObject/Private/MuCO/CustomizableObjectSystem.cpp
--- a/Engine/Plugins/Experimental/Mutable/Source/CustomizableObject/Private/MuCO/CustomizableObjectSystem.cpp
+++ b/Engine/Plugins/Experimental/Mutable/Source/CustomizableObject/Private/MuCO/CustomizableObjectSystem.cpp
@@ -163,7 +163,6 @@ void FUpdateContextPrivate::AddReferencedObjects(FReferenceCollector& Collector)
 FMutablePendingInstanceUpdate::FMutablePendingInstanceUpdate(const TSharedRef<FUpdateContextPrivate>& InContext) :
 	Context(InContext)
 {
-	SecondsAtUpdate = FPlatformTime::Seconds();
 }
 
 
@@ -185,7 +184,7 @@ bool FMutablePendingInstanceUpdate::operator<(const FMutablePendingInstanceUpdat
 	}
 	else
 	{
-		return SecondsAtUpdate < Other.SecondsAtUpdate;
+		return Context->StartQueueTime < Other.Context->StartQueueTime;
 	}
 }
 
@@ -214,12 +213,6 @@ uint32 FPendingInstanceUpdateKeyFuncs::GetKeyHash(const TWeakObjectPtr<const UCu
 }
 
 
-bool FMutablePendingInstanceWork::ArePendingUpdatesEmpty() const
-{
-	return PendingInstanceUpdates.Num() == 0;
-}
-
-
 int32 FMutablePendingInstanceWork::Num() const
 {
 	return PendingInstanceUpdates.Num() + PendingInstanceDiscards.Num() + PendingIDsToRelease.Num() + NumLODUpdatesLastTick;
@@ -234,16 +227,18 @@ void FMutablePendingInstanceWork::SetLODUpdatesLastTick(int32 NumLODUpdates)
 
 void FMutablePendingInstanceWork::AddUpdate(const FMutablePendingInstanceUpdate& UpdateToAdd)
 {
+	UpdateToAdd.Context->StartQueueTime = FPlatformTime::Seconds();
+	
 	if (const FMutablePendingInstanceUpdate* ExistingUpdate = PendingInstanceUpdates.Find(UpdateToAdd.Context->Instance))
 	{
 		ExistingUpdate->Context->UpdateResult = EUpdateResult::ErrorReplaced;
 		FinishUpdateGlobal(ExistingUpdate->Context);
 
-		FMutablePendingInstanceUpdate TaskToEnqueue = UpdateToAdd;
+		const FMutablePendingInstanceUpdate TaskToEnqueue = UpdateToAdd;
 		TaskToEnqueue.Context->PriorityType = FMath::Min(ExistingUpdate->Context->PriorityType, UpdateToAdd.Context->PriorityType);
-		TaskToEnqueue.SecondsAtUpdate = FMath::Min(ExistingUpdate->SecondsAtUpdate, UpdateToAdd.SecondsAtUpdate);
+		TaskToEnqueue.Context->StartQueueTime = FMath::Min(ExistingUpdate->Context->StartQueueTime, UpdateToAdd.Context->StartQueueTime);
 		
-		PendingInstanceUpdates.Remove(ExistingUpdate->Context->Instance);
+		RemoveUpdate(ExistingUpdate->Context->Instance);
 		PendingInstanceUpdates.Add(TaskToEnqueue);
 	}
 	else
@@ -263,7 +258,11 @@ void FMutablePendingInstanceWork::AddUpdate(const FMutablePendingInstanceUpdate&
 
 void FMutablePendingInstanceWork::RemoveUpdate(const TWeakObjectPtr<UCustomizableObjectInstance>& Instance)
 {
-	PendingInstanceUpdates.Remove(Instance);
+	if (const FMutablePendingInstanceUpdate* Update = PendingInstanceUpdates.Find(Instance))
+	{
+		Update->Context->QueueTime = FPlatformTime::Seconds() - Update->Context->StartQueueTime;
+		PendingInstanceUpdates.Remove(Instance);
+	}	
 }
 
 
@@ -279,7 +278,7 @@ void FMutablePendingInstanceWork::AddDiscard(const FMutablePendingInstanceDiscar
 	{
 		ExistingUpdate->Context->UpdateResult = EUpdateResult::ErrorReplaced;
 		FinishUpdateGlobal(ExistingUpdate->Context);
-		PendingInstanceUpdates.Remove(ExistingUpdate->Context->Instance);
+		RemoveUpdate(ExistingUpdate->Context->Instance);
 	}
 
 	PendingInstanceDiscards.Add(TaskToEnqueue);
@@ -292,14 +291,6 @@ void FMutablePendingInstanceWork::AddIDRelease(mu::Instance::ID IDToRelease)
 }
 
 
-void FMutablePendingInstanceWork::RemoveAllUpdatesAndDiscardsAndReleases()
-{
-	PendingInstanceUpdates.Empty();
-	PendingInstanceDiscards.Empty();
-	PendingIDsToRelease.Empty();
-}
-
-
 UCustomizableObjectSystem* UCustomizableObjectSystem::GetInstance()
 {
 	if (!FCustomizableObjectSystemPrivate::SSystem)
@@ -558,8 +549,6 @@ void UCustomizableObjectSystem::BeginDestroy()
 
 		Private->CurrentInstanceBeingUpdated = nullptr;
 
-		Private->MutablePendingInstanceWork.RemoveAllUpdatesAndDiscardsAndReleases();
-
 		FCustomizableObjectSystemPrivate::SSystem = nullptr;
 
 		Private = nullptr;
@@ -839,13 +828,31 @@ void FinishUpdateGlobal(const TSharedRef<FUpdateContextPrivate>& Context)
 		SystemPrivate->MutableTaskGraph.AllowLaunchingMutableTaskLowPriority(true, false);
 	}
 	
+	if (Context->StartUpdateTime != 0.0) // Update started.
+	{
+		Context->UpdateTime = FPlatformTime::Seconds() - Context->StartUpdateTime;		
+	}
+	
 	const uint32 InstanceId = Instance ? Instance->GetUniqueID() : 0;
-	Context->UpdateTime = FPlatformTime::Seconds() - Context->StartUpdateTime;
-	UE_LOG(LogMutable, Log, TEXT("Finished UpdateSkeletalMesh Async. of Instance %d, Frame=%d, UpdateTime=%f"), InstanceId, GFrameNumber, Context->UpdateTime);
+	UE_LOG(LogMutable, Log, TEXT("Finished UpdateSkeletalMesh Async. Instance=%d, Frame=%d, QueueTime=%f, UpdateTime=%f"), InstanceId, GFrameNumber, Context->QueueTime, Context->UpdateTime);
 
-	if (SystemPrivate)
+	if (SystemPrivate &&
+		CVarEnableBenchmark.GetValueOnAnyThread())
 	{
-		SystemPrivate->LogBenchmarkUtil.FinishUpdate(Context);		
+		FFunctionGraphTask::CreateAndDispatchWhenReady( // Calling Benchmark in a task so we make sure we exited all scopes.
+		[Context]()
+		{
+			UCustomizableObjectSystem* System = UCustomizableObjectSystem::GetInstance();
+			if (!System)
+			{
+				return;
+			}
+
+			System->GetPrivateChecked()->LogBenchmarkUtil.FinishUpdateMesh(Context);
+		},
+		TStatId{},
+		nullptr,
+		ENamedThreads::GameThread);
 	}
 	
 	if (Context->UpdateStarted)
@@ -1121,6 +1128,12 @@ void FCustomizableObjectSystemPrivate::EnqueueUpdateSkeletalMesh(const TSharedRe
 	check(Instance);
 	
 	UCustomizableInstancePrivateData* InstancePrivate = Instance->GetPrivate();
+
+	const EQueuePriorityType Priority = GetUpdatePriority(*Instance, Context->bForceHighPriority);
+	const uint32 InstanceId = Instance->GetUniqueID();
+	const float Distance = FMath::Sqrt(InstancePrivate->LastMinSquareDistFromComponentToPlayer);
+	const bool bIsPlayerOrNearIt = InstancePrivate->HasCOInstanceFlags(UsedByPlayerOrNearIt);
+	UE_LOG(LogMutable, Log, TEXT("Enqueue UpdateSkeletalMesh Async. Instance=%d, Frame=%d, Priority=%d, dist=%f, bIsPlayerOrNearIt=%d"), InstanceId, GFrameNumber, static_cast<int32>(Priority), Distance, bIsPlayerOrNearIt);				
 	
 	if (!Instance->CanUpdateInstance())
 	{
@@ -1131,7 +1144,6 @@ void FCustomizableObjectSystemPrivate::EnqueueUpdateSkeletalMesh(const TSharedRe
 
 	UCustomizableObjectSystem* System = UCustomizableObjectSystem::GetInstance();
 
-	
 	const EUpdateRequired UpdateRequired = IsUpdateRequired(*Instance, Context->bOnlyUpdateIfNotGenerated, Context->bIgnoreCloseDist);
 	switch (UpdateRequired)
 	{
@@ -1143,13 +1155,6 @@ void FCustomizableObjectSystemPrivate::EnqueueUpdateSkeletalMesh(const TSharedRe
 	}		
 	case EUpdateRequired::Update:
 	{
-		EQueuePriorityType Priority = GetUpdatePriority(*Instance, Context->bForceHighPriority);
-
-		const uint32 InstanceId = Instance->GetUniqueID();
-		const float Distance = FMath::Sqrt(InstancePrivate->LastMinSquareDistFromComponentToPlayer);
-		const bool bIsPlayerOrNearIt = InstancePrivate->HasCOInstanceFlags(UsedByPlayerOrNearIt);
-		UE_LOG(LogMutable, Log, TEXT("Enqueued UpdateSkeletalMesh Async. of Instance %d with priority %d at dist %f bIsPlayerOrNearIt=%d, frame=%d"), InstanceId, static_cast<int32>(Priority), Distance, bIsPlayerOrNearIt, GFrameNumber);				
-
 		if (InstancePrivate->HasCOInstanceFlags(PendingLODsUpdate))
 		{
 			UE_LOG(LogMutable, Verbose, TEXT("LOD change: %d, %d -> %d, %d"), Instance->GetCurrentMinLOD(), Instance->GetCurrentMaxLOD(), Instance->GetMinLODToLoad(), Instance->GetMaxLODToLoad());
@@ -3145,11 +3150,11 @@ bool UCustomizableObjectSystem::Tick(float DeltaTime)
 						const double MinSquareDistFromComponentToPlayer = PendingUpdate.Context->Instance->GetPrivate()->MinSquareDistFromComponentToPlayer;
 						
 						if (MinSquareDistFromComponentToPlayer < MaxSquareDistanceFound ||
-							(MinSquareDistFromComponentToPlayer == MaxSquareDistanceFound && PendingUpdate.SecondsAtUpdate < MinTimeFound))
+							(MinSquareDistFromComponentToPlayer == MaxSquareDistanceFound && PendingUpdate.Context->StartQueueTime < MinTimeFound))
 						{
 							MaxPriorityFound = PriorityType;
 							MaxSquareDistanceFound = MinSquareDistFromComponentToPlayer;
-							MinTimeFound = PendingUpdate.SecondsAtUpdate;
+							MinTimeFound = PendingUpdate.Context->StartQueueTime;
 							PendingInstanceUpdateFound = &PendingUpdate;
 							LODUpdateCandidateFound = nullptr;
 						}
@@ -3560,13 +3565,13 @@ void UCustomizableObjectSystem::AddUncompiledCOWarning(const UCustomizableObject
 
 void UCustomizableObjectSystem::EnableBenchmark()
 {
-	GetPrivate()->LogBenchmarkUtil.SetEnable(true);
+	CVarEnableBenchmark->Set(true);
 }
 
 
 void UCustomizableObjectSystem::EndBenchmark()
 {
-	GetPrivate()->LogBenchmarkUtil.SetEnable(false);
+	CVarEnableBenchmark->Set(false);
 }
 
 
@@ -3822,7 +3827,10 @@ void FCustomizableObjectSystemPrivate::StartUpdateSkeletalMesh(const TSharedRef<
 
 	check(!CurrentMutableOperation); // Can not start an update if there is already another in progress
 	check(Context->Instance.IsValid()) // The instance has to be alive to start the update
-		
+
+	const uint32 InstanceId = Context->Instance->GetUniqueID();
+	UE_LOG(LogMutable, Log, TEXT("Started UpdateSkeletalMesh Async. Instance=%d, Frame=%d"), InstanceId, GFrameNumber);				
+			
 	CurrentMutableOperation = Context;
 }
 
diff --git a/Engine/Plugins/Experimental/Mutable/Source/CustomizableObject/Private/MuCO/CustomizableObjectSystemPrivate.h b/Engine/Plugins/Experimental/Mutable/Source/CustomizableObject/Private/MuCO/CustomizableObjectSystemPrivate.h
--- a/Engine/Plugins/Experimental/Mutable/Source/CustomizableObject/Private/MuCO/CustomizableObjectSystemPrivate.h
+++ b/Engine/Plugins/Experimental/Mutable/Source/CustomizableObject/Private/MuCO/CustomizableObjectSystemPrivate.h
@@ -38,8 +38,6 @@ struct FMutablePendingInstanceUpdate
 {
 	TSharedRef<FUpdateContextPrivate> Context;
 
-	double SecondsAtUpdate = 0;
-
 	FMutablePendingInstanceUpdate(const TSharedRef<FUpdateContextPrivate>& InContext);
 
 	bool operator==(const FMutablePendingInstanceUpdate& Other) const;
@@ -117,9 +115,6 @@ class FMutablePendingInstanceWork
 	int32 NumLODUpdatesLastTick = 0;
 
 public:
-	// Returns true if there are no pending instance updates. Doesn't take into account discards.
-	bool ArePendingUpdatesEmpty() const;
-
 	// Returns the number of pending instance updates, LOD Updates, discards and releases last tick.
 	int32 Num() const;
 
@@ -479,6 +474,9 @@ public:
 	bool UpdateStarted = false;
 
 	// Update stats
+	double StartQueueTime = 0.0;
+	double QueueTime = 0.0;
+	
 	double StartUpdateTime = 0.0;
 	double UpdateTime = 0.0;
 
diff --git a/Engine/Plugins/Experimental/Mutable/Source/CustomizableObject/Private/MuCO/LogBenchmarkUtil.cpp b/Engine/Plugins/Experimental/Mutable/Source/CustomizableObject/Private/MuCO/LogBenchmarkUtil.cpp
--- a/Engine/Plugins/Experimental/Mutable/Source/CustomizableObject/Private/MuCO/LogBenchmarkUtil.cpp
+++ b/Engine/Plugins/Experimental/Mutable/Source/CustomizableObject/Private/MuCO/LogBenchmarkUtil.cpp
@@ -24,11 +24,12 @@
 extern ENGINE_API UEngine* GEngine;
 
 
-static TAutoConsoleVariable<bool> CVarEnableBenchmark(
+TAutoConsoleVariable<bool> CVarEnableBenchmark(
 	TEXT("mutable.EnableBenchmark"),
 	false,
 	TEXT("Enable or disable the benchmarking."));
 
+
 namespace LogBenchmarkUtil
 {
 	void Write(FArchive& Archive, const FString& Text )
@@ -40,21 +41,12 @@ namespace LogBenchmarkUtil
 		const FStringView StringView = ComposedString;
 		Archive.Serialize(const_cast<ANSICHAR*>(StringCast<ANSICHAR>(StringView.GetData(), StringView.Len()).Get()), StringView.Len() * sizeof(ANSICHAR));
 	}
-
-	static FString LocalBenchmarkFilePath = TEXT("");
-	static FAutoConsoleVariableRef CVarBenchmarkFilePath(
-		TEXT("mutable.BenchmarkFilePath"),
-		LocalBenchmarkFilePath,
-		TEXT("Sets the path where to store the generated mutable benchmark report file"));
 }
 
+
 TSharedPtr<FArchive> CreateFile()
 {
-	FString Directory = LogBenchmarkUtil::LocalBenchmarkFilePath;
-	if (Directory.IsEmpty())
-	{
-		Directory = FPaths::ProfilingDir() + TEXT("Mutable/Benchmark");
-	}
+	const FString Directory = FPaths::ProfilingDir() + TEXT("Mutable/Benchmark");
 	IFileManager::Get().MakeDirectory(*Directory, true);
 
 	const FDateTime FileDate = FDateTime::Now();
@@ -63,19 +55,12 @@ TSharedPtr<FArchive> CreateFile()
 	TSharedPtr<FArchive> Archive = MakeShareable(IFileManager::Get().CreateFileWriter(*Filename, FILEWRITE_AllowRead | FILEWRITE_NoFail));
 	check(Archive);
 
-	const FString HeaderRow = TEXT("ID_CO,ID_COI,ID_UpdateType,Time_Update,Time_TaskGetMesh,Time_TaskLockCache,Time_TaskGetImages,Time_TaskConvertResources,Time_TaskCallbacks");
-	LogBenchmarkUtil::Write(*Archive, HeaderRow);		
+	const FString HeaderRow = TEXT("ID_CO,ID_COI,ID_UpdateType,ID_UpdateResult,Time_Queue,Time_Update,Time_TaskGetMesh,Time_TaskLockCache,Time_TaskGetImages,Time_TaskConvertResources,Time_TaskCallbacks,Time_TaskUpdateImages");
+	LogBenchmarkUtil::Write(*Archive, HeaderRow);
 
 	return Archive;
 }
 
-FLogBenchmarkUtil::FLogBenchmarkUtil()
-{
-	if (LogBenchmarkUtil::LocalBenchmarkFilePath.IsEmpty())
-	{
-		LogBenchmarkUtil::LocalBenchmarkFilePath = FPaths::ProfilingDir() + TEXT("Mutable/Benchmark");
-	}
-}
 
 FLogBenchmarkUtil::~FLogBenchmarkUtil()
 {
@@ -86,12 +71,6 @@ FLogBenchmarkUtil::~FLogBenchmarkUtil()
 }
 
 
-void FLogBenchmarkUtil::SetEnable(bool bEnabled)
-{
-	CVarEnableBenchmark->Set(bEnabled);
-}
-
-
 void FLogBenchmarkUtil::GetInstancesStats(int32& OutNumInstances, int32& OutNumBuiltInstances, int32& OutNumInstancesLOD0, int32& OutNumInstancesLOD1, int32& OutNumInstancesLOD2, int32& OutNumAllocatedSkeletalMeshes) const
 {
 	OutNumInstances = 0;
@@ -286,7 +265,7 @@ void FLogBenchmarkUtil::UpdateStats()
 }
 
 
-void FLogBenchmarkUtil::FinishUpdate(const TSharedRef<FUpdateContextPrivate>& Context)
+void FLogBenchmarkUtil::FinishUpdateMesh(const TSharedRef<FUpdateContextPrivate>& Context)
 {
 	check(IsInGameThread());
 
@@ -320,16 +299,37 @@ void FLogBenchmarkUtil::FinishUpdate(const TSharedRef<FUpdateContextPrivate>& Co
 	const FString ID_CO = Context->Instance->GetCustomizableObject()->GetPathName();
 	const FString ID_COI = Instance->GetPathName();
 	const FString ID_UpdateType = TEXT("Mesh");
+	const FString ID_UpdateResult = StaticEnum<EUpdateResult>()->GetValueAsString(Context->UpdateResult);
+	const double Time_Queue = Context->QueueTime * 1000;
 	const double Time_Update = Context->UpdateTime * 1000;
 	const double Time_TaskGetMesh = Context->TaskGetMeshTime * 1000;
 	const double Time_TaskLockCache = Context->TaskLockCacheTime * 1000;
 	const double Time_TaskGetImages = Context->TaskGetImagesTime * 1000;
 	const double Time_TaskConvertResources = Context->TaskConvertResourcesTime * 1000;
-	const double Time_TaskCallbacks = Context->TaskCallbacksTime * 1000;
+	const double Time_TaskCallbacks =  Context->TaskCallbacksTime * 1000;
 
-	const FString UpdateString = FString::Printf(TEXT("%s,%s,%s,%f,%f,%f,%f,%f,%f"), *ID_CO, *ID_COI, *ID_UpdateType, Time_Update, Time_TaskGetMesh, Time_TaskLockCache, Time_TaskGetImages, Time_TaskConvertResources, Time_TaskCallbacks);
+	const FString UpdateString = FString::Printf(TEXT("%s,%s,%s,%s,%f,%f,%f,%f,%f,%f,%f"), *ID_CO, *ID_COI, *ID_UpdateType, *ID_UpdateResult, Time_Queue, Time_Update, Time_TaskGetMesh, Time_TaskLockCache, Time_TaskGetImages, Time_TaskConvertResources, Time_TaskCallbacks);
 	LogBenchmarkUtil::Write(*Archive, UpdateString);
 	Archive->Flush();
-	
+}
+
+
+void FLogBenchmarkUtil::FinishUpdateImage(const FString& CustomizableObjectPathName, const FString& InstancePathName, const double TaskUpdateImageTime) const
+{
+	check(IsInGameThread());
+
+	if (!CVarEnableBenchmark.GetValueOnGameThread())
+	{
+		return;
+	}
+
+	const FString& ID_CO = CustomizableObjectPathName;
+	const FString& ID_COI = InstancePathName;
+	const FString ID_UpdateType = TEXT("Image");
+	const double Time_TaskUpdateImage = TaskUpdateImageTime * 1000;
+		
+	const FString UpdateString = FString::Printf(TEXT("%s,%s,%s,,,,,,,,,%f"), *ID_CO, *ID_COI, *ID_UpdateType, Time_TaskUpdateImage);
+	LogBenchmarkUtil::Write(*Archive, UpdateString);
+	Archive->Flush();	
 }
 
diff --git a/Engine/Plugins/Experimental/Mutable/Source/CustomizableObject/Private/MuCO/LogBenchmarkUtil.h b/Engine/Plugins/Experimental/Mutable/Source/CustomizableObject/Private/MuCO/LogBenchmarkUtil.h
--- a/Engine/Plugins/Experimental/Mutable/Source/CustomizableObject/Private/MuCO/LogBenchmarkUtil.h
+++ b/Engine/Plugins/Experimental/Mutable/Source/CustomizableObject/Private/MuCO/LogBenchmarkUtil.h
@@ -5,9 +5,12 @@
 #include "Containers/Map.h"
 #include "UObject/NameTypes.h"
 #include "Stats/Stats.h"
+#include "HAL/IConsoleManager.h"
 
+class UCustomizableObjectInstance;
 class UTexture2D;
 
+
 /** Stat which automatically gets updated to Insights when modified. */
 #define DECLARE_BENCHMARK_STAT(Name, Type) \
 	struct F##Name \
@@ -62,16 +65,15 @@ DECLARE_BENCHMARK_INSIGHTS(NumBuiltInstances, TEXT("Num Built Instances"))
 DECLARE_BENCHMARK_INSIGHTS(InstanceBuildTimeAvrg, TEXT("Avrg Instance Build Time"));
 
 
+extern TAutoConsoleVariable<bool> CVarEnableBenchmark;
+
+
 /** Benchmarking system. Gathers stats and send it to Insights an Benchmarking Files. */
 class FLogBenchmarkUtil
 {
 public:
-	FLogBenchmarkUtil();
 	~FLogBenchmarkUtil();
 
-	/** Enable or disable the system. */
-	void SetEnable(bool bEnable);
-
 	/** Get stats. */
 	void GetInstancesStats(int32& OutNumInstances, int32& OutNumBuiltInstances, int32& OutNumInstancesLOD0, int32& OutNumInstancesLOD1, int32& OutNumInstancesLOD2, int32& OutNumAllocatedSkeletalMeshes) const;
 
@@ -81,8 +83,10 @@ public:
 	/** Update stats which can only be updated on the tick. */
 	void UpdateStats();
 
-	/** Gathers stats update stats when it has finished. */
-	void FinishUpdate(const TSharedRef<FUpdateContextPrivate>& Context);
+	/** Gathers update stats when it has finished. */
+	void FinishUpdateMesh(const TSharedRef<FUpdateContextPrivate>& Context);
+
+	void FinishUpdateImage(const FString& CustomizableObjectPathName, const FString& InstancePathName, double TaskUpdateImageTime) const;
 	
 	DECLARE_BENCHMARK_STAT(NumAllocatedTextures, uint32);
 	DECLARE_BENCHMARK_STAT(NumAllocatedSkeletalMeshes, int32);
diff --git a/Engine/Plugins/Experimental/Mutable/Source/CustomizableObject/Public/MuCO/CustomizableObjectMipDataProvider.h b/Engine/Plugins/Experimental/Mutable/Source/CustomizableObject/Public/MuCO/CustomizableObjectMipDataProvider.h
--- a/Engine/Plugins/Experimental/Mutable/Source/CustomizableObject/Public/MuCO/CustomizableObjectMipDataProvider.h
+++ b/Engine/Plugins/Experimental/Mutable/Source/CustomizableObject/Public/MuCO/CustomizableObjectMipDataProvider.h
@@ -33,10 +33,15 @@ class FMutableUpdateContext
 public:
 	FMutableUpdateContext() = default;
 	
-	FMutableUpdateContext(mu::Ptr<mu::System> InSystem, TSharedPtr<mu::Model, ESPMode::ThreadSafe> InModel, mu::Ptr<const mu::Parameters> InParameters, int32 InState);
+	FMutableUpdateContext(const FString& InCustomizableObjectPathName, const FString& InInstancePathName, mu::Ptr<mu::System> InSystem,
+		TSharedPtr<mu::Model, ESPMode::ThreadSafe> InModel, mu::Ptr<const mu::Parameters> InParameters, int32 InState);
 
 	~FMutableUpdateContext();
 
+	const FString& GetCustomizableObjectPathName() const;
+
+	const FString& GetInstancePathName() const;
+	
 	mu::Ptr<mu::System> GetSystem() const;
 	
 	TSharedPtr<mu::Model, ESPMode::ThreadSafe> GetModel() const;
@@ -48,6 +53,9 @@ public:
 	const TArray<mu::Ptr<const mu::Image>>& GetImageParameterValues() const;
 	
 private:
+	FString CustomizableObjectPathName;
+	FString InstancePathName;
+
 	mu::Ptr<mu::System> System;
 	TSharedPtr<mu::Model, ESPMode::ThreadSafe> Model;
 	mu::Ptr<const mu::Parameters> Parameters;
