[Mutable] Fix crashes when accessing GCed materials. BuildMaterials will only override the materials of cached meshes instead of replacing their base materials.

#rnx

diff --git a/Engine/Plugins/Experimental/Mutable/Source/CustomizableObject/Private/MuCO/CustomizableObjectInstance.cpp b/Engine/Plugins/Experimental/Mutable/Source/CustomizableObject/Private/MuCO/CustomizableObjectInstance.cpp
--- a/Engine/Plugins/Experimental/Mutable/Source/CustomizableObject/Private/MuCO/CustomizableObjectInstance.cpp
+++ b/Engine/Plugins/Experimental/Mutable/Source/CustomizableObject/Private/MuCO/CustomizableObjectInstance.cpp
@@ -5375,26 +5375,34 @@ void UCustomizableInstancePrivateData::BuildMaterials(const TSharedRef<FUpdateCo
 			continue;
 		}
 
+		// It is not safe to replace the materials of a SkeletalMesh whose resources are initialized. Use overrides instead.
+		const bool bUseOverrideMaterialsOnly = OperationData->bUseMeshCache && SkeletalMesh->GetResourceForRendering()->IsInitialized();
+
 		ComponentsData[ComponentIndex].OverrideMaterials.Reset();
 
+		if (!bUseOverrideMaterialsOnly)
 		{
-			// TEMP: Keep a reference to the previous materials for n frames to avoid GC of materials in use in the render thread.
-			// TODO: MTBL-1632 - Implement a proper fix to replace this hotfix
-			const TArray<FSkeletalMaterial>& Materials = SkeletalMesh->GetMaterials();
-			const int32 NumMaterials = Materials.Num();
-			
-			TArray<TObjectPtr<UMaterialInterface>> MaterialsToRelease;
-			MaterialsToRelease.Reserve(NumMaterials);
-
-			for (const FSkeletalMaterial& Material : Materials)
 			{
-				MaterialsToRelease.Add(Material.MaterialInterface);
+				// TEMP: Keep a reference to the previous materials for n frames to avoid GC of materials in use in the render thread.
+				// TODO: MTBL-1632 - Implement a proper fix to replace this hotfix
+				const TArray<FSkeletalMaterial>& Materials = SkeletalMesh->GetMaterials();
+				const int32 NumMaterials = Materials.Num();
+
+				TArray<TObjectPtr<UMaterialInterface>> MaterialsToRelease;
+				MaterialsToRelease.Reserve(NumMaterials);
+
+				for (const FSkeletalMaterial& Material : Materials)
+				{
+					MaterialsToRelease.Add(Material.MaterialInterface);
+				}
+
+				UCustomizableObjectSystem::GetInstance()->AddPendingReleaseMaterials(MaterialsToRelease);
 			}
-			
-			UCustomizableObjectSystem::GetInstance()->AddPendingReleaseMaterials(MaterialsToRelease);
+
+			SkeletalMesh->GetMaterials().Reset();
 		}
 
-		SkeletalMesh->GetMaterials().Reset();
+		TArray<FSkeletalMaterial> Materials;
 
 		// Maps serializations of FMutableMaterialPlaceholder to Created Dynamic Material instances, used to reuse materials across LODs
 		TMap<uint32, TSharedPtr<FMutableMaterialPlaceholder>> ReuseMaterialCache;
@@ -5839,11 +5847,11 @@ void UCustomizableInstancePrivateData::BuildMaterials(const TSharedRef<FUpdateCo
 
 					if (SkeletalMesh)
 					{
-						MatIndex = SkeletalMesh->GetMaterials().Num();
+						MatIndex = Materials.Num();
 						MutableMaterialPlaceholder.MatIndex = MatIndex;
 
 						// Set up SkeletalMaterial data
-						FSkeletalMaterial& SkeletalMaterial = SkeletalMesh->GetMaterials().Add_GetRef(Material.MaterialInterface.Get());
+						FSkeletalMaterial& SkeletalMaterial = Materials.Add_GetRef(Material.MaterialInterface.Get());
 						SkeletalMaterial.MaterialSlotName = CustomizableObject->ReferencedMaterialSlotNames[Surface.MaterialIndex];
 						SetMeshUVChannelDensity(SkeletalMaterial.UVChannelData, RefSkeletalMeshData->Settings.DefaultUVChannelDensity);
 					}
@@ -5916,6 +5924,14 @@ void UCustomizableInstancePrivateData::BuildMaterials(const TSharedRef<FUpdateCo
 			}
 		}
 
+		if (!bUseOverrideMaterialsOnly)
+		{
+			SkeletalMesh->SetMaterials(Materials);
+		}
+
+		// Ensure the number of materials is the same on both sides when using overrides. 
+		check(SkeletalMesh->GetMaterials().Num() == Materials.Num());
+
 		{
 			// Copy data from valid LODs into the skipped ones.
 			int32 LastValidLODIndex = OperationData->CurrentMaxLOD;
