[Mutable] Fix small issues with resize optimizations for compressed image formats.

- Sampling artifacts could appear at batch borders.

#rnx

diff --git a/Engine/Plugins/Experimental/Mutable/Source/MutableRuntime/Private/MuR/OpImageResize.cpp b/Engine/Plugins/Experimental/Mutable/Source/MutableRuntime/Private/MuR/OpImageResize.cpp
--- a/Engine/Plugins/Experimental/Mutable/Source/MutableRuntime/Private/MuR/OpImageResize.cpp
+++ b/Engine/Plugins/Experimental/Mutable/Source/MutableRuntime/Private/MuR/OpImageResize.cpp
@@ -724,7 +724,8 @@ namespace
 {
 
 	template<uint32 NumChannels>
-	FORCENOINLINE void ImageResizeLinearDownByLessThanTwoSubImageVectorImpl(FVector2f InOneOverScalingFactor,
+	FORCENOINLINE void ImageResizeLinearDownByLessThanTwoSubImageVectorImpl(
+			FVector2f InOneOverScalingFactor, FVector2f InSrcSubPixelOffset,
 			FImageSize DestSize, FImageSize DestSubSize, FImageSize SrcSize, FImageSize SrcSubSize,
 		  	uint8* RESTRICT SubDest, const uint8* RESTRICT SubSrc)
 	{	
@@ -733,6 +734,7 @@ namespace
 		// For some reason (probably cache related), doing this copy yields marginally but consistently 
 		// better performance than using the function argument. 
 		const FVector2f OneOverScalingFactor = InOneOverScalingFactor;
+		const FVector2f SrcSubPixelOffset = InSrcSubPixelOffset;
 		
 		check(OneOverScalingFactor.X <= 2.0f && OneOverScalingFactor.Y <= 2.0f && 
 			  OneOverScalingFactor.X >= 1.0f && OneOverScalingFactor.Y >= 1.0f);
@@ -759,7 +761,7 @@ namespace
 					return VectorIntToFloat(MakeVectorRegisterInt(PixelData[0], PixelData[1], PixelData[2], PixelData[3]));
 				};
 				
-				const FVector2f CoordsF = FVector2f(X, Y)*OneOverScalingFactor;
+				const FVector2f CoordsF = FVector2f(X, Y)*OneOverScalingFactor + SrcSubPixelOffset;
 	
 				const FIntVector2 Coords = FIntVector2(FMath::FloorToInt(CoordsF.X), FMath::FloorToInt(CoordsF.Y));
 				const FIntVector2 CoordsPlusOne = FIntVector2(
@@ -807,7 +809,8 @@ namespace
 
 
 	template<uint32 NumChannels>
-	FORCENOINLINE void ImageResizeLinearDownByLessThanTwoSubImageNonVectorImpl(FVector2f InOneOverScalingFactor,
+	FORCENOINLINE void ImageResizeLinearDownByLessThanTwoSubImageNonVectorImpl(
+			FVector2f InOneOverScalingFactor, FVector2f InSrcSubPixelOffset,
 			FImageSize DestSize, FImageSize DestSubSize, FImageSize SrcSize, FImageSize SrcSubSize,
 		  	uint8* RESTRICT SubDest, const uint8* RESTRICT SubSrc)
 	{
@@ -821,6 +824,7 @@ namespace
 		// For some reason (probably cache related), doing this copy yields marginally but consistently 
 		// better performance than using the function argument. 
 		const FVector2f OneOverScalingFactor = InOneOverScalingFactor;
+		const FVector2f SrcSubPixelOffset = InSrcSubPixelOffset;
 		
 		check(OneOverScalingFactor.X <= 2.0f && OneOverScalingFactor.Y <= 2.0f && 
 			  OneOverScalingFactor.X >= 1.0f && OneOverScalingFactor.Y >= 1.0f);
@@ -853,7 +857,7 @@ namespace
 					return Result;
 				};
 				
-				const FVector2f CoordsF = FVector2f(X, Y)*OneOverScalingFactor;
+				const FVector2f CoordsF = FVector2f(X, Y)*OneOverScalingFactor + SrcSubPixelOffset;
 				
 				using FUint16Vector2 = UE::Math::TIntVector2<uint16>;
 				const FUint16Vector2 FracINorm = FUint16Vector2(FMath::Frac(CoordsF.X) * 255.0f, FMath::Frac(CoordsF.Y) * 255.0f);
@@ -900,19 +904,22 @@ namespace
 	}
 
 	template<int32 NumChannels, bool bUseVectorImpl>
-	FORCEINLINE void ImageResizeLinearDownByLessThanTwoSubImage(FVector2f OneOverScalingFactor,
+	FORCEINLINE void ImageResizeLinearDownByLessThanTwoSubImage(
+			FVector2f OneOverScalingFactor, FVector2f SrcSubPixelOffset,
 			FImageSize DestSize, FImageSize DestSubSize, FImageSize SrcSize, FImageSize SrcSubSize,
-		  	uint8* RESTRICT SubDest, const uint8* RESTRICT SubSrc)
+			uint8* RESTRICT SubDest, const uint8* RESTRICT SubSrc)
 	{
 		if constexpr (bUseVectorImpl)
 		{
 			ImageResizeLinearDownByLessThanTwoSubImageVectorImpl<NumChannels>(
-					OneOverScalingFactor, DestSize, DestSubSize, SrcSize, SrcSubSize, SubDest, SubSrc);
+					OneOverScalingFactor, SrcSubPixelOffset,
+					DestSize, DestSubSize, SrcSize, SrcSubSize, SubDest, SubSrc);
 		}
 		else
 		{
 			ImageResizeLinearDownByLessThanTwoSubImageNonVectorImpl<NumChannels>(
-					OneOverScalingFactor, DestSize, DestSubSize, SrcSize, SrcSubSize, SubDest, SubSrc);
+					OneOverScalingFactor, SrcSubPixelOffset,
+					DestSize, DestSubSize, SrcSize, SrcSubSize, SubDest, SubSrc);
 		}
 	}
 
@@ -941,14 +948,15 @@ namespace
 		{
 			MUTABLE_CPUPROFILER_SCOPE(ResizeUnCompressed);
 
+			const FVector2f SrcSubPixelOffset = FVector2f(0.0f);
 			// Uncompressed format, resize in one go directly to Dest.
 			switch(SrcFormat)
 			{
-
 			case EImageFormat::IF_L_UBYTE:
 			{
 				constexpr int32 NumChannels = 1;
-				ImageResizeLinearDownByLessThanTwoSubImage<NumChannels, bUseVectorImpl>(OneOverScalingFactor,
+				ImageResizeLinearDownByLessThanTwoSubImage<NumChannels, bUseVectorImpl>(
+						OneOverScalingFactor, SrcSubPixelOffset,
 						DestSize, DestSize, SrcSize, SrcSize, DestData, SrcData);
 				break;
 			}
@@ -956,7 +964,8 @@ namespace
 			case EImageFormat::IF_RGB_UBYTE:
 			{
 				constexpr int32 NumChannels = 3;
-				ImageResizeLinearDownByLessThanTwoSubImage<NumChannels, bUseVectorImpl>(OneOverScalingFactor,
+				ImageResizeLinearDownByLessThanTwoSubImage<NumChannels, bUseVectorImpl>(
+						OneOverScalingFactor, SrcSubPixelOffset,
 						DestSize, DestSize, SrcSize, SrcSize, DestData, SrcData);
 				break;
 			}
@@ -965,7 +974,8 @@ namespace
 			case EImageFormat::IF_BGRA_UBYTE:
 			{
 				constexpr int32 NumChannels = 4;
-				ImageResizeLinearDownByLessThanTwoSubImage<NumChannels, bUseVectorImpl>(OneOverScalingFactor,
+				ImageResizeLinearDownByLessThanTwoSubImage<NumChannels, bUseVectorImpl>(
+						OneOverScalingFactor, SrcSubPixelOffset,
 						DestSize, DestSize, SrcSize, SrcSize, DestData, SrcData);
 				break;
 			}
@@ -985,7 +995,6 @@ namespace
 					FMath::DivideAndRoundUp(SrcSize.X, uint16(PixelsPerBlock.X)),
 					FMath::DivideAndRoundUp(SrcSize.Y, uint16(PixelsPerBlock.Y)));
 
-
 			constexpr int32 BatchSizeInBlocksX = 1 << 5;
 			constexpr int32 BatchSizeInBlocksY = 1 << 4;
 
@@ -1066,6 +1075,10 @@ namespace
 								FMath::DivideAndRoundUp(BatchBeginInBlocks.X*PixelsPerBlock.X*DestSize.X, int32(SrcSize.X)),
 								FMath::DivideAndRoundUp(BatchBeginInBlocks.Y*PixelsPerBlock.Y*DestSize.Y, int32(SrcSize.Y)));
 
+						const FVector2f SrcSubPixelOffset = FVector2f(
+								FMath::Frac(float(BatchOutBeginInPixels.X) * OneOverScalingFactor.X),
+								FMath::Frac(float(BatchOutBeginInPixels.Y) * OneOverScalingFactor.Y));
+
 						const FIntVector2 BatchOutEndInPixels = FIntVector2(
 								FMath::Min(
 									FMath::DivideAndRoundUp(BatchEndInBlocks.X*PixelsPerBlock.X*DestSize.X, int32(SrcSize.X)), 
@@ -1073,7 +1086,7 @@ namespace
 								FMath::Min(
 									FMath::DivideAndRoundUp(BatchEndInBlocks.Y*PixelsPerBlock.Y*DestSize.Y, int32(SrcSize.Y)), 
 									int32(DestSize.Y))); 
-					
+							
 						const FIntVector2 BatchSizeInPixels = (BatchEndInBlocks - BatchBeginInBlocks)*PixelsPerBlock;
 						
 						const FImageSize SrcSubSize  = FImageSize(BatchSizeInPixels.X, BatchSizeInPixels.Y);
@@ -1081,7 +1094,7 @@ namespace
 								BatchOutEndInPixels.X - BatchOutBeginInPixels.X, 
 								BatchOutEndInPixels.Y - BatchOutBeginInPixels.Y);
 		
-						const FIntVector2& DestSubOffset = FIntVector2(BatchOutBeginInPixels.X, BatchOutBeginInPixels.Y);
+						const FIntVector2& DestSubOffset = BatchOutBeginInPixels;
 					
 						const uint8* SrcSubData = JobStagingMemoryData;
 						switch(DestNumChannels)
@@ -1092,7 +1105,8 @@ namespace
 						
 							uint8* DestSubData = DestData + (DestSubOffset.Y*DestSize.X + DestSubOffset.X)*NumChannels;
 
-							ImageResizeLinearDownByLessThanTwoSubImage<NumChannels, bUseVectorImpl>(OneOverScalingFactor, 
+							ImageResizeLinearDownByLessThanTwoSubImage<NumChannels, bUseVectorImpl>( 
+									OneOverScalingFactor, SrcSubPixelOffset,
 									DestSize, DestSubSize, DecSize, SrcSubSize, DestSubData, SrcSubData);
 							break;
 						}
@@ -1103,7 +1117,8 @@ namespace
 
 							uint8* DestSubData = DestData + (DestSubOffset.Y*DestSize.X + DestSubOffset.X)*NumChannels;
 							
-							ImageResizeLinearDownByLessThanTwoSubImage<NumChannels, bUseVectorImpl>(OneOverScalingFactor, 
+							ImageResizeLinearDownByLessThanTwoSubImage<NumChannels, bUseVectorImpl>( 
+									OneOverScalingFactor, SrcSubPixelOffset,
 									DestSize, DestSubSize, DecSize, SrcSubSize, DestSubData, SrcSubData);
 							break;
 						}
@@ -1114,7 +1129,8 @@ namespace
 
 							uint8* DestSubData = DestData + (DestSubOffset.Y*DestSize.X + DestSubOffset.X)*NumChannels;
 							
-							ImageResizeLinearDownByLessThanTwoSubImage<NumChannels, bUseVectorImpl>(OneOverScalingFactor,
+							ImageResizeLinearDownByLessThanTwoSubImage<NumChannels, bUseVectorImpl>(
+									OneOverScalingFactor, SrcSubPixelOffset,
 									DestSize, DestSubSize, DecSize, SrcSubSize, DestSubData, SrcSubData);
 							break;
 						}
