[mutable] Improved documentation of blueprint-exposed functions.
Cleanup old projection mask and coverage query code.
#rnx

diff --git a/Engine/Plugins/Experimental/Mutable/Source/CustomizableObject/Classes/MuCO/CustomizableObject.h b/Engine/Plugins/Experimental/Mutable/Source/CustomizableObject/Classes/MuCO/CustomizableObject.h
--- a/Engine/Plugins/Experimental/Mutable/Source/CustomizableObject/Classes/MuCO/CustomizableObject.h
+++ b/Engine/Plugins/Experimental/Mutable/Source/CustomizableObject/Classes/MuCO/CustomizableObject.h
@@ -444,46 +444,6 @@ struct FMutableModelParameterProperties
 };
 
 
-// Represents a texture used for masking-out areas of an object from projectors
-USTRUCT()
-struct FMaskOutTexture
-{
-	GENERATED_USTRUCT_BODY()
-
-	bool operator ==(const FMaskOutTexture& Other) const
-	{ 
-		return SizeX == Other.SizeX && SizeY == Other.SizeY && Data == Other.Data;
-	}
-
-	void SetTextureSize(int32 InSizeX, int32 InSizeY)
-	{
-		SizeX = InSizeX;
-		SizeY = InSizeY;
-		Data.SetNumUninitialized(FMath::DivideAndRoundUp(InSizeX * InSizeY, NumBitsPerDWORD)); 
-	}
-
-	int32 GetSizeX() const { return SizeX; }
-	int32 GetSizeY() const { return SizeY; }
-
-	FBitReference GetTexelReference(int32 Index)
-	{
-		check(Index >= 0 && Index < SizeX * SizeY);
-		return FBitReference(Data[Index / NumBitsPerDWORD], 1 << (Index & (NumBitsPerDWORD - 1)));
-	}
-
-private:
-	UPROPERTY()
-	int32 SizeX = 0;
-
-	UPROPERTY()
-	int32 SizeY = 0;
-
-	UPROPERTY()
-	TArray<uint32> Data; // Only alpha channel of PF_R8G8B8A8 is stored as binary bits
-};
-
-
-
 USTRUCT()
 struct CUSTOMIZABLEOBJECT_API FAnimBpOverridePhysicsAssetsInfo
 {
@@ -729,21 +689,6 @@ struct FMutableSkinWeightProfileInfo
 };
 
 
-
-UCLASS()
-class CUSTOMIZABLEOBJECT_API UMutableMaskOutCache : public UObject
-{
-public:
-	GENERATED_BODY()
-
-	UPROPERTY()
-	TMap<FString, FString> Materials; // Maps a UMaterial's asset path to a UTexture's asset path
-
-	UPROPERTY()
-	TMap<FString, FMaskOutTexture> Textures; // Maps a UTexture's asset path to the cached mask-out texture data
-};
-
-
 USTRUCT()
 struct CUSTOMIZABLEOBJECT_API FMutableStreamableBlock
 {
@@ -1303,49 +1248,45 @@ public:
 	UPROPERTY()
 	int32 NumMeshComponentsInRoot = 0;
 	
-	/** Method to query the amount of components this Customizable Object has.
-	 * @warning It must be invoked from a COInstance to ensure that the CO has been compiled
-	 */
+	/** Get the number of components this Customizable Object has. */
 	UFUNCTION(BlueprintCallable, Category = CustomizableObject)
 	int32 GetComponentCount() const;
 	
-	// Object parameters interface. This is used to query static data about the parameters available
-	// in instances of this object.
-
-	// Get the number of parameters available in any instance.
+	/** Get the number of parameters available in instances of this object. */
 	UFUNCTION(BlueprintCallable, Category = CustomizableObject)
-	int GetParameterCount() const;
+	int32 GetParameterCount() const;
 
-	// Get the index of a parameter
+	/** Get the index of a parameter from its name. Return -1 if the parameter is not found. */
 	UFUNCTION(BlueprintCallable, Category = CustomizableObject)
 	int32 FindParameter(const FString& Name) const;
 
-	// Get the type of a parameter
+	/** Get the type of a parameter from its index. */
 	UFUNCTION(BlueprintCallable, Category = CustomizableObject)
 	EMutableParameterType GetParameterType(int32 ParamIndex) const;
+
+	/** Get the type of a parameter from its name. */
 	UFUNCTION(BlueprintCallable, Category = CustomizableObject)
 	EMutableParameterType GetParameterTypeByName(const FString& Name) const;
 
-	// Get the name of a parameter
+	/** Get the name of a parameter from its index. */
 	UFUNCTION(BlueprintCallable, Category = CustomizableObject)
 	const FString& GetParameterName(int32 ParamIndex) const;
 
-	// Get the number of description images available for a parameter
+	/** Deprecated. It will always return 0. */
 	UFUNCTION(BlueprintCallable, Category = CustomizableObject, meta = (DeprecatedFunction, DeprecationMessage = "Parameter decorations have been removed. This method will be removed in future versions."))
 	int32 GetParameterDescriptionCount(const FString& ParamName) const;
 
-	// Returns how many possible options an int parameter has
+	/** Returns how many possible options an int parameter has, if the parameter is an enumeration. Otherwise return 0. */
 	UFUNCTION(BlueprintCallable, Category = CustomizableObject)
 	int32 GetIntParameterNumOptions(int32 ParamIndex) const;
 
-	// Gets the Name of the option at position K in the list of available options for the int parameter.
-	// Useful to enumerate the int parameter's possible options (Ex: "Hat1", "Hat2", "Cap", "Nothing")
+	/** Gets the Name of the option at position K in the list of available options for the int parameter.
+	 * Useful to enumerate the int parameter's possible options (Ex: "Hat1", "Hat2", "Cap", "Nothing")
+	 */
 	UFUNCTION(BlueprintCallable, Category = CustomizableObject)
 	const FString& GetIntParameterAvailableOption(int32 ParamIndex, int32 K) const;
 
-	//UFUNCTION(BlueprintCallable, Category = CustomizableObject)
 	int32 FindIntParameterValue( int32 ParamIndex, const FString& Value ) const;
-	//UFUNCTION(BlueprintCallable, Category = CustomizableObject)
 	FString FindIntParameterValueName(int32 ParamIndex, int32 ParamValue) const;
 
 	//
@@ -1365,32 +1306,8 @@ public:
 	// Remove skeletons that have been destroyed by the garbage collector from the cache.
 	void UnCacheInvalidSkeletons();
 
-	// Call before using Mutable's Projector testing with mask out features. It should only be loaded when needed because it can spend quite a lot of memory
-	// Can cause a loading hitch
-	UFUNCTION(BlueprintCallable, Category = CustomizableObject)
-	void LoadMaskOutCache();
-
-	// Call after having used Mutable's Projector testing with mask out features. It should be unloaded because it can spend quite a lot of memory
-	UFUNCTION(BlueprintCallable, Category = CustomizableObject)
-	void UnloadMaskOutCache();
-
-	// Called to load the reference SkeletalMesh if it needs to be used as a placeholder and it's not loaded.
-	UFUNCTION(BlueprintCallable, Category = CustomizableObject)
-	void LoadReferenceSkeletalMeshesAsync();
-
-	UFUNCTION(BlueprintCallable, Category = CustomizableObject)
-	void UnloadReferenceSkeletalMeshes();
-
-	// Callback of LoadReferenceSkeletalMeshesAsync
-	void OnReferenceSkeletalMeshesAsyncLoaded();
-
 private:
 
-#if !WITH_EDITORONLY_DATA
-	// Handle used to store a streaming request operation.
-	TSharedPtr<FStreamableHandle> RefSkeletalMeshStreamingHandle;
-#endif
-
 	
 	/** Returns true or false if the parameter with name can be located and it has the type the caller is looking for. It will also
 	 * check if the model has been set to ensure access to it can take place at the calculated parameter index.
@@ -1442,7 +1359,7 @@ public:
 	UFUNCTION(BlueprintCallable, Category = CustomizableObject)
 	void GetProjectorParameterDefaultValue (
 		UPARAM(DisplayName = "Parameter Name") const FString& InParameterName,
-		UPARAM(DisplayName = "Possition") FVector3f& OutPos,
+		UPARAM(DisplayName = "Position") FVector3f& OutPos,
 		UPARAM(DisplayName = "Direction") FVector3f& OutDirection,
 		UPARAM(DisplayName = "Up") FVector3f& OutUp,
 		UPARAM(DisplayName = "Scale") FVector3f& OutScale,
@@ -1461,13 +1378,13 @@ public:
 	 */
 	FName GetTextureParameterDefaultValue (const FString& InParameterName) const;
 
-	/** Return true or false depending if the parameter at the index provided is multidimensional or not.
+	/** Return true if the parameter at the index provided is multidimensional.
 	 * @param InParamIndex The index of the parameter to check.
 	 * @return True if the parameter is multidimensional and false if it is not.
 	 */
 	bool IsParameterMultidimensional(const int32& InParamIndex) const;
 	
-	/** Return true or false depending if the parameter at the index provided is multidimensional or not.
+	/** Return true if the parameter at the index provided is multidimensional.
 	 * @param InParameterName The name of the parameter to check.
 	 * @return True if the parameter is multidimensional and false if it is not.
 	 */
@@ -1571,43 +1488,46 @@ public:
 	void LoadEmbeddedData(FArchive& Ar);
 
 	void PostLoad() override;
-	void BeginDestroy() override;
 
 	void Serialize(FArchive& Ar) override;
 
-	// 
-	void SerializeClothingDerivedData(FMemoryWriter64& Ar);
-	void DeserializeClothingDerivedData(FMemoryReaderView& Ar);
-
 	FGuid GetCompilationGuid() const;
 
 	int32 FindState( const FString& Name ) const;
 
+	/** Return the number of object states that are defined in the CustomizableObject. */
 	UFUNCTION(BlueprintCallable, Category = CustomizableObject)
 	int32 GetStateCount() const;
 
+	/** Return the name of an object state from its index. */
 	UFUNCTION(BlueprintCallable, Category = CustomizableObject)
 	FString GetStateName(int32 StateIndex) const;
 
 	int32 GetStateParameterCount(int32 StateIndex) const;
 	int32 GetStateParameterIndex(int32 StateIndex, int32 ParameterIndex) const;
 
+	/** Return the number of parameters that are editable at runtime for a specific state. */
 	UFUNCTION(BlueprintCallable, Category = CustomizableObject)
 	int32 GetStateParameterCount(const FString& StateName) const;
 
+	/** Return the name of one of the state's runtime parameters, by its index (from 0 to GetStateParameterCount-1). */
 	UFUNCTION(BlueprintCallable, Category = CustomizableObject)
 	FString GetStateParameterName(const FString& StateName, int32 ParameterIndex) const;
 	FString GetStateParameterName(int32 StateIndex, int32 ParameterIndex) const;
 
+	/** Return the metadata associated to an object state by name. */
 	UFUNCTION(BlueprintCallable, Category = CustomizableObject)
 	FParameterUIData GetStateUIMetadata(const FString& StateName) const;
 
+	/** Return the metadata associated to an object state by state index. */
 	UFUNCTION(BlueprintCallable, Category = CustomizableObject)
 	FParameterUIData GetStateUIMetadataFromIndex(int32 StateIndex) const;
 
+	/** Return the metadata associated to an object parameter by parameter name. */
 	UFUNCTION(BlueprintCallable, Category = CustomizableObject)
 	FParameterUIData GetParameterUIMetadata(const FString& ParamName) const;
 
+	/** Return the metadata associated to an object parameter by parameter index (from 0 to GetParameterCount-1). */
 	UFUNCTION(BlueprintCallable, Category = CustomizableObject)
 	FParameterUIData GetParameterUIMetadataFromIndex(int32 ParamIndex) const;
 
@@ -1616,7 +1536,7 @@ public:
 	/** See bPreserveUserLODsOnFirstGeneration. */
 	bool IsPreserveUserLODsOnFirstGeneration() const;
 
-	/** Stores all the parameter UI metadata information for all the dependencies of this Customizable Object */
+	/** Stores all the parameter UI metadata information for all the dependencies of this Customizable Object. */
 	UPROPERTY()
 	TMap<FString, FParameterUIData> ParameterUIDataMap;
 
@@ -1644,10 +1564,6 @@ public:
 	/** Stores the sockets provided by the part skeletal meshes, to be merged in the generated meshes */
 	TArray<FMutableRefSocket> SocketArray;
 
-	/** Stores the textures that will be used to mask-out areas in projectors. The cache isn't used for rendering, but for coverage testing */
-	UPROPERTY()
-	TSoftObjectPtr<UMutableMaskOutCache> MaskOutCache;
-
 	/** Map of Hash to Streaming blocks, used to stream a block of data representing a resource from the BulkData */
 	UPROPERTY()
 	TMap<uint64, FMutableStreamableBlock> HashToStreamableBlock;
@@ -1687,6 +1603,7 @@ public:
 	void PostCompile();
 #endif
 
+	/** Create a new instance of this object. The instance parameters will be initialized with the object default values. */
 	UFUNCTION(BlueprintCallable, Category = CustomizableObject)
 	class UCustomizableObjectInstance* CreateInstance();
 
@@ -1695,7 +1612,7 @@ public:
 
 	FCustomizableObjectPrivateData* GetPrivate() const;
 
-	// This will always return true in a packaged game
+	/** Check if the CustomizableObject asset has been compiled. This will always be true in a packaged game, but it could be false in the editor. */
 	UFUNCTION(BlueprintCallable, Category = CustomizableObject)
 	bool IsCompiled() const;
 
@@ -1738,10 +1655,6 @@ private:
 	/** Cache of merged skeletons */
 	TArray<FMergedSkeleton> MergedSkeletons;
 
-	/** Used to prevent GC of MaskOutCache and keep it in memory while it's needed */
-	UPROPERTY(Transient)
-	TObjectPtr<UMutableMaskOutCache> MaskOutCache_HardRef;
-
 	/** Unique identifier. Regenerated each time the object is compiled. */
 	UPROPERTY()
 	FGuid CompilationGuid;
diff --git a/Engine/Plugins/Experimental/Mutable/Source/CustomizableObject/Classes/MuCO/CustomizableObjectInstance.h b/Engine/Plugins/Experimental/Mutable/Source/CustomizableObject/Classes/MuCO/CustomizableObjectInstance.h
--- a/Engine/Plugins/Experimental/Mutable/Source/CustomizableObject/Classes/MuCO/CustomizableObjectInstance.h
+++ b/Engine/Plugins/Experimental/Mutable/Source/CustomizableObject/Classes/MuCO/CustomizableObjectInstance.h
@@ -24,11 +24,6 @@ struct FFrame;
 struct FGameplayTagContainer;
 struct FPropertyChangedEvent;
 
-//! Order of the unreal vertex buffers when in mutable data
-#define MUTABLE_VERTEXBUFFER_POSITION	0
-#define MUTABLE_VERTEXBUFFER_TANGENT	1
-#define MUTABLE_VERTEXBUFFER_TEXCOORDS	2
-
 
 // Priority for the mutable update queue, Low is the normal distance-based priority, High is normally used for discards and Mid for LOD downgrades
 enum class EQueuePriorityType : uint8 { High, Med, Med_Low, Low };
@@ -160,75 +155,94 @@ public:
 	virtual FString GetDesc() override;
 	virtual bool IsEditorOnly() const override;
 
+	/** Set the CustomizableObject this instance will be generated from. 
+	* It is usually not necessary to call this since instances are already generated from a CustomizableObject.
+	*/
 	UFUNCTION(BlueprintCallable, Category = CustomizableObjectInstance)
 	void SetObject(UCustomizableObject* InObject);
 
-	//Get the current CustomizableObject 
+	/** Get the CustomizableObject that this is an instance of. */
 	UFUNCTION(BlueprintCallable, Category = CustomizableObjectInstance)
 	UCustomizableObject* GetCustomizableObject() const;
 
+	/** Deprecated. */
 	UFUNCTION(BlueprintCallable, Category = CustomizableObjectInstance, meta = (DeprecatedFunction, DeprecationMessage = "Parameter decorations have been removed. This method will be removed in future versions."))
 	bool GetBuildParameterDecorations() const;
 
+	/** Deprecated. */
 	UFUNCTION(BlueprintCallable, Category = CustomizableObjectInstance, meta = (DeprecatedFunction, DeprecationMessage = "Parameter decorations have been removed. This method will be removed in future versions."))
 	void SetBuildParameterDecorations(bool Value);
 
+	/** Return true if the parameter relevancy will be updated when this instance is generated. */
 	UFUNCTION(BlueprintCallable, Category = CustomizableObjectInstance)
 	bool GetBuildParameterRelevancy() const;
 
+	/** Set the flag that controls if parameter relevancy will be updated when this instance is generated. */
 	UFUNCTION(BlueprintCallable, Category = CustomizableObjectInstance)
 	void SetBuildParameterRelevancy(bool Value);
 
 	int32 GetState() const;
 	void SetState(int32 InState);
 
+	/** Return the name of the current CustomizableObject state this is instance is set to. */
 	UFUNCTION(BlueprintCallable, Category = CustomizableObjectInstance)
 	FString GetCurrentState() const;
 
+	/** Set the CustomizableObject state that this instance will be generated into. */
 	UFUNCTION(BlueprintCallable, Category = CustomizableObjectInstance)
 	void SetCurrentState(const FString& StateName);
 
+	/** Get the skeletal mesh generated for this instance. 
+	 * If the object has multiple components, an index of the component can be specified. 
+	 */
 	UFUNCTION(BlueprintCallable, Category = CustomizableObjectInstance)
 	USkeletalMesh* GetSkeletalMesh(int32 ComponentIndex = 0) const;
 	
+	/** Return true if a skeletal mesh has been generated for any component of this instance. */
 	UFUNCTION(BlueprintCallable, Category = CustomizableObjectInstance)
 	bool HasAnySkeletalMesh() const;
 
+	/** Get the array of parameters in the instance that are of type "bool". */
 	UFUNCTION(BlueprintCallable, Category = CustomizableObjectInstance)
 	TArray<FCustomizableObjectBoolParameterValue>& GetBoolParameters();
 
 	const TArray<FCustomizableObjectBoolParameterValue>& GetBoolParameters() const;
 
+	/** Get the array of parameters in the instance that are of type "integer". */
 	UFUNCTION(BlueprintCallable, Category = CustomizableObjectInstance)
 	TArray<FCustomizableObjectIntParameterValue>& GetIntParameters();
 
 	const TArray<FCustomizableObjectIntParameterValue>& GetIntParameters() const;
 
+	/** Get the array of parameters in the instance that are of type "float". */
 	UFUNCTION(BlueprintCallable, Category = CustomizableObjectInstance)
 	TArray<FCustomizableObjectFloatParameterValue>& GetFloatParameters();
 
 	const TArray<FCustomizableObjectFloatParameterValue>& GetFloatParameters() const;
 	
+	/** Get the array of parameters in the instance that are of type "texture". */
 	UFUNCTION(BlueprintCallable, Category = CustomizableObjectInstance)
 	TArray<FCustomizableObjectTextureParameterValue>& GetTextureParameters();
 
 	const TArray<FCustomizableObjectTextureParameterValue>& GetTextureParameters() const;
 
+	/** Get the array of parameters in the instance that are of type "vector". */
 	UFUNCTION(BlueprintCallable, Category = CustomizableObjectInstance)
 	TArray<FCustomizableObjectVectorParameterValue>& GetVectorParameters();
 
 	const TArray<FCustomizableObjectVectorParameterValue>& GetVectorParameters() const;
 
+	/** Get the array of parameters in the instance that are of type "projector". */
 	UFUNCTION(BlueprintCallable, Category = CustomizableObjectInstance)
 	TArray<FCustomizableObjectProjectorParameterValue>& GetProjectorParameters();
 	
 	const TArray<FCustomizableObjectProjectorParameterValue>& GetProjectorParameters() const;
 
-	/** See FCustomizableObjectInstanceDescriptor::HasAnyParameters. */
+	/** Return true if the instance has any parameters. */
 	UFUNCTION(BlueprintCallable, Category = CustomizableObjectInstance)
 	bool HasAnyParameters() const;
 	
-	/** Set random values to the parameters. Useful for testing only. */
+	/** Set random values to the parameters. */
 	UFUNCTION(BlueprintCallable, Category = CustomizableObjectInstance)
 	void SetRandomValues();
 
@@ -236,12 +250,12 @@ public:
 	UFUNCTION(BlueprintCallable, Category = CustomizableObjectInstance)
 	void SetRandomValuesFromStream(const FRandomStream& Stream);
 
-	/**  Set random values to the parameters using a seed. Useful for testing only. */
+	/**  Set random values to the parameters using a seed. */
 	void SetRandomValues(const int32 InRandomizationSeed);
 
 	/** Returns the AssetUserData that was gathered from all the constituent mesh parts during the last update. 
 	 *  It requires that the CustomizableObject had the bEnableAssetUserDataMerge set to true during compilation. 
-	*/
+	 */
 	UFUNCTION(BlueprintCallable, Category = CustomizableObjectInstance)
 	TSet<UAssetUserData*> GetMergedAssetUserData(int32 ComponentIndex) const;
 
@@ -255,11 +269,13 @@ public:
 	/** Return true if the instance is not locked and if it's compiled */
 	bool CanUpdateInstance() const;
 
-	/** Update Skeletal Mesh asynchronously. */
+	/** Generate the instance with the current parameters and update all the components Skeletal Meshes asynchronously. */
 	UFUNCTION(BlueprintCallable, Category = CustomizableObjectInstance)
 	void UpdateSkeletalMeshAsync(bool bIgnoreCloseDist = false, bool bForceHighPriority = false);
 		
-	/** Update Skeletal Mesh asynchronously. Callback will be called once the update finishes, even if it fails. */
+	/** Generate the instance with the current parameters and update all the components Skeletal Meshes asynchronously.
+	  * Callback will be called once the update finishes, even if it fails. 
+	  */
 	UFUNCTION(BlueprintCallable, Category = CustomizableObjectInstance)
 	void UpdateSkeletalMeshAsyncResult(FInstanceUpdateDelegate Callback, bool bIgnoreCloseDist = false, bool bForceHighPriority = false);
 
@@ -267,7 +283,7 @@ public:
 	UFUNCTION(BlueprintCallable, Category = CustomizableObjectInstance)
 	UCustomizableObjectInstance* Clone();
 
-	// Clones the instance creating a new identical static instance.
+	// Clones the instance creating a new identical static instance with the given Outer.
 	UFUNCTION(BlueprintCallable, Category = CustomizableObjectInstance)
 	UCustomizableObjectInstance* CloneStatic(UObject* Outer);
 
@@ -277,31 +293,35 @@ public:
 	/** Immediately destroy the Mutable Core Live Update Instance attached to this (if exists). */
 	void DestroyLiveUpdateInstance();
 
-	// Returns de description texture (ex: color bar) for this parameter and DescIndex
-	// This will only be valid if bBuildParameterDecorations was set to true before the last update.
+	/** Deprecated. */
 	UFUNCTION(BlueprintCallable, Category = CustomizableObjectInstance, meta = (DeprecatedFunction, DeprecationMessage = "Parameter decorations have been removed. This method will be removed in future versions."))
 	UTexture2D* GetParameterDescription(const FString& ParamName, int32 DescIndex);
 
 	//! 
 	bool IsParameterRelevant( int32 ParameterIndex ) const;
 
-	//! 
+	/** Return true if the given parameter has any effect in the current object state, and considering the current values of the other parameters. */
 	UFUNCTION(BlueprintCallable, Category = CustomizableObjectInstance)
 	bool IsParameterRelevant(const FString& ParamName) const;
 
+	/** Return true if the given parameter supports multiple dimensions, and so, multiple values. */
 	UFUNCTION(BlueprintCallable, Category = CustomizableObjectInstance, 
 		meta = (DeprecatedFunction, DeprecationMessage = "Use the method with the same name located at the Customizable Object instead."))
 	bool IsParamMultidimensional(const FString& ParamName) const;
-	
+
+	/** For multidimensional parameters, return the number of dimensions that the given projector parameter supports. */
 	UFUNCTION(BlueprintCallable, Category = CustomizableObjectInstance)
 	int32 GetProjectorValueRange(const FString& ParamName) const;
 
+	/** For multidimensional parameters, return the number of dimensions that the given int parameter supports. */
 	UFUNCTION(BlueprintCallable, Category = CustomizableObjectInstance)
 	int32 GetIntValueRange(const FString& ParamName) const;
 
+	/** For multidimensional parameters, return the number of dimensions that the given float parameter supports. */
 	UFUNCTION(BlueprintCallable, Category = CustomizableObjectInstance)
 	int32 GetFloatValueRange(const FString& ParamName) const;
 
+	/** For multidimensional parameters, return the number of dimensions that the given texture parameter supports. */
 	UFUNCTION(BlueprintCallable, Category = CustomizableObjectInstance)
 	int32 GetTextureValueRange(const FString& ParamName) const;
 	
@@ -316,21 +336,22 @@ public:
 	// Gets the Name of the option at position K in the list of available options for the int parameter. Useful to enumerate the int parameter's possible options (Ex: "Hat1", "Hat2", "Cap", "Nothing")
 	//const FString& GetIntParameterAvailableOption(int32 IntParamIndex, int32 K);
 
+	/** Return the name of the option currently set in the given parameter. */
 	UFUNCTION(BlueprintCallable, Category = CustomizableObjectInstance)
 	const FString& GetIntParameterSelectedOption(const FString& ParamName, int32 RangeIndex = -1) const;
 
-	// Sets the selected option of an int parameter by the option's name
+	/** Set the currently selected option value for the given parameter, by parameter index and option name. */
 	void SetIntParameterSelectedOption(int32 IntParamIndex, const FString& SelectedOption, int32 RangeIndex = -1);
 
-	// Sets the selected option of an int parameter, by the option's name
+	/** Set the currently selected option value for the given parameter, by parameter name and option name. */
 	UFUNCTION(BlueprintCallable, Category = CustomizableObjectInstance)
 	void SetIntParameterSelectedOption(const FString& ParamName, const FString& SelectedOptionName, int32 RangeIndex = -1);
 
-	// Gets the value of a float parameter with name "FloatParamName"
+	/** Gets the value of a float parameter with name "FloatParamName". */
 	UFUNCTION(BlueprintCallable, Category = CustomizableObjectInstance)
 	float GetFloatParameterSelectedOption(const FString& FloatParamName, int32 RangeIndex = -1) const;
 
-	// Sets the float value "FloatValue" of a float parameter with index "FloatParamIndex"
+	/** Sets the float value "FloatValue" of a float parameter with index "FloatParamIndex". */
 	UFUNCTION(BlueprintCallable, Category = CustomizableObjectInstance)
 	void SetFloatParameterSelectedOption(const FString& FloatParamName, float FloatValue, int32 RangeIndex = -1);
 	
@@ -342,54 +363,54 @@ public:
 	UFUNCTION(BlueprintCallable, Category = CustomizableObjectInstance)
 	void SetTextureParameterSelectedOption(const FString& TextureParamName, const FString& TextureValue, int32 RangeIndex = -1);
 
-	// Gets the value of a color parameter with name "ColorParamName"
+	/** Gets the value of a color parameter with name "ColorParamName". */
 	UFUNCTION(BlueprintCallable, Category = CustomizableObjectInstance)
 	FLinearColor GetColorParameterSelectedOption(const FString& ColorParamName) const;
 
-	// Sets the color value "ColorValue" of a color parameter with index "ColorParamIndex"
+	/** Sets the color value "ColorValue" of a color parameter with index "ColorParamIndex". */
 	UFUNCTION(BlueprintCallable, Category = CustomizableObjectInstance)
 	void SetColorParameterSelectedOption(const FString& ColorParamName, const FLinearColor& ColorValue);
 
-	// Sets the bool value "BoolValue" of a bool parameter with name "BoolParamName"
+	/** Sets the bool value "BoolValue" of a bool parameter with name "BoolParamName". */
 	UFUNCTION(BlueprintCallable, Category = CustomizableObjectInstance)
 	bool GetBoolParameterSelectedOption(const FString& BoolParamName) const;
 
-	// Sets the bool value "BoolValue" of a bool parameter with name "BoolParamName"
+	/** Sets the bool value "BoolValue" of a bool parameter with name "BoolParamName". */
 	UFUNCTION(BlueprintCallable, Category = CustomizableObjectInstance)
 	void SetBoolParameterSelectedOption(const FString& BoolParamName, bool BoolValue);
 
-	// Sets the vector value "VectorValue" of a bool parameter with index "VectorParamIndex"
+	/** Sets the vector value "VectorValue" of a bool parameter with index "VectorParamIndex". */
 	UFUNCTION(BlueprintCallable, Category = CustomizableObjectInstance)
 	void SetVectorParameterSelectedOption(const FString& VectorParamName, const FLinearColor& VectorValue);
 
-	// Sets the projector values of a projector parameter with index "ProjectorParamIndex"
+	/** Sets the projector values of a projector parameter with index "ProjectorParamIndex". */
 	UFUNCTION(BlueprintCallable, Category = CustomizableObjectInstance)
 	void SetProjectorValue(const FString& ProjectorParamName,
 		const FVector& OutPos, const FVector& OutDirection, const FVector& OutUp, const FVector& OutScale,
 		float OutAngle,
 		int32 RangeIndex = -1);
 
-	/** Set only the projector position. */
+	/** Set only the projector position keeping the rest of values. */
 	UFUNCTION(BlueprintCallable, Category = CustomizableObjectInstance)
 	void SetProjectorPosition(const FString& ProjectorParamName, const FVector3f& Pos, int32 RangeIndex = -1); // TODO GMTFuture Make API consistent (FVector3f -> FVector)
 
-	/** Set only the projector direction vector. */
+	/** Set only the projector direction vector keeping the rest of values. */
 	UFUNCTION(BlueprintCallable, Category = CustomizableObjectInstance)
 	void SetProjectorDirection(const FString& ProjectorParamName, const FVector& Direction, int32 RangeIndex = -1);
 	
-	/** Set only the projector up vector. */
+	/** Set only the projector up vector keeping the rest of values. */
 	UFUNCTION(BlueprintCallable, Category = CustomizableObjectInstance)
 	void SetProjectorUp(const FString& ProjectorParamName, const FVector& Up, int32 RangeIndex = -1);
 
-	/** Set only the projector scale. */
+	/** Set only the projector scale keeping the rest of values. */
 	UFUNCTION(BlueprintCallable, Category = CustomizableObjectInstance)
 	void SetProjectorScale(const FString& ProjectorParamName, const FVector& Scale, int32 RangeIndex = -1);
 
-	/** Set only the cylindrical projector angle. */
+	/** Set only the cylindrical projector angle keeping the rest of values. */
 	UFUNCTION(BlueprintCallable, Category = CustomizableObjectInstance)
 	void SetProjectorAngle(const FString& ProjectorParamName, float Angle, int32 RangeIndex = -1);
 	
-	// Get the projector values of a projector parameter with index "ProjectorParamIndex"
+	/** Get the projector values of a projector parameter with index "ProjectorParamIndex". */
 	UFUNCTION(BlueprintCallable, Category = CustomizableObjectInstance)
 	void GetProjectorValue(const FString& ProjectorParamName,
 		FVector& OutPos, FVector& OutDirection, FVector& OutUp, FVector& OutScale,
@@ -401,87 +422,90 @@ public:
 		float& Angle, ECustomizableObjectProjectorType& Type,
 		int32 RangeIndex = -1) const;
 
-	// Get the current projector position for the parameter with the given name
+	/** Get the current projector position for the parameter with the given name. */
 	UFUNCTION(BlueprintCallable, Category = CustomizableObjectInstance)
 	FVector GetProjectorPosition(const FString & ParamName, int32 RangeIndex = -1) const;
 
-	// Get the current projector direction vector for the parameter with the given name
+	/** Get the current projector direction vector for the parameter with the given name. */
 	UFUNCTION(BlueprintCallable, Category = CustomizableObjectInstance)
 	FVector GetProjectorDirection(const FString & ParamName, int32 RangeIndex = -1) const;
 
-	// Get the current projector up vector for the parameter with the given name
+	/** Get the current projector up vector for the parameter with the given name. */
 	UFUNCTION(BlueprintCallable, Category = CustomizableObjectInstance)
 	FVector GetProjectorUp(const FString & ParamName, int32 RangeIndex = -1) const;
 
-	// Get the current projector scale for the parameter with the given name
+	/** Get the current projector scale for the parameter with the given name. */
 	UFUNCTION(BlueprintCallable, Category = CustomizableObjectInstance)
 	FVector GetProjectorScale(const FString & ParamName, int32 RangeIndex = -1) const;
 
-	// Get the current cylindrical projector angle for the parameter with the given name
+	/** Get the current cylindrical projector angle for the parameter with the given name. */
 	UFUNCTION(BlueprintCallable, Category = CustomizableObjectInstance)
 	float GetProjectorAngle(const FString& ParamName, int32 RangeIndex = -1) const;
 
-	// Get the current projector type for the parameter with the given name
+	/** Get the current projector type for the parameter with the given name. */
 	UFUNCTION(BlueprintCallable, Category = CustomizableObjectInstance)
 	ECustomizableObjectProjectorType GetProjectorParameterType(const FString& ParamName, int32 RangeIndex = -1) const;
 
-	/** See FCustomizableObjectInstanceDescriptor::GetProjector. */
+	/** Get the current projector for the parameter with the given name. */
 	FCustomizableObjectProjector GetProjector(const FString& ParamName, int32 RangeIndex) const;
 	
-	// Finds in IntParameters a parameter with name ParamName, returns the index if found, -1 otherwise
+	/** Finds in IntParameters a parameter with name ParamName, returns the index if found, -1 otherwise. */
 	UFUNCTION(BlueprintCallable, Category = CustomizableObjectInstance)
 	int32 FindIntParameterNameIndex(const FString& ParamName) const;
 
-	// Finds in FloatParameters a parameter with name ParamName, returns the index if found, -1 otherwise
+	/** Finds in FloatParameters a parameter with name ParamName, returns the index if found, -1 otherwise. */
 	UFUNCTION(BlueprintCallable, Category = CustomizableObjectInstance)
 	int32 FindFloatParameterNameIndex(const FString& ParamName) const;
 
-	// Finds in BoolParameters a parameter with name ParamName, returns the index if found, -1 otherwise
+	/** Finds in BoolParameters a parameter with name ParamName, returns the index if found, -1 otherwise. */
 	UFUNCTION(BlueprintCallable, Category = CustomizableObjectInstance)
 	int32 FindBoolParameterNameIndex(const FString& ParamName) const;
 
-	// Finds in VectorParameters a parameter with name ParamName, returns the index if found, -1 otherwise
+	/** Finds in VectorParameters a parameter with name ParamName, returns the index if found, -1 otherwise. */
 	UFUNCTION(BlueprintCallable, Category = CustomizableObjectInstance)
 	int32 FindVectorParameterNameIndex(const FString& ParamName) const;
 
-	// Finds in ProjectorParameters a parameter with name ParamName, returns the index if found, -1 otherwise
+	/** Finds in ProjectorParameters a parameter with name ParamName, returns the index if found, -1 otherwise. */
 	UFUNCTION(BlueprintCallable, Category = CustomizableObjectInstance)
 	int32 FindProjectorParameterNameIndex(const FString& ParamName) const;
 
-	// Increases the range of values of the integer with ParamName, returns the index of the new integer value, -1 otherwise.
-	// The added value is initialized with the first integer option and is the last one of the range.
+	/** Increases the range of values of the integer with ParamName and returns the index of the new integer value, -1 otherwise.
+	* The added value is initialized with the first integer option and is the last one of the range.
+	*/
 	UFUNCTION(BlueprintCallable, Category = CustomizableObjectInstance)
 	int32 AddValueToIntRange(const FString& ParamName);
 
-	// Increases the range of values of the float with ParamName, returns the index of the new float value, -1 otherwise.
-	// The added value is initialized with 0.5f and is the last one of the range.
+	/** Increases the range of values of the float with ParamName, returns the index of the new float value, -1 otherwise.
+	* The added value is initialized with 0.5f and is the last one of the range.
+	*/
 	UFUNCTION(BlueprintCallable, Category = CustomizableObjectInstance)
 	int32 AddValueToFloatRange(const FString& ParamName);
 
-	// Increases the range of values of the projector with ParamName, returns the index of the new projector value, -1 otherwise.
-	// The added value is initialized with the default projector as set up in the editor and is the last one of the range.
+	/** Increases the range of values of the projector with ParamName, returns the index of the new projector value, -1 otherwise.
+	* The added value is initialized with the default projector as set up in the editor and is the last one of the range.
+	*/
 	UFUNCTION(BlueprintCallable, Category = CustomizableObjectInstance)
 	int32 AddValueToProjectorRange(const FString& ParamName);
 
-	// Remove the last of the integer range of values from the parameter ParamName, returns the index of the last valid integer, -1 if no values left.
+	/** Remove the last of the integer range of values from the parameter ParamName, returns the index of the last valid integer, -1 if no values left. */
 	UFUNCTION(BlueprintCallable, Category = CustomizableObjectInstance)
 	int32 RemoveValueFromIntRange(const FString& ParamName);
 
-	// Remove the RangeIndex element of the integer range of values from the parameter ParamName, returns the index of the last valid integer, -1 if no values left.
+	/** Remove the RangeIndex element of the integer range of values from the parameter ParamName, returns the index of the last valid integer, -1 if no values left. */
 	int32 RemoveValueFromIntRange(const FString& ParamName, int32 RangeIndex);
 
-	// Remove the last of the float range of values from the parameter ParamName, returns the index of the last valid float, -1 if no values left.
+	/** Remove the last of the float range of values from the parameter ParamName, returns the index of the last valid float, -1 if no values left. */
 	UFUNCTION(BlueprintCallable, Category = CustomizableObjectInstance)
 	int32 RemoveValueFromFloatRange(const FString& ParamName);
 
-	// Remove the RangeIndex element of the float range of values from the parameter ParamName, returns the index of the last valid float, -1 if no values left.
+	/** Remove the RangeIndex element of the float range of values from the parameter ParamName, returns the index of the last valid float, -1 if no values left. */
 	int32 RemoveValueFromFloatRange(const FString& ParamName, int32 RangeIndex);
 
-	// Remove the last of the projector range of values from the parameter ParamName, returns the index of the last valid projector, -1 if no values left.
+	/** Remove the last of the projector range of values from the parameter ParamName, returns the index of the last valid projector, -1 if no values left. */
 	UFUNCTION(BlueprintCallable, Category = CustomizableObjectInstance)
 	int32 RemoveValueFromProjectorRange(const FString& ParamName);
 
-	// Remove the RangeIndex element of the projector range of values from the parameter ParamName, returns the index of the last valid projector, -1 if no values left.
+	/** Remove the RangeIndex element of the projector range of values from the parameter ParamName, returns the index of the last valid projector, -1 if no values left. */
 	int32 RemoveValueFromProjectorRange(const FString& ParamName, int32 RangeIndex);
 
 	// ------------------------------------------------------------
@@ -489,8 +513,7 @@ public:
 	// ------------------------------------------------------------
 	
 	/** Given Multilayer Projector name, create a new Multilayer Projector Helper (if non-existent). See FMultilayerProjector.
-	 *
-	 * @return ture if successfully created (or was already created).
+	 * @return true if successfully created (or was already created).
 	 */
 	UFUNCTION(BlueprintCallable, Category = CustomizableObjectInstance)
 	bool CreateMultiLayerProjector(const FName& ProjectorParamName);
@@ -549,13 +572,15 @@ public:
 	
 	// ------------------------------------------------------------
 	
-	// Returns the animation BP for the parameter component and slot, gathered from all the meshes that compose this instance
+	/** Returns the animation BP for the parameter component and slot, gathered from all the meshes that compose this instance. */
 	UFUNCTION(BlueprintCallable, Category = CustomizableObjectInstance)
 	TSubclassOf<UAnimInstance> GetAnimBP(int32 ComponentIndex, const FName& Slot) const;
 
+	/** Return the list of tags for this instance. */
 	UFUNCTION(BlueprintCallable, Category = CustomizableObjectInstance)
 	const FGameplayTagContainer& GetAnimationGameplayTags() const;
 	
+	/** Execute a delegate for each animation instance invovled in this customizable object instance. */
 	UFUNCTION(BlueprintCallable, Category = CustomizableObjectInstance)
 	void ForEachAnimInstance(int32 ComponentIndex, FEachComponentAnimInstanceClassDelegate Delegate) const;
 
@@ -570,7 +595,7 @@ public:
 	/** See FCustomizableObjectInstanceDescriptor::LoadDescriptor. */
 	void LoadDescriptor(FArchive &CustomizableObjectDescriptor);
 
-	// Enable physics asset replacement so that generated skeletal meshes have the merged physics assets of their skeletal mesh parts and reference mesh
+	/** Enable physics asset replacement so that generated skeletal meshes have the merged physics assets of their skeletal mesh parts and reference mesh. */
 	UFUNCTION(BlueprintCallable, Category = CustomizableObjectInstance)
 	void SetReplacePhysicsAssets(bool bReplaceEnabled);
 
@@ -582,16 +607,6 @@ public:
 	// If enabled, low-priority textures will generate resident mipmaps too.
 	void SetForceGenerateResidentMips(bool bForceGenerateResidentMips);
 
-	// Adds/removes a texture channel coverage query to the instance update process. Every material updated by mutable with a channel named TextureName will be checked
-	// It only works properly when the Instance is in a state with texture compression disabled.
-	void AddQueryTextureCoverage(const FString& TextureName, const FString* MaskOutChannelName = nullptr);
-	void RemoveQueryTextureCoverage(const FString& TextureName);
-
-	// Returns the result of a texture coverage query previously registered with AddQueryTextureCoverage. The query is run during an update of the instance.
-	// It only works properly when the Instance is in a state with texture compression disabled.
-	float GetQueryResultTextureCoverage(const FString& TextureName);
-	float GetQueryResultTextureCoverageMasked(const FString& TextureName); // Same as the previous query but this time intersected with the mask texture
-
 	void AdditionalAssetsAsyncLoaded( FGraphEventRef CompletionEvent );
 
 	const TArray<TObjectPtr<UMaterialInterface>>* GetOverrideMaterials(int32 ComponentIndex) const;
@@ -605,6 +620,7 @@ public:
 	float GetMinSquareDistToPlayer() const;
 	void SetMinSquareDistToPlayer(float NewValue);
 
+	/** Return the number of components generated in this instance. */
 	UFUNCTION(BlueprintCallable, Category = CustomizableObjectInstance)
 	int32 GetNumComponents() const;
 
diff --git a/Engine/Plugins/Experimental/Mutable/Source/CustomizableObject/Private/MuCO/CustomizableInstancePrivateData.h b/Engine/Plugins/Experimental/Mutable/Source/CustomizableObject/Private/MuCO/CustomizableInstancePrivateData.h
--- a/Engine/Plugins/Experimental/Mutable/Source/CustomizableObject/Private/MuCO/CustomizableInstancePrivateData.h
+++ b/Engine/Plugins/Experimental/Mutable/Source/CustomizableObject/Private/MuCO/CustomizableInstancePrivateData.h
@@ -255,11 +255,6 @@ public:
 
 	static void ReleaseMutableTexture(const FMutableImageCacheKey& MutableTextureKey, UTexture2D* Texture, struct FMutableResourceCache& Cache);
 
-	// Prepare the data for the unreal textures, but don't create them because
-	// it runs in the mutable thread
-	// \TODO: CustomizableObject shouldn't be here
-	static void ProcessTextureCoverageQueries(const TSharedRef<FUpdateContextPrivate>& OperationData, UCustomizableObject* CustomizableObject, const FString& ImageKeyName, FTexturePlatformData *PlatformData, UMaterialInterface* Material);
-
 	// Copy data generated in the mutable thread over to the instance and initializes additional data required during the update
 	void PrepareForUpdate(const TSharedRef<FUpdateContextPrivate>& OperationData);
 
@@ -334,8 +329,6 @@ public:
 	// Maximum number of SkeletalMesh LODs to stream
 	uint8 NumMaxLODsToStream;
 	
-	TMap<FString, FTextureCoverageQueryData> TextureCoverageQueries;
-
 	UPROPERTY(Transient)
 	TArray<FCustomizableInstanceComponentData> ComponentsData;
 
diff --git a/Engine/Plugins/Experimental/Mutable/Source/CustomizableObject/Private/MuCO/CustomizableObject.cpp b/Engine/Plugins/Experimental/Mutable/Source/CustomizableObject/Private/MuCO/CustomizableObject.cpp
--- a/Engine/Plugins/Experimental/Mutable/Source/CustomizableObject/Private/MuCO/CustomizableObject.cpp
+++ b/Engine/Plugins/Experimental/Mutable/Source/CustomizableObject/Private/MuCO/CustomizableObject.cpp
@@ -216,21 +216,6 @@ void UCustomizableObject::PostLoad()
 }
 
 
-void UCustomizableObject::BeginDestroy()
-{
-#if !WITH_EDITORONLY_DATA
-	if (RefSkeletalMeshStreamingHandle.IsValid() && RefSkeletalMeshStreamingHandle->IsActive())
-	{
-		RefSkeletalMeshStreamingHandle->CancelHandle();
-	}
-	
-	RefSkeletalMeshStreamingHandle = nullptr;
-#endif
-
-	Super::BeginDestroy();
-}
-
-
 bool UCustomizableObject::IsLocked() const
 {
 	if (PrivateData.IsValid())
@@ -1029,9 +1014,9 @@ FString UCustomizableObject::GetCompiledDataFileName(bool bIsModel, const ITarge
 
 FString UCustomizableObject::GetDesc()
 {
-	int states = GetStateCount();
-	int params = GetParameterCount();
-	return FString::Printf(TEXT("%d States, %d Parameters"), states, params);
+	int32 States = GetStateCount();
+	int32 Params = GetParameterCount();
+	return FString::Printf(TEXT("%d States, %d Parameters"), States, Params);
 }
 
 
@@ -1197,60 +1182,6 @@ TSoftObjectPtr<USkeleton> UCustomizableObject::GetReferencedSkeletonAssetPtr( ui
 }
 
 
-void UCustomizableObject::LoadReferenceSkeletalMeshesAsync()
-{
-#if !WITH_EDITORONLY_DATA
-	if (!RefSkeletalMeshStreamingHandle)
-	{
-		TArray<FSoftObjectPath> MeshesToStream;
-
-		for (const FMutableRefSkeletalMeshData& RefSkeletalMeshData : ReferenceSkeletalMeshesData)
-		{
-			MeshesToStream.Add(RefSkeletalMeshData.SkeletalMeshAssetPath);
-		}
-
-		if (!MeshesToStream.IsEmpty())
-		{
-			TWeakObjectPtr<UCustomizableObject> WeakCO(this);
-
-			FStreamableManager& StreamableManager = UCustomizableObjectSystem::GetInstance()->GetStreamableManager();
-			RefSkeletalMeshStreamingHandle = StreamableManager.RequestAsyncLoad(MeshesToStream, FStreamableDelegate::CreateUObject(this, &UCustomizableObject::OnReferenceSkeletalMeshesAsyncLoaded),
-				FStreamableManager::AsyncLoadHighPriority);
-		}
-
-	}
-#endif
-}
-
-
-void UCustomizableObject::UnloadReferenceSkeletalMeshes()
-{
-#if !WITH_EDITORONLY_DATA
-	for (FMutableRefSkeletalMeshData& Data : ReferenceSkeletalMeshesData)
-	{
-		Data.SkeletalMesh = nullptr;
-	}
-#endif
-}
-
-
-void UCustomizableObject::OnReferenceSkeletalMeshesAsyncLoaded()
-{
-#if !WITH_EDITORONLY_DATA
-	if (RefSkeletalMeshStreamingHandle)
-	{
-		for (FMutableRefSkeletalMeshData& Data : ReferenceSkeletalMeshesData)
-		{
-			Data.SkeletalMesh = TSoftObjectPtr<USkeletalMesh>(Data.SkeletalMeshAssetPath).Get();
-			check(Data.SkeletalMesh);
-		}
-
-		RefSkeletalMeshStreamingHandle.Reset();
-	}
-#endif
-}
-
-
 TObjectPtr<USkeleton> UCustomizableObject::GetCachedMergedSkeleton(const int32 ComponentIndex, const TArray<uint16>& SkeletonIds) const
 {
 	const FMergedSkeleton* MergedSkeleton = MergedSkeletons.FindByPredicate([&ComponentIndex, &SkeletonIds](const FMergedSkeleton& MergedSkeleton) 
@@ -1427,7 +1358,7 @@ int32 UCustomizableObject::GetComponentCount() const
 	return IsCompiled() ? NumMeshComponentsInRoot : 0;
 }
 
-int UCustomizableObject::GetParameterCount() const
+int32 UCustomizableObject::GetParameterCount() const
 {
 	return ParameterProperties.Num();
 }
@@ -1713,18 +1644,6 @@ FParameterUIData UCustomizableObject::GetStateUIMetadataFromIndex(int32 StateInd
 }
 
 
-void UCustomizableObject::LoadMaskOutCache()
-{
-	MaskOutCache_HardRef = MaskOutCache.LoadSynchronous();
-}
-
-
-void UCustomizableObject::UnloadMaskOutCache()
-{
-	MaskOutCache_HardRef = nullptr;
-}
-
-
 float UCustomizableObject::GetFloatParameterDefaultValue(const FString& InParameterName) const
 {
 	const int32 ParameterIndex = FindParameter(InParameterName);
diff --git a/Engine/Plugins/Experimental/Mutable/Source/CustomizableObject/Private/MuCO/CustomizableObjectInstance.cpp b/Engine/Plugins/Experimental/Mutable/Source/CustomizableObject/Private/MuCO/CustomizableObjectInstance.cpp
--- a/Engine/Plugins/Experimental/Mutable/Source/CustomizableObject/Private/MuCO/CustomizableObjectInstance.cpp
+++ b/Engine/Plugins/Experimental/Mutable/Source/CustomizableObject/Private/MuCO/CustomizableObjectInstance.cpp
@@ -5995,81 +5995,6 @@ void UCustomizableInstancePrivateData::BuildMaterials(const TSharedRef<FUpdateCo
 }
 
 
-void UCustomizableInstancePrivateData::ProcessTextureCoverageQueries(const TSharedRef<FUpdateContextPrivate>& OperationData, UCustomizableObject* CustomizableObject, const FString& ImageKeyName, FTexturePlatformData *PlatformData, UMaterialInterface* Material)
-{
-	if (OperationData->TextureCoverageQueries_MutableThreadParams.Num() && PlatformData->Mips.Num())
-	{
-		int32 ImageKey = FCString::Atoi(*ImageKeyName);
-		int32 ScannedMipLevel = 0;
-		check(PlatformData->Mips.IsValidIndex(ScannedMipLevel));
-
-		if (ImageKey >= 0 && ImageKey < CustomizableObject->ImageProperties.Num())
-		{
-			const FMutableModelImageProperties& Props = CustomizableObject->ImageProperties[ImageKey];
-
-			FTextureCoverageQueryData* TextureCoverageQueryData = OperationData->TextureCoverageQueries_MutableThreadParams.Find(Props.TextureParameterName);
-
-			if (TextureCoverageQueryData)
-			{
-				FTexture2DMipMap& ScannedMip = PlatformData->Mips[ScannedMipLevel];
-				uint32 TotalTexels = ScannedMip.SizeX * ScannedMip.SizeY;
-				uint32 CoveredTexels = 0;
-				uint32 MaskedOutCoveredTexels = 0;
-
-				auto MaskOutCache = CustomizableObject->MaskOutCache.Get();
-
-				if (!MaskOutCache)
-				{
-					UE_LOG(LogMutable, Error, TEXT("The CustomizableObject->MaskOutCache has to be manually loaded by the programmer by calling CustomizableObject->LoadMaskOutCache()."));
-				}
-
-				FString* MaskOutTexturePath = MaskOutCache ? MaskOutCache->Materials.Find(Material->GetPathName()) : nullptr;
-				FMaskOutTexture* MaskOutTexture = MaskOutTexturePath ? MaskOutCache->Textures.Find(*MaskOutTexturePath) : nullptr;
-
-				// Textures must be uncompressed, the CustomizableInstance's state should be set to one with
-				// texture compression disabled for Texture Coverage Queries
-				ensure(PlatformData->PixelFormat == PF_R8G8B8A8 || PlatformData->PixelFormat == PF_B8G8R8A8);
-
-				const void* pData = ScannedMip.BulkData.Lock(LOCK_READ_ONLY);
-				uint8_t* pDest = (uint8_t*)pData;
-
-				for (size_t px = 0; px < ScannedMip.SizeX; ++px)
-				{
-					for (size_t py = 0; py < ScannedMip.SizeY; ++py)
-					{
-						if (pDest[(py * ScannedMip.SizeX + px) * 4 + 3] > 0) // Check if alpha channel of PF_R8G8B8A8 is not zero
-						{
-							CoveredTexels++;
-
-							if (MaskOutTexture)
-							{
-								float u = float(px) / ScannedMip.SizeX;
-								float v = float(py) / ScannedMip.SizeY;
-
-								size_t pu = FMath::RoundToInt(u * MaskOutTexture->GetSizeX());
-								size_t pv = FMath::RoundToInt(v * MaskOutTexture->GetSizeY());
-
-								if (MaskOutTexture->GetTexelReference(pv * MaskOutTexture->GetSizeX() + pu) == 0)
-								{
-									MaskedOutCoveredTexels++;
-								}
-							}
-						}
-					}
-				}
-
-				ScannedMip.BulkData.Unlock();
-
-				auto& ResultData = OperationData->TextureCoverageQueries_MutableThreadResults.FindOrAdd(Props.TextureParameterName);
-				ResultData.CoveredTexels += CoveredTexels;
-				ResultData.MaskedOutCoveredTexels += MaskedOutCoveredTexels;
-				ResultData.TotalTexels += TotalTexels;
-			}
-		}
-	}
-}
-
-
 void UCustomizableObjectInstance::SetReplacePhysicsAssets(bool bReplaceEnabled)
 {
 	bReplaceEnabled ? GetPrivate()->SetCOInstanceFlags(ReplacePhysicsAssets) : GetPrivate()->ClearCOInstanceFlags(ReplacePhysicsAssets);
@@ -6088,39 +6013,6 @@ void UCustomizableObjectInstance::SetForceGenerateResidentMips(bool bForceGenera
 }
 
 
-void UCustomizableObjectInstance::AddQueryTextureCoverage(const FString& TextureName, const FString* MaskOutChannelName)
-{
-	FTextureCoverageQueryData& QueryData = GetPrivate()->TextureCoverageQueries.FindOrAdd(TextureName);
-
-	if (MaskOutChannelName)
-	{
-		QueryData.MaskOutChannelName = *MaskOutChannelName;
-	}
-}
-
-
-void UCustomizableObjectInstance::RemoveQueryTextureCoverage(const FString& TextureName)
-{
-	GetPrivate()->TextureCoverageQueries.Remove(TextureName);
-}
-
-
-float UCustomizableObjectInstance::GetQueryResultTextureCoverage(const FString& TextureName)
-{
-	FTextureCoverageQueryData* QueryData = GetPrivate()->TextureCoverageQueries.Find(TextureName);
-
-	return QueryData ? QueryData->GetCoverage() : 0.f;
-}
-
-
-float UCustomizableObjectInstance::GetQueryResultTextureCoverageMasked(const FString& TextureName)
-{
-	FTextureCoverageQueryData* QueryData = GetPrivate()->TextureCoverageQueries.Find(TextureName);
-
-	return QueryData ? QueryData->GetMaskedOutCoverage() : 0.f;
-}
-
-
 void UCustomizableObjectInstance::SetIsBeingUsedByComponentInPlay(bool bIsUsedByComponentInPlay)
 {
 	bIsUsedByComponentInPlay ? GetPrivate()->SetCOInstanceFlags(UsedByComponentInPlay) : GetPrivate()->ClearCOInstanceFlags(UsedByComponentInPlay);
diff --git a/Engine/Plugins/Experimental/Mutable/Source/CustomizableObject/Private/MuCO/CustomizableObjectSystem.cpp b/Engine/Plugins/Experimental/Mutable/Source/CustomizableObject/Private/MuCO/CustomizableObjectSystem.cpp
--- a/Engine/Plugins/Experimental/Mutable/Source/CustomizableObject/Private/MuCO/CustomizableObjectSystem.cpp
+++ b/Engine/Plugins/Experimental/Mutable/Source/CustomizableObject/Private/MuCO/CustomizableObjectSystem.cpp
@@ -2033,7 +2033,6 @@ namespace impl
 						{
 							FTexturePlatformData* PlatformData = UCustomizableInstancePrivateData::MutableCreateImagePlatformData(MutableImage, -1, Image.FullImageSizeX, Image.FullImageSizeY);
 							OperationData->ImageToPlatformDataMap.Add(Image.ImageID, PlatformData);
-							OperationData->PendingTextureCoverageQueries.Add({ KeyName.ToString(), Surface.MaterialIndex, PlatformData});
 						}
 						else
 						{
@@ -2379,70 +2378,39 @@ namespace impl
 		{
 			UCustomizableInstancePrivateData* CustomizableInstancePrivateData = CustomizableObjectInstance->GetPrivate();
 
-			// Process the pending texture coverage queries
-			{
-				MUTABLE_CPUPROFILER_SCOPE(GameTextureQueries);
-				for (const FPendingTextureCoverageQuery& Query : OperationData->PendingTextureCoverageQueries)
-				{
-					UMaterialInterface* Material = nullptr;
-					const uint32* InstanceIndex = CustomizableInstancePrivateData->ObjectToInstanceIndexMap.Find(Query.MaterialIndex);
-					if (InstanceIndex && CustomizableInstancePrivateData->ReferencedMaterials.IsValidIndex(*InstanceIndex))
-					{
-						Material = CustomizableInstancePrivateData->ReferencedMaterials[*InstanceIndex];
-					}
+#if WITH_EDITOR
+			CustomizableObjectInstance->LastUpdateMutableRuntimeCycles = OperationData->MutableRuntimeCycles;
+#endif
 
-					UCustomizableInstancePrivateData::ProcessTextureCoverageQueries(OperationData, CustomizableObjectInstance->GetCustomizableObject(), Query.KeyName, Query.PlatformData, Material);
-				}
-				OperationData->PendingTextureCoverageQueries.Empty();
-			}
+			// Convert Step
+			//-------------------------------------------------------------
 
-			// Process texture coverage queries because it's safe to do now that the Mutable thread is stopped
+			// \TODO: Bring that code here instead of keeping it in the UCustomizableObjectInstance
+			if (CustomizableInstancePrivateData->UpdateSkeletalMesh_PostBeginUpdate0(CustomizableObjectInstance, OperationData))
 			{
-				if (OperationData->TextureCoverageQueries_MutableThreadResults.Num() > 0)
+				// This used to be CustomizableObjectInstance::UpdateSkeletalMesh_PostBeginUpdate1
 				{
-					for (auto& Result : OperationData->TextureCoverageQueries_MutableThreadResults)
-					{
-						FTextureCoverageQueryData* FinalResultData = CustomizableInstancePrivateData->TextureCoverageQueries.Find(Result.Key);
-						*FinalResultData = Result.Value;
-					}
+					MUTABLE_CPUPROFILER_SCOPE(UpdateSkeletalMesh_PostBeginUpdate1);
 
-					OperationData->TextureCoverageQueries_MutableThreadResults.Empty();
+					// \TODO: Bring here
+					CustomizableInstancePrivateData->BuildMaterials(OperationData, CustomizableObjectInstance);
 				}
 
-#if WITH_EDITOR
-				CustomizableObjectInstance->LastUpdateMutableRuntimeCycles = OperationData->MutableRuntimeCycles;
-#endif
-
-				// Convert Step
-				//-------------------------------------------------------------
-
-				// \TODO: Bring that code here instead of keeping it in the UCustomizableObjectInstance
-				if (CustomizableInstancePrivateData->UpdateSkeletalMesh_PostBeginUpdate0(CustomizableObjectInstance, OperationData))
+				// This used to be CustomizableObjectInstance::UpdateSkeletalMesh_PostBeginUpdate2
 				{
-					// This used to be CustomizableObjectInstance::UpdateSkeletalMesh_PostBeginUpdate1
-					{
-						MUTABLE_CPUPROFILER_SCOPE(UpdateSkeletalMesh_PostBeginUpdate1);
+					MUTABLE_CPUPROFILER_SCOPE(UpdateSkeletalMesh_PostBeginUpdate2);
 
-						// \TODO: Bring here
-						CustomizableInstancePrivateData->BuildMaterials(OperationData, CustomizableObjectInstance);
-					}
-
-					// This used to be CustomizableObjectInstance::UpdateSkeletalMesh_PostBeginUpdate2
+					for (int32 Component = 0; Component < CustomizableObjectInstance->SkeletalMeshes.Num(); ++Component)
 					{
-						MUTABLE_CPUPROFILER_SCOPE(UpdateSkeletalMesh_PostBeginUpdate2);
-
-						for (int32 Component = 0; Component < CustomizableObjectInstance->SkeletalMeshes.Num(); ++Component)
+						if (CustomizableObjectInstance->SkeletalMeshes[Component] && CustomizableObjectInstance->SkeletalMeshes[Component]->GetLODInfoArray().Num())
 						{
-							if (CustomizableObjectInstance->SkeletalMeshes[Component] && CustomizableObjectInstance->SkeletalMeshes[Component]->GetLODInfoArray().Num())
-							{
-								MUTABLE_CPUPROFILER_SCOPE(UpdateSkeletalMesh_PostEditChangeProperty);
+							MUTABLE_CPUPROFILER_SCOPE(UpdateSkeletalMesh_PostEditChangeProperty);
 
-								CustomizableInstancePrivateData->PostEditChangePropertyWithoutEditor(CustomizableObjectInstance->SkeletalMeshes[Component]);
-							}
+							CustomizableInstancePrivateData->PostEditChangePropertyWithoutEditor(CustomizableObjectInstance->SkeletalMeshes[Component]);
 						}
 					}
 				}
-			} // END - Process texture coverage queries
+			}
 		} // if (!bInstanceValid)
 
 		FCustomizableObjectSystemPrivate* CustomizableObjectSystemPrivateData = System->GetPrivateChecked();
@@ -2892,7 +2860,6 @@ namespace impl
 		
 		// Task: Mutable Update and GetMesh
 		//-------------------------------------------------------------
-		Operation->TextureCoverageQueries_MutableThreadParams = CandidateInstancePrivateData->TextureCoverageQueries;
 		Operation->CurrentMinLOD = Operation->InstanceDescriptorRuntimeHash.GetMinLOD();
 		Operation->CurrentMaxLOD = Operation->InstanceDescriptorRuntimeHash.GetMaxLOD();
 		Operation->InstanceID = Operation->bLiveUpdateMode ? CandidateInstancePrivateData->LiveUpdateModeInstanceID : 0;
diff --git a/Engine/Plugins/Experimental/Mutable/Source/CustomizableObject/Private/MuCO/CustomizableObjectSystemPrivate.h b/Engine/Plugins/Experimental/Mutable/Source/CustomizableObject/Private/MuCO/CustomizableObjectSystemPrivate.h
--- a/Engine/Plugins/Experimental/Mutable/Source/CustomizableObject/Private/MuCO/CustomizableObjectSystemPrivate.h
+++ b/Engine/Plugins/Experimental/Mutable/Source/CustomizableObject/Private/MuCO/CustomizableObjectSystemPrivate.h
@@ -370,28 +370,6 @@ struct FInstanceUpdateData
 };
 
 
-struct FTextureCoverageQueryData
-{
-	uint32 CoveredTexels;
-	uint32 MaskedOutCoveredTexels;
-	uint32 TotalTexels;
-	FString MaskOutChannelName;
-
-	FTextureCoverageQueryData() : CoveredTexels(0), MaskedOutCoveredTexels(0), TotalTexels(0) { }
-
-	float GetCoverage() const { return TotalTexels > 0 ? (float(CoveredTexels) / TotalTexels) : 0.f; }
-	float GetMaskedOutCoverage() const { return CoveredTexels > 0 ? (float(MaskedOutCoveredTexels) / CoveredTexels) : 0.f; }
-};
-
-
-struct FPendingTextureCoverageQuery
-{
-	FString KeyName;
-	uint32 MaterialIndex = 0;
-	FTexturePlatformData* PlatformData = nullptr;
-};
-
-
 /** Update Context.
  *
  * Alive from the start to the end of the update (both API and LOD update). */
@@ -456,14 +434,8 @@ public:
 
 	TArray<uint16> RequestedLODs;
 
-	TMap<FString, FTextureCoverageQueryData> TextureCoverageQueries_MutableThreadParams;
-	TMap<FString, FTextureCoverageQueryData> TextureCoverageQueries_MutableThreadResults;
-
 	TMap<uint32, FTexturePlatformData*> ImageToPlatformDataMap;
 
-	/** This list of queries is generated in the update mutable task, and consumed later in the game thread. */
-	TArray<FPendingTextureCoverageQuery> PendingTextureCoverageQueries;
-
 	EUpdateResult UpdateResult = EUpdateResult::Success;
 
 	mu::FImageOperator::FImagePixelFormatFunc PixelFormatOverride;
diff --git a/Engine/Plugins/Experimental/Mutable/Source/CustomizableObject/Public/MuCO/UnrealConversionUtils.h b/Engine/Plugins/Experimental/Mutable/Source/CustomizableObject/Public/MuCO/UnrealConversionUtils.h
--- a/Engine/Plugins/Experimental/Mutable/Source/CustomizableObject/Public/MuCO/UnrealConversionUtils.h
+++ b/Engine/Plugins/Experimental/Mutable/Source/CustomizableObject/Public/MuCO/UnrealConversionUtils.h
@@ -4,6 +4,11 @@
 
 #include "MuCO/CustomizableObjectSystemPrivate.h"
 
+//! Order of the unreal vertex buffers when in mutable data
+#define MUTABLE_VERTEXBUFFER_POSITION	0
+#define MUTABLE_VERTEXBUFFER_TANGENT	1
+#define MUTABLE_VERTEXBUFFER_TEXCOORDS	2
+
 namespace mu { class FMeshBufferSet; }
 struct FReferenceSkeleton;
 
diff --git a/Engine/Plugins/Experimental/Mutable/Source/CustomizableObjectEditor/Private/MuCOE/CustomizableObjectCompiler.cpp b/Engine/Plugins/Experimental/Mutable/Source/CustomizableObjectEditor/Private/MuCOE/CustomizableObjectCompiler.cpp
--- a/Engine/Plugins/Experimental/Mutable/Source/CustomizableObjectEditor/Private/MuCOE/CustomizableObjectCompiler.cpp
+++ b/Engine/Plugins/Experimental/Mutable/Source/CustomizableObjectEditor/Private/MuCOE/CustomizableObjectCompiler.cpp
@@ -1231,55 +1231,6 @@ void FCustomizableObjectCompiler::CompileInternal(UCustomizableObject* Object, c
 			ParticipatingObjects = MoveTemp(GenerationContext.ParticipatingObjects);
 		}
 
-		if (bIsRootObject && (GenerationContext.MaskOutMaterialCache.Num() > 0 || GenerationContext.MaskOutTextureCache.Num() > 0))
-		{
-			// Load MaskOutCache, if can't, create it.
-			if (!Object->MaskOutCache.Get())
-			{
-				if (!Object->MaskOutCache.LoadSynchronous())
-				{
-					FString MaskOutCacheName = TEXT("MaskOutCache");
-					FString PackageName = Object->GetOutermost()->GetPathName() + FString("_") + MaskOutCacheName;
-					UPackage* Package = CreatePackage(*PackageName);
-					Package->FullyLoad();
-
-					FString ObjectName = Object->GetName() + FString("_") + MaskOutCacheName;
-					Object->MaskOutCache = NewObject<UMutableMaskOutCache>(Package, *ObjectName, RF_Public | RF_Standalone);
-
-					if (!ParamNamesToSelectedOptions.Num()) // Don't marked the object as modified because of a partial compilation
-					{
-						Object->MarkPackageDirty();
-					}
-				}
-			}
-
-			check(Object->MaskOutCache.Get());
-
-			if (!Object->MaskOutCache->Materials.OrderIndependentCompareEqual(GenerationContext.MaskOutMaterialCache) ||
-				!Object->MaskOutCache->Textures.OrderIndependentCompareEqual(GenerationContext.MaskOutTextureCache))
-			{
-				Object->MaskOutCache->Materials = GenerationContext.MaskOutMaterialCache;
-				Object->MaskOutCache->Textures = GenerationContext.MaskOutTextureCache;
-
-				if (!ParamNamesToSelectedOptions.Num()) // Don't marked the object as modified because of a partial compilation
-				{
-					Object->MarkPackageDirty();
-					Object->MaskOutCache->Modify();
-				}
-			}
-		}
-		else
-		{
-			if (!Object->MaskOutCache.ToString().IsEmpty() || Object->MaskOutCache.Get())
-			{
-				if (!ParamNamesToSelectedOptions.Num()) // Don't mark the object as modified because of a partial compilation
-				{
-					Object->MaskOutCache = nullptr;
-					Object->MarkPackageDirty();
-				}
-			}
-		}
-
 		if (CompileTask.IsValid()) // Don't start compilation if there's a compilation running
 		{
 			// TODO : warning?
diff --git a/Engine/Plugins/Experimental/Mutable/Source/CustomizableObjectEditor/Private/MuCOE/GenerateMutableSource/GenerateMutableSource.h b/Engine/Plugins/Experimental/Mutable/Source/CustomizableObjectEditor/Private/MuCOE/GenerateMutableSource/GenerateMutableSource.h
--- a/Engine/Plugins/Experimental/Mutable/Source/CustomizableObjectEditor/Private/MuCOE/GenerateMutableSource/GenerateMutableSource.h
+++ b/Engine/Plugins/Experimental/Mutable/Source/CustomizableObjectEditor/Private/MuCOE/GenerateMutableSource/GenerateMutableSource.h
@@ -736,10 +736,6 @@ struct FMutableGraphGenerationContext
 	*/
 	TArray<int32> SocketPriorityStack;
 
-	// Stores the textures that will be used to mask-out areas in the projection. The cache isn't used for rendering, but for coverage testing
-	TMap<FString, FString> MaskOutMaterialCache; // Maps a UMaterial's asset path to a UTexture's asset path
-	TMap<FString, FMaskOutTexture> MaskOutTextureCache; // Maps a UTexture's asset path to the cached mask-out texture data
-
 	// Stores the only option of an Int Param that should be compiled in a partial compilation
 	TMap<FString, FString> ParamNamesToSelectedOptions;
 
diff --git a/Engine/Plugins/Experimental/Mutable/Source/CustomizableObjectEditor/Private/MuCOE/GenerateMutableSource/GenerateMutableSourceGroupProjector.cpp b/Engine/Plugins/Experimental/Mutable/Source/CustomizableObjectEditor/Private/MuCOE/GenerateMutableSource/GenerateMutableSourceGroupProjector.cpp
--- a/Engine/Plugins/Experimental/Mutable/Source/CustomizableObjectEditor/Private/MuCOE/GenerateMutableSource/GenerateMutableSourceGroupProjector.cpp
+++ b/Engine/Plugins/Experimental/Mutable/Source/CustomizableObjectEditor/Private/MuCOE/GenerateMutableSource/GenerateMutableSourceGroupProjector.cpp
@@ -214,64 +214,6 @@ mu::NodeImagePtr GenerateMutableGroupProjection(const int32 NodeLOD, const int32
 
 			ImageNodes.Add(ImageNode);
 			ImageNodes_ProjectorTempData.Add(ProjectorTempData);
-
-			// Generate the projection mask-out texture cache
-			const FString& MaskName = ProjectorTempData.CustomizableObjectNodeGroupProjectorParameter->MaskedOutAreaMaterialChannelName;
-
-			if (!MaskName.IsEmpty())
-			{
-				const UMaterialInterface* MaskMaterial = TypedNodeMat ? TypedNodeMat->Material : ParentMaterial->Material;
-				if (MaskMaterial)
-				{
-					UTexture* MaskTexture = nullptr;
-
-					if (MaskMaterial->GetTextureParameterValue(FName(*MaskName), MaskTexture))
-					{
-						UTexture2D* MaskTexture2D = Cast<UTexture2D>(MaskTexture);
-
-						if (MaskTexture2D && MaskTexture2D->Source.GetNumMips() > 0)
-						{
-							TArray64<uint8> TempData;
-							verify( MaskTexture2D->Source.GetMipData(TempData, 0, 0, 0) );
-							uint32 TotalTexels = MaskTexture2D->Source.GetSizeX() * MaskTexture2D->Source.GetSizeY();
-
-							ETextureSourceFormat PixelFormat = MaskTexture2D->Source.GetFormat();
-							int32 BytesPerPixel = MaskTexture2D->Source.GetBytesPerPixel();
-							ensure((PixelFormat == TSF_BGRA8 && BytesPerPixel == 4) || BytesPerPixel == 1);
-
-							{
-								uint32 DataSize = TotalTexels * MaskTexture2D->Source.GetBytesPerPixel();
-								FString MaskTexture2DPath = MaskTexture2D->GetPathName();
-
-								if (!GenerationContext.MaskOutTextureCache.Find(MaskTexture2DPath) && DataSize > 1024) // Discard default or too small textures to have any mask detail
-								{
-									FMaskOutTexture& CachedTexture = GenerationContext.MaskOutTextureCache.Add(MaskTexture2DPath);
-									
-									CachedTexture.SetTextureSize(MaskTexture2D->Source.GetSizeX(), MaskTexture2D->Source.GetSizeY());
-
-									for (size_t p = 0; p < DataSize / BytesPerPixel; ++p)
-									{
-										if (BytesPerPixel == 4)
-										{
-											CachedTexture.GetTexelReference(p) = TempData[p * 4 + 3] > 0; // Copy alpha channel of PF_R8G8B8A8
-										}
-										else if (BytesPerPixel == 1)
-										{
-											CachedTexture.GetTexelReference(p) = TempData[p] > 0;
-										}
-										else
-										{
-											check(false);
-										}
-									}
-								}
-
-								GenerationContext.MaskOutMaterialCache.Add(MaskMaterial->GetPathName(), MaskTexture2DPath);
-							}
-						}
-					}
-				}
-			}
 		}
 	}
 
diff --git a/Engine/Plugins/Experimental/Mutable/Source/CustomizableObjectEditor/Private/MuCOE/GenerateMutableSource/GenerateMutableSourceMesh.cpp b/Engine/Plugins/Experimental/Mutable/Source/CustomizableObjectEditor/Private/MuCOE/GenerateMutableSource/GenerateMutableSourceMesh.cpp
--- a/Engine/Plugins/Experimental/Mutable/Source/CustomizableObjectEditor/Private/MuCOE/GenerateMutableSource/GenerateMutableSourceMesh.cpp
+++ b/Engine/Plugins/Experimental/Mutable/Source/CustomizableObjectEditor/Private/MuCOE/GenerateMutableSource/GenerateMutableSourceMesh.cpp
@@ -13,6 +13,7 @@
 #include "Modules/ModuleManager.h"
 #include "MuCO/CustomizableObjectInstance.h"
 #include "MuCO/CustomizableObjectSystem.h"
+#include "MuCO/UnrealConversionUtils.h"
 #include "MuCOE/CustomizableObjectCompiler.h"
 #include "MuCOE/CustomizableObjectLayout.h"
 #include "MuCOE/GenerateMutableSource/GenerateMutableSourceFloat.h"
diff --git a/Engine/Plugins/Experimental/Mutable/Source/CustomizableObjectEditor/Private/MuCOE/GenerateMutableSource/GenerateMutableSourceSurface.cpp b/Engine/Plugins/Experimental/Mutable/Source/CustomizableObjectEditor/Private/MuCOE/GenerateMutableSource/GenerateMutableSourceSurface.cpp
--- a/Engine/Plugins/Experimental/Mutable/Source/CustomizableObjectEditor/Private/MuCOE/GenerateMutableSource/GenerateMutableSourceSurface.cpp
+++ b/Engine/Plugins/Experimental/Mutable/Source/CustomizableObjectEditor/Private/MuCOE/GenerateMutableSource/GenerateMutableSourceSurface.cpp
@@ -16,6 +16,7 @@
 
 #include "MuCO/CustomizableObjectInstance.h"
 #include "MuCO/MutableMeshBufferUtils.h"
+#include "MuCO/UnrealConversionUtils.h"
 #include "MuCOE/CustomizableObjectCompiler.h"
 #include "MuCOE/CustomizableObjectLayout.h"
 #include "MuCOE/GenerateMutableSource/GenerateMutableSourceColor.h"
diff --git a/Engine/Plugins/Experimental/Mutable/Source/CustomizableObjectPopulation/Classes/MuCOP/CustomizableObjectPopulation.h b/Engine/Plugins/Experimental/Mutable/Source/CustomizableObjectPopulation/Classes/MuCOP/CustomizableObjectPopulation.h
--- a/Engine/Plugins/Experimental/Mutable/Source/CustomizableObjectPopulation/Classes/MuCOP/CustomizableObjectPopulation.h
+++ b/Engine/Plugins/Experimental/Mutable/Source/CustomizableObjectPopulation/Classes/MuCOP/CustomizableObjectPopulation.h
@@ -48,19 +48,23 @@ public:
 
 	UCustomizableObjectPopulation();
 
-	/** Name of the customizable object population */
+	/** Name of the customizable object population. */
 	UPROPERTY(Category = "CustomizablePopulation", EditAnywhere, BlueprintReadOnly)
 	FString Name;
 
-	/** List of Pupulation classes that define this population */
+	/** List of Pupulation classes that define this population. */
 	UPROPERTY(Category = "CustomizablePopulation", EditAnywhere)
 	TArray<FClassWeightPair> ClassWeights;
-
-	// Returns the seed that was used to create instances in this call of the GeneratePopulation. Returns -1 if no generator was found.
+	
+	/** Create a number of instances of the target population.
+	 * @return the seed that was used to create instances in this call of the GeneratePopulation. Returns -1 if there in an error with the population asset.
+	 */
 	UFUNCTION(BlueprintPure, Category = "CustomizablePopulation")
 	int32 GeneratePopulation(TArray<UCustomizableObjectInstance*>& OutInstances, int32 NumInstancesToGenerate = 1) const;
 
-	// Returns false if no generator was found
+	/** Update the instances in the array with new parameter values for the target population.
+	* @return false if there was an error with the population asset.
+	*/
 	UFUNCTION(BlueprintPure, Category = "CustomizablePopulation")
 	bool RegeneratePopulation(int32 Seed, TArray<UCustomizableObjectInstance*>& OutInstances, int32 NumInstancesToGenerate = 1) const;
 	
