[Mutable] Fix ReloadParameters does not verifying that an int parameter option still exists.

#rnx

diff --git a/Engine/Plugins/Experimental/Mutable/Source/CustomizableObject/Private/MuCO/CustomizableObjectInstanceDescriptor.cpp b/Engine/Plugins/Experimental/Mutable/Source/CustomizableObject/Private/MuCO/CustomizableObjectInstanceDescriptor.cpp
--- a/Engine/Plugins/Experimental/Mutable/Source/CustomizableObject/Private/MuCO/CustomizableObjectInstanceDescriptor.cpp
+++ b/Engine/Plugins/Experimental/Mutable/Source/CustomizableObject/Private/MuCO/CustomizableObjectInstanceDescriptor.cpp
@@ -122,7 +122,7 @@ void FCustomizableObjectInstanceDescriptor::SaveDescriptor(FArchive& Ar, bool bU
 
 				if (bIsParamMultidimensional)
 				{
-					for (int i = 0; i < P.ParameterRangeValueNames.Num(); ++i)
+					for (int32 i = 0; i < P.ParameterRangeValueNames.Num(); ++i)
 					{
 						ValueNames.Add(P.ParameterRangeValueNames[i]);
 						Values.Add(CustomizableObject->FindIntParameterValue(ModelParameterIndex, P.ParameterRangeValueNames[i]));
@@ -332,7 +332,7 @@ void FCustomizableObjectInstanceDescriptor::LoadDescriptor(FArchive& Ar)
 						//check((P.ParameterValueName.IsEmpty() && ValueName.Equals(FString("None"))) || P.ParameterValueName.Equals(ValueName));
 						P.ParameterRangeValueNames.SetNum(Values.Num());
 
-						for (int ParamIndex = 0; ParamIndex < Values.Num(); ++ParamIndex)
+						for (int32 ParamIndex = 0; ParamIndex < Values.Num(); ++ParamIndex)
 						{
 							P.ParameterRangeValueNames[ParamIndex] = CustomizableObject->FindIntParameterValueName(ModelParameterIndex, Values[ParamIndex]);
 						}
@@ -491,7 +491,7 @@ mu::ParametersPtr FCustomizableObjectInstanceDescriptor::GetParameters() const
 				{
 					if (mu::RangeIndexPtr RangeIdxPtr = MutableParameters->NewRangeIndex(ParamIndex))
 					{
-						for (int RangeIndex = 0; RangeIndex < IntParameter.ParameterRangeValueNames.Num(); ++RangeIndex)
+						for (int32 RangeIndex = 0; RangeIndex < IntParameter.ParameterRangeValueNames.Num(); ++RangeIndex)
 						{
 							RangeIdxPtr->SetPosition(0, RangeIndex);
 
@@ -521,7 +521,7 @@ mu::ParametersPtr FCustomizableObjectInstanceDescriptor::GetParameters() const
 				{
 					if (mu::RangeIndexPtr RangeIdxPtr = MutableParameters->NewRangeIndex(ParamIndex))
 					{
-						for (int RangeIndex = 0; RangeIndex < FloatParameter.ParameterRangeValues.Num(); ++RangeIndex)
+						for (int32 RangeIndex = 0; RangeIndex < FloatParameter.ParameterRangeValues.Num(); ++RangeIndex)
 						{
 							RangeIdxPtr->SetPosition(0, RangeIndex);
 							MutableParameters->SetFloatValue(ParamIndex, FloatParameter.ParameterRangeValues[RangeIndex], RangeIdxPtr);
@@ -599,7 +599,7 @@ mu::ParametersPtr FCustomizableObjectInstanceDescriptor::GetParameters() const
 
 					if (const mu::RangeIndexPtr RangeIdxPtr = MutableParameters->NewRangeIndex(ParamIndex))
 					{
-						for (int RangeIndex = 0; RangeIndex < ProjectorParameter.RangeValues.Num(); ++RangeIndex)
+						for (int32 RangeIndex = 0; RangeIndex < ProjectorParameter.RangeValues.Num(); ++RangeIndex)
 						{
 							RangeIdxPtr->SetPosition(0, RangeIndex);
 							CopyProjector(ProjectorParameter.RangeValues[RangeIndex], RangeIdxPtr);
@@ -623,7 +623,7 @@ mu::ParametersPtr FCustomizableObjectInstanceDescriptor::GetParameters() const
 				{
 					if (mu::RangeIndexPtr RangeIdxPtr = MutableParameters->NewRangeIndex(ParamIndex))
 					{
-						for (int RangeIndex = 0; RangeIndex < TextureParameter.ParameterRangeValues.Num(); ++RangeIndex)
+						for (int32 RangeIndex = 0; RangeIndex < TextureParameter.ParameterRangeValues.Num(); ++RangeIndex)
 						{
 							RangeIdxPtr->SetPosition(0, RangeIndex);
 							MutableParameters->SetImageValue(ParamIndex, TextureParameter.ParameterRangeValues[RangeIndex], RangeIdxPtr);
@@ -818,14 +818,63 @@ void FCustomizableObjectInstanceDescriptor::ReloadParameters()
 			};
 			
 			if (FCustomizableObjectIntParameterValue* Result = OldIntParameters.FindByPredicate(FindByNameAndUid))
-			{	
-				if (mu::RangeIndexPtr RangeIdxPtr = MutableParameters->NewRangeIndex(ParamIndex))
+			{
+				const int32 NumValueIndex = MutableParameters->GetIntPossibleValueCount(ParamIndex);
+
+				auto ValueExists = [&](const FString& ValueName)
+				{
+					for (int32 ValueIndex = 0; ValueIndex < NumValueIndex; ++ValueIndex)
+					{
+						if (ValueName == MutableParameters->GetIntPossibleValueName(ParamIndex, ValueIndex))
+						{
+							return true;
+						}
+					}
+
+					return false;
+				};
+				
+				if (mu::RangeIndexPtr RangeIdxPtr = MutableParameters->NewRangeIndex(ParamIndex)) // Is multidimensional
 				{
-					Param.ParameterRangeValueNames = Result->ParameterRangeValueNames;
+					// Find the max RangeIndex of this parameter 
+					int32 RangeNum = 0;
+					const int32 ValueCount = MutableParameters->GetValueCount(ParamIndex);
+					for (int32 ValueIndex = 0; ValueIndex < ValueCount; ++ValueIndex)
+					{
+						mu::RangeIndexPtr RangeValueIdxPtr = MutableParameters->GetValueIndex(ParamIndex, ValueIndex);
+						RangeNum = FMath::Max(RangeNum, RangeValueIdxPtr->GetPosition(0));
+					}
+					
+					Param.ParameterRangeValueNames.Reserve(RangeNum); 
+
+					for (int32 RangeIndex = 0; RangeIndex < RangeNum; ++RangeIndex)
+					{
+						if (const FString& OldValue = Result->ParameterRangeValueNames[RangeIndex];
+							ValueExists(OldValue)) // Value still exits
+						{
+							Param.ParameterRangeValueNames.Add(OldValue);							
+						}
+						else
+						{
+							RangeIdxPtr->SetPosition(0, RangeIndex);
+							
+							const int32 Value = MutableParameters->GetIntValue(ParamIndex, RangeIdxPtr);
+							const FString AuxParameterValueName = CustomizableObject->FindIntParameterValueName(ParamIndex, Value);
+							Param.ParameterRangeValueNames.Add(AuxParameterValueName);
+						}
+					}
 				}
 				else
 				{
-					Param.ParameterValueName = Result->ParameterValueName;
+					if (ValueExists(Result->ParameterValueName))
+					{
+						Param.ParameterValueName = Result->ParameterValueName;
+					}
+					else
+					{
+						const int32 ParamValue = MutableParameters->GetIntValue(ParamIndex);
+						Param.ParameterValueName = CustomizableObject->FindIntParameterValueName(ParamIndex, ParamValue);
+					}
 				}
 			} 
 			else // Not found in Instance Parameters. Use Mutable Parameters.
@@ -834,7 +883,7 @@ void FCustomizableObjectInstanceDescriptor::ReloadParameters()
 				{
 					const int32 ValueCount = MutableParameters->GetValueCount(ParamIndex);
 
-					for (int ValueIndex = 0; ValueIndex < ValueCount; ++ValueIndex)
+					for (int32 ValueIndex = 0; ValueIndex < ValueCount; ++ValueIndex)
 					{	
 						const mu::RangeIndexPtr RangeValueIdxPtr = MutableParameters->GetValueIndex(ParamIndex, ValueIndex);
 						const int32 RangeIndex = RangeValueIdxPtr->GetPosition(0);
@@ -889,7 +938,7 @@ void FCustomizableObjectInstanceDescriptor::ReloadParameters()
 				{
 					const int32 ValueCount = MutableParameters->GetValueCount(ParamIndex);
 
-					for (int ValueIndex = 0; ValueIndex < ValueCount; ++ValueIndex)
+					for (int32 ValueIndex = 0; ValueIndex < ValueCount; ++ValueIndex)
 					{	
 						mu::RangeIndexPtr RangeValueIdxPtr = MutableParameters->GetValueIndex(ParamIndex, ValueIndex);
 						int32 RangeIndex = RangeValueIdxPtr->GetPosition(0);
@@ -989,7 +1038,7 @@ void FCustomizableObjectInstanceDescriptor::ReloadParameters()
 				{
 					const int32 ValueCount = MutableParameters->GetValueCount(ParamIndex);
 
-					for (int ValueIndex = 0; ValueIndex < ValueCount; ++ValueIndex)
+					for (int32 ValueIndex = 0; ValueIndex < ValueCount; ++ValueIndex)
 					{	
 						mu::RangeIndexPtr RangeValueIdxPtr = MutableParameters->GetValueIndex(ParamIndex, ValueIndex);
 						int32 RangeIndex = RangeValueIdxPtr->GetPosition(0);
@@ -1040,7 +1089,7 @@ void FCustomizableObjectInstanceDescriptor::ReloadParameters()
 				{
 					const int32 ValueCount = MutableParameters->GetValueCount(ParamIndex);
 
-					for (int ValueIndex = 0; ValueIndex < ValueCount; ++ValueIndex)
+					for (int32 ValueIndex = 0; ValueIndex < ValueCount; ++ValueIndex)
 					{	
 						mu::RangeIndexPtr RangeValueIdxPtr = MutableParameters->GetValueIndex(ParamIndex, ValueIndex);
 						int32 RangeIndex = RangeValueIdxPtr->GetPosition(0);
