[mutable] Updated COE module loading time to be less restrictive.
- Moved the loading time to Default from PostEngineInit. A valid use case was pointed out by an external team.
- Updated COE Module config to register the COI factory once the editor has been fully loaded.

diff --git a/Engine/Plugins/Experimental/Mutable/Mutable.uplugin b/Engine/Plugins/Experimental/Mutable/Mutable.uplugin
--- a/Engine/Plugins/Experimental/Mutable/Mutable.uplugin
+++ b/Engine/Plugins/Experimental/Mutable/Mutable.uplugin
@@ -31,7 +31,7 @@
 		{
 			"Name": "CustomizableObjectEditor",
 			"Type": "Editor",
-			"LoadingPhase": "PostEngineInit"
+			"LoadingPhase": "Default"
 		},
 		{
 			"Name": "CustomizableObjectPopulation",
diff --git a/Engine/Plugins/Experimental/Mutable/Source/CustomizableObjectEditor/Private/MuCOE/CustomizableObjectEditorModule.cpp b/Engine/Plugins/Experimental/Mutable/Source/CustomizableObjectEditor/Private/MuCOE/CustomizableObjectEditorModule.cpp
--- a/Engine/Plugins/Experimental/Mutable/Source/CustomizableObjectEditor/Private/MuCOE/CustomizableObjectEditorModule.cpp
+++ b/Engine/Plugins/Experimental/Mutable/Source/CustomizableObjectEditor/Private/MuCOE/CustomizableObjectEditorModule.cpp
@@ -127,6 +127,9 @@ private:
 	// Command to look for Customizable Object Instance in the player pawn of the current world and open its Customizable Object Instance Editor
 	IConsoleCommand* LaunchCOIECommand;
 	static void OpenCOIE(const TArray<FString>& Arguments);
+
+	/** Register the COI factory */
+	void RegisterFactory();
 };
 
 IMPLEMENT_MODULE( FCustomizableObjectEditorModule, CustomizableObjectEditor );
@@ -191,15 +194,7 @@ void FCustomizableObjectEditorModule::StartupModule()
 	PropertyModule.NotifyCustomizationModuleChanged();
 
 	// Register factory
-	if (GEditor)
-	{
-		GEditor->ActorFactories.Add(NewObject<UCustomizableObjectInstanceFactory>());
-		if (UPlacementSubsystem* PlacementSubsystem = GEditor->GetEditorSubsystem<UPlacementSubsystem>())
-		{
-			PlacementSubsystem->RegisterAssetFactory(NewObject<UCustomizableObjectInstanceFactory>());
-		}
-	}
-	
+	FCoreDelegates::OnPostEngineInit.AddRaw(this,&FCustomizableObjectEditorModule::RegisterFactory);
 
 	// Additional UI style
 	FCustomizableObjectEditorStyle::Initialize();
@@ -237,6 +232,8 @@ void FCustomizableObjectEditorModule::ShutdownModule()
 	CustomizableObjectEditor_ToolBarExtensibilityManager.Reset();
 	CustomizableObjectEditor_MenuExtensibilityManager.Reset();
 
+	FCoreDelegates::OnPostEngineInit.RemoveAll(this);
+
 	FCustomizableObjectEditorStyle::Shutdown();
 }
 
@@ -381,6 +378,19 @@ void FCustomizableObjectEditorModule::OpenCOIE(const TArray<FString>& Arguments)
 }
 
 
+void FCustomizableObjectEditorModule::RegisterFactory()
+{
+	if (GEditor)
+	{
+		GEditor->ActorFactories.Add(NewObject<UCustomizableObjectInstanceFactory>());
+		if (UPlacementSubsystem* PlacementSubsystem = GEditor->GetEditorSubsystem<UPlacementSubsystem>())
+		{
+			PlacementSubsystem->RegisterAssetFactory(NewObject<UCustomizableObjectInstanceFactory>());
+		}
+	}
+}
+
+
 EAssetTypeCategories::Type FCustomizableObjectEditorModule::GetAssetCategory() const
 {
 	return CustomizableObjectAssetCategory;
