[mutable] Remove unnecessary boilerplate code that was only used in one debug tool.
#rnx

diff --git a/Engine/Plugins/Experimental/Mutable/Source/CustomizableObjectEditor/Private/MuCOE/SMutableGraphViewer.cpp b/Engine/Plugins/Experimental/Mutable/Source/CustomizableObjectEditor/Private/MuCOE/SMutableGraphViewer.cpp
--- a/Engine/Plugins/Experimental/Mutable/Source/CustomizableObjectEditor/Private/MuCOE/SMutableGraphViewer.cpp
+++ b/Engine/Plugins/Experimental/Mutable/Source/CustomizableObjectEditor/Private/MuCOE/SMutableGraphViewer.cpp
@@ -12,12 +12,22 @@
 #include "MuCOE/CustomizableObjectEditorStyle.h"
 #include "MuCOE/SMutableCodeViewer.h"
 #include "MuCOE/UnrealEditorPortabilityHelpers.h"
-#include "MuT/Streams.h"
 #include "Widgets/Docking/SDockTab.h"
 #include "Widgets/Input/STextComboBox.h"
 #include "Widgets/Input/SNumericDropDown.h"
 #include "Widgets/Views/STreeView.h"
 #include "ScopedTransaction.h"
+#include "MuT/Streams.h"
+
+// This is necessary because of problems with rtti information in other platforms. In any case, this part of the debugger is only useful in the standard editor.
+#if PLATFORM_WINDOWS
+#include "MuT/NodeObjectNewPrivate.h"
+#include "MuT/NodeObjectGroupPrivate.h"
+#include "MuT/NodeSurfaceNewPrivate.h"
+#include "MuT/NodeSurfaceEditPrivate.h"
+#include "MuT/NodeSurfaceSwitchPrivate.h"
+#include "MuT/NodeSurfaceVariationPrivate.h"
+#endif
 
 class FExtender;
 class FReferenceCollector;
@@ -47,7 +57,11 @@ public:
 
 		const char* TypeName = RowItem->MutableNode->GetType()->m_strName;
 
-		FText MainLabel = FText::FromString(StringCast<TCHAR>(TypeName).Get());
+		FString LabelString = RowItem->Prefix.IsEmpty() 
+			? StringCast<TCHAR>(TypeName).Get() 
+			: FString::Printf( TEXT("%s : %s"), *RowItem->Prefix, StringCast<TCHAR>(TypeName).Get() );
+
+		FText MainLabel = FText::FromString(LabelString);
 		if (RowItem->DuplicatedOf)
 		{
 			MainLabel = FText::FromString( FString::Printf(TEXT("%s (Duplicated)"), StringCast<TCHAR>(TypeName).Get()));
@@ -389,6 +403,8 @@ TSharedRef<ITableRow> SMutableGraphViewer::GenerateRowForNodeTree(TSharedPtr<FMu
 
 void SMutableGraphViewer::GetChildrenForInfo(TSharedPtr<FMutableGraphTreeElement> InInfo, TArray<TSharedPtr<FMutableGraphTreeElement>>& OutChildren)
 {
+// This is necessary because of problems with rtti information in other platforms. In any case, this part of the debugger is only useful in the standard editor.
+#if PLATFORM_WINDOWS
 	if (!InInfo->MutableNode)
 	{
 		return;
@@ -400,14 +416,14 @@ void SMutableGraphViewer::GetChildrenForInfo(TSharedPtr<FMutableGraphTreeElement
 		return;
 	}
 
-	// Generic node case
-	int32 InputCount = InInfo->MutableNode->GetInputCount();
-	for (int32 InputIndex=0; InputIndex<InputCount; ++InputIndex)
+	mu::Node* ParentNode = InInfo->MutableNode.get();
+	uint32 InputIndex = 0;
+
+	auto AddChildFunc = [this, ParentNode, &InputIndex, &OutChildren](mu::Node* ChildNode, const FString& Prefix)
 	{
-		mu::NodePtr ChildNode = InInfo->MutableNode->GetInputNode( InputIndex );
 		if (ChildNode)
 		{
-			FItemCacheKey Key = { InInfo->MutableNode.get(),ChildNode.get(), uint32(InputIndex) };
+			FItemCacheKey Key = { ParentNode, ChildNode, InputIndex };
 			TSharedPtr<FMutableGraphTreeElement>* CachedItem = ItemCache.Find(Key);
 
 			if (CachedItem)
@@ -416,18 +432,130 @@ void SMutableGraphViewer::GetChildrenForInfo(TSharedPtr<FMutableGraphTreeElement
 			}
 			else
 			{
-				TSharedPtr<FMutableGraphTreeElement>* MainItemPtr = MainItemPerNode.Find(ChildNode.get());
-				TSharedPtr<FMutableGraphTreeElement> Item = MakeShareable(new FMutableGraphTreeElement(ChildNode, MainItemPtr));
+				TSharedPtr<FMutableGraphTreeElement>* MainItemPtr = MainItemPerNode.Find(ChildNode);
+				TSharedPtr<FMutableGraphTreeElement> Item = MakeShareable(new FMutableGraphTreeElement(ChildNode, MainItemPtr, Prefix));
 				OutChildren.Add(Item);
 				ItemCache.Add(Key, Item);
 
 				if (!MainItemPtr)
 				{
-					MainItemPerNode.Add(ChildNode.get(),Item);
+					MainItemPerNode.Add(ChildNode, Item);
 				}
 			}
 		}
+		++InputIndex;
+	};
+
+	if (ParentNode->GetType() == mu::NodeObjectNew::GetStaticType())
+	{
+		mu::NodeObjectNew* ObjectNew = reinterpret_cast<mu::NodeObjectNew*>(ParentNode);
+		mu::NodeObjectNew::Private* Private = ObjectNew->GetPrivate();
+		for (int32 l = 0; l < Private->m_lods.Num(); ++l)
+		{
+			AddChildFunc(Private->m_lods[l].get(), TEXT("LOD") );
+		}
+
+		for (int32 l = 0; l < Private->m_children.Num(); ++l)
+		{
+			AddChildFunc(Private->m_children[l].get(), TEXT("CHILD"));
+		}
+	}
+
+	else if (ParentNode->GetType() == mu::NodeObjectGroup::GetStaticType())
+	{
+		mu::NodeObjectGroup* ObjectGroup = reinterpret_cast<mu::NodeObjectGroup*>(ParentNode);
+		mu::NodeObjectGroup::Private* Private = ObjectGroup->GetPrivate();
+		for (int32 l = 0; l < Private->m_children.Num(); ++l)
+		{
+			AddChildFunc(Private->m_children[l].get(), TEXT("CHILD"));
+		}
+	}
+
+	else if (ParentNode->GetType() == mu::NodeSurfaceNew::GetStaticType())
+	{
+		mu::NodeSurfaceNew* SurfaceNew = reinterpret_cast<mu::NodeSurfaceNew*>(ParentNode);
+		mu::NodeSurfaceNew::Private* Private = SurfaceNew->GetPrivate();
+		for (int32 l = 0; l < Private->m_meshes.Num(); ++l)
+		{
+			AddChildFunc(Private->m_meshes[l].m_pMesh.get(), TEXT("MESH"));
+		}
+
+		for (int32 l = 0; l < Private->m_images.Num(); ++l)
+		{
+			AddChildFunc(Private->m_images[l].m_pImage.get(), FString::Printf(TEXT("IMAGE [%s]"), *Private->m_images[l].m_name));
+		}
+
+		for (int32 l = 0; l < Private->m_vectors.Num(); ++l)
+		{
+			AddChildFunc(Private->m_vectors[l].m_pVector.get(), FString::Printf(TEXT("VECTOR [%s]"), *Private->m_vectors[l].m_name));
+		}
+
+		for (int32 l = 0; l < Private->m_scalars.Num(); ++l)
+		{
+			AddChildFunc(Private->m_scalars[l].m_pScalar.get(), FString::Printf(TEXT("SCALAR [%s]"), *Private->m_scalars[l].m_name));
+		}
+	}
+
+	else if (ParentNode->GetType() == mu::NodeSurfaceEdit::GetStaticType())
+	{
+		mu::NodeSurfaceEdit* SurfaceEdit = reinterpret_cast<mu::NodeSurfaceEdit*>(ParentNode);
+		mu::NodeSurfaceEdit::Private* Private = SurfaceEdit->GetPrivate();
+		AddChildFunc(Private->m_pMesh.get(), TEXT("MESH"));
+		AddChildFunc(Private->m_pMorph.get(), TEXT("MORPH"));
+		AddChildFunc(Private->m_pFactor.get(), TEXT("MORPH_FACTOR"));
+
+		for (int32 l = 0; l < Private->m_textures.Num(); ++l)
+		{
+			AddChildFunc(Private->m_textures[l].m_pExtend.get(), FString::Printf(TEXT("EXTEND [%d]"), l));
+			AddChildFunc(Private->m_textures[l].m_pPatch.get(), FString::Printf(TEXT("PATCH [%d]"), l));
+		}
+	}
+
+	else if (ParentNode->GetType() == mu::NodeSurfaceSwitch::GetStaticType())
+	{
+		mu::NodeSurfaceSwitch* SurfaceSwitch = reinterpret_cast<mu::NodeSurfaceSwitch*>(ParentNode);
+		mu::NodeSurfaceSwitch::Private* Private = SurfaceSwitch->GetPrivate();
+		AddChildFunc(Private->Parameter.get(), TEXT("PARAM"));
+
+		for (int32 l = 0; l < Private->Options.Num(); ++l)
+		{
+			AddChildFunc(Private->Options[l].get(), FString::Printf(TEXT("OPTION [%d]"), l));
+		}
+	}
+
+	else if (ParentNode->GetType() == mu::NodeSurfaceVariation::GetStaticType())
+	{
+		mu::NodeSurfaceVariation* SurfaceVar = reinterpret_cast<mu::NodeSurfaceVariation*>(ParentNode);
+		mu::NodeSurfaceVariation::Private* Private = SurfaceVar->GetPrivate();
+		for (int32 l = 0; l < Private->m_defaultSurfaces.Num(); ++l)
+		{
+			AddChildFunc(Private->m_defaultSurfaces[l].get(), FString::Printf(TEXT("DEF SURF [%d]"), l));
+		}
+		for (int32 l = 0; l < Private->m_defaultModifiers.Num(); ++l)
+		{
+			AddChildFunc(Private->m_defaultModifiers[l].get(), FString::Printf(TEXT("DEF MOD [%d]"), l));
+		}
+
+		for (int32 v = 0; v < Private->m_variations.Num(); ++v)
+		{
+			const mu::NodeSurfaceVariation::Private::FVariation Var = Private->m_variations[v];
+			for (int32 l = 0; l < Var.m_surfaces.Num(); ++l)
+			{
+				AddChildFunc(Var.m_surfaces[l].get(), FString::Printf(TEXT("VAR [%s] SURF [%d]"), *Var.m_tag, l));
+			}
+			for (int32 l = 0; l < Var.m_modifiers.Num(); ++l)
+			{
+				AddChildFunc(Var.m_modifiers[l].get(), FString::Printf(TEXT("VAR [%s] MOD [%d]"), *Var.m_tag, l));
+			}
+		}
+	}
+
+	else
+	{
+		// This node type has not been implemented, so its children won't be added to the tree.
+		ensure(false);
 	}
+#endif
 }
 
 
diff --git a/Engine/Plugins/Experimental/Mutable/Source/CustomizableObjectEditor/Private/MuCOE/SMutableGraphViewer.h b/Engine/Plugins/Experimental/Mutable/Source/CustomizableObjectEditor/Private/MuCOE/SMutableGraphViewer.h
--- a/Engine/Plugins/Experimental/Mutable/Source/CustomizableObjectEditor/Private/MuCOE/SMutableGraphViewer.h
+++ b/Engine/Plugins/Experimental/Mutable/Source/CustomizableObjectEditor/Private/MuCOE/SMutableGraphViewer.h
@@ -127,9 +127,10 @@ private:
 class FMutableGraphTreeElement : public TSharedFromThis<FMutableGraphTreeElement>
 {
 public:
-	FMutableGraphTreeElement(const mu::NodePtr& InNode, TSharedPtr<FMutableGraphTreeElement>* InDuplicatedOf=nullptr )
+	FMutableGraphTreeElement(const mu::NodePtr& InNode, TSharedPtr<FMutableGraphTreeElement>* InDuplicatedOf=nullptr, const FString& InPrefix=FString() )
 	{
 		MutableNode = InNode;
+		Prefix = InPrefix;
 		if (InDuplicatedOf)
 		{
 			DuplicatedOf = *InDuplicatedOf;
@@ -144,4 +145,6 @@ public:
 	/** If this tree element is a duplicated of another node, this is the node. */
 	TSharedPtr<FMutableGraphTreeElement> DuplicatedOf;
 
+	/** Optional label prefix */
+	FString Prefix;
 };
diff --git a/Engine/Plugins/Experimental/Mutable/Source/MutableTools/Private/MuT/CodeGenerator.cpp b/Engine/Plugins/Experimental/Mutable/Source/MutableTools/Private/MuT/CodeGenerator.cpp
--- a/Engine/Plugins/Experimental/Mutable/Source/MutableTools/Private/MuT/CodeGenerator.cpp
+++ b/Engine/Plugins/Experimental/Mutable/Source/MutableTools/Private/MuT/CodeGenerator.cpp
@@ -2130,7 +2130,7 @@ namespace mu
         m_pErrorLog->GetPrivate()->Add( Msg, ELMT_ERROR, errorContext );
 
         // Create a constant node
-        NodeBoolConstantPtr pNode = new NodeBoolConstant();
+        Ptr<NodeBoolConstant> pNode = new NodeBoolConstant();
         pNode->SetValue( value );
 
         Ptr<ASTOp> result = Generate( pNode );
diff --git a/Engine/Plugins/Experimental/Mutable/Source/MutableTools/Private/MuT/CodeGenerator.h b/Engine/Plugins/Experimental/Mutable/Source/MutableTools/Private/MuT/CodeGenerator.h
--- a/Engine/Plugins/Experimental/Mutable/Source/MutableTools/Private/MuT/CodeGenerator.h
+++ b/Engine/Plugins/Experimental/Mutable/Source/MutableTools/Private/MuT/CodeGenerator.h
@@ -652,7 +652,7 @@ namespace mu
 		typedef TMap<FVisitedKeyMap, FBoolGenerationResult> GeneratedBoolsMap;
 		GeneratedBoolsMap m_generatedBools;
 
-		void GenerateBool(FBoolGenerationResult&, const NodeBoolPtrConst&);
+		void GenerateBool(FBoolGenerationResult&, const Ptr<const NodeBool>&);
 		void GenerateBool_Constant(FBoolGenerationResult&, const Ptr<const NodeBoolConstant>&);
 		void GenerateBool_Parameter(FBoolGenerationResult&, const Ptr<const NodeBoolParameter>&);
 		void GenerateBool_IsNull(FBoolGenerationResult&, const Ptr<const NodeBoolIsNull>&);
diff --git a/Engine/Plugins/Experimental/Mutable/Source/MutableTools/Private/MuT/CodeGenerator_Bool.cpp b/Engine/Plugins/Experimental/Mutable/Source/MutableTools/Private/MuT/CodeGenerator_Bool.cpp
--- a/Engine/Plugins/Experimental/Mutable/Source/MutableTools/Private/MuT/CodeGenerator_Bool.cpp
+++ b/Engine/Plugins/Experimental/Mutable/Source/MutableTools/Private/MuT/CodeGenerator_Bool.cpp
@@ -26,7 +26,7 @@ namespace mu
 
 
 	//-------------------------------------------------------------------------------------------------
-	void CodeGenerator::GenerateBool(FBoolGenerationResult& result, const NodeBoolPtrConst& untyped)
+	void CodeGenerator::GenerateBool(FBoolGenerationResult& result, const Ptr<const NodeBool>& untyped)
 	{
 		if (!untyped)
 		{
diff --git a/Engine/Plugins/Experimental/Mutable/Source/MutableTools/Private/MuT/NodeBool.cpp b/Engine/Plugins/Experimental/Mutable/Source/MutableTools/Private/MuT/NodeBool.cpp
--- a/Engine/Plugins/Experimental/Mutable/Source/MutableTools/Private/MuT/NodeBool.cpp
+++ b/Engine/Plugins/Experimental/Mutable/Source/MutableTools/Private/MuT/NodeBool.cpp
@@ -31,7 +31,7 @@ namespace mu
 
         
 	//---------------------------------------------------------------------------------------------
-	NodeBoolPtr NodeBool::StaticUnserialise( InputArchive& arch )
+	Ptr<NodeBool> NodeBool::StaticUnserialise( InputArchive& arch )
 	{
         uint32_t ver;
 		arch >> ver;
diff --git a/Engine/Plugins/Experimental/Mutable/Source/MutableTools/Private/MuT/NodeBoolAnd.cpp b/Engine/Plugins/Experimental/Mutable/Source/MutableTools/Private/MuT/NodeBoolAnd.cpp
--- a/Engine/Plugins/Experimental/Mutable/Source/MutableTools/Private/MuT/NodeBoolAnd.cpp
+++ b/Engine/Plugins/Experimental/Mutable/Source/MutableTools/Private/MuT/NodeBoolAnd.cpp
@@ -25,70 +25,31 @@ namespace mu
 	MUTABLE_IMPLEMENT_NODE( NodeBoolAnd, EType::And, Node, Node::EType::Bool);
 
 
-	//---------------------------------------------------------------------------------------------
-	// Node Interface
-	//---------------------------------------------------------------------------------------------
-	int NodeBoolAnd::GetInputCount() const
-	{
-		return 2;
-	}
-
-
-	//---------------------------------------------------------------------------------------------
-	Node* NodeBoolAnd::GetInputNode( int i ) const
-	{
-		check( i>=0 && i<GetInputCount() );
-		if ( i == 0)
-		{
-			return m_pD->m_pA.get();
-		}
-		else
-		{
-			return m_pD->m_pB.get();
-		}
-	}
-
-
-	//---------------------------------------------------------------------------------------------
-	void NodeBoolAnd::SetInputNode( int i, NodePtr p )
-	{
-		check( i>=0 && i<GetInputCount() );
-		if ( i == 0)
-		{
-			m_pD->m_pA = dynamic_cast<NodeBool*>( p.get() );
-		}
-		else
-		{
-			m_pD->m_pB = dynamic_cast<NodeBool*>( p.get() );
-		}
-	}
-
-
 	//---------------------------------------------------------------------------------------------
 	// Own Interface
 	//---------------------------------------------------------------------------------------------
-	NodeBoolPtr NodeBoolAnd::GetA() const
+	Ptr<NodeBool> NodeBoolAnd::GetA() const
 	{
 		return m_pD->m_pA;
 	}
 
 
 	//---------------------------------------------------------------------------------------------
-	void NodeBoolAnd::SetA( NodeBoolPtr p )
+	void NodeBoolAnd::SetA(Ptr<NodeBool> p )
 	{
 		m_pD->m_pA = p;
 	}
 
 
 	//---------------------------------------------------------------------------------------------
-	NodeBoolPtr NodeBoolAnd::GetB() const
+	Ptr<NodeBool> NodeBoolAnd::GetB() const
 	{
 		return m_pD->m_pB;
 	}
 
 
 	//---------------------------------------------------------------------------------------------
-	void NodeBoolAnd::SetB( NodeBoolPtr p )
+	void NodeBoolAnd::SetB(Ptr<NodeBool> p )
 	{
 		m_pD->m_pB = p;
 	}
diff --git a/Engine/Plugins/Experimental/Mutable/Source/MutableTools/Private/MuT/NodeBoolConstant.cpp b/Engine/Plugins/Experimental/Mutable/Source/MutableTools/Private/MuT/NodeBoolConstant.cpp
--- a/Engine/Plugins/Experimental/Mutable/Source/MutableTools/Private/MuT/NodeBoolConstant.cpp
+++ b/Engine/Plugins/Experimental/Mutable/Source/MutableTools/Private/MuT/NodeBoolConstant.cpp
@@ -26,30 +26,6 @@ namespace mu
 	MUTABLE_IMPLEMENT_NODE( NodeBoolConstant, EType::Constant, Node, Node::EType::Bool);
 
 
-	//---------------------------------------------------------------------------------------------
-	// Node Interface
-	//---------------------------------------------------------------------------------------------
-	int NodeBoolConstant::GetInputCount() const
-	{
-		return 0;
-	}
-
-
-	//---------------------------------------------------------------------------------------------
-    Node* NodeBoolConstant::GetInputNode( int ) const
-	{
-		check( false );
-		return 0;
-	}
-
-
-	//---------------------------------------------------------------------------------------------
-    void NodeBoolConstant::SetInputNode( int, NodePtr )
-	{
-		check( false );
-	}
-
-
 	//---------------------------------------------------------------------------------------------
 	// Own Interface
 	//---------------------------------------------------------------------------------------------
diff --git a/Engine/Plugins/Experimental/Mutable/Source/MutableTools/Private/MuT/NodeBoolIsNull.cpp b/Engine/Plugins/Experimental/Mutable/Source/MutableTools/Private/MuT/NodeBoolIsNull.cpp
--- a/Engine/Plugins/Experimental/Mutable/Source/MutableTools/Private/MuT/NodeBoolIsNull.cpp
+++ b/Engine/Plugins/Experimental/Mutable/Source/MutableTools/Private/MuT/NodeBoolIsNull.cpp
@@ -23,29 +23,6 @@ namespace mu
 
 	MUTABLE_IMPLEMENT_NODE( NodeBoolIsNull, EType::IsNull, Node, Node::EType::Bool);
 
-
-	//---------------------------------------------------------------------------------------------
-	// Node Interface
-	//---------------------------------------------------------------------------------------------
-	int NodeBoolIsNull::GetInputCount() const
-	{
-		return 1;
-	}
-
-
-	//---------------------------------------------------------------------------------------------
-    Node* NodeBoolIsNull::GetInputNode( int ) const
-	{
-		return m_pD->m_pSource.get();
-	}
-
-
-	//---------------------------------------------------------------------------------------------
-    void NodeBoolIsNull::SetInputNode( int, NodePtr p )
-	{
-		m_pD->m_pSource = p;
-	}
-
 }
 
 
diff --git a/Engine/Plugins/Experimental/Mutable/Source/MutableTools/Private/MuT/NodeBoolNot.cpp b/Engine/Plugins/Experimental/Mutable/Source/MutableTools/Private/MuT/NodeBoolNot.cpp
--- a/Engine/Plugins/Experimental/Mutable/Source/MutableTools/Private/MuT/NodeBoolNot.cpp
+++ b/Engine/Plugins/Experimental/Mutable/Source/MutableTools/Private/MuT/NodeBoolNot.cpp
@@ -25,44 +25,17 @@ namespace mu
 	MUTABLE_IMPLEMENT_NODE( NodeBoolNot, EType::Not, Node, Node::EType::Bool);
 
 
-	//---------------------------------------------------------------------------------------------
-	// Node Interface
-	//---------------------------------------------------------------------------------------------
-	int NodeBoolNot::GetInputCount() const
-	{
-		return 1;
-	}
-
-
-	//---------------------------------------------------------------------------------------------
-	Node* NodeBoolNot::GetInputNode( int i ) const
-	{
-		check( i>=0 && i<GetInputCount() );
-        (void)i;
-        return m_pD->m_pSource.get();
-	}
-
-
-	//---------------------------------------------------------------------------------------------
-	void NodeBoolNot::SetInputNode( int i, NodePtr p)
-	{
-		check( i>=0 && i<GetInputCount() );
-        (void)i;
-        m_pD->m_pSource = dynamic_cast<NodeBool*>( p.get() );
-	}
-
-
 	//---------------------------------------------------------------------------------------------
 	// Own Interface
 	//---------------------------------------------------------------------------------------------
-	NodeBoolPtr NodeBoolNot::GetInput() const
+	Ptr<NodeBool> NodeBoolNot::GetInput() const
 	{
 		return m_pD->m_pSource;
 	}
 
 
 	//---------------------------------------------------------------------------------------------
-	void NodeBoolNot::SetInput( NodeBoolPtr p )
+	void NodeBoolNot::SetInput(Ptr<NodeBool> p )
 	{
 		m_pD->m_pSource = p;
 	}
diff --git a/Engine/Plugins/Experimental/Mutable/Source/MutableTools/Private/MuT/NodeBoolParameter.cpp b/Engine/Plugins/Experimental/Mutable/Source/MutableTools/Private/MuT/NodeBoolParameter.cpp
--- a/Engine/Plugins/Experimental/Mutable/Source/MutableTools/Private/MuT/NodeBoolParameter.cpp
+++ b/Engine/Plugins/Experimental/Mutable/Source/MutableTools/Private/MuT/NodeBoolParameter.cpp
@@ -31,38 +31,6 @@ namespace mu
 	MUTABLE_IMPLEMENT_NODE( NodeBoolParameter, EType::Parameter, Node, Node::EType::Bool);
 
 
-	//---------------------------------------------------------------------------------------------
-	// Node Interface
-	//---------------------------------------------------------------------------------------------
-	int NodeBoolParameter::GetInputCount() const
-	{
-        return int( m_pD->m_ranges.Num() );
-    }
-
-
-	//---------------------------------------------------------------------------------------------
-    Node* NodeBoolParameter::GetInputNode( int i ) const
-	{
-        check( i<GetInputCount() );
-        if (i<GetInputCount())
-        {
-            return m_pD->m_ranges[i].get();
-        }
-        return nullptr;
-	}
-
-
-	//---------------------------------------------------------------------------------------------
-    void NodeBoolParameter::SetInputNode( int i , NodePtr n )
-	{
-        check( i<GetInputCount() );
-        if (i<GetInputCount())
-        {
-            m_pD->m_ranges[i] = dynamic_cast<NodeRange*>(n.get());
-        }
-    }
-
-
 	//---------------------------------------------------------------------------------------------
 	// Own Interface
 	//---------------------------------------------------------------------------------------------
diff --git a/Engine/Plugins/Experimental/Mutable/Source/MutableTools/Private/MuT/NodeBoolPrivate.h b/Engine/Plugins/Experimental/Mutable/Source/MutableTools/Private/MuT/NodeBoolPrivate.h
--- a/Engine/Plugins/Experimental/Mutable/Source/MutableTools/Private/MuT/NodeBoolPrivate.h
+++ b/Engine/Plugins/Experimental/Mutable/Source/MutableTools/Private/MuT/NodeBoolPrivate.h
@@ -34,7 +34,7 @@ namespace mu
 		//!
 		void Serialise( OutputArchive& arch ) const
 		{
-            uint32_t ver = 0;
+            uint32 ver = 0;
 			arch << ver;
 
 			arch << m_value;
@@ -43,7 +43,7 @@ namespace mu
 		//!
 		void Unserialise( InputArchive& arch )
 		{
-            uint32_t ver;
+            uint32 ver;
 			arch >> ver;
 			check(ver==0);
 
@@ -71,7 +71,7 @@ namespace mu
 		//!
 		void Serialise( OutputArchive& arch ) const
 		{
-            uint32_t ver = 3;
+            uint32 ver = 3;
 			arch << ver;
 
 			arch << m_defaultValue;
@@ -83,7 +83,7 @@ namespace mu
 		//!
 		void Unserialise( InputArchive& arch )
 		{
-            uint32_t ver;
+            uint32 ver;
 			arch >> ver;
             check(ver>=2 && ver<=3);
 
@@ -122,7 +122,7 @@ namespace mu
 		//!
 		void Serialise( OutputArchive& arch ) const
 		{
-            uint32_t ver = 0;
+            uint32 ver = 0;
 			arch << ver;
 
 			arch << m_pSource;
@@ -131,7 +131,7 @@ namespace mu
 		//!
 		void Unserialise( InputArchive& arch )
 		{
-            uint32_t ver;
+            uint32 ver;
 			arch >> ver;
             check(ver<=0);
 
@@ -150,12 +150,12 @@ namespace mu
 
 		static NODE_TYPE s_type;
 
-		NodeBoolPtr m_pSource;
+		Ptr<NodeBool> m_pSource;
 
 		//!
 		void Serialise( OutputArchive& arch ) const
 		{
-            uint32_t ver = 0;
+            uint32 ver = 0;
 			arch << ver;
 
 			arch << m_pSource;
@@ -164,7 +164,7 @@ namespace mu
 		//!
 		void Unserialise( InputArchive& arch )
 		{
-            uint32_t ver;
+            uint32 ver;
 			arch >> ver;
             check(ver<=0);
 
@@ -183,13 +183,13 @@ namespace mu
 
 		static NODE_TYPE s_type;
 
-		NodeBoolPtr m_pA;
-		NodeBoolPtr m_pB;
+		Ptr<NodeBool> m_pA;
+		Ptr<NodeBool> m_pB;
 
 		//!
 		void Serialise( OutputArchive& arch ) const
 		{
-            uint32_t ver = 0;
+            uint32 ver = 0;
 			arch << ver;
 
 			arch << m_pA;
@@ -199,7 +199,7 @@ namespace mu
 		//!
 		void Unserialise( InputArchive& arch )
 		{
-            uint32_t ver;
+            uint32 ver;
 			arch >> ver;
             check(ver<=0);
 
diff --git a/Engine/Plugins/Experimental/Mutable/Source/MutableTools/Private/MuT/NodeColourArithmeticOperation.cpp b/Engine/Plugins/Experimental/Mutable/Source/MutableTools/Private/MuT/NodeColourArithmeticOperation.cpp
--- a/Engine/Plugins/Experimental/Mutable/Source/MutableTools/Private/MuT/NodeColourArithmeticOperation.cpp
+++ b/Engine/Plugins/Experimental/Mutable/Source/MutableTools/Private/MuT/NodeColourArithmeticOperation.cpp
@@ -35,45 +35,6 @@ namespace mu
     };
 
 
-	//---------------------------------------------------------------------------------------------
-	// Node Interface
-	//---------------------------------------------------------------------------------------------
-	int NodeColourArithmeticOperation::GetInputCount() const
-	{
-		return 2;
-	}
-
-
-	//---------------------------------------------------------------------------------------------
-	Node* NodeColourArithmeticOperation::GetInputNode( int i ) const
-	{
-		check( i>=0 && i< GetInputCount());
-
-		Node* pResult = 0;
-
-		switch (i)
-		{
-		case 0: pResult = m_pD->m_pA.get(); break;
-		case 1: pResult = m_pD->m_pB.get(); break;
-		}
-
-		return pResult;
-	}
-
-
-	//---------------------------------------------------------------------------------------------
-	void NodeColourArithmeticOperation::SetInputNode( int i, NodePtr pNode )
-	{
-		check( i>=0 && i< GetInputCount());
-
-		switch (i)
-		{
-		case 0: m_pD->m_pA = dynamic_cast<NodeColour*>(pNode.get()); break;
-		case 1: m_pD->m_pB = dynamic_cast<NodeColour*>(pNode.get()); break;
-		}
-	}
-
-
 	//---------------------------------------------------------------------------------------------
 	// Own Interface
 	//---------------------------------------------------------------------------------------------
diff --git a/Engine/Plugins/Experimental/Mutable/Source/MutableTools/Private/MuT/NodeColourConstant.cpp b/Engine/Plugins/Experimental/Mutable/Source/MutableTools/Private/MuT/NodeColourConstant.cpp
--- a/Engine/Plugins/Experimental/Mutable/Source/MutableTools/Private/MuT/NodeColourConstant.cpp
+++ b/Engine/Plugins/Experimental/Mutable/Source/MutableTools/Private/MuT/NodeColourConstant.cpp
@@ -37,30 +37,6 @@ namespace mu
 	}
 
 
-	//---------------------------------------------------------------------------------------------
-	// Node Interface
-	//---------------------------------------------------------------------------------------------
-	int NodeColourConstant::GetInputCount() const
-	{
-		return 0;
-	}
-
-
-	//---------------------------------------------------------------------------------------------
-    Node* NodeColourConstant::GetInputNode( int ) const
-	{
-		check( false );
-		return 0;
-	}
-
-
-	//---------------------------------------------------------------------------------------------
-    void NodeColourConstant::SetInputNode( int, NodePtr )
-	{
-		check( false );
-	}
-
-
 	//---------------------------------------------------------------------------------------------
 	// Own Interface
 	//---------------------------------------------------------------------------------------------
diff --git a/Engine/Plugins/Experimental/Mutable/Source/MutableTools/Private/MuT/NodeColourFromScalars.cpp b/Engine/Plugins/Experimental/Mutable/Source/MutableTools/Private/MuT/NodeColourFromScalars.cpp
--- a/Engine/Plugins/Experimental/Mutable/Source/MutableTools/Private/MuT/NodeColourFromScalars.cpp
+++ b/Engine/Plugins/Experimental/Mutable/Source/MutableTools/Private/MuT/NodeColourFromScalars.cpp
@@ -26,49 +26,6 @@ namespace mu
 	MUTABLE_IMPLEMENT_NODE( NodeColourFromScalars, EType::FromScalars, Node, Node::EType::Colour)
 
 
-	//---------------------------------------------------------------------------------------------
-	// Node Interface
-	//---------------------------------------------------------------------------------------------
-	int NodeColourFromScalars::GetInputCount() const
-	{
-		return 4;
-	}
-
-
-	//---------------------------------------------------------------------------------------------
-	Node* NodeColourFromScalars::GetInputNode( int i ) const
-	{
-		check( i>=0 && i< GetInputCount());
-
-		Node* pResult = 0;
-
-		switch (i)
-		{
-		case 0: pResult = m_pD->m_pX.get(); break;
-		case 1: pResult = m_pD->m_pY.get(); break;
-		case 2: pResult = m_pD->m_pZ.get(); break;
-		case 3: pResult = m_pD->m_pW.get(); break;
-		}
-
-		return pResult;
-	}
-
-
-	//---------------------------------------------------------------------------------------------
-	void NodeColourFromScalars::SetInputNode( int i, NodePtr pNode )
-	{
-		check( i>=0 && i< GetInputCount());
-
-		switch (i)
-		{
-		case 0: m_pD->m_pX = dynamic_cast<NodeScalar*>(pNode.get()); break;
-		case 1: m_pD->m_pY = dynamic_cast<NodeScalar*>(pNode.get()); break;
-		case 2: m_pD->m_pZ = dynamic_cast<NodeScalar*>(pNode.get()); break;
-		case 3: m_pD->m_pW = dynamic_cast<NodeScalar*>(pNode.get()); break;
-		}
-	}
-
-
 	//---------------------------------------------------------------------------------------------
 	// Own Interface
 	//---------------------------------------------------------------------------------------------
diff --git a/Engine/Plugins/Experimental/Mutable/Source/MutableTools/Private/MuT/NodeColourParameter.cpp b/Engine/Plugins/Experimental/Mutable/Source/MutableTools/Private/MuT/NodeColourParameter.cpp
--- a/Engine/Plugins/Experimental/Mutable/Source/MutableTools/Private/MuT/NodeColourParameter.cpp
+++ b/Engine/Plugins/Experimental/Mutable/Source/MutableTools/Private/MuT/NodeColourParameter.cpp
@@ -27,38 +27,6 @@ namespace mu
 	MUTABLE_IMPLEMENT_NODE( NodeColourParameter, EType::Parameter, Node, Node::EType::Colour)
 
 
-	//---------------------------------------------------------------------------------------------
-	// Node Interface
-	//---------------------------------------------------------------------------------------------
-	int NodeColourParameter::GetInputCount() const
-	{
-        return  m_pD->m_ranges.Num();
-    }
-
-
-	//---------------------------------------------------------------------------------------------
-    Node* NodeColourParameter::GetInputNode( int i ) const
-	{
-        check( i<GetInputCount() );
-        if (i<GetInputCount())
-        {
-            return m_pD->m_ranges[i].get();
-        }
-        return nullptr;
-    }
-
-
-	//---------------------------------------------------------------------------------------------
-    void NodeColourParameter::SetInputNode( int i, NodePtr n )
-	{
-        check( i<GetInputCount() );
-        if (i<GetInputCount())
-        {
-            m_pD->m_ranges[i] = dynamic_cast<NodeRange*>(n.get());
-        }
-    }
-
-
 	//---------------------------------------------------------------------------------------------
 	// Own Interface
 	//---------------------------------------------------------------------------------------------
diff --git a/Engine/Plugins/Experimental/Mutable/Source/MutableTools/Private/MuT/NodeColourSampleImage.cpp b/Engine/Plugins/Experimental/Mutable/Source/MutableTools/Private/MuT/NodeColourSampleImage.cpp
--- a/Engine/Plugins/Experimental/Mutable/Source/MutableTools/Private/MuT/NodeColourSampleImage.cpp
+++ b/Engine/Plugins/Experimental/Mutable/Source/MutableTools/Private/MuT/NodeColourSampleImage.cpp
@@ -27,47 +27,6 @@ namespace mu
 	MUTABLE_IMPLEMENT_NODE( NodeColourSampleImage, EType::SampleImage, Node, Node::EType::Colour)
 
 
-	//---------------------------------------------------------------------------------------------
-	// Node Interface
-	//---------------------------------------------------------------------------------------------
-	int NodeColourSampleImage::GetInputCount() const
-	{
-		return 3;
-	}
-
-
-	//---------------------------------------------------------------------------------------------
-	Node* NodeColourSampleImage::GetInputNode( int i ) const
-	{
-		check( i>=0 && i< GetInputCount());
-
-		Node* pResult = 0;
-
-		switch (i)
-		{
-		case 0: pResult = m_pD->m_pImage.get(); break;
-		case 1: pResult = m_pD->m_pX.get(); break;
-		case 2: pResult = m_pD->m_pY.get(); break;
-		}
-
-		return pResult;
-	}
-
-
-	//---------------------------------------------------------------------------------------------
-	void NodeColourSampleImage::SetInputNode( int i, NodePtr pNode )
-	{
-		check( i>=0 && i< GetInputCount());
-
-		switch (i)
-		{
-		case 0: m_pD->m_pImage = dynamic_cast<NodeImage*>(pNode.get()); break;
-		case 1: m_pD->m_pX = dynamic_cast<NodeScalar*>(pNode.get()); break;
-		case 2: m_pD->m_pY = dynamic_cast<NodeScalar*>(pNode.get()); break;
-		}
-	}
-
-
 	//---------------------------------------------------------------------------------------------
 	// Own Interface
 	//---------------------------------------------------------------------------------------------
diff --git a/Engine/Plugins/Experimental/Mutable/Source/MutableTools/Private/MuT/NodeColourSwitch.cpp b/Engine/Plugins/Experimental/Mutable/Source/MutableTools/Private/MuT/NodeColourSwitch.cpp
--- a/Engine/Plugins/Experimental/Mutable/Source/MutableTools/Private/MuT/NodeColourSwitch.cpp
+++ b/Engine/Plugins/Experimental/Mutable/Source/MutableTools/Private/MuT/NodeColourSwitch.cpp
@@ -25,55 +25,6 @@ namespace mu
     MUTABLE_IMPLEMENT_NODE( NodeColourSwitch, EType::Switch, Node, Node::EType::Colour)
 
 
-	//---------------------------------------------------------------------------------------------
-	// Node Interface
-	//---------------------------------------------------------------------------------------------
-    int NodeColourSwitch::GetInputCount() const
-	{
-		return 1 + m_pD->m_options.Num();
-	}
-
-
-	//---------------------------------------------------------------------------------------------
-    Node* NodeColourSwitch::GetInputNode( int i ) const
-	{
-		check( i>=0 && i<GetInputCount() );
-
-		Node* pResult = 0;
-
-		switch (i)
-		{
-		case 0:
-			pResult = m_pD->m_pParameter.get();
-			break;
-
-		default:
-			pResult = m_pD->m_options[i-1].get();
-			break;
-		}
-
-		return pResult;
-	}
-
-
-	//---------------------------------------------------------------------------------------------
-    void NodeColourSwitch::SetInputNode( int i, NodePtr pNode )
-	{
-		check( i>=0 && i<GetInputCount() );
-
-		switch (i)
-		{
-		case 0:
-			m_pD->m_pParameter = dynamic_cast<NodeScalar*>(pNode.get());
-			break;
-
-		default:
-            m_pD->m_options[i-1] = dynamic_cast<NodeColour*>(pNode.get());
-			break;
-		}
-	}
-
-
 	//---------------------------------------------------------------------------------------------
 	// Own Interface
 	//---------------------------------------------------------------------------------------------
diff --git a/Engine/Plugins/Experimental/Mutable/Source/MutableTools/Private/MuT/NodeColourTable.cpp b/Engine/Plugins/Experimental/Mutable/Source/MutableTools/Private/MuT/NodeColourTable.cpp
--- a/Engine/Plugins/Experimental/Mutable/Source/MutableTools/Private/MuT/NodeColourTable.cpp
+++ b/Engine/Plugins/Experimental/Mutable/Source/MutableTools/Private/MuT/NodeColourTable.cpp
@@ -28,33 +28,6 @@ namespace mu
 	MUTABLE_IMPLEMENT_NODE( NodeColourTable, EType::Table, Node, Node::EType::Colour)
 
 
-	//---------------------------------------------------------------------------------------------
-	// Node Interface
-	//---------------------------------------------------------------------------------------------
-	int NodeColourTable::GetInputCount() const
-	{
-		return 0;
-	}
-
-
-	//---------------------------------------------------------------------------------------------
-	Node* NodeColourTable::GetInputNode( int i ) const
-	{
-		check( i >=0 && i < 0 );
-        (void)i;
-        return 0;
-	}
-
-
-	//---------------------------------------------------------------------------------------------
-    void NodeColourTable::SetInputNode( int i, NodePtr )
-	{
-		check( i >=0 && i < 0 );
-        (void)i;
-        //m_pD->m_pObject = dynamic_cast<NodeObject*>( pNode.get() );
-	}
-
-
 	//---------------------------------------------------------------------------------------------
 	// Own Interface
 	//---------------------------------------------------------------------------------------------
diff --git a/Engine/Plugins/Experimental/Mutable/Source/MutableTools/Private/MuT/NodeColourVariation.cpp b/Engine/Plugins/Experimental/Mutable/Source/MutableTools/Private/MuT/NodeColourVariation.cpp
--- a/Engine/Plugins/Experimental/Mutable/Source/MutableTools/Private/MuT/NodeColourVariation.cpp
+++ b/Engine/Plugins/Experimental/Mutable/Source/MutableTools/Private/MuT/NodeColourVariation.cpp
@@ -18,60 +18,9 @@ namespace mu
         NODE_TYPE( "ColourVariation", NodeColour::GetStaticType() );
 
 
-    //---------------------------------------------------------------------------------------------
-    //---------------------------------------------------------------------------------------------
-    //---------------------------------------------------------------------------------------------
-
     MUTABLE_IMPLEMENT_NODE( NodeColourVariation, EType::Variation, Node, Node::EType::Colour)
 
 
-    //---------------------------------------------------------------------------------------------
-    int NodeColourVariation::GetInputCount() const { return 1 + m_pD->m_variations.Num(); }
-
-
-    //---------------------------------------------------------------------------------------------
-    Node* NodeColourVariation::GetInputNode( int i ) const
-    {
-        check( i >= 0 && i < GetInputCount() );
-
-        if ( i == 0 )
-        {
-            return m_pD->m_defaultColour.get();
-        }
-        i -= 1;
-
-        if ( i < int( m_pD->m_variations.Num() ) )
-        {
-            return m_pD->m_variations[i].m_colour.get();
-        }
-        i -= int( m_pD->m_variations.Num() );
-
-        return nullptr;
-    }
-
-
-    //---------------------------------------------------------------------------------------------
-    void NodeColourVariation::SetInputNode( int i, NodePtr pNode )
-    {
-        check( i >= 0 && i < GetInputCount() );
-
-        if ( i == 0 )
-        {
-            m_pD->m_defaultColour = dynamic_cast<NodeColour*>( pNode.get() );
-            return;
-        }
-
-        i -= 1;
-        if ( i < int( m_pD->m_variations.Num() ) )
-        {
-
-            m_pD->m_variations[i].m_colour = dynamic_cast<NodeColour*>( pNode.get() );
-            return;
-        }
-        i -= (int)m_pD->m_variations.Num();
-    }
-
-
     //---------------------------------------------------------------------------------------------
     // Own Interface
     //---------------------------------------------------------------------------------------------
diff --git a/Engine/Plugins/Experimental/Mutable/Source/MutableTools/Private/MuT/NodeComponentEdit.cpp b/Engine/Plugins/Experimental/Mutable/Source/MutableTools/Private/MuT/NodeComponentEdit.cpp
--- a/Engine/Plugins/Experimental/Mutable/Source/MutableTools/Private/MuT/NodeComponentEdit.cpp
+++ b/Engine/Plugins/Experimental/Mutable/Source/MutableTools/Private/MuT/NodeComponentEdit.cpp
@@ -28,43 +28,6 @@ namespace mu
     MUTABLE_IMPLEMENT_NODE( NodeComponentEdit, EType::Edit, Node, Node::EType::Component)
 
 
-    //---------------------------------------------------------------------------------------------
-    // Node Interface
-    //---------------------------------------------------------------------------------------------
-    int NodeComponentEdit::GetInputCount() const
-    {
-        return m_pD->m_surfaces.Num();
-    }
-
-
-    //---------------------------------------------------------------------------------------------
-    Node* NodeComponentEdit::GetInputNode( int i ) const
-    {
-        check( i >=0 && i < GetInputCount() );
-
-        NodePtr pResult;
-
-        if ( i<m_pD->m_surfaces.Num() )
-        {
-            pResult = m_pD->m_surfaces[i];
-        }
-
-        return pResult.get();
-    }
-
-
-    //---------------------------------------------------------------------------------------------
-    void NodeComponentEdit::SetInputNode( int i, NodePtr pNode )
-    {
-        check( i >=0 && i < GetInputCount() );
-
-        if ( i<m_pD->m_surfaces.Num() )
-        {
-            m_pD->m_surfaces[ i ] = dynamic_cast<NodeSurface*>(pNode.get());
-        }
-    }
-
-
 	//---------------------------------------------------------------------------------------------
 	// Own Interface
 	//---------------------------------------------------------------------------------------------
diff --git a/Engine/Plugins/Experimental/Mutable/Source/MutableTools/Private/MuT/NodeComponentNew.cpp b/Engine/Plugins/Experimental/Mutable/Source/MutableTools/Private/MuT/NodeComponentNew.cpp
--- a/Engine/Plugins/Experimental/Mutable/Source/MutableTools/Private/MuT/NodeComponentNew.cpp
+++ b/Engine/Plugins/Experimental/Mutable/Source/MutableTools/Private/MuT/NodeComponentNew.cpp
@@ -27,43 +27,6 @@ namespace mu
     MUTABLE_IMPLEMENT_NODE( NodeComponentNew, EType::New, Node, Node::EType::Component)
 
 
-	//---------------------------------------------------------------------------------------------
-	// Node Interface
-	//---------------------------------------------------------------------------------------------
-	int NodeComponentNew::GetInputCount() const
-	{
-        return m_pD->m_surfaces.Num();
-	}
-
-
-	//---------------------------------------------------------------------------------------------
-	Node* NodeComponentNew::GetInputNode( int i ) const
-	{
-        check( i >=0 && i < GetInputCount() );
-
-		NodePtr pResult;
-
-        if ( i<m_pD->m_surfaces.Num() )
-		{
-            pResult = m_pD->m_surfaces[i];
-		}
-
-		return pResult.get();
-	}
-
-
-	//---------------------------------------------------------------------------------------------
-	void NodeComponentNew::SetInputNode( int i, NodePtr pNode )
-	{
-        check( i >=0 && i < GetInputCount() );
-
-        if ( i<m_pD->m_surfaces.Num() )
-		{
-            m_pD->m_surfaces[ i ] = dynamic_cast<NodeSurface*>(pNode.get());
-		}
-	}
-
-
 	//---------------------------------------------------------------------------------------------
 	// Own Interface
 	//---------------------------------------------------------------------------------------------
diff --git a/Engine/Plugins/Experimental/Mutable/Source/MutableTools/Private/MuT/NodeExtensionDataConstant.cpp b/Engine/Plugins/Experimental/Mutable/Source/MutableTools/Private/MuT/NodeExtensionDataConstant.cpp
--- a/Engine/Plugins/Experimental/Mutable/Source/MutableTools/Private/MuT/NodeExtensionDataConstant.cpp
+++ b/Engine/Plugins/Experimental/Mutable/Source/MutableTools/Private/MuT/NodeExtensionDataConstant.cpp
@@ -25,30 +25,6 @@ namespace mu
 	MUTABLE_IMPLEMENT_NODE(NodeExtensionDataConstant, EType::Constant, Node, Node::EType::ExtensionData);
 
 
-	//---------------------------------------------------------------------------------------------
-	// Node Interface
-	//---------------------------------------------------------------------------------------------
-	int NodeExtensionDataConstant::GetInputCount() const
-	{
-		return 0;
-	}
-
-
-	//---------------------------------------------------------------------------------------------
-	Node* NodeExtensionDataConstant::GetInputNode(int i) const
-	{
-		check(false);
-		return nullptr;
-	}
-
-
-	//---------------------------------------------------------------------------------------------
-	void NodeExtensionDataConstant::SetInputNode(int i, NodePtr pNode)
-	{
-		check(false);
-	}
-
-
 	//---------------------------------------------------------------------------------------------
 	// Own Interface
 	//---------------------------------------------------------------------------------------------
@@ -58,7 +34,6 @@ namespace mu
 	}
 
 
-	//---------------------------------------------------------------------------------------------
 	void NodeExtensionDataConstant::SetValue(ExtensionDataPtrConst Value)
 	{
 		m_pD->Value = Value;
diff --git a/Engine/Plugins/Experimental/Mutable/Source/MutableTools/Private/MuT/NodeExtensionDataSwitch.cpp b/Engine/Plugins/Experimental/Mutable/Source/MutableTools/Private/MuT/NodeExtensionDataSwitch.cpp
--- a/Engine/Plugins/Experimental/Mutable/Source/MutableTools/Private/MuT/NodeExtensionDataSwitch.cpp
+++ b/Engine/Plugins/Experimental/Mutable/Source/MutableTools/Private/MuT/NodeExtensionDataSwitch.cpp
@@ -21,81 +21,28 @@ namespace mu
 
 	MUTABLE_IMPLEMENT_NODE(NodeExtensionDataSwitch, EType::Switch, Node, Node::EType::ExtensionData);
 
-	//---------------------------------------------------------------------------------------------
-	// Node Interface
-	//---------------------------------------------------------------------------------------------
-
-	//---------------------------------------------------------------------------------------------
-	int NodeExtensionDataSwitch::GetInputCount() const
-	{
-		return 1 + m_pD->Options.Num();
-	}
-
-	//---------------------------------------------------------------------------------------------
-	Node* NodeExtensionDataSwitch::GetInputNode(int i) const
-	{
-		check(i >=0 && i < GetInputCount());
-
-
-		Node* Result = 0;
-		switch (i)
-		{
-			case 0:
-				Result = m_pD->Parameter.get();
-				break;
-
-			default:
-				Result = m_pD->Options[i - 1].get();
-				break;
-		}
-
-		return Result;
-	}
-
-	//---------------------------------------------------------------------------------------------
-	void NodeExtensionDataSwitch::SetInputNode(int i, NodePtr pNode)
-	{
-		check(i >= 0 && i < GetInputCount());
-
-		switch (i)
-		{
-			case 0:
-				m_pD->Parameter = dynamic_cast<NodeScalar*>(pNode.get());
-				break;
-
-			default:
-				m_pD->Options[i - 1] = dynamic_cast<NodeExtensionData*>(pNode.get());
-				break;
-		}
-	}
-
-	//---------------------------------------------------------------------------------------------
 	NodeScalarPtr NodeExtensionDataSwitch::GetParameter() const
 	{
 		return m_pD->Parameter.get();
 	}
 
-	//---------------------------------------------------------------------------------------------
 	void NodeExtensionDataSwitch::SetParameter(NodeScalarPtr pNode)
 	{
 		m_pD->Parameter = pNode;
 	}
 
-	//---------------------------------------------------------------------------------------------
 	void NodeExtensionDataSwitch::SetOptionCount(int InNumOptions)
 	{
 		check(InNumOptions >= 0);
 		m_pD->Options.SetNum(InNumOptions);
 	}
 
-	//---------------------------------------------------------------------------------------------
 	NodeExtensionDataPtr NodeExtensionDataSwitch::GetOption(int t) const
 	{
 		check(m_pD->Options.IsValidIndex(t));
 		return m_pD->Options[t].get();
 	}
 
-	//---------------------------------------------------------------------------------------------
 	void NodeExtensionDataSwitch::SetOption(int t, NodeExtensionDataPtr pNode)
 	{
 		check(m_pD->Options.IsValidIndex(t));
diff --git a/Engine/Plugins/Experimental/Mutable/Source/MutableTools/Private/MuT/NodeExtensionDataVariation.cpp b/Engine/Plugins/Experimental/Mutable/Source/MutableTools/Private/MuT/NodeExtensionDataVariation.cpp
--- a/Engine/Plugins/Experimental/Mutable/Source/MutableTools/Private/MuT/NodeExtensionDataVariation.cpp
+++ b/Engine/Plugins/Experimental/Mutable/Source/MutableTools/Private/MuT/NodeExtensionDataVariation.cpp
@@ -18,87 +18,32 @@ namespace mu
 	//---------------------------------------------------------------------------------------------
 	MUTABLE_IMPLEMENT_NODE(NodeExtensionDataVariation, EType::Variation, Node, Node::EType::ExtensionData);
 
-	//---------------------------------------------------------------------------------------------
-	// Node Interface
-	//---------------------------------------------------------------------------------------------
-
-	//---------------------------------------------------------------------------------------------
-	int NodeExtensionDataVariation::GetInputCount() const
-	{
-		return 1 + m_pD->Variations.Num();
-	}
-
-	//---------------------------------------------------------------------------------------------
-	Node* NodeExtensionDataVariation::GetInputNode(int i) const
-	{
-		check(i >= 0 && i < GetInputCount());
-
-		if (i == 0)
-		{
-			return m_pD->DefaultValue.get();
-		}
-
-		i -= 1;
-
-		if (m_pD->Variations.IsValidIndex(i))
-		{
-			return m_pD->Variations[i].Value.get();
-		}
-
-		return nullptr;
-	}
-
-	//---------------------------------------------------------------------------------------------
-	void NodeExtensionDataVariation::SetInputNode(int i, NodePtr pNode)
-	{
-		check(i >= 0 && i < GetInputCount());
-
-		if (i == 0)
-		{
-			m_pD->DefaultValue = dynamic_cast<NodeExtensionData*>(pNode.get());
-			return;
-		}
-
-		i -= 1;
-
-		if (m_pD->Variations.IsValidIndex(i))
-		{
-			m_pD->Variations[i].Value = dynamic_cast<NodeExtensionData*>(pNode.get());
-			return;
-		}
-	}
 
 	//---------------------------------------------------------------------------------------------
 	// Own Interface
-	//---------------------------------------------------------------------------------------------
-
 	//---------------------------------------------------------------------------------------------
 	void NodeExtensionDataVariation::SetDefaultValue(NodeExtensionDataPtr InValue)
 	{
 		m_pD->DefaultValue = InValue;
 	}
 
-	//---------------------------------------------------------------------------------------------
 	void NodeExtensionDataVariation::SetVariationCount(int InCount)
 	{
 		check(InCount >= 0);
 		m_pD->Variations.SetNum(InCount);
 	}
 
-	//---------------------------------------------------------------------------------------------
 	int NodeExtensionDataVariation::GetVariationCount() const
 	{
 		return m_pD->Variations.Num();
 	}
 
-	//---------------------------------------------------------------------------------------------
 	void NodeExtensionDataVariation::SetVariationTag(int InIndex, const FString& Tag)
 	{
 		check(m_pD->Variations.IsValidIndex(InIndex));
 
 		m_pD->Variations[InIndex].Tag = Tag;
 	}
-
 	//---------------------------------------------------------------------------------------------
 	void NodeExtensionDataVariation::SetVariationValue(int InIndex, NodeExtensionDataPtr InValue)
 	{
diff --git a/Engine/Plugins/Experimental/Mutable/Source/MutableTools/Private/MuT/NodeImageBinarise.cpp b/Engine/Plugins/Experimental/Mutable/Source/MutableTools/Private/MuT/NodeImageBinarise.cpp
--- a/Engine/Plugins/Experimental/Mutable/Source/MutableTools/Private/MuT/NodeImageBinarise.cpp
+++ b/Engine/Plugins/Experimental/Mutable/Source/MutableTools/Private/MuT/NodeImageBinarise.cpp
@@ -23,45 +23,6 @@ namespace mu
 	MUTABLE_IMPLEMENT_NODE( NodeImageBinarise, EType::Binarise, Node, Node::EType::Image)
 
 
-	//---------------------------------------------------------------------------------------------
-	// Node Interface
-	//---------------------------------------------------------------------------------------------
-	void NodeImageBinarise::SetInputNode( int i, NodePtr pNode )
-	{
-		check( i>=0 && i< GetInputCount());
-
-		switch (i)
-		{
-        case 0: m_pD->m_pBase = dynamic_cast<NodeImage*>(pNode.get()); break;
-        case 1: m_pD->m_pThreshold = dynamic_cast<NodeScalar*>(pNode.get()); break;
-		}
-	}
-
-
-	//---------------------------------------------------------------------------------------------
-	int NodeImageBinarise::GetInputCount() const
-	{
-		return 2;
-	}
-
-
-	//---------------------------------------------------------------------------------------------
-	Node* NodeImageBinarise::GetInputNode( int i ) const
-	{
-		check( i>=0 && i< GetInputCount());
-
-		Node* pResult = 0;
-
-		switch (i)
-		{
-        case 0: pResult = m_pD->m_pBase.get(); break;
-        case 1: pResult = m_pD->m_pThreshold.get(); break;
-		}
-
-		return pResult;
-	}
-
-
 	//---------------------------------------------------------------------------------------------
 	// Own Interface
 	//---------------------------------------------------------------------------------------------
diff --git a/Engine/Plugins/Experimental/Mutable/Source/MutableTools/Private/MuT/NodeImageColourMap.cpp b/Engine/Plugins/Experimental/Mutable/Source/MutableTools/Private/MuT/NodeImageColourMap.cpp
--- a/Engine/Plugins/Experimental/Mutable/Source/MutableTools/Private/MuT/NodeImageColourMap.cpp
+++ b/Engine/Plugins/Experimental/Mutable/Source/MutableTools/Private/MuT/NodeImageColourMap.cpp
@@ -25,47 +25,6 @@ namespace mu
 	MUTABLE_IMPLEMENT_NODE( NodeImageColourMap, EType::ColourMap, Node, Node::EType::Image)
 
 
-	//---------------------------------------------------------------------------------------------
-	// Node Interface
-	//---------------------------------------------------------------------------------------------
-	int NodeImageColourMap::GetInputCount() const
-	{
-		return 3;
-	}
-
-
-	//---------------------------------------------------------------------------------------------
-	Node* NodeImageColourMap::GetInputNode( int i ) const
-	{
-		check( i>=0 && i< GetInputCount());
-
-		Node* pResult = 0;
-
-		switch (i)
-		{
-        case 0: pResult = m_pD->m_pBase.get(); break;
-        case 1: pResult = m_pD->m_pMask.get(); break;
-        case 2: pResult = m_pD->m_pMap.get(); break;
-		}
-
-		return pResult;
-	}
-
-
-	//---------------------------------------------------------------------------------------------
-	void NodeImageColourMap::SetInputNode( int i, NodePtr pNode )
-	{
-		check( i>=0 && i< GetInputCount());
-
-		switch (i)
-		{
-        case 0: m_pD->m_pBase = dynamic_cast<NodeImage*>( pNode.get() ); break;
-        case 1: m_pD->m_pMask = dynamic_cast<NodeImage*>( pNode.get() ); break;
-        case 2: m_pD->m_pMap = dynamic_cast<NodeImage*>( pNode.get() ); break;
-		}
-	}
-
-
 	//---------------------------------------------------------------------------------------------
 	// Own Interface
 	//---------------------------------------------------------------------------------------------
diff --git a/Engine/Plugins/Experimental/Mutable/Source/MutableTools/Private/MuT/NodeImageConditional.cpp b/Engine/Plugins/Experimental/Mutable/Source/MutableTools/Private/MuT/NodeImageConditional.cpp
--- a/Engine/Plugins/Experimental/Mutable/Source/MutableTools/Private/MuT/NodeImageConditional.cpp
+++ b/Engine/Plugins/Experimental/Mutable/Source/MutableTools/Private/MuT/NodeImageConditional.cpp
@@ -26,69 +26,6 @@ namespace mu
     MUTABLE_IMPLEMENT_NODE( NodeImageConditional, EType::Conditional, Node, Node::EType::Image)
 
 
-	//---------------------------------------------------------------------------------------------
-	// Node Interface
-	//---------------------------------------------------------------------------------------------
-    int NodeImageConditional::GetInputCount() const
-	{
-        return 3;
-	}
-
-
-	//---------------------------------------------------------------------------------------------
-    Node* NodeImageConditional::GetInputNode( int i ) const
-	{
-		check( i>=0 && i<GetInputCount() );
-
-		Node* pResult = 0;
-
-        switch (i)
-		{
-        case 0:
-            pResult = m_pD->m_parameter.get();
-            break;
-
-        case 1:
-            pResult = m_pD->m_true.get();
-            break;
-
-        case 2:
-            pResult = m_pD->m_false.get();
-            break;
-
-        default:
-			break;
-		}
-
-		return pResult;
-	}
-
-
-	//---------------------------------------------------------------------------------------------
-    void NodeImageConditional::SetInputNode( int i, NodePtr pNode )
-	{
-		check( i>=0 && i<GetInputCount() );
-
-        switch (i)
-		{
-        case 0:
-            m_pD->m_parameter = dynamic_cast<NodeBool*>(pNode.get());
-            break;
-
-        case 1:
-            m_pD->m_true = dynamic_cast<NodeImage*>(pNode.get());
-            break;
-
-        case 2:
-            m_pD->m_false = dynamic_cast<NodeImage*>(pNode.get());
-            break;
-
-        default:
-			break;
-		}
-	}
-
-
 	//---------------------------------------------------------------------------------------------
 	// Own Interface
 	//---------------------------------------------------------------------------------------------
diff --git a/Engine/Plugins/Experimental/Mutable/Source/MutableTools/Private/MuT/NodeImageConditionalPrivate.h b/Engine/Plugins/Experimental/Mutable/Source/MutableTools/Private/MuT/NodeImageConditionalPrivate.h
--- a/Engine/Plugins/Experimental/Mutable/Source/MutableTools/Private/MuT/NodeImageConditionalPrivate.h
+++ b/Engine/Plugins/Experimental/Mutable/Source/MutableTools/Private/MuT/NodeImageConditionalPrivate.h
@@ -23,7 +23,7 @@ namespace mu
 
 		static NODE_TYPE s_type;
 
-        NodeBoolPtr m_parameter;
+        Ptr<NodeBool> m_parameter;
         NodeImagePtr m_true;
         NodeImagePtr m_false;
 
diff --git a/Engine/Plugins/Experimental/Mutable/Source/MutableTools/Private/MuT/NodeImageConstant.cpp b/Engine/Plugins/Experimental/Mutable/Source/MutableTools/Private/MuT/NodeImageConstant.cpp
--- a/Engine/Plugins/Experimental/Mutable/Source/MutableTools/Private/MuT/NodeImageConstant.cpp
+++ b/Engine/Plugins/Experimental/Mutable/Source/MutableTools/Private/MuT/NodeImageConstant.cpp
@@ -24,28 +24,6 @@ namespace mu
 	MUTABLE_IMPLEMENT_NODE(NodeImageConstant, EType::Constant, Node, Node::EType::Image);
 
 
-	//---------------------------------------------------------------------------------------------
-	int NodeImageConstant::GetInputCount() const
-	{
-		return 0;
-	}
-
-
-	//---------------------------------------------------------------------------------------------
-	Node* NodeImageConstant::GetInputNode(int) const
-	{
-		check(false);
-		return 0;
-	}
-
-
-	//---------------------------------------------------------------------------------------------
-	void NodeImageConstant::SetInputNode(int, NodePtr)
-	{
-		check(false);
-	}
-
-
 	//---------------------------------------------------------------------------------------------
 	ImagePtrConst NodeImageConstant::GetValue() const
 	{
diff --git a/Engine/Plugins/Experimental/Mutable/Source/MutableTools/Private/MuT/NodeImageFormat.cpp b/Engine/Plugins/Experimental/Mutable/Source/MutableTools/Private/MuT/NodeImageFormat.cpp
--- a/Engine/Plugins/Experimental/Mutable/Source/MutableTools/Private/MuT/NodeImageFormat.cpp
+++ b/Engine/Plugins/Experimental/Mutable/Source/MutableTools/Private/MuT/NodeImageFormat.cpp
@@ -25,35 +25,6 @@ namespace mu
 	MUTABLE_IMPLEMENT_NODE( NodeImageFormat, EType::Format, Node, Node::EType::Image)
 
 
-	//---------------------------------------------------------------------------------------------
-	// Node Interface
-	//---------------------------------------------------------------------------------------------
-	int NodeImageFormat::GetInputCount() const
-	{
-		return 1;
-	}
-
-
-	//---------------------------------------------------------------------------------------------
-	Node* NodeImageFormat::GetInputNode( int i ) const
-	{
-		check( i>=0 && i<GetInputCount() );
-        (void)i;
-
-		return m_pD->m_source.get();
-	}
-
-
-	//---------------------------------------------------------------------------------------------
-	void NodeImageFormat::SetInputNode( int i, NodePtr pNode )
-	{
-		check( i>=0 && i<GetInputCount() );
-        (void)i;
-
-		m_pD->m_source = dynamic_cast<NodeImage*>(pNode.get());
-	}
-
-
 	//---------------------------------------------------------------------------------------------
 	// Own Interface
 	//---------------------------------------------------------------------------------------------
diff --git a/Engine/Plugins/Experimental/Mutable/Source/MutableTools/Private/MuT/NodeImageGradient.cpp b/Engine/Plugins/Experimental/Mutable/Source/MutableTools/Private/MuT/NodeImageGradient.cpp
--- a/Engine/Plugins/Experimental/Mutable/Source/MutableTools/Private/MuT/NodeImageGradient.cpp
+++ b/Engine/Plugins/Experimental/Mutable/Source/MutableTools/Private/MuT/NodeImageGradient.cpp
@@ -27,45 +27,6 @@ namespace mu
 	MUTABLE_IMPLEMENT_NODE( NodeImageGradient, EType::Gradient, Node, Node::EType::Image)
 
 
-	//---------------------------------------------------------------------------------------------
-	// Node Interface
-	//---------------------------------------------------------------------------------------------
-	int NodeImageGradient::GetInputCount() const
-	{
-		return 2;
-	}
-
-
-	//---------------------------------------------------------------------------------------------
-	Node* NodeImageGradient::GetInputNode( int i ) const
-	{
-		check( i>=0 && i< GetInputCount());
-
-		Node* pResult = 0;
-
-		switch (i)
-		{
-        case 0: pResult = m_pD->m_pColour0.get(); break;
-        case 1: pResult = m_pD->m_pColour1.get(); break;
-		}
-
-		return pResult;
-	}
-
-
-	//---------------------------------------------------------------------------------------------
-	void NodeImageGradient::SetInputNode( int i, NodePtr pNode )
-	{
-		check( i>=0 && i< GetInputCount());
-
-		switch (i)
-		{
-        case 0: m_pD->m_pColour0 = dynamic_cast<NodeColour*>( pNode.get() ); break;
-        case 1: m_pD->m_pColour1 = dynamic_cast<NodeColour*>( pNode.get() ); break;
-		}
-	}
-
-
 	//---------------------------------------------------------------------------------------------
 	// Own Interface
 	//---------------------------------------------------------------------------------------------
diff --git a/Engine/Plugins/Experimental/Mutable/Source/MutableTools/Private/MuT/NodeImageInterpolate.cpp b/Engine/Plugins/Experimental/Mutable/Source/MutableTools/Private/MuT/NodeImageInterpolate.cpp
--- a/Engine/Plugins/Experimental/Mutable/Source/MutableTools/Private/MuT/NodeImageInterpolate.cpp
+++ b/Engine/Plugins/Experimental/Mutable/Source/MutableTools/Private/MuT/NodeImageInterpolate.cpp
@@ -26,55 +26,6 @@ namespace mu
 	MUTABLE_IMPLEMENT_NODE( NodeImageInterpolate, EType::Interpolate, Node, Node::EType::Image)
 
 
-	//---------------------------------------------------------------------------------------------
-	// Node Interface
-	//---------------------------------------------------------------------------------------------
-	int NodeImageInterpolate::GetInputCount() const
-	{
-		return 1 + (int)m_pD->m_targets.Num();
-	}
-
-
-	//---------------------------------------------------------------------------------------------
-	Node* NodeImageInterpolate::GetInputNode( int i ) const
-	{
-		check( i>=0 && i<GetInputCount() );
-
-		Node* pResult = 0;
-
-		switch (i)
-		{
-		case 0:
-			pResult = m_pD->m_pFactor.get();
-			break;
-
-		default:
-			pResult = m_pD->m_targets[i-1].get();
-			break;
-		}
-
-		return pResult;
-	}
-
-
-	//---------------------------------------------------------------------------------------------
-	void NodeImageInterpolate::SetInputNode( int i, NodePtr pNode )
-	{
-		check( i>=0 && i<GetInputCount() );
-
-		switch (i)
-		{
-		case 0:
-			m_pD->m_pFactor = dynamic_cast<NodeScalar*>(pNode.get());
-			break;
-
-		default:
-			m_pD->m_targets[i-1] = dynamic_cast<NodeImage*>(pNode.get());
-			break;
-		}
-	}
-
-
 	//---------------------------------------------------------------------------------------------
 	// Own Interface
 	//---------------------------------------------------------------------------------------------
diff --git a/Engine/Plugins/Experimental/Mutable/Source/MutableTools/Private/MuT/NodeImageInvert.cpp b/Engine/Plugins/Experimental/Mutable/Source/MutableTools/Private/MuT/NodeImageInvert.cpp
--- a/Engine/Plugins/Experimental/Mutable/Source/MutableTools/Private/MuT/NodeImageInvert.cpp
+++ b/Engine/Plugins/Experimental/Mutable/Source/MutableTools/Private/MuT/NodeImageInvert.cpp
@@ -21,33 +21,6 @@ namespace mu
 
 	MUTABLE_IMPLEMENT_NODE( NodeImageInvert, EType::Invert, Node, Node::EType::Image)
 
-	//---------------------------------------------------------------------------------------------
-	// Node Interface
-	//---------------------------------------------------------------------------------------------
-	void NodeImageInvert::SetInputNode(int i, NodePtr pNode)
-	{
-		check(i >= 0 && i < GetInputCount());
-
-		m_pD->m_pBase = dynamic_cast<NodeImage*>(pNode.get());
-	}
-
-	//---------------------------------------------------------------------------------------------
-	int NodeImageInvert::GetInputCount() const
-	{
-		return 1;
-	}
-
-	//---------------------------------------------------------------------------------------------
-	Node* NodeImageInvert::GetInputNode(int i)const
-	{
-		check(i >= 0 && i < GetInputCount());
-
-		Node* pResult = 0;
-		pResult = m_pD->m_pBase.get();
-
-		return pResult;
-	}
-
 	//---------------------------------------------------------------------------------------------
 	// Own Interface
 	//---------------------------------------------------------------------------------------------
diff --git a/Engine/Plugins/Experimental/Mutable/Source/MutableTools/Private/MuT/NodeImageLayer.cpp b/Engine/Plugins/Experimental/Mutable/Source/MutableTools/Private/MuT/NodeImageLayer.cpp
--- a/Engine/Plugins/Experimental/Mutable/Source/MutableTools/Private/MuT/NodeImageLayer.cpp
+++ b/Engine/Plugins/Experimental/Mutable/Source/MutableTools/Private/MuT/NodeImageLayer.cpp
@@ -25,47 +25,6 @@ namespace mu
 	MUTABLE_IMPLEMENT_NODE( NodeImageLayer, EType::Layer, Node, Node::EType::Image)
 
 
-	//---------------------------------------------------------------------------------------------
-	// Node Interface
-	//---------------------------------------------------------------------------------------------
-	int NodeImageLayer::GetInputCount() const
-	{
-		return 3;
-	}
-
-
-	//---------------------------------------------------------------------------------------------
-	Node* NodeImageLayer::GetInputNode( int i ) const
-	{
-		check( i>=0 && i< GetInputCount());
-
-		Node* pResult = 0;
-
-		switch (i)
-		{
-        case 0: pResult = m_pD->m_pBase.get(); break;
-        case 1: pResult = m_pD->m_pMask.get(); break;
-        case 2: pResult = m_pD->m_pBlended.get(); break;
-		}
-
-		return pResult;
-	}
-
-
-	//---------------------------------------------------------------------------------------------
-	void NodeImageLayer::SetInputNode( int i, NodePtr pNode )
-	{
-		check( i>=0 && i< GetInputCount());
-
-		switch (i)
-		{
-        case 0: m_pD->m_pBase = dynamic_cast<NodeImage*>( pNode.get() ); break;
-        case 1: m_pD->m_pMask = dynamic_cast<NodeImage*>( pNode.get() ); break;
-        case 2: m_pD->m_pBlended = dynamic_cast<NodeImage*>( pNode.get() ); break;
-		}
-	}
-
-
 	//---------------------------------------------------------------------------------------------
 	// Own Interface
 	//---------------------------------------------------------------------------------------------
diff --git a/Engine/Plugins/Experimental/Mutable/Source/MutableTools/Private/MuT/NodeImageLayerColour.cpp b/Engine/Plugins/Experimental/Mutable/Source/MutableTools/Private/MuT/NodeImageLayerColour.cpp
--- a/Engine/Plugins/Experimental/Mutable/Source/MutableTools/Private/MuT/NodeImageLayerColour.cpp
+++ b/Engine/Plugins/Experimental/Mutable/Source/MutableTools/Private/MuT/NodeImageLayerColour.cpp
@@ -26,47 +26,6 @@ namespace mu
 	MUTABLE_IMPLEMENT_NODE( NodeImageLayerColour, EType::LayerColour, Node, Node::EType::Image);
 
 
-	//---------------------------------------------------------------------------------------------
-	// Node Interface
-	//---------------------------------------------------------------------------------------------
-	int NodeImageLayerColour::GetInputCount() const
-	{
-		return 3;
-	}
-
-
-	//---------------------------------------------------------------------------------------------
-	Node* NodeImageLayerColour::GetInputNode( int i ) const
-	{
-		check( i>=0 && i< GetInputCount());
-
-		Node* pResult = 0;
-
-		switch (i)
-		{
-        case 0: pResult = m_pD->m_pBase.get(); break;
-        case 1: pResult = m_pD->m_pMask.get(); break;
-        case 2: pResult = m_pD->m_pColour.get(); break;
-		}
-
-		return pResult;
-	}
-
-
-	//---------------------------------------------------------------------------------------------
-	void NodeImageLayerColour::SetInputNode( int i, NodePtr pNode )
-	{
-		check( i>=0 && i< GetInputCount());
-
-		switch (i)
-		{
-        case 0: m_pD->m_pBase = dynamic_cast<NodeImage*>( pNode.get() ); break;
-        case 1: m_pD->m_pMask = dynamic_cast<NodeImage*>( pNode.get() ); break;
-        case 2: m_pD->m_pColour = dynamic_cast<NodeColour*>( pNode.get() ); break;
-		}
-	}
-
-
 	//---------------------------------------------------------------------------------------------
 	// Own Interface
 	//---------------------------------------------------------------------------------------------
diff --git a/Engine/Plugins/Experimental/Mutable/Source/MutableTools/Private/MuT/NodeImageLuminance.cpp b/Engine/Plugins/Experimental/Mutable/Source/MutableTools/Private/MuT/NodeImageLuminance.cpp
--- a/Engine/Plugins/Experimental/Mutable/Source/MutableTools/Private/MuT/NodeImageLuminance.cpp
+++ b/Engine/Plugins/Experimental/Mutable/Source/MutableTools/Private/MuT/NodeImageLuminance.cpp
@@ -25,43 +25,6 @@ namespace mu
 	MUTABLE_IMPLEMENT_NODE( NodeImageLuminance, EType::Luminance, Node, Node::EType::Image)
 
 
-	//---------------------------------------------------------------------------------------------
-	// Node Interface
-	//---------------------------------------------------------------------------------------------
-	int NodeImageLuminance::GetInputCount() const
-	{
-		return 1;
-	}
-
-
-	//---------------------------------------------------------------------------------------------
-	Node* NodeImageLuminance::GetInputNode( int i ) const
-	{
-		check( i>=0 && i< GetInputCount());
-
-		Node* pResult = 0;
-
-		switch (i)
-		{
-        case 0: pResult = m_pD->m_pSource.get(); break;
-		}
-
-		return pResult;
-	}
-
-
-	//---------------------------------------------------------------------------------------------
-	void NodeImageLuminance::SetInputNode( int i, NodePtr pNode )
-	{
-		check( i>=0 && i< GetInputCount());
-
-		switch (i)
-		{
-        case 0: m_pD->m_pSource = dynamic_cast<NodeImage*>(pNode.get()); break;
-		}
-	}
-
-
 	//---------------------------------------------------------------------------------------------
 	// Own Interface
 	//---------------------------------------------------------------------------------------------
diff --git a/Engine/Plugins/Experimental/Mutable/Source/MutableTools/Private/MuT/NodeImageMipmap.cpp b/Engine/Plugins/Experimental/Mutable/Source/MutableTools/Private/MuT/NodeImageMipmap.cpp
--- a/Engine/Plugins/Experimental/Mutable/Source/MutableTools/Private/MuT/NodeImageMipmap.cpp
+++ b/Engine/Plugins/Experimental/Mutable/Source/MutableTools/Private/MuT/NodeImageMipmap.cpp
@@ -26,43 +26,6 @@ namespace mu
     MUTABLE_IMPLEMENT_NODE( NodeImageMipmap, EType::Mipmap, Node, Node::EType::Image)
 
 
-	//---------------------------------------------------------------------------------------------
-	// Node Interface
-	//---------------------------------------------------------------------------------------------
-    int NodeImageMipmap::GetInputCount() const
-	{
-		return 1;
-	}
-
-
-	//---------------------------------------------------------------------------------------------
-    Node* NodeImageMipmap::GetInputNode( int i ) const
-	{
-		check( i>=0 && i< GetInputCount());
-
-		Node* pResult = 0;
-
-		switch (i)
-		{
-        case 0: pResult = m_pD->m_pSource.get(); break;
-		}
-
-		return pResult;
-	}
-
-
-	//---------------------------------------------------------------------------------------------
-    void NodeImageMipmap::SetInputNode( int i, NodePtr pNode )
-	{
-		check( i>=0 && i< GetInputCount());
-
-		switch (i)
-		{
-        case 0: m_pD->m_pSource = dynamic_cast<NodeImage*>(pNode.get()); break;
-		}
-	}
-
-
 	//---------------------------------------------------------------------------------------------
 	// Own Interface
 	//---------------------------------------------------------------------------------------------
diff --git a/Engine/Plugins/Experimental/Mutable/Source/MutableTools/Private/MuT/NodeImageMultiLayer.cpp b/Engine/Plugins/Experimental/Mutable/Source/MutableTools/Private/MuT/NodeImageMultiLayer.cpp
--- a/Engine/Plugins/Experimental/Mutable/Source/MutableTools/Private/MuT/NodeImageMultiLayer.cpp
+++ b/Engine/Plugins/Experimental/Mutable/Source/MutableTools/Private/MuT/NodeImageMultiLayer.cpp
@@ -25,49 +25,6 @@ namespace mu
     MUTABLE_IMPLEMENT_NODE( NodeImageMultiLayer, EType::MultiLayer, Node, Node::EType::Image)
 
 
-	//---------------------------------------------------------------------------------------------
-	// Node Interface
-	//---------------------------------------------------------------------------------------------
-    int NodeImageMultiLayer::GetInputCount() const
-	{
-		return 4;
-	}
-
-
-	//---------------------------------------------------------------------------------------------
-    Node* NodeImageMultiLayer::GetInputNode( int i ) const
-	{
-		check( i>=0 && i< GetInputCount());
-
-		Node* pResult = 0;
-
-		switch (i)
-		{
-        case 0: pResult = m_pD->m_pBase.get(); break;
-        case 1: pResult = m_pD->m_pMask.get(); break;
-        case 2: pResult = m_pD->m_pBlended.get(); break;
-        case 3: pResult = m_pD->m_pRange.get(); break;
-        }
-
-		return pResult;
-	}
-
-
-	//---------------------------------------------------------------------------------------------
-    void NodeImageMultiLayer::SetInputNode( int i, NodePtr pNode )
-	{
-		check( i>=0 && i< GetInputCount());
-
-		switch (i)
-		{
-        case 0: m_pD->m_pBase = dynamic_cast<NodeImage*>( pNode.get() ); break;
-        case 1: m_pD->m_pMask = dynamic_cast<NodeImage*>( pNode.get() ); break;
-        case 2: m_pD->m_pBlended = dynamic_cast<NodeImage*>( pNode.get() ); break;
-        case 3: m_pD->m_pRange = dynamic_cast<NodeRange*>( pNode.get() ); break;
-        }
-	}
-
-
 	//---------------------------------------------------------------------------------------------
 	// Own Interface
 	//---------------------------------------------------------------------------------------------
diff --git a/Engine/Plugins/Experimental/Mutable/Source/MutableTools/Private/MuT/NodeImageNormalComposite.cpp b/Engine/Plugins/Experimental/Mutable/Source/MutableTools/Private/MuT/NodeImageNormalComposite.cpp
--- a/Engine/Plugins/Experimental/Mutable/Source/MutableTools/Private/MuT/NodeImageNormalComposite.cpp
+++ b/Engine/Plugins/Experimental/Mutable/Source/MutableTools/Private/MuT/NodeImageNormalComposite.cpp
@@ -25,45 +25,6 @@ namespace mu
 	MUTABLE_IMPLEMENT_NODE( NodeImageNormalComposite, EType::NormalComposite, Node, Node::EType::Image )
 
 
-	//---------------------------------------------------------------------------------------------
-	// Node Interface
-	//---------------------------------------------------------------------------------------------
-	int NodeImageNormalComposite::GetInputCount() const
-	{
-		return 2;
-	}
-
-
-	//---------------------------------------------------------------------------------------------
-	Node* NodeImageNormalComposite::GetInputNode( int i ) const
-	{
-		check( i>=0 && i< GetInputCount());
-
-		Node* pResult = 0;
-
-		switch (i)
-		{
-        case 0: pResult = m_pD->m_pBase.get(); break;
-        case 1: pResult = m_pD->m_pNormal.get(); break;
-		}
-
-		return pResult;
-	}
-
-
-	//---------------------------------------------------------------------------------------------
-	void NodeImageNormalComposite::SetInputNode( int i, NodePtr pNode )
-	{
-		check( i>=0 && i< GetInputCount());
-
-		switch (i)
-		{
-        case 0: m_pD->m_pBase = dynamic_cast<NodeImage*>( pNode.get() ); break;
-        case 1: m_pD->m_pNormal = dynamic_cast<NodeImage*>( pNode.get() ); break;
-		}
-	}
-
-
 	//---------------------------------------------------------------------------------------------
 	// Own Interface
 	//---------------------------------------------------------------------------------------------
diff --git a/Engine/Plugins/Experimental/Mutable/Source/MutableTools/Private/MuT/NodeImageParameter.cpp b/Engine/Plugins/Experimental/Mutable/Source/MutableTools/Private/MuT/NodeImageParameter.cpp
--- a/Engine/Plugins/Experimental/Mutable/Source/MutableTools/Private/MuT/NodeImageParameter.cpp
+++ b/Engine/Plugins/Experimental/Mutable/Source/MutableTools/Private/MuT/NodeImageParameter.cpp
@@ -25,30 +25,6 @@ namespace mu
     MUTABLE_IMPLEMENT_NODE( NodeImageParameter, EType::Parameter, Node, Node::EType::Image)
 
 
-	//---------------------------------------------------------------------------------------------
-	// Node Interface
-	//---------------------------------------------------------------------------------------------
-    int NodeImageParameter::GetInputCount() const
-	{
-		return 0;
-	}
-
-
-	//---------------------------------------------------------------------------------------------
-    Node* NodeImageParameter::GetInputNode( int ) const
-	{
-		check( false );
-		return 0;
-	}
-
-
-	//---------------------------------------------------------------------------------------------
-    void NodeImageParameter::SetInputNode( int, NodePtr )
-	{
-		check( false );
-	}
-
-
 	//---------------------------------------------------------------------------------------------
 	// Own Interface
 	//---------------------------------------------------------------------------------------------
diff --git a/Engine/Plugins/Experimental/Mutable/Source/MutableTools/Private/MuT/NodeImagePlainColour.cpp b/Engine/Plugins/Experimental/Mutable/Source/MutableTools/Private/MuT/NodeImagePlainColour.cpp
--- a/Engine/Plugins/Experimental/Mutable/Source/MutableTools/Private/MuT/NodeImagePlainColour.cpp
+++ b/Engine/Plugins/Experimental/Mutable/Source/MutableTools/Private/MuT/NodeImagePlainColour.cpp
@@ -24,43 +24,6 @@ namespace mu
 	MUTABLE_IMPLEMENT_NODE( NodeImagePlainColour, EType::PlainColour, Node, Node::EType::Image)
 
 
-	//---------------------------------------------------------------------------------------------
-	// Node Interface
-	//---------------------------------------------------------------------------------------------
-	void NodeImagePlainColour::SetInputNode( int i, NodePtr pNode )
-	{
-		check( i>=0 && i< GetInputCount());
-
-		switch (i)
-		{
-        case 0: m_pD->m_pColour = dynamic_cast<NodeColour*>(pNode.get()); break;
-		}
-	}
-
-
-	//---------------------------------------------------------------------------------------------
-	int NodeImagePlainColour::GetInputCount() const
-	{
-		return 1;
-	}
-
-
-	//---------------------------------------------------------------------------------------------
-	Node* NodeImagePlainColour::GetInputNode( int i ) const
-	{
-		check( i>=0 && i< GetInputCount());
-
-		Node* pResult = 0;
-
-		switch (i)
-		{
-        case 0: pResult = m_pD->m_pColour.get(); break;
-		}
-
-		return pResult;
-	}
-
-
 	//---------------------------------------------------------------------------------------------
 	// Own Interface
 	//---------------------------------------------------------------------------------------------
diff --git a/Engine/Plugins/Experimental/Mutable/Source/MutableTools/Private/MuT/NodeImageProject.cpp b/Engine/Plugins/Experimental/Mutable/Source/MutableTools/Private/MuT/NodeImageProject.cpp
--- a/Engine/Plugins/Experimental/Mutable/Source/MutableTools/Private/MuT/NodeImageProject.cpp
+++ b/Engine/Plugins/Experimental/Mutable/Source/MutableTools/Private/MuT/NodeImageProject.cpp
@@ -30,87 +30,6 @@ namespace mu
 	MUTABLE_IMPLEMENT_NODE( NodeImageProject, EType::Project, Node, Node::EType::Image)
 
 
-	//---------------------------------------------------------------------------------------------
-	// Node Interface
-	//---------------------------------------------------------------------------------------------
-	int NodeImageProject::GetInputCount() const
-	{
-		return 6;
-	}
-
-
-	//---------------------------------------------------------------------------------------------
-	Node* NodeImageProject::GetInputNode( int i ) const
-	{
-		check( i>=0 && i<GetInputCount() );
-
-		Node* pResult = 0;
-
-		switch (i)
-		{
-		case 0:
-			pResult = m_pD->m_pMesh.get();
-			break;
-
-		case 1:
-			pResult = m_pD->m_pProjector.get();
-			break;
-
-		case 2:
-			pResult = m_pD->m_pAngleFadeStart.get();
-			break;
-
-		case 3:
-			pResult = m_pD->m_pAngleFadeEnd.get();
-			break;
-
-		case 4:
-			pResult = m_pD->m_pMask.get();
-			break;
-
-		default:
-			pResult = m_pD->m_pImage.get();
-			break;
-		}
-
-		return pResult;
-	}
-
-
-	//---------------------------------------------------------------------------------------------
-	void NodeImageProject::SetInputNode( int i, Ptr<Node> InNode)
-	{
-		check( i>=0 && i<GetInputCount() );
-
-		switch (i)
-		{
-		case 0:
-			m_pD->m_pMesh = dynamic_cast<NodeMesh*>(InNode.get());
-			break;
-
-		case 1:
-			m_pD->m_pProjector = dynamic_cast<NodeProjector*>(InNode.get());
-			break;
-
-		case 2:
-			m_pD->m_pAngleFadeStart = dynamic_cast<NodeScalar*>(InNode.get());
-			break;
-
-		case 3:
-			m_pD->m_pAngleFadeEnd = dynamic_cast<NodeScalar*>(InNode.get());
-			break;
-
-		case 4:
-			m_pD->m_pMask = dynamic_cast<NodeImage*>(InNode.get());
-			break;
-
-		default:
-			m_pD->m_pImage = dynamic_cast<NodeImage*>(InNode.get());
-			break;
-		}
-	}
-
-
 	//---------------------------------------------------------------------------------------------
 	// Own Interface
 	//---------------------------------------------------------------------------------------------
diff --git a/Engine/Plugins/Experimental/Mutable/Source/MutableTools/Private/MuT/NodeImageResize.cpp b/Engine/Plugins/Experimental/Mutable/Source/MutableTools/Private/MuT/NodeImageResize.cpp
--- a/Engine/Plugins/Experimental/Mutable/Source/MutableTools/Private/MuT/NodeImageResize.cpp
+++ b/Engine/Plugins/Experimental/Mutable/Source/MutableTools/Private/MuT/NodeImageResize.cpp
@@ -22,43 +22,6 @@ namespace mu
 	MUTABLE_IMPLEMENT_NODE( NodeImageResize, EType::Resize, Node, Node::EType::Image)
 
 
-	//---------------------------------------------------------------------------------------------
-	// Node Interface
-	//---------------------------------------------------------------------------------------------
-	void NodeImageResize::SetInputNode( int i, NodePtr pNode )
-	{
-		check( i>=0 && i< GetInputCount());
-
-		switch (i)
-		{
-        case 0: m_pD->m_pBase = dynamic_cast<NodeImage*>(pNode.get()); break;
-		}
-	}
-
-
-	//---------------------------------------------------------------------------------------------
-	int NodeImageResize::GetInputCount() const
-	{
-		return 1;
-	}
-
-
-	//---------------------------------------------------------------------------------------------
-	Node* NodeImageResize::GetInputNode( int i ) const
-	{
-		check( i>=0 && i< GetInputCount());
-
-		Node* pResult = 0;
-
-		switch (i)
-		{
-        case 0: pResult = m_pD->m_pBase.get(); break;
-		}
-
-		return pResult;
-	}
-
-
 	//---------------------------------------------------------------------------------------------
 	// Own Interface
 	//---------------------------------------------------------------------------------------------
diff --git a/Engine/Plugins/Experimental/Mutable/Source/MutableTools/Private/MuT/NodeImageSaturate.cpp b/Engine/Plugins/Experimental/Mutable/Source/MutableTools/Private/MuT/NodeImageSaturate.cpp
--- a/Engine/Plugins/Experimental/Mutable/Source/MutableTools/Private/MuT/NodeImageSaturate.cpp
+++ b/Engine/Plugins/Experimental/Mutable/Source/MutableTools/Private/MuT/NodeImageSaturate.cpp
@@ -26,45 +26,6 @@ namespace mu
 	MUTABLE_IMPLEMENT_NODE( NodeImageSaturate, EType::Saturate, Node, Node::EType::Image)
 
 
-	//---------------------------------------------------------------------------------------------
-	// Node Interface
-	//---------------------------------------------------------------------------------------------
-	int NodeImageSaturate::GetInputCount() const
-	{
-		return 2;
-	}
-
-
-	//---------------------------------------------------------------------------------------------
-	Node* NodeImageSaturate::GetInputNode( int i ) const
-	{
-		check( i>=0 && i< GetInputCount());
-
-		Node* pResult = 0;
-
-		switch (i)
-		{
-		case 0: pResult = m_pD->m_pFactor.get();
-		case 1: pResult = m_pD->m_pSource.get();
-		}
-
-		return pResult;
-	}
-
-
-	//---------------------------------------------------------------------------------------------
-	void NodeImageSaturate::SetInputNode( int i, NodePtr pNode )
-	{
-		check( i>=0 && i< GetInputCount());
-
-		switch (i)
-		{
-		case 0: m_pD->m_pFactor = dynamic_cast<NodeScalar*>(pNode.get());
-		case 1: m_pD->m_pSource = dynamic_cast<NodeImage*>(pNode.get());
-		}
-	}
-
-
 	//---------------------------------------------------------------------------------------------
 	// Own Interface
 	//---------------------------------------------------------------------------------------------
diff --git a/Engine/Plugins/Experimental/Mutable/Source/MutableTools/Private/MuT/NodeImageSwitch.cpp b/Engine/Plugins/Experimental/Mutable/Source/MutableTools/Private/MuT/NodeImageSwitch.cpp
--- a/Engine/Plugins/Experimental/Mutable/Source/MutableTools/Private/MuT/NodeImageSwitch.cpp
+++ b/Engine/Plugins/Experimental/Mutable/Source/MutableTools/Private/MuT/NodeImageSwitch.cpp
@@ -26,55 +26,6 @@ namespace mu
 	MUTABLE_IMPLEMENT_NODE( NodeImageSwitch, EType::Switch, Node, Node::EType::Image)
 
 
-	//---------------------------------------------------------------------------------------------
-	// Node Interface
-	//---------------------------------------------------------------------------------------------
-	int NodeImageSwitch::GetInputCount() const
-	{
-		return 1 + (int)m_pD->m_options.Num();
-	}
-
-
-	//---------------------------------------------------------------------------------------------
-	Node* NodeImageSwitch::GetInputNode( int i ) const
-	{
-		check( i>=0 && i<GetInputCount() );
-
-		Node* pResult = 0;
-
-		switch (i)
-		{
-		case 0:
-			pResult = m_pD->m_pParameter.get();
-			break;
-
-		default:
-			pResult = m_pD->m_options[i-1].get();
-			break;
-		}
-
-		return pResult;
-	}
-
-
-	//---------------------------------------------------------------------------------------------
-	void NodeImageSwitch::SetInputNode( int i, NodePtr pNode )
-	{
-		check( i>=0 && i<GetInputCount() );
-
-		switch (i)
-		{
-		case 0:
-			m_pD->m_pParameter = dynamic_cast<NodeScalar*>(pNode.get());
-			break;
-
-		default:
-			m_pD->m_options[i-1] = dynamic_cast<NodeImage*>(pNode.get());
-			break;
-		}
-	}
-
-
 	//---------------------------------------------------------------------------------------------
 	// Own Interface
 	//---------------------------------------------------------------------------------------------
diff --git a/Engine/Plugins/Experimental/Mutable/Source/MutableTools/Private/MuT/NodeImageSwizzle.cpp b/Engine/Plugins/Experimental/Mutable/Source/MutableTools/Private/MuT/NodeImageSwizzle.cpp
--- a/Engine/Plugins/Experimental/Mutable/Source/MutableTools/Private/MuT/NodeImageSwizzle.cpp
+++ b/Engine/Plugins/Experimental/Mutable/Source/MutableTools/Private/MuT/NodeImageSwizzle.cpp
@@ -26,37 +26,6 @@ namespace mu
 	MUTABLE_IMPLEMENT_NODE( NodeImageSwizzle, EType::Swizzle, Node, Node::EType::Image)
 
 
-	//---------------------------------------------------------------------------------------------
-	// Node Interface
-	//---------------------------------------------------------------------------------------------
-	int NodeImageSwizzle::GetInputCount() const
-	{
-		return (int)m_pD->m_sources.Num();
-	}
-
-
-	//---------------------------------------------------------------------------------------------
-	Node* NodeImageSwizzle::GetInputNode( int i ) const
-	{
-		check( i>=0 && i<GetInputCount() );
-
-		Node* pResult = 0;
-
-		pResult = m_pD->m_sources[i].get();
-
-		return pResult;
-	}
-
-
-	//---------------------------------------------------------------------------------------------
-	void NodeImageSwizzle::SetInputNode( int i, NodePtr pNode )
-	{
-		check( i>=0 && i<GetInputCount() );
-
-		m_pD->m_sources[i] = dynamic_cast<NodeImage*>(pNode.get());
-	}
-
-
 	//---------------------------------------------------------------------------------------------
 	// Own Interface
 	//---------------------------------------------------------------------------------------------
diff --git a/Engine/Plugins/Experimental/Mutable/Source/MutableTools/Private/MuT/NodeImageTable.cpp b/Engine/Plugins/Experimental/Mutable/Source/MutableTools/Private/MuT/NodeImageTable.cpp
--- a/Engine/Plugins/Experimental/Mutable/Source/MutableTools/Private/MuT/NodeImageTable.cpp
+++ b/Engine/Plugins/Experimental/Mutable/Source/MutableTools/Private/MuT/NodeImageTable.cpp
@@ -27,33 +27,6 @@ namespace mu
 	MUTABLE_IMPLEMENT_NODE( NodeImageTable, EType::Table, Node, Node::EType::Image);
 
 
-	//---------------------------------------------------------------------------------------------
-	// Node Interface
-	//---------------------------------------------------------------------------------------------
-	int NodeImageTable::GetInputCount() const
-	{
-		return 0;
-	}
-
-
-	//---------------------------------------------------------------------------------------------
-	Node* NodeImageTable::GetInputNode( int i ) const
-	{
-		check( i >=0 && i < 0 );
-        (void)i;
-        return 0;
-	}
-
-
-	//---------------------------------------------------------------------------------------------
-    void NodeImageTable::SetInputNode( int i, NodePtr )
-	{
-		check( i >=0 && i < 0 );
-        (void)i;
-        //m_pD->m_pObject = dynamic_cast<NodeObject*>( pNode.get() );
-	}
-
-
 	//---------------------------------------------------------------------------------------------
 	// Own Interface
 	//---------------------------------------------------------------------------------------------
diff --git a/Engine/Plugins/Experimental/Mutable/Source/MutableTools/Private/MuT/NodeImageTransform.cpp b/Engine/Plugins/Experimental/Mutable/Source/MutableTools/Private/MuT/NodeImageTransform.cpp
--- a/Engine/Plugins/Experimental/Mutable/Source/MutableTools/Private/MuT/NodeImageTransform.cpp
+++ b/Engine/Plugins/Experimental/Mutable/Source/MutableTools/Private/MuT/NodeImageTransform.cpp
@@ -23,53 +23,6 @@ namespace mu
 	MUTABLE_IMPLEMENT_NODE( NodeImageTransform, EType::Transform, Node, Node::EType::Image)
 
 
-	//---------------------------------------------------------------------------------------------
-	// Node Interface
-	//---------------------------------------------------------------------------------------------
-	void NodeImageTransform::SetInputNode( int i, NodePtr pNode )
-	{
-		check( i>=0 && i< GetInputCount());
-
-		switch (i)
-		{
-        case 0: m_pD->m_pBase = dynamic_cast<NodeImage*>(pNode.get()); break;
-        case 1: m_pD->m_pOffsetX = dynamic_cast<NodeScalar*>(pNode.get()); break;
-        case 2: m_pD->m_pOffsetY = dynamic_cast<NodeScalar*>(pNode.get()); break;
-        case 3: m_pD->m_pScaleX = dynamic_cast<NodeScalar*>(pNode.get()); break;
-        case 4: m_pD->m_pScaleY = dynamic_cast<NodeScalar*>(pNode.get()); break;
-        case 5: m_pD->m_pRotation = dynamic_cast<NodeScalar*>(pNode.get()); break;
-		}
-	}
-
-
-	//---------------------------------------------------------------------------------------------
-	int NodeImageTransform::GetInputCount() const
-	{
-		return 6;
-	}
-
-
-	//---------------------------------------------------------------------------------------------
-	Node* NodeImageTransform::GetInputNode( int i ) const
-	{
-		check( i>=0 && i< GetInputCount());
-
-		Node* pResult = 0;
-
-		switch (i)
-		{
-        case 0: pResult = m_pD->m_pBase.get(); break;
-		case 1: pResult = m_pD->m_pOffsetX.get(); break;
-		case 2: pResult = m_pD->m_pOffsetY.get(); break;
-		case 3: pResult = m_pD->m_pScaleX.get(); break;
-		case 4: pResult = m_pD->m_pScaleY.get(); break;
-		case 5: pResult = m_pD->m_pRotation.get(); break;
-		}
-
-		return pResult;
-	}
-
-
 	//---------------------------------------------------------------------------------------------
 	// Own Interface
 	//---------------------------------------------------------------------------------------------
diff --git a/Engine/Plugins/Experimental/Mutable/Source/MutableTools/Private/MuT/NodeImageVariation.cpp b/Engine/Plugins/Experimental/Mutable/Source/MutableTools/Private/MuT/NodeImageVariation.cpp
--- a/Engine/Plugins/Experimental/Mutable/Source/MutableTools/Private/MuT/NodeImageVariation.cpp
+++ b/Engine/Plugins/Experimental/Mutable/Source/MutableTools/Private/MuT/NodeImageVariation.cpp
@@ -25,72 +25,29 @@ namespace mu
     MUTABLE_IMPLEMENT_NODE( NodeImageVariation, EType::Variation, Node, Node::EType::Image)
 
 
-    //---------------------------------------------------------------------------------------------
-    int NodeImageVariation::GetInputCount() const { return 1 + int( m_pD->m_variations.Num() ); }
-
-
-    //---------------------------------------------------------------------------------------------
-    Node* NodeImageVariation::GetInputNode( int i ) const
-    {
-        check( i >= 0 && i < GetInputCount() );
-
-        if ( i == 0 )
-        {
-            return m_pD->m_defaultImage.get();
-        }
-        i -= 1;
-
-        if ( i < int( m_pD->m_variations.Num() ) )
-        {
-            return m_pD->m_variations[i].m_image.get();
-        }
-        i -= int( m_pD->m_variations.Num() );
-
-        return nullptr;
-    }
-
-
-    //---------------------------------------------------------------------------------------------
-    void NodeImageVariation::SetInputNode( int i, NodePtr pNode )
-    {
-        check( i >= 0 && i < GetInputCount() );
-
-        if ( i == 0 )
-        {
-            m_pD->m_defaultImage = dynamic_cast<NodeImage*>( pNode.get() );
-            return;
-        }
-
-        i -= 1;
-        if ( i < int( m_pD->m_variations.Num() ) )
-        {
-
-            m_pD->m_variations[i].m_image = dynamic_cast<NodeImage*>( pNode.get() );
-            return;
-        }
-        i -= (int)m_pD->m_variations.Num();
-    }
-
-
     //---------------------------------------------------------------------------------------------
     // Own Interface
     //---------------------------------------------------------------------------------------------
-    void NodeImageVariation::SetDefaultImage( NodeImage* p ) { m_pD->m_defaultImage = p; }
+    void NodeImageVariation::SetDefaultImage( NodeImage* p ) 
+	{ 
+		m_pD->m_defaultImage = p; 
+	}
 
 
-    //---------------------------------------------------------------------------------------------
-    int NodeImageVariation::GetVariationCount() const { return int( m_pD->m_variations.Num() ); }
+	int NodeImageVariation::GetVariationCount() const 
+	{ 
+		return int( m_pD->m_variations.Num() ); 
+	}
 
 
-    //---------------------------------------------------------------------------------------------
     void NodeImageVariation::SetVariationCount( int num )
     {
         check( num >= 0 );
         m_pD->m_variations.SetNum( num );
     }
 
-    //---------------------------------------------------------------------------------------------
-    void NodeImageVariation::SetVariationTag( int index, const FString& Tag )
+
+	void NodeImageVariation::SetVariationTag( int index, const FString& Tag )
     {
         check( index >= 0 && index < (int)m_pD->m_variations.Num() );
 
@@ -98,7 +55,6 @@ namespace mu
     }
 
 
-    //---------------------------------------------------------------------------------------------
     void NodeImageVariation::SetVariationImage( int index, NodeImage* pNode )
     {
         check( index >= 0 && index < (int)m_pD->m_variations.Num() );
diff --git a/Engine/Plugins/Experimental/Mutable/Source/MutableTools/Private/MuT/NodeLOD.cpp b/Engine/Plugins/Experimental/Mutable/Source/MutableTools/Private/MuT/NodeLOD.cpp
--- a/Engine/Plugins/Experimental/Mutable/Source/MutableTools/Private/MuT/NodeLOD.cpp
+++ b/Engine/Plugins/Experimental/Mutable/Source/MutableTools/Private/MuT/NodeLOD.cpp
@@ -29,55 +29,6 @@ namespace mu
 	MUTABLE_IMPLEMENT_NODE( NodeLOD, EType::LOD, Node, Node::EType::LOD);
 
 
-	//---------------------------------------------------------------------------------------------
-	// Node Interface
-	//---------------------------------------------------------------------------------------------
-	int NodeLOD::GetInputCount() const
-	{
-		return m_pD->m_components.Num()+m_pD->m_modifiers.Num();
-	}
-
-
-	//---------------------------------------------------------------------------------------------
-	Node* NodeLOD::GetInputNode( int i ) const
-	{
-        check(i >= 0 && i < GetInputCount());
-
-        if (i < int(m_pD->m_components.Num()))
-        {
-            return m_pD->m_components[i].get();
-		}
-
-        i -= int(m_pD->m_components.Num());
-
-        if (i < int(m_pD->m_modifiers.Num()))
-        {
-            return m_pD->m_modifiers[i].get();
-        }
-
-        return nullptr;
-    }
-
-
-    //---------------------------------------------------------------------------------------------
-	void NodeLOD::SetInputNode( int i, NodePtr pNode )
-	{
-        check(i >= 0 && i < GetInputCount());
-
-        if (i<int(m_pD->m_components.Num()))
-		{
-            m_pD->m_components[i] = dynamic_cast<NodeComponent*>(pNode.get());
-            return;
-        }
-
-        i -= int(m_pD->m_components.Num());
-        if (i < int(m_pD->m_modifiers.Num()))
-        {
-            m_pD->m_modifiers[i] = dynamic_cast<NodeModifier*>(pNode.get());
-        }
-    }
-
-
 	//---------------------------------------------------------------------------------------------
 	// Own Interface
 	//---------------------------------------------------------------------------------------------
diff --git a/Engine/Plugins/Experimental/Mutable/Source/MutableTools/Private/MuT/NodeLayout.cpp b/Engine/Plugins/Experimental/Mutable/Source/MutableTools/Private/MuT/NodeLayout.cpp
--- a/Engine/Plugins/Experimental/Mutable/Source/MutableTools/Private/MuT/NodeLayout.cpp
+++ b/Engine/Plugins/Experimental/Mutable/Source/MutableTools/Private/MuT/NodeLayout.cpp
@@ -98,30 +98,6 @@ namespace mu
 	MUTABLE_IMPLEMENT_NODE( NodeLayoutBlocks, EType::Blocks, Node, Node::EType::Layout)
 
 
-	//---------------------------------------------------------------------------------------------
-	// Node Interface
-	//---------------------------------------------------------------------------------------------
-	int NodeLayoutBlocks::GetInputCount() const
-	{
-		return 0;
-	}
-
-
-	//---------------------------------------------------------------------------------------------
-    Node* NodeLayoutBlocks::GetInputNode( int ) const
-	{
-		check(false);
-		return 0;
-	}
-
-
-	//---------------------------------------------------------------------------------------------
-    void NodeLayoutBlocks::SetInputNode( int, NodePtr )
-	{
-		check(false);
-	}
-
-
 	//---------------------------------------------------------------------------------------------
 	// Own Interface
 	//---------------------------------------------------------------------------------------------
