[mutable] Remove unnecessary boilerplate code that was only used in one debug tool.

diff --git a/Engine/Plugins/Experimental/Mutable/Source/MutableTools/Private/MuT/NodeMeshApplyPose.cpp b/Engine/Plugins/Experimental/Mutable/Source/MutableTools/Private/MuT/NodeMeshApplyPose.cpp
--- a/Engine/Plugins/Experimental/Mutable/Source/MutableTools/Private/MuT/NodeMeshApplyPose.cpp
+++ b/Engine/Plugins/Experimental/Mutable/Source/MutableTools/Private/MuT/NodeMeshApplyPose.cpp
@@ -28,61 +28,6 @@ namespace mu
     MUTABLE_IMPLEMENT_NODE( NodeMeshApplyPose, EType::ApplyPose, Node, Node::EType::Mesh)
 
 
-	//---------------------------------------------------------------------------------------------
-	// Node Interface
-	//---------------------------------------------------------------------------------------------
-    int NodeMeshApplyPose::GetInputCount() const
-	{
-        return 2;
-	}
-
-
-	//---------------------------------------------------------------------------------------------
-    Node* NodeMeshApplyPose::GetInputNode( int i ) const
-	{
-		check( i>=0 && i<GetInputCount() );
-
-		Node* pResult = 0;
-
-		switch (i)
-		{
-		case 0:
-            pResult = m_pD->m_pBase.get();
-            break;
-        case 1:
-            pResult = m_pD->m_pPose.get();
-			break;
-        default:
-            check(false);
-            break;
-		}
-
-		return pResult;
-	}
-
-
-	//---------------------------------------------------------------------------------------------
-    void NodeMeshApplyPose::SetInputNode( int i, NodePtr pNode )
-	{
-		check( i>=0 && i<GetInputCount() );
-
-		switch (i)
-		{
-		case 0:
-            m_pD->m_pBase = dynamic_cast<NodeMesh*>(pNode.get());
-            break;
-
-		case 1:
-            m_pD->m_pPose = dynamic_cast<NodeMesh*>(pNode.get());
-			break;
-
-        default:
-            check(false);
-            break;
-        }
-	}
-
-
 	//---------------------------------------------------------------------------------------------
 	// Own Interface
 	//---------------------------------------------------------------------------------------------
diff --git a/Engine/Plugins/Experimental/Mutable/Source/MutableTools/Private/MuT/NodeMeshClipDeform.cpp b/Engine/Plugins/Experimental/Mutable/Source/MutableTools/Private/MuT/NodeMeshClipDeform.cpp
--- a/Engine/Plugins/Experimental/Mutable/Source/MutableTools/Private/MuT/NodeMeshClipDeform.cpp
+++ b/Engine/Plugins/Experimental/Mutable/Source/MutableTools/Private/MuT/NodeMeshClipDeform.cpp
@@ -27,49 +27,6 @@ namespace mu
 	MUTABLE_IMPLEMENT_NODE( NodeMeshClipDeform, EType::ClipDeform, Node, Node::EType::Mesh)
 
 
-	//---------------------------------------------------------------------------------------------
-	// Node Interface
-	//---------------------------------------------------------------------------------------------
-	int NodeMeshClipDeform::GetInputCount() const
-	{
-		return 2;
-	}
-
-
-	//---------------------------------------------------------------------------------------------
-	Node* NodeMeshClipDeform::GetInputNode( int i ) const
-	{
-		check( i>=0 && i<GetInputCount() );
-
-		Node* pResult = nullptr;
-
-		switch (i)
-		{
-		case 0: pResult = m_pD->m_pBaseMesh.get(); break;
-		case 1: pResult = m_pD->m_pClipShape.get(); break;
-		default:
-			break;
-		}
-
-		return pResult;
-	}
-
-
-	//---------------------------------------------------------------------------------------------
-	void NodeMeshClipDeform::SetInputNode( int i, NodePtr pNode )
-	{
-		check( i>=0 && i<GetInputCount() );
-
-		switch (i)
-		{
-		case 0: m_pD->m_pBaseMesh = dynamic_cast<NodeMesh*>(pNode.get()); break;
-		case 1: m_pD->m_pClipShape = dynamic_cast<NodeMesh*>(pNode.get()); break;
-		default:
-			break;
-		}
-	}
-
-
 	//---------------------------------------------------------------------------------------------
 	// Own Interface
 	//---------------------------------------------------------------------------------------------
diff --git a/Engine/Plugins/Experimental/Mutable/Source/MutableTools/Private/MuT/NodeMeshClipMorphPlane.cpp b/Engine/Plugins/Experimental/Mutable/Source/MutableTools/Private/MuT/NodeMeshClipMorphPlane.cpp
--- a/Engine/Plugins/Experimental/Mutable/Source/MutableTools/Private/MuT/NodeMeshClipMorphPlane.cpp
+++ b/Engine/Plugins/Experimental/Mutable/Source/MutableTools/Private/MuT/NodeMeshClipMorphPlane.cpp
@@ -28,53 +28,6 @@ namespace mu
     MUTABLE_IMPLEMENT_NODE( NodeMeshClipMorphPlane, EType::ClipMorphPlane, Node, Node::EType::Mesh)
 
 
-	//---------------------------------------------------------------------------------------------
-	// Node Interface
-	//---------------------------------------------------------------------------------------------
-    int NodeMeshClipMorphPlane::GetInputCount() const
-	{
-		return 1;
-	}
-
-
-	//---------------------------------------------------------------------------------------------
-    Node* NodeMeshClipMorphPlane::GetInputNode( int i ) const
-	{
-		check( i>=0 && i< GetInputCount());
-        (void)i;
-        return m_pD->m_pSource.get();
-	}
-
-
-	//---------------------------------------------------------------------------------------------
-    void NodeMeshClipMorphPlane::SetInputNode( int i, NodePtr pNode )
-	{
-		check( i>=0 && i< GetInputCount());
-		if (i==0)
-		{
-			m_pD->m_pSource = dynamic_cast<NodeMesh*>( pNode.get() );
-		}
-	}
-
-
-	//---------------------------------------------------------------------------------------------
-    const char* NodeMeshClipMorphPlane::GetInputName( int i ) const
-	{
-		check( i>=0 && i< GetInputCount());
-        (void)i;
-        return "Source";
-	}
-
-
-	//---------------------------------------------------------------------------------------------
-    const NODE_TYPE* NodeMeshClipMorphPlane::GetInputType( int i ) const
-	{
-		check( i>=0 && i< GetInputCount());
-        (void)i;
-        return NodeMesh::GetStaticType();
-	}
-
-
 	//---------------------------------------------------------------------------------------------
 	// Own Interface
 	//---------------------------------------------------------------------------------------------
diff --git a/Engine/Plugins/Experimental/Mutable/Source/MutableTools/Private/MuT/NodeMeshClipWithMesh.cpp b/Engine/Plugins/Experimental/Mutable/Source/MutableTools/Private/MuT/NodeMeshClipWithMesh.cpp
--- a/Engine/Plugins/Experimental/Mutable/Source/MutableTools/Private/MuT/NodeMeshClipWithMesh.cpp
+++ b/Engine/Plugins/Experimental/Mutable/Source/MutableTools/Private/MuT/NodeMeshClipWithMesh.cpp
@@ -27,53 +27,6 @@ namespace mu
     MUTABLE_IMPLEMENT_NODE( NodeMeshClipWithMesh, EType::ClipWithMesh, Node, Node::EType::Mesh)
 
 
-	//---------------------------------------------------------------------------------------------
-	// Node Interface
-	//---------------------------------------------------------------------------------------------
-    int NodeMeshClipWithMesh::GetInputCount() const
-	{
-		return 1;
-	}
-
-
-	//---------------------------------------------------------------------------------------------
-    Node* NodeMeshClipWithMesh::GetInputNode( int i ) const
-	{
-		check( i>=0 && i< GetInputCount());
-        (void)i;
-        return m_pD->m_pSource.get();
-	}
-
-
-	//---------------------------------------------------------------------------------------------
-    void NodeMeshClipWithMesh::SetInputNode( int i, NodePtr pNode )
-	{
-		check( i>=0 && i< GetInputCount());
-		if (i==0)
-		{
-			m_pD->m_pSource = dynamic_cast<NodeMesh*>( pNode.get() );
-		}
-	}
-
-
-	//---------------------------------------------------------------------------------------------
-    const char* NodeMeshClipWithMesh::GetInputName( int i ) const
-	{
-		check( i>=0 && i< GetInputCount());
-        (void)i;
-        return "Source";
-	}
-
-
-	//---------------------------------------------------------------------------------------------
-    const NODE_TYPE* NodeMeshClipWithMesh::GetInputType( int i ) const
-	{
-		check( i>=0 && i< GetInputCount());
-        (void)i;
-        return NodeMesh::GetStaticType();
-	}
-
-
 	//---------------------------------------------------------------------------------------------
 	// Own Interface
 	//---------------------------------------------------------------------------------------------
diff --git a/Engine/Plugins/Experimental/Mutable/Source/MutableTools/Private/MuT/NodeMeshConstant.cpp b/Engine/Plugins/Experimental/Mutable/Source/MutableTools/Private/MuT/NodeMeshConstant.cpp
--- a/Engine/Plugins/Experimental/Mutable/Source/MutableTools/Private/MuT/NodeMeshConstant.cpp
+++ b/Engine/Plugins/Experimental/Mutable/Source/MutableTools/Private/MuT/NodeMeshConstant.cpp
@@ -26,32 +26,6 @@ namespace mu
 	MUTABLE_IMPLEMENT_NODE( NodeMeshConstant, EType::Constant, Node, Node::EType::Mesh);
 
 
-	//---------------------------------------------------------------------------------------------
-	// Node Interface
-	//---------------------------------------------------------------------------------------------
-	int NodeMeshConstant::GetInputCount() const
-	{
-		return m_pD->m_layouts.Num();
-	}
-
-
-	//---------------------------------------------------------------------------------------------
-	Node* NodeMeshConstant::GetInputNode( int i ) const
-	{
-		check( i>=0 && i<m_pD->m_layouts.Num() );
-		return m_pD->m_layouts[i].get();
-	}
-
-
-	//---------------------------------------------------------------------------------------------
-	void NodeMeshConstant::SetInputNode( int i, NodePtr pNode )
-	{
-		check( i>=0 && i<m_pD->m_layouts.Num() );
-
-		m_pD->m_layouts[i] = dynamic_cast<NodeLayout*>( pNode.get() );
-	}
-
-
 	//---------------------------------------------------------------------------------------------
 	// Own Interface
 	//---------------------------------------------------------------------------------------------
diff --git a/Engine/Plugins/Experimental/Mutable/Source/MutableTools/Private/MuT/NodeMeshFormat.cpp b/Engine/Plugins/Experimental/Mutable/Source/MutableTools/Private/MuT/NodeMeshFormat.cpp
--- a/Engine/Plugins/Experimental/Mutable/Source/MutableTools/Private/MuT/NodeMeshFormat.cpp
+++ b/Engine/Plugins/Experimental/Mutable/Source/MutableTools/Private/MuT/NodeMeshFormat.cpp
@@ -28,33 +28,6 @@ class FMeshBufferSet;
 	MUTABLE_IMPLEMENT_NODE( NodeMeshFormat, EType::Format, Node, Node::EType::Mesh)
 
 
-	//---------------------------------------------------------------------------------------------
-	// Node Interface
-	//---------------------------------------------------------------------------------------------
-	int NodeMeshFormat::GetInputCount() const
-	{
-		return 1;
-	}
-
-
-	//---------------------------------------------------------------------------------------------
-	Node* NodeMeshFormat::GetInputNode( int i ) const
-	{
-		check( i==0 );
-        (void)i;
-        return m_pD->m_pSource.get();
-	}
-
-
-	//---------------------------------------------------------------------------------------------
-	void NodeMeshFormat::SetInputNode( int i, NodePtr pNode )
-	{
-		check( i==0 );
-        (void)i;
-        m_pD->m_pSource = dynamic_cast<NodeMesh*>( pNode.get() );
-	}
-
-
 	//---------------------------------------------------------------------------------------------
 	// Own Interface
 	//---------------------------------------------------------------------------------------------
diff --git a/Engine/Plugins/Experimental/Mutable/Source/MutableTools/Private/MuT/NodeMeshFragment.cpp b/Engine/Plugins/Experimental/Mutable/Source/MutableTools/Private/MuT/NodeMeshFragment.cpp
--- a/Engine/Plugins/Experimental/Mutable/Source/MutableTools/Private/MuT/NodeMeshFragment.cpp
+++ b/Engine/Plugins/Experimental/Mutable/Source/MutableTools/Private/MuT/NodeMeshFragment.cpp
@@ -27,46 +27,6 @@ namespace mu
 	MUTABLE_IMPLEMENT_NODE( NodeMeshFragment, EType::Fragment, Node, Node::EType::Mesh)
 
 
-	//---------------------------------------------------------------------------------------------
-	// Node Interface
-	//---------------------------------------------------------------------------------------------
-	int NodeMeshFragment::GetInputCount() const
-	{
-		return 1;
-	}
-
-
-	//---------------------------------------------------------------------------------------------
-	Node* NodeMeshFragment::GetInputNode( int i ) const
-	{
-		check( i>=0 && i<GetInputCount() );
-        (void)i;
-
-		Node* pResult = 0;
-
-		pResult = m_pD->m_pMesh.get();
-
-		return pResult;
-	}
-
-
-	//---------------------------------------------------------------------------------------------
-	void NodeMeshFragment::SetInputNode( int i, NodePtr pNode )
-	{
-		check( i>=0 && i<GetInputCount() );
-
-		switch (i)
-		{
-		case 0:
-			m_pD->m_pMesh = dynamic_cast<NodeMesh*>(pNode.get());
-			break;
-
-		default:
-			break;
-		}
-	}
-
-
 	//---------------------------------------------------------------------------------------------
 	// Own Interface
 	//---------------------------------------------------------------------------------------------
diff --git a/Engine/Plugins/Experimental/Mutable/Source/MutableTools/Private/MuT/NodeMeshGeometryOperation.cpp b/Engine/Plugins/Experimental/Mutable/Source/MutableTools/Private/MuT/NodeMeshGeometryOperation.cpp
--- a/Engine/Plugins/Experimental/Mutable/Source/MutableTools/Private/MuT/NodeMeshGeometryOperation.cpp
+++ b/Engine/Plugins/Experimental/Mutable/Source/MutableTools/Private/MuT/NodeMeshGeometryOperation.cpp
@@ -29,53 +29,6 @@ namespace mu
 	MUTABLE_IMPLEMENT_NODE( NodeMeshGeometryOperation, EType::GeometryOperation, Node, Node::EType::Mesh)
 
 
-	//---------------------------------------------------------------------------------------------
-	// Node Interface
-	//---------------------------------------------------------------------------------------------
-	int NodeMeshGeometryOperation::GetInputCount() const
-	{
-		return 4;
-	}
-
-
-	//---------------------------------------------------------------------------------------------
-	Node* NodeMeshGeometryOperation::GetInputNode( int i ) const
-	{
-		check( i>=0 && i<GetInputCount() );
-
-		Node* pResult = nullptr;
-
-		switch (i)
-		{
-		case 0: pResult = m_pD->m_pMeshA.get(); break;
-		case 1: pResult = m_pD->m_pMeshB.get(); break;
-		case 2: pResult = m_pD->m_pScalarA.get(); break;
-		case 3: pResult = m_pD->m_pScalarB.get(); break;
-		default:
-			break;
-		}
-
-		return pResult;
-	}
-
-
-	//---------------------------------------------------------------------------------------------
-	void NodeMeshGeometryOperation::SetInputNode( int i, NodePtr pNode )
-	{
-		check( i>=0 && i<GetInputCount() );
-
-		switch (i)
-		{
-		case 0: m_pD->m_pMeshA = dynamic_cast<NodeMesh*>(pNode.get()); break;
-		case 1: m_pD->m_pMeshB = dynamic_cast<NodeMesh*>(pNode.get()); break;
-		case 2: m_pD->m_pScalarA = dynamic_cast<NodeScalar*>(pNode.get()); break;
-		case 3: m_pD->m_pScalarB = dynamic_cast<NodeScalar*>(pNode.get()); break;
-		default:
-			break;
-		}
-	}
-
-
 	//---------------------------------------------------------------------------------------------
 	// Own Interface
 	//---------------------------------------------------------------------------------------------
diff --git a/Engine/Plugins/Experimental/Mutable/Source/MutableTools/Private/MuT/NodeMeshInterpolate.cpp b/Engine/Plugins/Experimental/Mutable/Source/MutableTools/Private/MuT/NodeMeshInterpolate.cpp
--- a/Engine/Plugins/Experimental/Mutable/Source/MutableTools/Private/MuT/NodeMeshInterpolate.cpp
+++ b/Engine/Plugins/Experimental/Mutable/Source/MutableTools/Private/MuT/NodeMeshInterpolate.cpp
@@ -28,55 +28,6 @@ namespace mu
 	MUTABLE_IMPLEMENT_NODE( NodeMeshInterpolate, EType::Interpolate, Node, Node::EType::Mesh)
 
 
-	//---------------------------------------------------------------------------------------------
-	// Node Interface
-	//---------------------------------------------------------------------------------------------
-	int NodeMeshInterpolate::GetInputCount() const
-	{
-		return 1 + m_pD->m_targets.Num();
-	}
-
-
-	//---------------------------------------------------------------------------------------------
-	Node* NodeMeshInterpolate::GetInputNode( int i ) const
-	{
-		check( i>=0 && i<GetInputCount() );
-
-		Node* pResult = 0;
-
-		switch (i)
-		{
-		case 0:
-			pResult = m_pD->m_pFactor.get();
-			break;
-
-		default:
-			pResult = m_pD->m_targets[i-1].get();
-			break;
-		}
-
-		return pResult;
-	}
-
-
-	//---------------------------------------------------------------------------------------------
-	void NodeMeshInterpolate::SetInputNode( int i, NodePtr pNode )
-	{
-		check( i>=0 && i<GetInputCount() );
-
-		switch (i)
-		{
-		case 0:
-			m_pD->m_pFactor = dynamic_cast<NodeScalar*>(pNode.get());
-			break;
-
-		default:
-			m_pD->m_targets[i-1] = dynamic_cast<NodeMesh*>(pNode.get());
-			break;
-		}
-	}
-
-
 	//---------------------------------------------------------------------------------------------
 	// Own Interface
 	//---------------------------------------------------------------------------------------------
diff --git a/Engine/Plugins/Experimental/Mutable/Source/MutableTools/Private/MuT/NodeMeshMakeMorph.cpp b/Engine/Plugins/Experimental/Mutable/Source/MutableTools/Private/MuT/NodeMeshMakeMorph.cpp
--- a/Engine/Plugins/Experimental/Mutable/Source/MutableTools/Private/MuT/NodeMeshMakeMorph.cpp
+++ b/Engine/Plugins/Experimental/Mutable/Source/MutableTools/Private/MuT/NodeMeshMakeMorph.cpp
@@ -27,38 +27,6 @@ namespace mu
     MUTABLE_IMPLEMENT_NODE( NodeMeshMakeMorph, EType::MakeMorph, Node, Node::EType::Mesh)
 
 
-	//---------------------------------------------------------------------------------------------
-	// Node Interface
-	//---------------------------------------------------------------------------------------------
-    int NodeMeshMakeMorph::GetInputCount() const
-	{
-		return 2;
-	}
-
-
-	//---------------------------------------------------------------------------------------------
-    Node* NodeMeshMakeMorph::GetInputNode( int i ) const
-	{
-		check( i>=0 && i< GetInputCount());
-		return i == 0 ? m_pD->m_pBase.get() : m_pD->m_pTarget.get();
-	}
-
-
-	//---------------------------------------------------------------------------------------------
-    void NodeMeshMakeMorph::SetInputNode( int i, NodePtr pNode )
-	{
-		check( i>=0 && i< GetInputCount());
-		if (i==0)
-		{
-			m_pD->m_pBase = dynamic_cast<NodeMesh*>( pNode.get() );
-		}
-		else
-		{
-			m_pD->m_pTarget = dynamic_cast<NodeMesh*>( pNode.get() );
-		}
-	}
-
-
 	//---------------------------------------------------------------------------------------------
 	// Own Interface
 	//---------------------------------------------------------------------------------------------
diff --git a/Engine/Plugins/Experimental/Mutable/Source/MutableTools/Private/MuT/NodeMeshMorph.cpp b/Engine/Plugins/Experimental/Mutable/Source/MutableTools/Private/MuT/NodeMeshMorph.cpp
--- a/Engine/Plugins/Experimental/Mutable/Source/MutableTools/Private/MuT/NodeMeshMorph.cpp
+++ b/Engine/Plugins/Experimental/Mutable/Source/MutableTools/Private/MuT/NodeMeshMorph.cpp
@@ -29,69 +29,6 @@ namespace mu
 	MUTABLE_IMPLEMENT_NODE( NodeMeshMorph, EType::Morph, Node, Node::EType::Mesh)
 
 
-	//---------------------------------------------------------------------------------------------
-	// Node Interface
-	//---------------------------------------------------------------------------------------------
-	int NodeMeshMorph::GetInputCount() const
-	{
-		return 3;
-	}
-
-
-	//---------------------------------------------------------------------------------------------
-	Node* NodeMeshMorph::GetInputNode( int i ) const
-	{
-		check( i>=0 && i<GetInputCount() );
-
-		Node* pResult = 0;
-
-		switch (i)
-		{
-		case 0:
-			pResult = m_pD->Factor.get();
-			break;
-
-		case 1:
-			pResult = m_pD->Base.get();
-			break;
-
-		case 2:
-			pResult = m_pD->Morph.get();
-			break;
-
-		default:
-			check(false);
-		}
-
-		return pResult;
-	}
-
-
-	//---------------------------------------------------------------------------------------------
-	void NodeMeshMorph::SetInputNode( int i, Ptr<Node> pNode )
-	{
-		check( i>=0 && i<GetInputCount() );
-
-		switch (i)
-		{
-		case 0:
-			m_pD->Factor = dynamic_cast<NodeScalar*>(pNode.get());
-			break;
-
-		case 1:
-			m_pD->Base = dynamic_cast<NodeMesh*>(pNode.get());
-			break;
-
-		case 2:
-			m_pD->Morph = dynamic_cast<NodeMesh*>(pNode.get());
-			break;
-
-		default:
-			break;
-		}
-	}
-
-
 	//---------------------------------------------------------------------------------------------
 	// Own Interface
 	//---------------------------------------------------------------------------------------------
diff --git a/Engine/Plugins/Experimental/Mutable/Source/MutableTools/Private/MuT/NodeMeshReshape.cpp b/Engine/Plugins/Experimental/Mutable/Source/MutableTools/Private/MuT/NodeMeshReshape.cpp
--- a/Engine/Plugins/Experimental/Mutable/Source/MutableTools/Private/MuT/NodeMeshReshape.cpp
+++ b/Engine/Plugins/Experimental/Mutable/Source/MutableTools/Private/MuT/NodeMeshReshape.cpp
@@ -28,51 +28,6 @@ namespace mu
 	MUTABLE_IMPLEMENT_NODE(NodeMeshReshape, EType::Reshape, Node, Node::EType::Mesh)
 
 
-	//---------------------------------------------------------------------------------------------
-	// Node Interface
-	//---------------------------------------------------------------------------------------------
-	int NodeMeshReshape::GetInputCount() const
-	{
-		return 3;
-	}
-
-
-	//---------------------------------------------------------------------------------------------
-	Node* NodeMeshReshape::GetInputNode( int i ) const
-	{
-		check( i>=0 && i<GetInputCount() );
-
-		Node* pResult = nullptr;
-
-		switch (i)
-		{
-		case 0: pResult = m_pD->BaseMesh.get(); break;
-		case 1: pResult = m_pD->BaseShape.get(); break;
-		case 2: pResult = m_pD->TargetShape.get(); break;
-		default:
-			break;
-		}
-
-		return pResult;
-	}
-
-
-	//---------------------------------------------------------------------------------------------
-	void NodeMeshReshape::SetInputNode( int i, NodePtr pNode )
-	{
-		check( i>=0 && i<GetInputCount() );
-
-		switch (i)
-		{
-		case 0: m_pD->BaseMesh = dynamic_cast<NodeMesh*>(pNode.get()); break;
-		case 1: m_pD->BaseShape = dynamic_cast<NodeMesh*>(pNode.get()); break;
-		case 2: m_pD->TargetShape = dynamic_cast<NodeMesh*>(pNode.get()); break;
-		default:
-			break;
-		}
-	}
-
-
 	//---------------------------------------------------------------------------------------------
 	// Own Interface
 	//---------------------------------------------------------------------------------------------
diff --git a/Engine/Plugins/Experimental/Mutable/Source/MutableTools/Private/MuT/NodeMeshSwitch.cpp b/Engine/Plugins/Experimental/Mutable/Source/MutableTools/Private/MuT/NodeMeshSwitch.cpp
--- a/Engine/Plugins/Experimental/Mutable/Source/MutableTools/Private/MuT/NodeMeshSwitch.cpp
+++ b/Engine/Plugins/Experimental/Mutable/Source/MutableTools/Private/MuT/NodeMeshSwitch.cpp
@@ -28,55 +28,6 @@ namespace mu
 	MUTABLE_IMPLEMENT_NODE( NodeMeshSwitch, EType::Switch, Node, Node::EType::Mesh)
 
 
-	//---------------------------------------------------------------------------------------------
-	// Node Interface
-	//---------------------------------------------------------------------------------------------
-	int NodeMeshSwitch::GetInputCount() const
-	{
-		return 1 + m_pD->m_options.Num();
-	}
-
-
-	//---------------------------------------------------------------------------------------------
-	Node* NodeMeshSwitch::GetInputNode( int i ) const
-	{
-		check( i>=0 && i<GetInputCount() );
-
-		Node* pResult = 0;
-
-		switch (i)
-		{
-		case 0:
-			pResult = m_pD->m_pParameter.get();
-			break;
-
-		default:
-			pResult = m_pD->m_options[i-1].get();
-			break;
-		}
-
-		return pResult;
-	}
-
-
-	//---------------------------------------------------------------------------------------------
-	void NodeMeshSwitch::SetInputNode( int i, NodePtr pNode )
-	{
-		check( i>=0 && i<GetInputCount() );
-
-		switch (i)
-		{
-		case 0:
-			m_pD->m_pParameter = dynamic_cast<NodeScalar*>(pNode.get());
-			break;
-
-		default:
-			m_pD->m_options[i-1] = dynamic_cast<NodeMesh*>(pNode.get());
-			break;
-		}
-	}
-
-
 	//---------------------------------------------------------------------------------------------
 	// Own Interface
 	//---------------------------------------------------------------------------------------------
diff --git a/Engine/Plugins/Experimental/Mutable/Source/MutableTools/Private/MuT/NodeMeshTable.cpp b/Engine/Plugins/Experimental/Mutable/Source/MutableTools/Private/MuT/NodeMeshTable.cpp
--- a/Engine/Plugins/Experimental/Mutable/Source/MutableTools/Private/MuT/NodeMeshTable.cpp
+++ b/Engine/Plugins/Experimental/Mutable/Source/MutableTools/Private/MuT/NodeMeshTable.cpp
@@ -28,33 +28,6 @@ namespace mu
 	MUTABLE_IMPLEMENT_NODE( NodeMeshTable, EType::Table, Node, Node::EType::Mesh);
 
 
-	//---------------------------------------------------------------------------------------------
-	// Node Interface
-	//---------------------------------------------------------------------------------------------
-	int NodeMeshTable::GetInputCount() const
-	{
-		return 0;
-	}
-
-
-	//---------------------------------------------------------------------------------------------
-	Node* NodeMeshTable::GetInputNode( int i ) const
-	{
-		check( i >=0 && i < 0 );
-        (void)i;
-        return 0;
-	}
-
-
-	//---------------------------------------------------------------------------------------------
-    void NodeMeshTable::SetInputNode( int i, NodePtr )
-	{
-		check( i >=0 && i < 0 );
-        (void)i;
-        //m_pD->m_pObject = dynamic_cast<NodeObject*>( pNode.get() );
-	}
-
-
 	//---------------------------------------------------------------------------------------------
 	// Own Interface
 	//---------------------------------------------------------------------------------------------
diff --git a/Engine/Plugins/Experimental/Mutable/Source/MutableTools/Private/MuT/NodeMeshTangents.cpp b/Engine/Plugins/Experimental/Mutable/Source/MutableTools/Private/MuT/NodeMeshTangents.cpp
--- a/Engine/Plugins/Experimental/Mutable/Source/MutableTools/Private/MuT/NodeMeshTangents.cpp
+++ b/Engine/Plugins/Experimental/Mutable/Source/MutableTools/Private/MuT/NodeMeshTangents.cpp
@@ -27,36 +27,6 @@ namespace mu
 	MUTABLE_IMPLEMENT_NODE( NodeMeshTangents, EType::Tangents, Node, Node::EType::Mesh)
 
 
-	//---------------------------------------------------------------------------------------------
-	// Node Interface
-	//---------------------------------------------------------------------------------------------
-	int NodeMeshTangents::GetInputCount() const
-	{
-		return 1;
-	}
-
-
-	//---------------------------------------------------------------------------------------------
-	Node* NodeMeshTangents::GetInputNode( int i ) const
-	{
-		check( i>=0 && i< GetInputCount());
-        (void)i;
-		return m_pD->m_pSource.get();
-	}
-
-
-	//---------------------------------------------------------------------------------------------
-	void NodeMeshTangents::SetInputNode( int i, NodePtr pNode )
-	{
-		check( i>=0 && i< GetInputCount());
-        (void)i;
-        if (i==0)
-		{
-			m_pD->m_pSource = dynamic_cast<NodeMesh*>( pNode.get() );
-		}
-	}
-
-
 	//---------------------------------------------------------------------------------------------
 	// Own Interface
 	//---------------------------------------------------------------------------------------------
diff --git a/Engine/Plugins/Experimental/Mutable/Source/MutableTools/Private/MuT/NodeMeshTransform.cpp b/Engine/Plugins/Experimental/Mutable/Source/MutableTools/Private/MuT/NodeMeshTransform.cpp
--- a/Engine/Plugins/Experimental/Mutable/Source/MutableTools/Private/MuT/NodeMeshTransform.cpp
+++ b/Engine/Plugins/Experimental/Mutable/Source/MutableTools/Private/MuT/NodeMeshTransform.cpp
@@ -29,35 +29,6 @@ namespace mu
     MUTABLE_IMPLEMENT_NODE( NodeMeshTransform, EType::Transform, Node, Node::EType::Mesh)
 
 
-	//---------------------------------------------------------------------------------------------
-	// Node Interface
-	//---------------------------------------------------------------------------------------------
-    int NodeMeshTransform::GetInputCount() const
-	{
-		return 1;
-	}
-
-
-	//---------------------------------------------------------------------------------------------
-    Node* NodeMeshTransform::GetInputNode( int i ) const
-	{
-		check( i>=0 && i< GetInputCount());
-        (void)i;
-        return m_pD->m_pSource.get();
-	}
-
-
-	//---------------------------------------------------------------------------------------------
-    void NodeMeshTransform::SetInputNode( int i, NodePtr pNode )
-	{
-		check( i>=0 && i< GetInputCount());
-		if (i==0)
-		{
-			m_pD->m_pSource = dynamic_cast<NodeMesh*>( pNode.get() );
-		}
-	}
-
-
 	//---------------------------------------------------------------------------------------------
 	// Own Interface
 	//---------------------------------------------------------------------------------------------
diff --git a/Engine/Plugins/Experimental/Mutable/Source/MutableTools/Private/MuT/NodeMeshVariation.cpp b/Engine/Plugins/Experimental/Mutable/Source/MutableTools/Private/MuT/NodeMeshVariation.cpp
--- a/Engine/Plugins/Experimental/Mutable/Source/MutableTools/Private/MuT/NodeMeshVariation.cpp
+++ b/Engine/Plugins/Experimental/Mutable/Source/MutableTools/Private/MuT/NodeMeshVariation.cpp
@@ -24,53 +24,6 @@ namespace mu
     MUTABLE_IMPLEMENT_NODE( NodeMeshVariation, EType::Variation, Node, Node::EType::Mesh)
 
 
-    //---------------------------------------------------------------------------------------------
-    int NodeMeshVariation::GetInputCount() const { return 1 + m_pD->m_variations.Num(); }
-
-
-    //---------------------------------------------------------------------------------------------
-    Node* NodeMeshVariation::GetInputNode( int i ) const
-    {
-        check( i >= 0 && i < GetInputCount() );
-
-        if ( i == 0 )
-        {
-            return m_pD->m_defaultMesh.get();
-        }
-        i -= 1;
-
-        if ( i < m_pD->m_variations.Num() )
-        {
-            return m_pD->m_variations[i].m_mesh.get();
-        }
-        i -= m_pD->m_variations.Num();
-
-        return nullptr;
-    }
-
-
-    //---------------------------------------------------------------------------------------------
-    void NodeMeshVariation::SetInputNode( int i, NodePtr pNode )
-    {
-        check( i >= 0 && i < GetInputCount() );
-
-        if ( i == 0 )
-        {
-            m_pD->m_defaultMesh = dynamic_cast<NodeMesh*>( pNode.get() );
-            return;
-        }
-
-        i -= 1;
-        if ( i < m_pD->m_variations.Num() )
-        {
-
-            m_pD->m_variations[i].m_mesh = dynamic_cast<NodeMesh*>( pNode.get() );
-            return;
-        }
-        i -= m_pD->m_variations.Num();
-    }
-
-
     //---------------------------------------------------------------------------------------------
     // Own Interface
     //---------------------------------------------------------------------------------------------
diff --git a/Engine/Plugins/Experimental/Mutable/Source/MutableTools/Private/MuT/NodeModifierMeshClipDeform.cpp b/Engine/Plugins/Experimental/Mutable/Source/MutableTools/Private/MuT/NodeModifierMeshClipDeform.cpp
--- a/Engine/Plugins/Experimental/Mutable/Source/MutableTools/Private/MuT/NodeModifierMeshClipDeform.cpp
+++ b/Engine/Plugins/Experimental/Mutable/Source/MutableTools/Private/MuT/NodeModifierMeshClipDeform.cpp
@@ -27,32 +27,6 @@ namespace mu
 	MUTABLE_IMPLEMENT_NODE( NodeModifierMeshClipDeform, EType::MeshClipDeform, Node, Node::EType::Modifier)
 
 
-	//---------------------------------------------------------------------------------------------
-	// Node Interface
-	//---------------------------------------------------------------------------------------------
-	int NodeModifierMeshClipDeform::GetInputCount() const
-	{
-		return 0;
-	}
-
-
-	//---------------------------------------------------------------------------------------------
-	Node* NodeModifierMeshClipDeform::GetInputNode( int i ) const
-	{
-		check( i>=0 && i< GetInputCount());
-		(void)i;
-		return nullptr;
-	}
-
-
-	//---------------------------------------------------------------------------------------------
-	void NodeModifierMeshClipDeform::SetInputNode( int32 i, NodePtr )
-	{
-		check( i>=0 && i< GetInputCount());
-		(void)i;
-	}
-
-
 	//---------------------------------------------------------------------------------------------
 	// Own Interface
 	//---------------------------------------------------------------------------------------------
diff --git a/Engine/Plugins/Experimental/Mutable/Source/MutableTools/Private/MuT/NodeModifierMeshClipMorphPlane.cpp b/Engine/Plugins/Experimental/Mutable/Source/MutableTools/Private/MuT/NodeModifierMeshClipMorphPlane.cpp
--- a/Engine/Plugins/Experimental/Mutable/Source/MutableTools/Private/MuT/NodeModifierMeshClipMorphPlane.cpp
+++ b/Engine/Plugins/Experimental/Mutable/Source/MutableTools/Private/MuT/NodeModifierMeshClipMorphPlane.cpp
@@ -26,32 +26,6 @@ namespace mu
     MUTABLE_IMPLEMENT_NODE( NodeModifierMeshClipMorphPlane, EType::MeshClipMorphPlane, Node, Node::EType::Modifier)
 
 
-	//---------------------------------------------------------------------------------------------
-	// Node Interface
-	//---------------------------------------------------------------------------------------------
-    int NodeModifierMeshClipMorphPlane::GetInputCount() const
-	{
-		return 0;
-	}
-
-
-	//---------------------------------------------------------------------------------------------
-    Node* NodeModifierMeshClipMorphPlane::GetInputNode( int i ) const
-	{
-		check( i>=0 && i< GetInputCount());
-        (void)i;
-		return nullptr;
-	}
-
-
-	//---------------------------------------------------------------------------------------------
-    void NodeModifierMeshClipMorphPlane::SetInputNode( int i, NodePtr )
-	{
-		check( i>=0 && i< GetInputCount());
-        (void)i;
-    }
-
-
 	//---------------------------------------------------------------------------------------------
 	// Own Interface
 	//---------------------------------------------------------------------------------------------
diff --git a/Engine/Plugins/Experimental/Mutable/Source/MutableTools/Private/MuT/NodeModifierMeshClipWithMesh.cpp b/Engine/Plugins/Experimental/Mutable/Source/MutableTools/Private/MuT/NodeModifierMeshClipWithMesh.cpp
--- a/Engine/Plugins/Experimental/Mutable/Source/MutableTools/Private/MuT/NodeModifierMeshClipWithMesh.cpp
+++ b/Engine/Plugins/Experimental/Mutable/Source/MutableTools/Private/MuT/NodeModifierMeshClipWithMesh.cpp
@@ -26,34 +26,6 @@ namespace mu
     MUTABLE_IMPLEMENT_NODE( NodeModifierMeshClipWithMesh, EType::MeshClipWithMesh, Node, Node::EType::Modifier)
 
 
-	//---------------------------------------------------------------------------------------------
-	// Node Interface
-	//---------------------------------------------------------------------------------------------
-    int NodeModifierMeshClipWithMesh::GetInputCount() const
-	{
-		return 1;
-	}
-
-
-	//---------------------------------------------------------------------------------------------
-    Node* NodeModifierMeshClipWithMesh::GetInputNode( int i ) const
-	{
-		check( i>=0 && i< GetInputCount());
-        (void)i;
-        return m_pD->ClipMesh.get();
-	}
-
-
-	//---------------------------------------------------------------------------------------------
-    void NodeModifierMeshClipWithMesh::SetInputNode( int i, Ptr<Node> Input )
-	{
-		check( i>=0 && i< GetInputCount());
-        (void)i;
-		m_pD->ClipMesh = dynamic_cast<NodeMesh*>(Input.get());
-
-    }
-
-
 	//---------------------------------------------------------------------------------------------
 	// Own Interface
 	//---------------------------------------------------------------------------------------------
diff --git a/Engine/Plugins/Experimental/Mutable/Source/MutableTools/Private/MuT/NodeModifierMeshClipWithUVMask.cpp b/Engine/Plugins/Experimental/Mutable/Source/MutableTools/Private/MuT/NodeModifierMeshClipWithUVMask.cpp
--- a/Engine/Plugins/Experimental/Mutable/Source/MutableTools/Private/MuT/NodeModifierMeshClipWithUVMask.cpp
+++ b/Engine/Plugins/Experimental/Mutable/Source/MutableTools/Private/MuT/NodeModifierMeshClipWithUVMask.cpp
@@ -26,33 +26,6 @@ namespace mu
     MUTABLE_IMPLEMENT_NODE( NodeModifierMeshClipWithUVMask, EType::MeshClipWithMesh, Node, Node::EType::Modifier)
 
 
-	//---------------------------------------------------------------------------------------------
-	// Node Interface
-	//---------------------------------------------------------------------------------------------
-    int NodeModifierMeshClipWithUVMask::GetInputCount() const
-	{
-		return 1;
-	}
-
-
-	//---------------------------------------------------------------------------------------------
-    Node* NodeModifierMeshClipWithUVMask::GetInputNode( int i ) const
-	{
-		check( i>=0 && i< GetInputCount());
-        (void)i;
-        return m_pD->ClipMask.get();
-	}
-
-
-	//---------------------------------------------------------------------------------------------
-    void NodeModifierMeshClipWithUVMask::SetInputNode( int i, Ptr<Node> Input )
-	{
-		check( i>=0 && i< GetInputCount());
-        (void)i;
-		m_pD->ClipMask = dynamic_cast<NodeImage*>(Input.get());
-    }
-
-
 	//---------------------------------------------------------------------------------------------
 	// Own Interface
 	//---------------------------------------------------------------------------------------------
diff --git a/Engine/Plugins/Experimental/Mutable/Source/MutableTools/Private/MuT/NodeObjectGroup.cpp b/Engine/Plugins/Experimental/Mutable/Source/MutableTools/Private/MuT/NodeObjectGroup.cpp
--- a/Engine/Plugins/Experimental/Mutable/Source/MutableTools/Private/MuT/NodeObjectGroup.cpp
+++ b/Engine/Plugins/Experimental/Mutable/Source/MutableTools/Private/MuT/NodeObjectGroup.cpp
@@ -33,43 +33,6 @@ namespace mu
 	MUTABLE_IMPLEMENT_NODE( NodeObjectGroup, EType::Group, Node, Node::EType::Object)
 
 
-	//---------------------------------------------------------------------------------------------
-	// Node Interface
-	//---------------------------------------------------------------------------------------------
-	int NodeObjectGroup::GetInputCount() const
-	{
-		return m_pD->m_children.Num();
-	}
-
-
-	//---------------------------------------------------------------------------------------------
-	Node* NodeObjectGroup::GetInputNode( int i ) const
-	{
-		check( i>=0 && i<GetInputCount() );
-
-		Node* pResult = 0;
-
-		if ( i<m_pD->m_children.Num() )
-		{
-			pResult = m_pD->m_children[i].get();
-		}
-
-		return pResult;
-	}
-
-
-	//---------------------------------------------------------------------------------------------
-	void NodeObjectGroup::SetInputNode( int i, NodePtr pNode )
-	{
-		check( i>=0 && i<GetInputCount() );
-
-		if ( i<m_pD->m_children.Num() )
-		{
-			m_pD->m_children[i] = dynamic_cast<NodeObject*>( pNode.get() );
-		}
-	}
-
-
 	//---------------------------------------------------------------------------------------------
 	// Own Interface
 	//---------------------------------------------------------------------------------------------
diff --git a/Engine/Plugins/Experimental/Mutable/Source/MutableTools/Private/MuT/NodeObjectNew.cpp b/Engine/Plugins/Experimental/Mutable/Source/MutableTools/Private/MuT/NodeObjectNew.cpp
--- a/Engine/Plugins/Experimental/Mutable/Source/MutableTools/Private/MuT/NodeObjectNew.cpp
+++ b/Engine/Plugins/Experimental/Mutable/Source/MutableTools/Private/MuT/NodeObjectNew.cpp
@@ -40,71 +40,6 @@ namespace mu
 	MUTABLE_IMPLEMENT_NODE( NodeObjectNew, EType::New, Node, Node::EType::Object )
 
 
-	//---------------------------------------------------------------------------------------------
-	// Node Interface
-	//---------------------------------------------------------------------------------------------
-	int NodeObjectNew::GetInputCount() const
-	{
-		return (m_pD->m_lods.Num() + m_pD->m_children.Num());
-	}
-
-
-	//---------------------------------------------------------------------------------------------
-	Node* NodeObjectNew::GetInputNode( int i ) const
-	{
-		check( i>=0 && i<GetInputCount() );
-
-		Node* pResult = 0;
-
-		if ( i<(int)m_pD->m_lods.Num() )
-		{
-			pResult = m_pD->m_lods[i].get();
-		}
-		else
-		{
-			i -= (int)m_pD->m_lods.Num();
-			
-			if ( i<(int) m_pD->m_children.Num() )
-			{
-				pResult = m_pD->m_children[i].get();
-			}
-			else
-			{
-				i -= (int)m_pD->m_children.Num();
-
-				pResult = m_pD->m_extensionDataNodes[i].Node.get();
-			}
-		}
-
-		return pResult;
-	}
-
-
-	//---------------------------------------------------------------------------------------------
-	void NodeObjectNew::SetInputNode( int i, NodePtr pNode )
-	{
-		check( i>=0 && i<GetInputCount() );
-
-		if ( i<(int)m_pD->m_lods.Num() )
-		{
-			m_pD->m_lods[i] = dynamic_cast<NodeLOD*>( pNode.get() );
-		}
-		else
-		{
-			i -= (int)m_pD->m_lods.Num();
-
-			if (i < (int) m_pD->m_children.Num())
-			{
-				m_pD->m_children[i] = dynamic_cast<NodeObject*>( pNode.get() );
-			}
-			else
-			{
-				checkf(false, TEXT("Trying to set an input node for an extension input data. Please use AddExtensionDataNode for this"));
-			}
-		}
-	}
-
-
 	//---------------------------------------------------------------------------------------------
 	// Own Interface
 	//---------------------------------------------------------------------------------------------
diff --git a/Engine/Plugins/Experimental/Mutable/Source/MutableTools/Private/MuT/NodePatchImage.cpp b/Engine/Plugins/Experimental/Mutable/Source/MutableTools/Private/MuT/NodePatchImage.cpp
--- a/Engine/Plugins/Experimental/Mutable/Source/MutableTools/Private/MuT/NodePatchImage.cpp
+++ b/Engine/Plugins/Experimental/Mutable/Source/MutableTools/Private/MuT/NodePatchImage.cpp
@@ -30,51 +30,6 @@ namespace mu
     MUTABLE_IMPLEMENT_NODE( NodePatchImage, EType::PatchImage, Node, Node::EType::PatchImage);
 
 
-    //---------------------------------------------------------------------------------------------
-    // Node Interface
-    //---------------------------------------------------------------------------------------------
-    int NodePatchImage::GetInputCount() const
-    {
-        return 2;
-    }
-
-
-    //---------------------------------------------------------------------------------------------
-    Node* NodePatchImage::GetInputNode( int i ) const
-    {
-        check( i >=0 && i < GetInputCount());
-
-        Node* pResult = 0;
-
-        if ( i==0 )
-        {
-            pResult = m_pD->m_pImage.get();
-        }
-        else
-        {
-            pResult = m_pD->m_pMask.get();
-        }
-
-        return pResult;
-    }
-
-
-    //---------------------------------------------------------------------------------------------
-    void NodePatchImage::SetInputNode( int i, NodePtr pNode )
-    {
-        check( i >=0 && i < GetInputCount());
-
-        if ( i==0 )
-        {
-            m_pD->m_pImage = dynamic_cast<NodeImage*>( pNode.get() );
-        }
-        else
-        {
-            m_pD->m_pMask = dynamic_cast<NodeImage*>( pNode.get() );
-        }
-    }
-
-
     //---------------------------------------------------------------------------------------------
     // Own Interface
     //---------------------------------------------------------------------------------------------
diff --git a/Engine/Plugins/Experimental/Mutable/Source/MutableTools/Private/MuT/NodePatchMesh.cpp b/Engine/Plugins/Experimental/Mutable/Source/MutableTools/Private/MuT/NodePatchMesh.cpp
--- a/Engine/Plugins/Experimental/Mutable/Source/MutableTools/Private/MuT/NodePatchMesh.cpp
+++ b/Engine/Plugins/Experimental/Mutable/Source/MutableTools/Private/MuT/NodePatchMesh.cpp
@@ -26,55 +26,6 @@ namespace mu
 	MUTABLE_IMPLEMENT_NODE( NodePatchMesh, EType::PatchMesh, Node, Node::EType::PatchMesh);
 
 
-	//---------------------------------------------------------------------------------------------
-	// Node Interface
-	//---------------------------------------------------------------------------------------------
-	int NodePatchMesh::GetInputCount() const
-	{
-		return 2;
-	}
-
-
-	//---------------------------------------------------------------------------------------------
-	Node* NodePatchMesh::GetInputNode( int i ) const
-	{
-		Node* pResult = 0;
-
-		if( i==0 )
-		{
-			pResult = m_pD->m_pRemove.get();
-		}
-		else if (i==1)
-		{
-			pResult = m_pD->m_pAdd.get();
-		}
-		else
-		{
-			check( false );
-		}
-
-		return pResult;
-	}
-
-
-	//---------------------------------------------------------------------------------------------
-	void NodePatchMesh::SetInputNode( int i, NodePtr pNode )
-	{
-		if( i==0 )
-		{
-			m_pD->m_pRemove = dynamic_cast<NodeMesh*>( pNode.get() );
-		}
-		else if (i==1)
-		{
-			m_pD->m_pAdd = dynamic_cast<NodeMesh*>( pNode.get() );
-		}
-		else
-		{
-			check( false );
-		}
-	}
-
-
 	//---------------------------------------------------------------------------------------------
 	// Own Interface
 	//---------------------------------------------------------------------------------------------
diff --git a/Engine/Plugins/Experimental/Mutable/Source/MutableTools/Private/MuT/NodeProjectorConstant.cpp b/Engine/Plugins/Experimental/Mutable/Source/MutableTools/Private/MuT/NodeProjectorConstant.cpp
--- a/Engine/Plugins/Experimental/Mutable/Source/MutableTools/Private/MuT/NodeProjectorConstant.cpp
+++ b/Engine/Plugins/Experimental/Mutable/Source/MutableTools/Private/MuT/NodeProjectorConstant.cpp
@@ -27,30 +27,6 @@ namespace mu
 	MUTABLE_IMPLEMENT_NODE( NodeProjectorConstant, EType::Constant, Node, Node::EType::Projector)
 
 
-	//---------------------------------------------------------------------------------------------
-	// Node Interface
-	//---------------------------------------------------------------------------------------------
-	int NodeProjectorConstant::GetInputCount() const
-	{
-		return 0;
-	}
-
-
-	//---------------------------------------------------------------------------------------------
-    Node* NodeProjectorConstant::GetInputNode( int ) const
-	{
-		check( false );
-		return 0;
-	}
-
-
-	//---------------------------------------------------------------------------------------------
-    void NodeProjectorConstant::SetInputNode( int, NodePtr )
-	{
-		check( false );
-	}
-
-
 	//---------------------------------------------------------------------------------------------
 	// Own Interface
 	//---------------------------------------------------------------------------------------------
diff --git a/Engine/Plugins/Experimental/Mutable/Source/MutableTools/Private/MuT/NodeProjectorParameter.cpp b/Engine/Plugins/Experimental/Mutable/Source/MutableTools/Private/MuT/NodeProjectorParameter.cpp
--- a/Engine/Plugins/Experimental/Mutable/Source/MutableTools/Private/MuT/NodeProjectorParameter.cpp
+++ b/Engine/Plugins/Experimental/Mutable/Source/MutableTools/Private/MuT/NodeProjectorParameter.cpp
@@ -30,38 +30,6 @@ namespace mu
 	MUTABLE_IMPLEMENT_NODE( NodeProjectorParameter, EType::Parameter, Node, Node::EType::Projector)
 
 
-	//---------------------------------------------------------------------------------------------
-	// Node Interface
-	//---------------------------------------------------------------------------------------------
-	int NodeProjectorParameter::GetInputCount() const
-	{
-        return m_pD->m_ranges.Num();
-    }
-
-
-	//---------------------------------------------------------------------------------------------
-    Node* NodeProjectorParameter::GetInputNode( int i ) const
-	{
-        check( i<GetInputCount() );
-        if (i<GetInputCount())
-        {
-            return m_pD->m_ranges[i].get();
-        }
-        return nullptr;
-    }
-
-
-	//---------------------------------------------------------------------------------------------
-    void NodeProjectorParameter::SetInputNode( int i, NodePtr n )
-	{
-        check( i<GetInputCount() );
-        if (i<GetInputCount())
-        {
-            m_pD->m_ranges[i] = dynamic_cast<NodeRange*>(n.get());
-        }
-    }
-
-
 	//---------------------------------------------------------------------------------------------
 	// Own Interface
 	//---------------------------------------------------------------------------------------------
diff --git a/Engine/Plugins/Experimental/Mutable/Source/MutableTools/Private/MuT/NodeRangeFromScalar.cpp b/Engine/Plugins/Experimental/Mutable/Source/MutableTools/Private/MuT/NodeRangeFromScalar.cpp
--- a/Engine/Plugins/Experimental/Mutable/Source/MutableTools/Private/MuT/NodeRangeFromScalar.cpp
+++ b/Engine/Plugins/Experimental/Mutable/Source/MutableTools/Private/MuT/NodeRangeFromScalar.cpp
@@ -27,43 +27,6 @@ namespace mu
     MUTABLE_IMPLEMENT_NODE( NodeRangeFromScalar, EType::FromScalar, Node, Node::EType::Range)
 
 
-	//---------------------------------------------------------------------------------------------
-	// Node Interface
-	//---------------------------------------------------------------------------------------------
-    int NodeRangeFromScalar::GetInputCount() const
-	{
-		return 1;
-	}
-
-
-	//---------------------------------------------------------------------------------------------
-    Node* NodeRangeFromScalar::GetInputNode( int i ) const
-	{
-		check( i>=0 && i< GetInputCount());
-
-        Node* pResult = nullptr;
-
-		switch (i)
-		{
-        case 0: pResult = m_pD->m_pSize.get(); break;
-		}
-
-		return pResult;
-	}
-
-
-	//---------------------------------------------------------------------------------------------
-    void NodeRangeFromScalar::SetInputNode( int i, NodePtr pNode )
-	{
-		check( i>=0 && i< GetInputCount());
-
-		switch (i)
-		{
-        case 0: m_pD->m_pSize = dynamic_cast<NodeScalar*>(pNode.get()); break;
-		}
-	}
-
-
 	//---------------------------------------------------------------------------------------------
 	// Own Interface
 	//---------------------------------------------------------------------------------------------
diff --git a/Engine/Plugins/Experimental/Mutable/Source/MutableTools/Private/MuT/NodeScalarArithmeticOperation.cpp b/Engine/Plugins/Experimental/Mutable/Source/MutableTools/Private/MuT/NodeScalarArithmeticOperation.cpp
--- a/Engine/Plugins/Experimental/Mutable/Source/MutableTools/Private/MuT/NodeScalarArithmeticOperation.cpp
+++ b/Engine/Plugins/Experimental/Mutable/Source/MutableTools/Private/MuT/NodeScalarArithmeticOperation.cpp
@@ -39,45 +39,6 @@ namespace mu
     };
 
 
-	//---------------------------------------------------------------------------------------------
-	// Node Interface
-	//---------------------------------------------------------------------------------------------
-    int NodeScalarArithmeticOperation::GetInputCount() const
-	{
-		return 2;
-	}
-
-
-	//---------------------------------------------------------------------------------------------
-    Node* NodeScalarArithmeticOperation::GetInputNode( int i ) const
-	{
-		check( i>=0 && i< GetInputCount());
-
-		Node* pResult = 0;
-
-		switch (i)
-		{
-		case 0: pResult = m_pD->m_pA.get(); break;
-		case 1: pResult = m_pD->m_pB.get(); break;
-		}
-
-		return pResult;
-	}
-
-
-	//---------------------------------------------------------------------------------------------
-    void NodeScalarArithmeticOperation::SetInputNode( int i, NodePtr pNode )
-	{
-		check( i>=0 && i< GetInputCount());
-
-		switch (i)
-		{
-        case 0: m_pD->m_pA = dynamic_cast<NodeScalar*>(pNode.get()); break;
-        case 1: m_pD->m_pB = dynamic_cast<NodeScalar*>(pNode.get()); break;
-		}
-	}
-
-
 	//---------------------------------------------------------------------------------------------
 	// Own Interface
 	//---------------------------------------------------------------------------------------------
diff --git a/Engine/Plugins/Experimental/Mutable/Source/MutableTools/Private/MuT/NodeScalarConstant.cpp b/Engine/Plugins/Experimental/Mutable/Source/MutableTools/Private/MuT/NodeScalarConstant.cpp
--- a/Engine/Plugins/Experimental/Mutable/Source/MutableTools/Private/MuT/NodeScalarConstant.cpp
+++ b/Engine/Plugins/Experimental/Mutable/Source/MutableTools/Private/MuT/NodeScalarConstant.cpp
@@ -25,30 +25,6 @@ namespace mu
 	MUTABLE_IMPLEMENT_NODE( NodeScalarConstant, EType::Constant, Node, Node::EType::Scalar)
 
 
-	//---------------------------------------------------------------------------------------------
-	// Node Interface
-	//---------------------------------------------------------------------------------------------
-	int NodeScalarConstant::GetInputCount() const
-	{
-		return 0;
-	}
-
-
-	//---------------------------------------------------------------------------------------------
-    Node* NodeScalarConstant::GetInputNode( int ) const
-	{
-		check( false );
-		return 0;
-	}
-
-
-	//---------------------------------------------------------------------------------------------
-    void NodeScalarConstant::SetInputNode( int, NodePtr )
-	{
-		check( false );
-	}
-
-
 	//---------------------------------------------------------------------------------------------
 	// Own Interface
 	//---------------------------------------------------------------------------------------------
diff --git a/Engine/Plugins/Experimental/Mutable/Source/MutableTools/Private/MuT/NodeScalarCurve.cpp b/Engine/Plugins/Experimental/Mutable/Source/MutableTools/Private/MuT/NodeScalarCurve.cpp
--- a/Engine/Plugins/Experimental/Mutable/Source/MutableTools/Private/MuT/NodeScalarCurve.cpp
+++ b/Engine/Plugins/Experimental/Mutable/Source/MutableTools/Private/MuT/NodeScalarCurve.cpp
@@ -28,35 +28,6 @@ namespace mu
     MUTABLE_IMPLEMENT_NODE( NodeScalarCurve, EType::Curve, Node, Node::EType::Scalar)
 
 
-	//---------------------------------------------------------------------------------------------
-	// Node Interface
-	//---------------------------------------------------------------------------------------------
-    int NodeScalarCurve::GetInputCount() const
-	{
-		return 1;
-	}
-
-
-	//---------------------------------------------------------------------------------------------
-    Node* NodeScalarCurve::GetInputNode( int i ) const
-	{
-		check( i >=0 && i < GetInputCount() );
-        (void)i;
-
-		return m_pD->m_input_scalar.get();
-	}
-
-
-	//---------------------------------------------------------------------------------------------
-    void NodeScalarCurve::SetInputNode( int i, NodePtr pNode )
-	{
-		check( i >=0 && i < GetInputCount());
-        (void)i;
-
-		m_pD->m_input_scalar = dynamic_cast<NodeScalar*>(pNode.get());
-	}
-
-
 	//---------------------------------------------------------------------------------------------
 	// Own Interface
 	//---------------------------------------------------------------------------------------------
diff --git a/Engine/Plugins/Experimental/Mutable/Source/MutableTools/Private/MuT/NodeScalarEnumParameter.cpp b/Engine/Plugins/Experimental/Mutable/Source/MutableTools/Private/MuT/NodeScalarEnumParameter.cpp
--- a/Engine/Plugins/Experimental/Mutable/Source/MutableTools/Private/MuT/NodeScalarEnumParameter.cpp
+++ b/Engine/Plugins/Experimental/Mutable/Source/MutableTools/Private/MuT/NodeScalarEnumParameter.cpp
@@ -25,36 +25,6 @@ namespace mu
 	MUTABLE_IMPLEMENT_NODE( NodeScalarEnumParameter, EType::EnumParameter, Node, Node::EType::Scalar)
 
 
-	//---------------------------------------------------------------------------------------------
-	int NodeScalarEnumParameter::GetInputCount() const
-	{
-        return m_pD->m_ranges.Num();
-	}
-
-
-	//---------------------------------------------------------------------------------------------
-    Node* NodeScalarEnumParameter::GetInputNode( int i ) const
-	{
-        check( i<GetInputCount() );
-        if (i<GetInputCount())
-        {
-            return m_pD->m_ranges[i].get();
-        }
-        return nullptr;
-	}
-
-
-	//---------------------------------------------------------------------------------------------
-    void NodeScalarEnumParameter::SetInputNode( int i, NodePtr n )
-	{
-        check( i<GetInputCount() );
-        if (i<GetInputCount())
-        {
-            m_pD->m_ranges[i] = dynamic_cast<NodeRange*>(n.get());
-        }
-	}
-
-
 	//---------------------------------------------------------------------------------------------
 	const FString& NodeScalarEnumParameter::GetName() const
 	{
diff --git a/Engine/Plugins/Experimental/Mutable/Source/MutableTools/Private/MuT/NodeScalarParameter.cpp b/Engine/Plugins/Experimental/Mutable/Source/MutableTools/Private/MuT/NodeScalarParameter.cpp
--- a/Engine/Plugins/Experimental/Mutable/Source/MutableTools/Private/MuT/NodeScalarParameter.cpp
+++ b/Engine/Plugins/Experimental/Mutable/Source/MutableTools/Private/MuT/NodeScalarParameter.cpp
@@ -27,39 +27,6 @@ namespace mu
 	MUTABLE_IMPLEMENT_NODE( NodeScalarParameter, EType::Parameter, Node, Node::EType::Scalar)
 
 
-	//---------------------------------------------------------------------------------------------
-	int NodeScalarParameter::GetInputCount() const
-	{
-        return m_pD->m_ranges.Num();
-	}
-
-
-	//---------------------------------------------------------------------------------------------
-    Node* NodeScalarParameter::GetInputNode( int i ) const
-	{
-        check( i<GetInputCount() );
-
-        int rangeCount = int(m_pD->m_ranges.Num());
-		if ( i < rangeCount )
-        {
-            return m_pD->m_ranges[i].get();
-        }
-        return nullptr;
-	}
-
-
-	//---------------------------------------------------------------------------------------------
-    void NodeScalarParameter::SetInputNode( int i, NodePtr n )
-	{
-        check( i<GetInputCount() );
-        int rangeCount = int(m_pD->m_ranges.Num());
-        if ( i < rangeCount )
-        {
-            m_pD->m_ranges[i] = dynamic_cast<NodeRange*>(n.get());
-        }
-    }
-
-
 	//---------------------------------------------------------------------------------------------
 	void NodeScalarParameter::SetName( const FString& Name )
 	{
diff --git a/Engine/Plugins/Experimental/Mutable/Source/MutableTools/Private/MuT/NodeScalarSwitch.cpp b/Engine/Plugins/Experimental/Mutable/Source/MutableTools/Private/MuT/NodeScalarSwitch.cpp
--- a/Engine/Plugins/Experimental/Mutable/Source/MutableTools/Private/MuT/NodeScalarSwitch.cpp
+++ b/Engine/Plugins/Experimental/Mutable/Source/MutableTools/Private/MuT/NodeScalarSwitch.cpp
@@ -29,55 +29,6 @@ namespace mu
     MUTABLE_IMPLEMENT_NODE( NodeScalarSwitch, EType::Switch, Node, Node::EType::Scalar)
 
 
-	//---------------------------------------------------------------------------------------------
-	// Node Interface
-	//---------------------------------------------------------------------------------------------
-    int NodeScalarSwitch::GetInputCount() const
-	{
-		return 1 + (int)m_pD->m_options.Num();
-	}
-
-
-	//---------------------------------------------------------------------------------------------
-    Node* NodeScalarSwitch::GetInputNode( int i ) const
-	{
-		check( i>=0 && i<GetInputCount() );
-
-		Node* pResult = 0;
-
-		switch (i)
-		{
-		case 0:
-			pResult = m_pD->m_pParameter.get();
-			break;
-
-		default:
-			pResult = m_pD->m_options[i-1].get();
-			break;
-		}
-
-		return pResult;
-	}
-
-
-	//---------------------------------------------------------------------------------------------
-    void NodeScalarSwitch::SetInputNode( int i, NodePtr pNode )
-	{
-		check( i>=0 && i<GetInputCount() );
-
-		switch (i)
-		{
-		case 0:
-			m_pD->m_pParameter = dynamic_cast<NodeScalar*>(pNode.get());
-			break;
-
-		default:
-            m_pD->m_options[i-1] = dynamic_cast<NodeScalar*>(pNode.get());
-			break;
-		}
-	}
-
-
 	//---------------------------------------------------------------------------------------------
 	// Own Interface
 	//---------------------------------------------------------------------------------------------
diff --git a/Engine/Plugins/Experimental/Mutable/Source/MutableTools/Private/MuT/NodeScalarTable.cpp b/Engine/Plugins/Experimental/Mutable/Source/MutableTools/Private/MuT/NodeScalarTable.cpp
--- a/Engine/Plugins/Experimental/Mutable/Source/MutableTools/Private/MuT/NodeScalarTable.cpp
+++ b/Engine/Plugins/Experimental/Mutable/Source/MutableTools/Private/MuT/NodeScalarTable.cpp
@@ -27,30 +27,6 @@ namespace mu
 	MUTABLE_IMPLEMENT_NODE(NodeScalarTable, EType::Table, Node, Node::EType::Scalar)
 
 
-	//---------------------------------------------------------------------------------------------
-	// Node Interface
-	//---------------------------------------------------------------------------------------------
-	int NodeScalarTable::GetInputCount() const
-	{
-		return 0;
-	}
-
-
-	//---------------------------------------------------------------------------------------------
-	Node* NodeScalarTable::GetInputNode( int ) const
-	{
-		check( false);
-        return nullptr;
-	}
-
-
-	//---------------------------------------------------------------------------------------------
-    void NodeScalarTable::SetInputNode( int, NodePtr )
-	{
-		check( false);
-	}
-
-
 	//---------------------------------------------------------------------------------------------
 	// Own Interface
 	//---------------------------------------------------------------------------------------------
diff --git a/Engine/Plugins/Experimental/Mutable/Source/MutableTools/Private/MuT/NodeScalarVariation.cpp b/Engine/Plugins/Experimental/Mutable/Source/MutableTools/Private/MuT/NodeScalarVariation.cpp
--- a/Engine/Plugins/Experimental/Mutable/Source/MutableTools/Private/MuT/NodeScalarVariation.cpp
+++ b/Engine/Plugins/Experimental/Mutable/Source/MutableTools/Private/MuT/NodeScalarVariation.cpp
@@ -25,53 +25,6 @@ namespace mu
     MUTABLE_IMPLEMENT_NODE( NodeScalarVariation, EType::Variation, Node, Node::EType::Scalar)
 
 
-    //---------------------------------------------------------------------------------------------
-    int NodeScalarVariation::GetInputCount() const { return 1 + int( m_pD->m_variations.Num() ); }
-
-
-    //---------------------------------------------------------------------------------------------
-    Node* NodeScalarVariation::GetInputNode( int i ) const
-    {
-        check( i >= 0 && i < GetInputCount() );
-
-        if ( i == 0 )
-        {
-            return m_pD->m_defaultScalar.get();
-        }
-        i -= 1;
-
-        if ( i < int( m_pD->m_variations.Num() ) )
-        {
-            return m_pD->m_variations[i].m_scalar.get();
-        }
-        i -= int( m_pD->m_variations.Num() );
-
-        return nullptr;
-    }
-
-
-    //---------------------------------------------------------------------------------------------
-    void NodeScalarVariation::SetInputNode( int i, NodePtr pNode )
-    {
-        check( i >= 0 && i < GetInputCount() );
-
-        if ( i == 0 )
-        {
-            m_pD->m_defaultScalar = dynamic_cast<NodeScalar*>( pNode.get() );
-            return;
-        }
-
-        i -= 1;
-        if ( i < int( m_pD->m_variations.Num() ) )
-        {
-
-            m_pD->m_variations[i].m_scalar = dynamic_cast<NodeScalar*>( pNode.get() );
-            return;
-        }
-        i -= (int)m_pD->m_variations.Num();
-    }
-
-
     //---------------------------------------------------------------------------------------------
     // Own Interface
     //---------------------------------------------------------------------------------------------
diff --git a/Engine/Plugins/Experimental/Mutable/Source/MutableTools/Private/MuT/NodeStringConstant.cpp b/Engine/Plugins/Experimental/Mutable/Source/MutableTools/Private/MuT/NodeStringConstant.cpp
--- a/Engine/Plugins/Experimental/Mutable/Source/MutableTools/Private/MuT/NodeStringConstant.cpp
+++ b/Engine/Plugins/Experimental/Mutable/Source/MutableTools/Private/MuT/NodeStringConstant.cpp
@@ -25,30 +25,6 @@ namespace mu
 	MUTABLE_IMPLEMENT_NODE( NodeStringConstant, EType::Constant, Node, Node::EType::String)
 
 
-	//---------------------------------------------------------------------------------------------
-	// Node Interface
-	//---------------------------------------------------------------------------------------------
-	int NodeStringConstant::GetInputCount() const
-	{
-		return 0;
-	}
-
-
-	//---------------------------------------------------------------------------------------------
-    Node* NodeStringConstant::GetInputNode( int ) const
-	{
-		check( false );
-		return 0;
-	}
-
-
-	//---------------------------------------------------------------------------------------------
-    void NodeStringConstant::SetInputNode( int, NodePtr )
-	{
-		check( false );
-	}
-
-
 	//---------------------------------------------------------------------------------------------
 	// Own Interface
 	//---------------------------------------------------------------------------------------------
diff --git a/Engine/Plugins/Experimental/Mutable/Source/MutableTools/Private/MuT/NodeStringParameter.cpp b/Engine/Plugins/Experimental/Mutable/Source/MutableTools/Private/MuT/NodeStringParameter.cpp
--- a/Engine/Plugins/Experimental/Mutable/Source/MutableTools/Private/MuT/NodeStringParameter.cpp
+++ b/Engine/Plugins/Experimental/Mutable/Source/MutableTools/Private/MuT/NodeStringParameter.cpp
@@ -27,53 +27,6 @@ namespace mu
 	MUTABLE_IMPLEMENT_NODE( NodeStringParameter, EType::Parameter, Node, Node::EType::String)
 
 
-	//---------------------------------------------------------------------------------------------
-	int NodeStringParameter::GetInputCount() const
-	{
-        return int( m_pD->m_additionalImages.Num()
-                    +
-                    m_pD->m_ranges.Num() );
-	}
-
-
-	//---------------------------------------------------------------------------------------------
-    Node* NodeStringParameter::GetInputNode( int i ) const
-	{
-        check( i<GetInputCount() );
-
-        int imageCount = int(m_pD->m_additionalImages.Num());
-        int rangeCount = int(m_pD->m_ranges.Num());
-        if (i<imageCount)
-        {
-            return m_pD->m_additionalImages[i].get();
-        }
-        else if ( i < imageCount + rangeCount )
-        {
-            int r = i - imageCount;
-            return m_pD->m_ranges[r].get();
-        }
-        return nullptr;
-	}
-
-
-	//---------------------------------------------------------------------------------------------
-    void NodeStringParameter::SetInputNode( int i, NodePtr n )
-	{
-        check( i<GetInputCount() );
-        int imageCount = int(m_pD->m_additionalImages.Num());
-        int rangeCount = int(m_pD->m_ranges.Num());
-        if (i<imageCount)
-        {
-            m_pD->m_additionalImages[i] = dynamic_cast<NodeImage*>(n.get());
-        }
-        else if ( i < imageCount + rangeCount )
-        {
-            int r = i - imageCount;
-            m_pD->m_ranges[r] = dynamic_cast<NodeRange*>(n.get());
-        }
-    }
-
-
 	//---------------------------------------------------------------------------------------------
 	void NodeStringParameter::SetName( const FString& Name )
 	{
diff --git a/Engine/Plugins/Experimental/Mutable/Source/MutableTools/Private/MuT/NodeSurfaceEdit.cpp b/Engine/Plugins/Experimental/Mutable/Source/MutableTools/Private/MuT/NodeSurfaceEdit.cpp
--- a/Engine/Plugins/Experimental/Mutable/Source/MutableTools/Private/MuT/NodeSurfaceEdit.cpp
+++ b/Engine/Plugins/Experimental/Mutable/Source/MutableTools/Private/MuT/NodeSurfaceEdit.cpp
@@ -31,88 +31,6 @@ namespace mu
     MUTABLE_IMPLEMENT_NODE( NodeSurfaceEdit, EType::Edit, Node, Node::EType::Surface)
 
 
-	//---------------------------------------------------------------------------------------------
-	// Node Interface
-	//---------------------------------------------------------------------------------------------
-    int NodeSurfaceEdit::GetInputCount() const
-	{
-		return
-			3 	// Extend mesh, morph, factor
-			+
-			m_pD->m_textures.Num()*2
-			;
-	}
-
-
-	//---------------------------------------------------------------------------------------------
-    Node* NodeSurfaceEdit::GetInputNode( int i ) const
-	{
-		check( i >=0 && i < GetInputCount() );
-
-		Node* pResult = 0;
-
-		if ( i==0 )
-		{
-			pResult = m_pD->m_pMesh.get();
-		}
-		else if ( i==1 )
-		{
-			pResult = m_pD->m_pMorph.get();
-		}
-		else if (i == 2)
-		{
-			pResult = m_pD->m_pFactor.get();
-		}
-		else
-		{
-			int idx = i - 3;
-
-			if (idx%2 == 0)
-			{
-				pResult = m_pD->m_textures[idx/2].m_pExtend.get();
-			}
-			else
-			{
-				pResult = m_pD->m_textures[idx/2].m_pPatch.get();
-			}
-		}
-
-		return pResult;
-	}
-
-
-	//---------------------------------------------------------------------------------------------
-    void NodeSurfaceEdit::SetInputNode( int i, NodePtr pNode )
-	{
-		check( i >=0 && i < GetInputCount());
-
-		if ( i==0 )
-		{
-			m_pD->m_pMesh = dynamic_cast<NodePatchMesh*>( pNode.get() );
-		}
-		else if ( i==1 )
-		{
-			m_pD->m_pMorph = dynamic_cast<NodeMesh*>( pNode.get() );
-		}
-		else if ( i==2 )
-		{
-			m_pD->m_pFactor = dynamic_cast<NodeScalar*>( pNode.get() );
-		}
-		else
-		{
-			int idx = i - 3;
-			if (idx % 2==0)
-			{
-				m_pD->m_textures[idx/2].m_pExtend = dynamic_cast<NodeImage*>( pNode.get() );
-			}
-			else
-			{
-				m_pD->m_textures[idx/2].m_pPatch = dynamic_cast<NodePatchImage*>( pNode.get() );
-			}
-		}
-	}
-
-
 	//---------------------------------------------------------------------------------------------
 	// Own Interface
 	//---------------------------------------------------------------------------------------------
diff --git a/Engine/Plugins/Experimental/Mutable/Source/MutableTools/Private/MuT/NodeSurfaceNew.cpp b/Engine/Plugins/Experimental/Mutable/Source/MutableTools/Private/MuT/NodeSurfaceNew.cpp
--- a/Engine/Plugins/Experimental/Mutable/Source/MutableTools/Private/MuT/NodeSurfaceNew.cpp
+++ b/Engine/Plugins/Experimental/Mutable/Source/MutableTools/Private/MuT/NodeSurfaceNew.cpp
@@ -33,132 +33,6 @@ namespace mu
     MUTABLE_IMPLEMENT_NODE( NodeSurfaceNew, EType::New, Node, Node::EType::Surface)
 
 
-	//---------------------------------------------------------------------------------------------
-	// Node Interface
-	//---------------------------------------------------------------------------------------------
-    int NodeSurfaceNew::GetInputCount() const
-	{
-        return m_pD->m_meshes.Num() + m_pD->m_images.Num() + m_pD->m_vectors.Num() +
-                      m_pD->m_scalars.Num() + m_pD->m_strings.Num();
-    }
-
-
-	//---------------------------------------------------------------------------------------------
-    mu::Node* NodeSurfaceNew::GetInputNode( int i ) const
-	{
-		check( i >=0 && i < GetInputCount() );
-
-        int32 index = i;
-
-        if ( index < m_pD->m_meshes.Num() )
-        {
-            return m_pD->m_meshes[index].m_pMesh.get();
-        }
-		else
-		{
-            index -= m_pD->m_meshes.Num();
-        }
-
-        if ( index < m_pD->m_images.Num() )
-        {
-            return m_pD->m_images[index].m_pImage.get();
-        }
-        else 
-		{
-            index -= m_pD->m_images.Num();
-        }
-
-        if ( index < m_pD->m_vectors.Num() )
-        {
-            return m_pD->m_vectors[index].m_pVector.get();
-        }
-        else
-        {
-            index -= m_pD->m_vectors.Num();
-        }
-
-        if ( index < m_pD->m_scalars.Num() )
-        {
-            return m_pD->m_scalars[index].m_pScalar.get();
-        }
-        else
-        {
-            index -= m_pD->m_scalars.Num();
-        }
-
-        if ( index < m_pD->m_strings.Num() )
-        {
-            return m_pD->m_strings[index].m_pString.get();
-        }
-        else
-        {
-            index -= m_pD->m_strings.Num();
-        }
-
-        return nullptr;
-    }
-
-
-	//---------------------------------------------------------------------------------------------
-    void NodeSurfaceNew::SetInputNode( int i, NodePtr pNode )
-	{
-		check( i >=0 && i < GetInputCount());
-
-        std::size_t index = size_t( i );
-
-        if ( index<m_pD->m_meshes.Num() )
-		{
-			m_pD->m_meshes[ index ].m_pMesh = dynamic_cast<NodeMesh*>(pNode.get());
-            return;
-        }
-		else
-		{
-            index -= m_pD->m_meshes.Num();
-        }
-
-
-        if ( index < m_pD->m_images.Num())
-        {
-			m_pD->m_images[index].m_pImage = dynamic_cast<NodeImage*>(pNode.get());
-            return;
-        }
-        else
-        {
-            index -= m_pD->m_images.Num();
-        }
-        
-		if ( index < m_pD->m_vectors.Num() )
-        {
-            m_pD->m_vectors[index].m_pVector = dynamic_cast<NodeColour*>(pNode.get());
-            return;
-        }
-        else
-        {
-            index -= m_pD->m_vectors.Num();
-        }
-
-        if ( index < m_pD->m_scalars.Num() )
-        {
-            m_pD->m_scalars[index].m_pScalar = dynamic_cast<NodeScalar*>( pNode.get() );
-            return;
-        }
-        else
-        {
-            index -= m_pD->m_scalars.Num();
-        }
-
-        if ( index < m_pD->m_strings.Num() )
-        {
-            m_pD->m_strings[index].m_pString = dynamic_cast<NodeString*>( pNode.get() );
-            return;
-        }
-        else
-        {
-            index -= m_pD->m_strings.Num();
-        }
-    }
-
-
 	//---------------------------------------------------------------------------------------------
 	// Own Interface
 	//---------------------------------------------------------------------------------------------
diff --git a/Engine/Plugins/Experimental/Mutable/Source/MutableTools/Private/MuT/NodeSurfaceSwitch.cpp b/Engine/Plugins/Experimental/Mutable/Source/MutableTools/Private/MuT/NodeSurfaceSwitch.cpp
--- a/Engine/Plugins/Experimental/Mutable/Source/MutableTools/Private/MuT/NodeSurfaceSwitch.cpp
+++ b/Engine/Plugins/Experimental/Mutable/Source/MutableTools/Private/MuT/NodeSurfaceSwitch.cpp
@@ -26,55 +26,6 @@ namespace mu
 	MUTABLE_IMPLEMENT_NODE( NodeSurfaceSwitch, EType::Switch, Node, Node::EType::Surface)
 
 
-	//---------------------------------------------------------------------------------------------
-	// Node Interface
-	//---------------------------------------------------------------------------------------------
-	int NodeSurfaceSwitch::GetInputCount() const
-	{
-		return 1 + m_pD->Options.Num();
-	}
-
-
-	//---------------------------------------------------------------------------------------------
-	Node* NodeSurfaceSwitch::GetInputNode( int i ) const
-	{
-		check( i>=0 && i<GetInputCount() );
-
-		Node* pResult = 0;
-
-		switch (i)
-		{
-		case 0:
-			pResult = m_pD->Parameter.get();
-			break;
-
-		default:
-			pResult = m_pD->Options[i-1].get();
-			break;
-		}
-
-		return pResult;
-	}
-
-
-	//---------------------------------------------------------------------------------------------
-	void NodeSurfaceSwitch::SetInputNode( int i, NodePtr pNode )
-	{
-		check( i>=0 && i<GetInputCount() );
-
-		switch (i)
-		{
-		case 0:
-			m_pD->Parameter = dynamic_cast<NodeScalar*>(pNode.get());
-			break;
-
-		default:
-			m_pD->Options[i-1] = dynamic_cast<NodeSurface*>(pNode.get());
-			break;
-		}
-	}
-
-
 	//---------------------------------------------------------------------------------------------
 	// Own Interface
 	//---------------------------------------------------------------------------------------------
diff --git a/Engine/Plugins/Experimental/Mutable/Source/MutableTools/Private/MuT/NodeSurfaceVariation.cpp b/Engine/Plugins/Experimental/Mutable/Source/MutableTools/Private/MuT/NodeSurfaceVariation.cpp
--- a/Engine/Plugins/Experimental/Mutable/Source/MutableTools/Private/MuT/NodeSurfaceVariation.cpp
+++ b/Engine/Plugins/Experimental/Mutable/Source/MutableTools/Private/MuT/NodeSurfaceVariation.cpp
@@ -28,98 +28,6 @@ namespace mu
     MUTABLE_IMPLEMENT_NODE( NodeSurfaceVariation, EType::Variation, Node, Node::EType::Surface)
 
 
-	//---------------------------------------------------------------------------------------------
-	// Node Interface
-	//---------------------------------------------------------------------------------------------
-    int NodeSurfaceVariation::GetInputCount() const
-	{
-        int32 c = m_pD->m_defaultSurfaces.Num();
-        c += m_pD->m_defaultModifiers.Num();
-        for (const auto& v : m_pD->m_variations)
-		{
-            c += v.m_surfaces.Num();
-            c += v.m_modifiers.Num();
-        }
-
-		return (int)c;
-	}
-
-
-	//---------------------------------------------------------------------------------------------
-    Node* NodeSurfaceVariation::GetInputNode( int i ) const
-	{
-		check( i >=0 && i < GetInputCount() );
-
-        if ( i<m_pD->m_defaultSurfaces.Num())
-		{
-            return m_pD->m_defaultSurfaces[i].get();
-		}
-        i -= m_pD->m_defaultSurfaces.Num();
-
-        if ( i<m_pD->m_defaultModifiers.Num())
-        {
-            return m_pD->m_defaultModifiers[i].get();
-        }
-        i -= m_pD->m_defaultModifiers.Num();
-
-        for (const auto& v : m_pD->m_variations)
-        {
-            if (i < (int)v.m_surfaces.Num())
-            {
-                return v.m_surfaces[i].get();
-            }
-            i -= (int)v.m_surfaces.Num();
-
-            if (i < (int)v.m_modifiers.Num())
-            {
-                return v.m_modifiers[i].get();
-            }
-            i -= (int)v.m_modifiers.Num();
-        }
-
-		return nullptr;
-	}
-
-
-	//---------------------------------------------------------------------------------------------
-    void NodeSurfaceVariation::SetInputNode( int i, NodePtr pNode )
-	{
-		check( i >=0 && i <  GetInputCount());
-
-        if (i<m_pD->m_defaultSurfaces.Num())
-		{
-            m_pD->m_defaultSurfaces[i] = dynamic_cast<NodeSurface*>(pNode.get());
-			return;
-		}
-
-        i -= m_pD->m_defaultSurfaces.Num();
-        if (i<m_pD->m_defaultModifiers.Num())
-        {
-
-            m_pD->m_defaultModifiers[i] = dynamic_cast<NodeModifier*>(pNode.get());
-            return;
-        }
-        i -= m_pD->m_defaultModifiers.Num();
-
-        for (auto& v : m_pD->m_variations)
-        {
-            if (i < (int)v.m_surfaces.Num())
-            {
-                v.m_surfaces[i] = dynamic_cast<NodeSurface*>(pNode.get());
-                return;
-            }
-            i -= (int)v.m_surfaces.Num();
-
-            if (i < (int)v.m_modifiers.Num())
-            {
-                v.m_modifiers[i] = dynamic_cast<NodeModifier*>(pNode.get());
-                return;
-            }
-            i -= (int)v.m_modifiers.Num();
-        }
-	}
-
-
 	//---------------------------------------------------------------------------------------------
 	// Own Interface
 	//---------------------------------------------------------------------------------------------
diff --git a/Engine/Plugins/Experimental/Mutable/Source/MutableTools/Public/MuT/Node.h b/Engine/Plugins/Experimental/Mutable/Source/MutableTools/Public/MuT/Node.h
--- a/Engine/Plugins/Experimental/Mutable/Source/MutableTools/Public/MuT/Node.h
+++ b/Engine/Plugins/Experimental/Mutable/Source/MutableTools/Public/MuT/Node.h
@@ -99,11 +99,6 @@ namespace mu
         virtual const NODE_TYPE* GetType() const;
 		static const NODE_TYPE* GetStaticType();
 
-		//! Generic access to graph connections
-        virtual int GetInputCount() const = 0;
-		virtual Node* GetInputNode( int i ) const = 0;
-        virtual void SetInputNode( int i, NodePtr pNode ) = 0;
-
 		//! Set the opaque context returned in messages in the compiler log.
 		void SetMessageContext( const void* context );
 
diff --git a/Engine/Plugins/Experimental/Mutable/Source/MutableTools/Public/MuT/NodeBool.h b/Engine/Plugins/Experimental/Mutable/Source/MutableTools/Public/MuT/NodeBool.h
--- a/Engine/Plugins/Experimental/Mutable/Source/MutableTools/Public/MuT/NodeBool.h
+++ b/Engine/Plugins/Experimental/Mutable/Source/MutableTools/Public/MuT/NodeBool.h
@@ -10,36 +10,7 @@
 
 namespace mu
 {
-
-	// Forward definitions
-	class NodeBool;
-    using NodeBoolPtr = Ptr<NodeBool>;
-    using NodeBoolPtrConst = Ptr<const NodeBool>;
-
-	class NodeBoolConstant;
-    using NodeBoolConstantPtr = Ptr<NodeBoolConstant>;
-    using NodeBoolConstantPtrConst = Ptr<const NodeBoolConstant>;
-
-	class NodeBoolParameter;
-    using NodeBoolParameterPtr = Ptr<NodeBoolParameter>;
-    using NodeBoolParameterPtrConst = Ptr<const NodeBoolParameter>;
-
-	class NodeBoolIsNull;
-    using NodeBoolIsNullPtr = Ptr<NodeBoolIsNull>;
-    using NodeBoolIsNullPtrConst = Ptr<const NodeBoolIsNull>;
-
-	class NodeBoolNot;
-    using NodeBoolNotPtr = Ptr<NodeBoolNot>;
-    using NodeBoolNotPtrConst = Ptr<const NodeBoolNot>;
-
-	class NodeBoolAnd;
-    using NodeBoolAndPtr = Ptr<NodeBoolAnd>;
-    using NodeBoolAndPtrConst = Ptr<const NodeBoolAnd>;
-
-    class NodeRange;
-    using NodeRangePtr = Ptr<NodeRange>;
-    using NodeRangePtrConst = Ptr<const NodeRange>;
-
+	class NodeRange;
 
 	//---------------------------------------------------------------------------------------------
     //! %Base class of any node that outputs a Bool value.
@@ -66,7 +37,7 @@ namespace mu
 		//-----------------------------------------------------------------------------------------
 
 		static void Serialise( const NodeBool* pNode, OutputArchive& arch );
-		static NodeBoolPtr StaticUnserialise( InputArchive& arch );
+		static Ptr<NodeBool> StaticUnserialise( InputArchive& arch );
 
 
 		//-----------------------------------------------------------------------------------------
@@ -109,7 +80,7 @@ namespace mu
 
 		static void Serialise( const NodeBoolConstant* pNode, OutputArchive& arch );
 		void SerialiseWrapper(OutputArchive& arch) const override;
-		static NodeBoolConstantPtr StaticUnserialise( InputArchive& arch );
+		static Ptr<NodeBoolConstant> StaticUnserialise( InputArchive& arch );
 
 
 		//-----------------------------------------------------------------------------------------
@@ -119,10 +90,6 @@ namespace mu
         const NODE_TYPE* GetType() const override;
 		static const NODE_TYPE* GetStaticType();
 
-        virtual int GetInputCount() const override;
-        virtual Node* GetInputNode( int i ) const override;
-        void SetInputNode( int i, NodePtr pNode ) override;
-
 		//-----------------------------------------------------------------------------------------
 		// Own Interface
 		//-----------------------------------------------------------------------------------------
@@ -168,7 +135,7 @@ namespace mu
 
 		void SerialiseWrapper(OutputArchive& arch) const override;
 		static void Serialise( const NodeBoolParameter* pNode, OutputArchive& arch );
-		static NodeBoolParameterPtr StaticUnserialise( InputArchive& arch );
+		static Ptr<NodeBoolParameter> StaticUnserialise( InputArchive& arch );
 
 
 		//-----------------------------------------------------------------------------------------
@@ -178,10 +145,6 @@ namespace mu
         const NODE_TYPE* GetType() const override;
 		static const NODE_TYPE* GetStaticType();
 
-        int GetInputCount() const override;
-        Node* GetInputNode( int i ) const override;
-        void SetInputNode( int i, NodePtr pNode ) override;
-
 		//-----------------------------------------------------------------------------------------
 		// Own Interface
 		//-----------------------------------------------------------------------------------------
@@ -195,7 +158,7 @@ namespace mu
         //! Set the number of ranges (dimensions) for this parameter.
         //! By default a parameter has 0 ranges, meaning it only has one value.
         void SetRangeCount( int i );
-        void SetRange( int i, NodeRangePtr pRange );
+        void SetRange( int i, Ptr<NodeRange> );
 
 		//-----------------------------------------------------------------------------------------
 		// Interface pattern
@@ -233,22 +196,16 @@ namespace mu
 
 		void SerialiseWrapper(OutputArchive& arch) const override;
 		static void Serialise( const NodeBoolIsNull* pNode, OutputArchive& arch );
-		static NodeBoolIsNullPtr StaticUnserialise( InputArchive& arch );
+		static Ptr<NodeBoolIsNull> StaticUnserialise( InputArchive& arch );
 
 
 		//-----------------------------------------------------------------------------------------
 		// Node Interface
 		//-----------------------------------------------------------------------------------------
 
-        
-
         const NODE_TYPE* GetType() const override;
 		static const NODE_TYPE* GetStaticType();
 
-        int GetInputCount() const override;
-        Node* GetInputNode( int i ) const override;
-        void SetInputNode( int i, NodePtr pNode ) override;
-
 		//-----------------------------------------------------------------------------------------
 		// Own Interface
 		//-----------------------------------------------------------------------------------------
@@ -288,29 +245,23 @@ namespace mu
 
 		void SerialiseWrapper(OutputArchive& arch) const override;
 		static void Serialise( const NodeBoolNot* pNode, OutputArchive& arch );
-		static NodeBoolNotPtr StaticUnserialise( InputArchive& arch );
+		static Ptr<NodeBoolNot> StaticUnserialise( InputArchive& arch );
 
 
 		//-----------------------------------------------------------------------------------------
 		// Node Interface
 		//-----------------------------------------------------------------------------------------
 
-        
-
         const NODE_TYPE* GetType() const override;
 		static const NODE_TYPE* GetStaticType();
 
-        int GetInputCount() const override;
-        Node* GetInputNode( int i ) const override;
-        void SetInputNode( int i, NodePtr pNode ) override;
-
 		//-----------------------------------------------------------------------------------------
 		// Own Interface
 		//-----------------------------------------------------------------------------------------
 		
 		//! Input
-		NodeBoolPtr GetInput() const;
-		void SetInput( NodeBoolPtr );
+		Ptr<NodeBool> GetInput() const;
+		void SetInput( Ptr<NodeBool> );
 
 		//-----------------------------------------------------------------------------------------
 		// Interface pattern
@@ -347,32 +298,26 @@ namespace mu
 
 		void SerialiseWrapper(OutputArchive& arch) const override;
 		static void Serialise( const NodeBoolAnd* pNode, OutputArchive& arch );
-		static NodeBoolAndPtr StaticUnserialise( InputArchive& arch );
+		static Ptr<NodeBoolAnd> StaticUnserialise( InputArchive& arch );
 
 
 		//-----------------------------------------------------------------------------------------
 		// Node Interface
 		//-----------------------------------------------------------------------------------------
 
-        
-
         const NODE_TYPE* GetType() const override;
 		static const NODE_TYPE* GetStaticType();
 
-        int GetInputCount() const override;
-        Node* GetInputNode( int i ) const override;
-        void SetInputNode( int i, NodePtr pNode ) override;
-
 		//-----------------------------------------------------------------------------------------
 		// Own Interface
 		//-----------------------------------------------------------------------------------------
 		
 		//! Inputs
-		NodeBoolPtr GetA() const;
-		void SetA( NodeBoolPtr );
+		Ptr<NodeBool> GetA() const;
+		void SetA(Ptr<NodeBool>);
 
-		NodeBoolPtr GetB() const;
-		void SetB( NodeBoolPtr );
+		Ptr<NodeBool> GetB() const;
+		void SetB(Ptr<NodeBool>);
 
 		//-----------------------------------------------------------------------------------------
 		// Interface pattern
diff --git a/Engine/Plugins/Experimental/Mutable/Source/MutableTools/Public/MuT/NodeColourArithmeticOperation.h b/Engine/Plugins/Experimental/Mutable/Source/MutableTools/Public/MuT/NodeColourArithmeticOperation.h
--- a/Engine/Plugins/Experimental/Mutable/Source/MutableTools/Public/MuT/NodeColourArithmeticOperation.h
+++ b/Engine/Plugins/Experimental/Mutable/Source/MutableTools/Public/MuT/NodeColourArithmeticOperation.h
@@ -42,11 +42,6 @@ namespace mu
 		virtual const NODE_TYPE* GetType() const override;
 		static const NODE_TYPE* GetStaticType();
 
-		virtual int GetInputCount() const override;
-		virtual Node* GetInputNode( int i ) const override;
-		virtual void SetInputNode( int i, NodePtr pNode ) override;
-
-
 		//-----------------------------------------------------------------------------------------
 		// Own Interface
 		//-----------------------------------------------------------------------------------------
diff --git a/Engine/Plugins/Experimental/Mutable/Source/MutableTools/Public/MuT/NodeColourConstant.h b/Engine/Plugins/Experimental/Mutable/Source/MutableTools/Public/MuT/NodeColourConstant.h
--- a/Engine/Plugins/Experimental/Mutable/Source/MutableTools/Public/MuT/NodeColourConstant.h
+++ b/Engine/Plugins/Experimental/Mutable/Source/MutableTools/Public/MuT/NodeColourConstant.h
@@ -37,15 +37,9 @@ namespace mu
 		// Node Interface
 		//-----------------------------------------------------------------------------------------
 
-        
-
         const NODE_TYPE* GetType() const override;
 		static const NODE_TYPE* GetStaticType();
 
-        virtual int GetInputCount() const override;
-        virtual Node* GetInputNode( int i ) const override;
-        void SetInputNode( int i, NodePtr pNode ) override;
-
 		//-----------------------------------------------------------------------------------------
 		// Own Interface
 		//-----------------------------------------------------------------------------------------
diff --git a/Engine/Plugins/Experimental/Mutable/Source/MutableTools/Public/MuT/NodeColourFromScalars.h b/Engine/Plugins/Experimental/Mutable/Source/MutableTools/Public/MuT/NodeColourFromScalars.h
--- a/Engine/Plugins/Experimental/Mutable/Source/MutableTools/Public/MuT/NodeColourFromScalars.h
+++ b/Engine/Plugins/Experimental/Mutable/Source/MutableTools/Public/MuT/NodeColourFromScalars.h
@@ -45,11 +45,6 @@ namespace mu
 		const NODE_TYPE* GetType() const override;
 		static const NODE_TYPE* GetStaticType();
 
-        int GetInputCount() const override;
-        Node* GetInputNode( int i ) const override;
-		void SetInputNode( int i, NodePtr pNode ) override;
-
-
 		//-----------------------------------------------------------------------------------------
 		// Own Interface
 		//-----------------------------------------------------------------------------------------
diff --git a/Engine/Plugins/Experimental/Mutable/Source/MutableTools/Public/MuT/NodeColourParameter.h b/Engine/Plugins/Experimental/Mutable/Source/MutableTools/Public/MuT/NodeColourParameter.h
--- a/Engine/Plugins/Experimental/Mutable/Source/MutableTools/Public/MuT/NodeColourParameter.h
+++ b/Engine/Plugins/Experimental/Mutable/Source/MutableTools/Public/MuT/NodeColourParameter.h
@@ -43,10 +43,6 @@ namespace mu
 		const NODE_TYPE* GetType() const override;
 		static const NODE_TYPE* GetStaticType();
 
-		int GetInputCount() const override;
-		Node* GetInputNode( int i ) const override;
-		void SetInputNode( int i, NodePtr pNode ) override;
-
 		//-----------------------------------------------------------------------------------------
 		// Own Interface
 		//-----------------------------------------------------------------------------------------
diff --git a/Engine/Plugins/Experimental/Mutable/Source/MutableTools/Public/MuT/NodeColourSampleImage.h b/Engine/Plugins/Experimental/Mutable/Source/MutableTools/Public/MuT/NodeColourSampleImage.h
--- a/Engine/Plugins/Experimental/Mutable/Source/MutableTools/Public/MuT/NodeColourSampleImage.h
+++ b/Engine/Plugins/Experimental/Mutable/Source/MutableTools/Public/MuT/NodeColourSampleImage.h
@@ -46,16 +46,9 @@ namespace mu
 		// Node Interface
 		//-----------------------------------------------------------------------------------------
 
-		
-
 		const NODE_TYPE* GetType() const override;
 		static const NODE_TYPE* GetStaticType();
 
-		int GetInputCount() const override;
-		Node* GetInputNode( int i ) const override;
-		void SetInputNode( int i, NodePtr pNode ) override;
-
-
 		//-----------------------------------------------------------------------------------------
 		// Own Interface
 		//-----------------------------------------------------------------------------------------
diff --git a/Engine/Plugins/Experimental/Mutable/Source/MutableTools/Public/MuT/NodeColourSwitch.h b/Engine/Plugins/Experimental/Mutable/Source/MutableTools/Public/MuT/NodeColourSwitch.h
--- a/Engine/Plugins/Experimental/Mutable/Source/MutableTools/Public/MuT/NodeColourSwitch.h
+++ b/Engine/Plugins/Experimental/Mutable/Source/MutableTools/Public/MuT/NodeColourSwitch.h
@@ -42,15 +42,9 @@ namespace mu
 		// Node Interface
 		//-----------------------------------------------------------------------------------------
 
-        
-
         const NODE_TYPE* GetType() const override;
 		static const NODE_TYPE* GetStaticType();
 
-        virtual int GetInputCount() const override;
-        virtual Node* GetInputNode( int i ) const override;
-        void SetInputNode( int i, NodePtr pNode ) override;
-
 		//-----------------------------------------------------------------------------------------
 		// Own Interface
 		//-----------------------------------------------------------------------------------------
diff --git a/Engine/Plugins/Experimental/Mutable/Source/MutableTools/Public/MuT/NodeColourTable.h b/Engine/Plugins/Experimental/Mutable/Source/MutableTools/Public/MuT/NodeColourTable.h
--- a/Engine/Plugins/Experimental/Mutable/Source/MutableTools/Public/MuT/NodeColourTable.h
+++ b/Engine/Plugins/Experimental/Mutable/Source/MutableTools/Public/MuT/NodeColourTable.h
@@ -44,10 +44,6 @@ namespace mu
 		const NODE_TYPE* GetType() const override;
 		static const NODE_TYPE* GetStaticType();
 
-		int GetInputCount() const override;
-		Node* GetInputNode( int i ) const override;
-		void SetInputNode( int i, NodePtr pNode ) override;
-
 		//-----------------------------------------------------------------------------------------
 		// Own Interface
 		//-----------------------------------------------------------------------------------------
diff --git a/Engine/Plugins/Experimental/Mutable/Source/MutableTools/Public/MuT/NodeComponentEdit.h b/Engine/Plugins/Experimental/Mutable/Source/MutableTools/Public/MuT/NodeComponentEdit.h
--- a/Engine/Plugins/Experimental/Mutable/Source/MutableTools/Public/MuT/NodeComponentEdit.h
+++ b/Engine/Plugins/Experimental/Mutable/Source/MutableTools/Public/MuT/NodeComponentEdit.h
@@ -43,15 +43,10 @@ namespace mu
 		//-----------------------------------------------------------------------------------------
         // Node interface
 		//-----------------------------------------------------------------------------------------
-		
 
 		const NODE_TYPE* GetType() const override;
 		static const NODE_TYPE* GetStaticType();
 
-		int GetInputCount() const override;
-		Node* GetInputNode( int i ) const override;
-		void SetInputNode( int i, NodePtr pNode ) override;
-
         //-----------------------------------------------------------------------------------------
         // NodeComponent interface
         //-----------------------------------------------------------------------------------------
diff --git a/Engine/Plugins/Experimental/Mutable/Source/MutableTools/Public/MuT/NodeComponentNew.h b/Engine/Plugins/Experimental/Mutable/Source/MutableTools/Public/MuT/NodeComponentNew.h
--- a/Engine/Plugins/Experimental/Mutable/Source/MutableTools/Public/MuT/NodeComponentNew.h
+++ b/Engine/Plugins/Experimental/Mutable/Source/MutableTools/Public/MuT/NodeComponentNew.h
@@ -44,17 +44,12 @@ namespace mu
 
 		//-----------------------------------------------------------------------------------------
         // Node interface
-		//-----------------------------------------------------------------------------------------
-		
+		//-----------------------------------------------------------------------------------------		
 
         const NODE_TYPE* GetType() const override;
 		static const NODE_TYPE* GetStaticType();
 
-		int GetInputCount() const override;
-		Node* GetInputNode( int i ) const override;
-		void SetInputNode( int i, NodePtr pNode ) override;
-
-        //-----------------------------------------------------------------------------------------
+		//-----------------------------------------------------------------------------------------
         // NodeComponent interface
         //-----------------------------------------------------------------------------------------
         int GetSurfaceCount() const override;
diff --git a/Engine/Plugins/Experimental/Mutable/Source/MutableTools/Public/MuT/NodeExtensionDataConstant.h b/Engine/Plugins/Experimental/Mutable/Source/MutableTools/Public/MuT/NodeExtensionDataConstant.h
--- a/Engine/Plugins/Experimental/Mutable/Source/MutableTools/Public/MuT/NodeExtensionDataConstant.h
+++ b/Engine/Plugins/Experimental/Mutable/Source/MutableTools/Public/MuT/NodeExtensionDataConstant.h
@@ -38,10 +38,6 @@ public:
 	const NODE_TYPE* GetType() const override;
 	static const NODE_TYPE* GetStaticType();
 
-	int GetInputCount() const override;
-	Node* GetInputNode(int i) const override;
-	void SetInputNode(int i, NodePtr pNode) override;
-
 	//-----------------------------------------------------------------------------------------
 	// Own Interface
 	//-----------------------------------------------------------------------------------------
diff --git a/Engine/Plugins/Experimental/Mutable/Source/MutableTools/Public/MuT/NodeExtensionDataSwitch.h b/Engine/Plugins/Experimental/Mutable/Source/MutableTools/Public/MuT/NodeExtensionDataSwitch.h
--- a/Engine/Plugins/Experimental/Mutable/Source/MutableTools/Public/MuT/NodeExtensionDataSwitch.h
+++ b/Engine/Plugins/Experimental/Mutable/Source/MutableTools/Public/MuT/NodeExtensionDataSwitch.h
@@ -37,10 +37,6 @@ namespace mu
 		const NODE_TYPE* GetType() const override;
 		static const NODE_TYPE* GetStaticType();
 
-		int GetInputCount() const override;
-		Node* GetInputNode(int i) const override;
-		void SetInputNode(int i, NodePtr pNode) override;
-
 		//-----------------------------------------------------------------------------------------
 		// Own Interface
 		//-----------------------------------------------------------------------------------------
diff --git a/Engine/Plugins/Experimental/Mutable/Source/MutableTools/Public/MuT/NodeImageConditional.h b/Engine/Plugins/Experimental/Mutable/Source/MutableTools/Public/MuT/NodeImageConditional.h
--- a/Engine/Plugins/Experimental/Mutable/Source/MutableTools/Public/MuT/NodeImageConditional.h
+++ b/Engine/Plugins/Experimental/Mutable/Source/MutableTools/Public/MuT/NodeImageConditional.h
@@ -13,9 +13,6 @@ namespace mu
 
 	// Forward definitions
     class NodeBool;
-    typedef Ptr<NodeBool> NodeBoolPtr;
-    typedef Ptr<const NodeBool> NodeBoolPtrConst;
-
     class NodeImageConditional;
     typedef Ptr<NodeImageConditional> NodeImageConditionalPtr;
     typedef Ptr<const NodeImageConditional> NodeImageConditionalPtrConst;
@@ -42,15 +39,9 @@ namespace mu
 		// Node Interface
 		//-----------------------------------------------------------------------------------------
 
-        
-
         const NODE_TYPE* GetType() const override;
 		static const NODE_TYPE* GetStaticType();
 
-        virtual int GetInputCount() const override;
-        virtual Node* GetInputNode( int i ) const override;
-        void SetInputNode( int i, NodePtr pNode ) override;
-
 		//-----------------------------------------------------------------------------------------
 		// Own Interface
 		//-----------------------------------------------------------------------------------------
