[mutable] Remove unnecessary boilerplate code that was only used in one debug tool.

diff --git a/Engine/Plugins/Experimental/Mutable/Source/MutableTools/Public/MuT/NodeImageConstant.h b/Engine/Plugins/Experimental/Mutable/Source/MutableTools/Public/MuT/NodeImageConstant.h
--- a/Engine/Plugins/Experimental/Mutable/Source/MutableTools/Public/MuT/NodeImageConstant.h
+++ b/Engine/Plugins/Experimental/Mutable/Source/MutableTools/Public/MuT/NodeImageConstant.h
@@ -38,7 +38,6 @@ namespace mu
 		static void Serialise( const NodeImageConstant* pNode, OutputArchive& arch );
 		static NodeImageConstantPtr StaticUnserialise( InputArchive& arch );
 
-
 		//-----------------------------------------------------------------------------------------
 		// Node Interface
 		//-----------------------------------------------------------------------------------------
@@ -46,10 +45,6 @@ namespace mu
         const NODE_TYPE* GetType() const override;
 		static const NODE_TYPE* GetStaticType();
 
-        virtual int GetInputCount() const override;
-        virtual Node* GetInputNode( int i ) const override;
-        void SetInputNode( int i, NodePtr pNode ) override;
-
 		//-----------------------------------------------------------------------------------------
 		// Own Interface
 		//-----------------------------------------------------------------------------------------
diff --git a/Engine/Plugins/Experimental/Mutable/Source/MutableTools/Public/MuT/NodeImageFormat.h b/Engine/Plugins/Experimental/Mutable/Source/MutableTools/Public/MuT/NodeImageFormat.h
--- a/Engine/Plugins/Experimental/Mutable/Source/MutableTools/Public/MuT/NodeImageFormat.h
+++ b/Engine/Plugins/Experimental/Mutable/Source/MutableTools/Public/MuT/NodeImageFormat.h
@@ -42,15 +42,9 @@ namespace mu
 		// Node Interface
 		//-----------------------------------------------------------------------------------------
 
-		
-
 		const NODE_TYPE* GetType() const override;
 		static const NODE_TYPE* GetStaticType();
 
-		int GetInputCount() const override;
-		Node* GetInputNode( int i ) const override;
-		void SetInputNode( int i, NodePtr pNode ) override;
-
 		//-----------------------------------------------------------------------------------------
 		// Own Interface
 		//-----------------------------------------------------------------------------------------
diff --git a/Engine/Plugins/Experimental/Mutable/Source/MutableTools/Public/MuT/NodeImageInvert.h b/Engine/Plugins/Experimental/Mutable/Source/MutableTools/Public/MuT/NodeImageInvert.h
--- a/Engine/Plugins/Experimental/Mutable/Source/MutableTools/Public/MuT/NodeImageInvert.h
+++ b/Engine/Plugins/Experimental/Mutable/Source/MutableTools/Public/MuT/NodeImageInvert.h
@@ -41,10 +41,6 @@ namespace mu
 		const NODE_TYPE* GetType() const override;
 		static const NODE_TYPE* GetStaticType();
 
-		virtual int GetInputCount() const override;
-		virtual Node* GetInputNode(int i)const override;
-		void SetInputNode(int i, NodePtr pNode) override;
-
 		//-----------------------------------------------------------------------------------------
 		// Own Interface
 		//-----------------------------------------------------------------------------------------
diff --git a/Engine/Plugins/Experimental/Mutable/Source/MutableTools/Public/MuT/NodeImageNormalComposite.h b/Engine/Plugins/Experimental/Mutable/Source/MutableTools/Public/MuT/NodeImageNormalComposite.h
--- a/Engine/Plugins/Experimental/Mutable/Source/MutableTools/Public/MuT/NodeImageNormalComposite.h
+++ b/Engine/Plugins/Experimental/Mutable/Source/MutableTools/Public/MuT/NodeImageNormalComposite.h
@@ -43,10 +43,6 @@ namespace mu
 		const NODE_TYPE* GetType() const override;
 		static const NODE_TYPE* GetStaticType();
 
-		virtual int GetInputCount() const override;
-		virtual Node* GetInputNode(int i)const override;
-		void SetInputNode(int i, NodePtr pNode) override;
-
 		//-----------------------------------------------------------------------------------------
 		// Own Interface
 		//-----------------------------------------------------------------------------------------
diff --git a/Engine/Plugins/Experimental/Mutable/Source/MutableTools/Public/MuT/NodeImageParameter.h b/Engine/Plugins/Experimental/Mutable/Source/MutableTools/Public/MuT/NodeImageParameter.h
--- a/Engine/Plugins/Experimental/Mutable/Source/MutableTools/Public/MuT/NodeImageParameter.h
+++ b/Engine/Plugins/Experimental/Mutable/Source/MutableTools/Public/MuT/NodeImageParameter.h
@@ -38,10 +38,6 @@ namespace mu
 		const NODE_TYPE* GetType() const override;
 		static const NODE_TYPE* GetStaticType();
 
-		int GetInputCount() const override;
-		Node* GetInputNode( int i ) const override;
-		void SetInputNode( int i, Ptr<Node> Node ) override;
-
 		//-----------------------------------------------------------------------------------------
 		// Own Interface
 		//-----------------------------------------------------------------------------------------
diff --git a/Engine/Plugins/Experimental/Mutable/Source/MutableTools/Public/MuT/NodeImageProject.h b/Engine/Plugins/Experimental/Mutable/Source/MutableTools/Public/MuT/NodeImageProject.h
--- a/Engine/Plugins/Experimental/Mutable/Source/MutableTools/Public/MuT/NodeImageProject.h
+++ b/Engine/Plugins/Experimental/Mutable/Source/MutableTools/Public/MuT/NodeImageProject.h
@@ -46,10 +46,6 @@ namespace mu
         const NODE_TYPE* GetType() const override;
 		static const NODE_TYPE* GetStaticType();
 
-        virtual int GetInputCount() const override;
-        virtual Node* GetInputNode( int i ) const override;
-        void SetInputNode( int i, Ptr<Node> pNode ) override;
-
 		//-----------------------------------------------------------------------------------------
 		// Own Interface
 		//-----------------------------------------------------------------------------------------
diff --git a/Engine/Plugins/Experimental/Mutable/Source/MutableTools/Public/MuT/NodeLOD.h b/Engine/Plugins/Experimental/Mutable/Source/MutableTools/Public/MuT/NodeLOD.h
--- a/Engine/Plugins/Experimental/Mutable/Source/MutableTools/Public/MuT/NodeLOD.h
+++ b/Engine/Plugins/Experimental/Mutable/Source/MutableTools/Public/MuT/NodeLOD.h
@@ -46,15 +46,9 @@ namespace mu
 		// Node Interface
 		//-----------------------------------------------------------------------------------------
 
-		
-
 		const NODE_TYPE* GetType() const override;
 		static const NODE_TYPE* GetStaticType();
 
-		int GetInputCount() const override;
-		Node* GetInputNode( int i ) const override;
-		void SetInputNode( int i, NodePtr pNode ) override;
-
 		//-----------------------------------------------------------------------------------------
 		// Own Interface
 		//-----------------------------------------------------------------------------------------
diff --git a/Engine/Plugins/Experimental/Mutable/Source/MutableTools/Public/MuT/NodeMeshClipMorphPlane.h b/Engine/Plugins/Experimental/Mutable/Source/MutableTools/Public/MuT/NodeMeshClipMorphPlane.h
--- a/Engine/Plugins/Experimental/Mutable/Source/MutableTools/Public/MuT/NodeMeshClipMorphPlane.h
+++ b/Engine/Plugins/Experimental/Mutable/Source/MutableTools/Public/MuT/NodeMeshClipMorphPlane.h
@@ -40,12 +40,6 @@ namespace mu
 		virtual const NODE_TYPE* GetType() const;
 		static const NODE_TYPE* GetStaticType();
 
-		virtual int GetInputCount() const;
-		virtual const char* GetInputName( int i ) const;
-		virtual const NODE_TYPE* GetInputType( int i ) const;
-		virtual Node* GetInputNode( int i ) const;
-		virtual void SetInputNode( int i, NodePtr pNode );
-
 		//-----------------------------------------------------------------------------------------
 		// Own Interface
 		//-----------------------------------------------------------------------------------------
diff --git a/Engine/Plugins/Experimental/Mutable/Source/MutableTools/Public/MuT/NodeMeshClipWithMesh.h b/Engine/Plugins/Experimental/Mutable/Source/MutableTools/Public/MuT/NodeMeshClipWithMesh.h
--- a/Engine/Plugins/Experimental/Mutable/Source/MutableTools/Public/MuT/NodeMeshClipWithMesh.h
+++ b/Engine/Plugins/Experimental/Mutable/Source/MutableTools/Public/MuT/NodeMeshClipWithMesh.h
@@ -40,12 +40,6 @@ namespace mu
 		virtual const NODE_TYPE* GetType() const;
 		static const NODE_TYPE* GetStaticType();
 
-		virtual int GetInputCount() const;
-		virtual const char* GetInputName( int i ) const;
-		virtual const NODE_TYPE* GetInputType( int i ) const;
-		virtual Node* GetInputNode( int i ) const;
-		virtual void SetInputNode( int i, NodePtr pNode );
-
 		//-----------------------------------------------------------------------------------------
 		// Own Interface
 		//-----------------------------------------------------------------------------------------
diff --git a/Engine/Plugins/Experimental/Mutable/Source/MutableTools/Public/MuT/NodeMeshMorph.h b/Engine/Plugins/Experimental/Mutable/Source/MutableTools/Public/MuT/NodeMeshMorph.h
--- a/Engine/Plugins/Experimental/Mutable/Source/MutableTools/Public/MuT/NodeMeshMorph.h
+++ b/Engine/Plugins/Experimental/Mutable/Source/MutableTools/Public/MuT/NodeMeshMorph.h
@@ -38,10 +38,6 @@ namespace mu
 		const NODE_TYPE* GetType() const override;
 		static const NODE_TYPE* GetStaticType();
 
-		int GetInputCount() const override;
-		Node* GetInputNode( int i ) const override;
-		void SetInputNode( int i, Ptr<Node> pNode ) override;
-
 		//-----------------------------------------------------------------------------------------
 		// Own Interface
 		//-----------------------------------------------------------------------------------------
diff --git a/Engine/Plugins/Experimental/Mutable/Source/MutableTools/Public/MuT/NodeModifierMeshClipWithMesh.h b/Engine/Plugins/Experimental/Mutable/Source/MutableTools/Public/MuT/NodeModifierMeshClipWithMesh.h
--- a/Engine/Plugins/Experimental/Mutable/Source/MutableTools/Public/MuT/NodeModifierMeshClipWithMesh.h
+++ b/Engine/Plugins/Experimental/Mutable/Source/MutableTools/Public/MuT/NodeModifierMeshClipWithMesh.h
@@ -35,15 +35,9 @@ namespace mu
         // Node interface
 		//-----------------------------------------------------------------------------------------
 
-        
-
         const NODE_TYPE* GetType() const override;
 		static const NODE_TYPE* GetStaticType();
 
-        int GetInputCount() const override;
-        Node* GetInputNode( int i ) const override;
-        void SetInputNode( int i, Ptr<Node> ) override;
-
 		//-----------------------------------------------------------------------------------------
         // Own interface
 		//-----------------------------------------------------------------------------------------
diff --git a/Engine/Plugins/Experimental/Mutable/Source/MutableTools/Public/MuT/NodeModifierMeshClipWithUVMask.h b/Engine/Plugins/Experimental/Mutable/Source/MutableTools/Public/MuT/NodeModifierMeshClipWithUVMask.h
--- a/Engine/Plugins/Experimental/Mutable/Source/MutableTools/Public/MuT/NodeModifierMeshClipWithUVMask.h
+++ b/Engine/Plugins/Experimental/Mutable/Source/MutableTools/Public/MuT/NodeModifierMeshClipWithUVMask.h
@@ -34,15 +34,9 @@ namespace mu
         // Node interface
 		//-----------------------------------------------------------------------------------------
 
-        
-
         const NODE_TYPE* GetType() const override;
 		static const NODE_TYPE* GetStaticType();
 
-        int GetInputCount() const override;
-        Node* GetInputNode( int i ) const override;
-        void SetInputNode( int i, Ptr<Node> pNode ) override;
-
 		//-----------------------------------------------------------------------------------------
         // Own interface
 		//-----------------------------------------------------------------------------------------
diff --git a/Engine/Plugins/Experimental/Mutable/Source/MutableTools/Public/MuT/NodePatchImage.h b/Engine/Plugins/Experimental/Mutable/Source/MutableTools/Public/MuT/NodePatchImage.h
--- a/Engine/Plugins/Experimental/Mutable/Source/MutableTools/Public/MuT/NodePatchImage.h
+++ b/Engine/Plugins/Experimental/Mutable/Source/MutableTools/Public/MuT/NodePatchImage.h
@@ -44,16 +44,9 @@ namespace mu
         // Node Interface
         //-----------------------------------------------------------------------------------------
 
-        
-
         const NODE_TYPE* GetType() const override;
         static const NODE_TYPE* GetStaticType();
 
-        int GetInputCount() const override;
-        Node* GetInputNode( int i ) const override;
-        void SetInputNode( int i, NodePtr pNode ) override;
-
-
         //-----------------------------------------------------------------------------------------
         // Own Interface
         //-----------------------------------------------------------------------------------------
diff --git a/Engine/Plugins/Experimental/Mutable/Source/MutableTools/Public/MuT/NodePatchMesh.h b/Engine/Plugins/Experimental/Mutable/Source/MutableTools/Public/MuT/NodePatchMesh.h
--- a/Engine/Plugins/Experimental/Mutable/Source/MutableTools/Public/MuT/NodePatchMesh.h
+++ b/Engine/Plugins/Experimental/Mutable/Source/MutableTools/Public/MuT/NodePatchMesh.h
@@ -37,15 +37,9 @@ namespace mu
 		// Node Interface
 		//-----------------------------------------------------------------------------------------
 
-        
-
-        const NODE_TYPE* GetType() const override;
+		const NODE_TYPE* GetType() const override;
 		static const NODE_TYPE* GetStaticType();
 
-        int GetInputCount() const override;
-        Node* GetInputNode( int i ) const override;
-        void SetInputNode( int i, NodePtr pNode ) override;
-
 		//-----------------------------------------------------------------------------------------
 		// Own Interface
 		//-----------------------------------------------------------------------------------------
diff --git a/Engine/Plugins/Experimental/Mutable/Source/MutableTools/Public/MuT/NodeProjector.h b/Engine/Plugins/Experimental/Mutable/Source/MutableTools/Public/MuT/NodeProjector.h
--- a/Engine/Plugins/Experimental/Mutable/Source/MutableTools/Public/MuT/NodeProjector.h
+++ b/Engine/Plugins/Experimental/Mutable/Source/MutableTools/Public/MuT/NodeProjector.h
@@ -96,15 +96,9 @@ namespace mu
 		// Node Interface
 		//-----------------------------------------------------------------------------------------
 
-        
-
         const NODE_TYPE* GetType() const override;
 		static const NODE_TYPE* GetStaticType();
 
-        int GetInputCount() const override;
-        Node* GetInputNode( int i ) const override;
-        void SetInputNode( int i, NodePtr pNode ) override;
-
 		//-----------------------------------------------------------------------------------------
 		// Own Interface
 		//-----------------------------------------------------------------------------------------
@@ -160,20 +154,13 @@ namespace mu
 		static void Serialise( const NodeProjectorParameter* pNode, OutputArchive& arch );
 		static NodeProjectorParameterPtr StaticUnserialise( InputArchive& arch );
 
-
 		//-----------------------------------------------------------------------------------------
 		// Node Interface
 		//-----------------------------------------------------------------------------------------
 
-        
-
         const NODE_TYPE* GetType() const override;
 		static const NODE_TYPE* GetStaticType();
 
-        int GetInputCount() const override;
-        Node* GetInputNode( int i ) const override;
-        void SetInputNode( int i, NodePtr pNode ) override;
-
 		//-----------------------------------------------------------------------------------------
 		// Own Interface
 		//-----------------------------------------------------------------------------------------
diff --git a/Engine/Plugins/Experimental/Mutable/Source/MutableTools/Public/MuT/NodeRangeFromScalar.h b/Engine/Plugins/Experimental/Mutable/Source/MutableTools/Public/MuT/NodeRangeFromScalar.h
--- a/Engine/Plugins/Experimental/Mutable/Source/MutableTools/Public/MuT/NodeRangeFromScalar.h
+++ b/Engine/Plugins/Experimental/Mutable/Source/MutableTools/Public/MuT/NodeRangeFromScalar.h
@@ -37,11 +37,6 @@ namespace mu
 		const NODE_TYPE* GetType() const override;
 		static const NODE_TYPE* GetStaticType();
 
-        int GetInputCount() const override;
-        Node* GetInputNode( int i ) const override;
-		void SetInputNode( int i, NodePtr pNode ) override;
-
-
 		//-----------------------------------------------------------------------------------------
 		// Own Interface
 		//-----------------------------------------------------------------------------------------
diff --git a/Engine/Plugins/Experimental/Mutable/Source/MutableTools/Public/MuT/NodeScalarEnumParameter.h b/Engine/Plugins/Experimental/Mutable/Source/MutableTools/Public/MuT/NodeScalarEnumParameter.h
--- a/Engine/Plugins/Experimental/Mutable/Source/MutableTools/Public/MuT/NodeScalarEnumParameter.h
+++ b/Engine/Plugins/Experimental/Mutable/Source/MutableTools/Public/MuT/NodeScalarEnumParameter.h
@@ -42,17 +42,12 @@ namespace mu
 		static void Serialise( const NodeScalarEnumParameter* pNode, OutputArchive& arch );
 		static NodeScalarEnumParameterPtr StaticUnserialise( InputArchive& arch );
 
-
 		//-----------------------------------------------------------------------------------------
 		// Node Interface
 		//-----------------------------------------------------------------------------------------
         const NODE_TYPE* GetType() const override;
 		static const NODE_TYPE* GetStaticType();
 
-        int GetInputCount() const override;
-        Node* GetInputNode( int i ) const override;
-        void SetInputNode( int i, NodePtr pNode ) override;
-
 		//-----------------------------------------------------------------------------------------
 		// Own Interface
 		//-----------------------------------------------------------------------------------------
diff --git a/Engine/Plugins/Experimental/Mutable/Source/MutableTools/Public/MuT/NodeScalarSwitch.h b/Engine/Plugins/Experimental/Mutable/Source/MutableTools/Public/MuT/NodeScalarSwitch.h
--- a/Engine/Plugins/Experimental/Mutable/Source/MutableTools/Public/MuT/NodeScalarSwitch.h
+++ b/Engine/Plugins/Experimental/Mutable/Source/MutableTools/Public/MuT/NodeScalarSwitch.h
@@ -37,7 +37,6 @@ namespace mu
 		static void Serialise( const NodeScalarSwitch* pNode, OutputArchive& arch );
         static NodeScalarSwitchPtr StaticUnserialise( InputArchive& arch );
 
-
 		//-----------------------------------------------------------------------------------------
 		// Node Interface
 		//-----------------------------------------------------------------------------------------
@@ -45,10 +44,6 @@ namespace mu
         const NODE_TYPE* GetType() const override;
 		static const NODE_TYPE* GetStaticType();
 
-        virtual int GetInputCount() const override;
-        virtual Node* GetInputNode( int i ) const override;
-        void SetInputNode( int i, NodePtr pNode ) override;
-
 		//-----------------------------------------------------------------------------------------
 		// Own Interface
 		//-----------------------------------------------------------------------------------------
diff --git a/Engine/Plugins/Experimental/Mutable/Source/MutableTools/Public/MuT/NodeSurfaceEdit.h b/Engine/Plugins/Experimental/Mutable/Source/MutableTools/Public/MuT/NodeSurfaceEdit.h
--- a/Engine/Plugins/Experimental/Mutable/Source/MutableTools/Public/MuT/NodeSurfaceEdit.h
+++ b/Engine/Plugins/Experimental/Mutable/Source/MutableTools/Public/MuT/NodeSurfaceEdit.h
@@ -62,10 +62,6 @@ namespace mu
         const NODE_TYPE* GetType() const override;
 		static const NODE_TYPE* GetStaticType();
 
-        virtual int GetInputCount() const override;
-        virtual Node* GetInputNode( int i ) const override;
-        void SetInputNode( int i, NodePtr pNode ) override;
-
 		//-----------------------------------------------------------------------------------------
 		// Own Interface
 		//-----------------------------------------------------------------------------------------
