[mutable] Fix problem with table image columns having images with different formats.
Fix bug with mutable debugger referenced textures.
#rnx

diff --git a/Engine/Plugins/Experimental/Mutable/Source/CustomizableObjectEditor/Private/MuCOE/CustomizableObjectCompiler.cpp b/Engine/Plugins/Experimental/Mutable/Source/CustomizableObjectEditor/Private/MuCOE/CustomizableObjectCompiler.cpp
--- a/Engine/Plugins/Experimental/Mutable/Source/CustomizableObjectEditor/Private/MuCOE/CustomizableObjectCompiler.cpp
+++ b/Engine/Plugins/Experimental/Mutable/Source/CustomizableObjectEditor/Private/MuCOE/CustomizableObjectCompiler.cpp
@@ -1305,7 +1305,9 @@ void FCustomizableObjectCompiler::CompileInternal(UCustomizableObject* Object, c
 }
 
 
-mu::NodePtr FCustomizableObjectCompiler::Export(UCustomizableObject* Object, const FCompilationOptions& InCompilerOptions, TArray<TSoftObjectPtr<UTexture>>& OutRuntimeReferencedTextures)
+mu::NodePtr FCustomizableObjectCompiler::Export(UCustomizableObject* Object, const FCompilationOptions& InCompilerOptions, 
+	TArray<TSoftObjectPtr<UTexture>>& OutRuntimeReferencedTextures,
+	TArray<TSoftObjectPtr<UTexture>>& OutCompilerReferencedTextures )
 {
 	UE_LOG(LogMutable, Log, TEXT("Started Customizable Object Export %s."), *Object->GetName());
 
@@ -1348,6 +1350,13 @@ mu::NodePtr FCustomizableObjectCompiler::Export(UCustomizableObject* Object, con
 		OutRuntimeReferencedTextures.Add(Pair.Key);
 	}
 
+	OutCompilerReferencedTextures.Empty();
+	for (const TPair<TSoftObjectPtr<UTexture>, FMutableGraphGenerationContext::FGeneratedReferencedTexture>& Pair : GenerationContext.CompileTimeTextureMap)
+	{
+		check(Pair.Value.ID == OutCompilerReferencedTextures.Num());
+		OutCompilerReferencedTextures.Add(Pair.Key);
+	}
+
 	return MutableRoot;
 }
 
diff --git a/Engine/Plugins/Experimental/Mutable/Source/CustomizableObjectEditor/Private/MuCOE/CustomizableObjectCompiler.h b/Engine/Plugins/Experimental/Mutable/Source/CustomizableObjectEditor/Private/MuCOE/CustomizableObjectCompiler.h
--- a/Engine/Plugins/Experimental/Mutable/Source/CustomizableObjectEditor/Private/MuCOE/CustomizableObjectCompiler.h
+++ b/Engine/Plugins/Experimental/Mutable/Source/CustomizableObjectEditor/Private/MuCOE/CustomizableObjectCompiler.h
@@ -43,7 +43,7 @@ public:
 	virtual UCustomizableObject* GetRootObject(UCustomizableObject* InObject) override;
 
 	/** Generate the Mutable Graph from the Unreal Graph. */
-	mu::NodePtr Export(UCustomizableObject* Object, const FCompilationOptions& Options, TArray<TSoftObjectPtr<UTexture>>& OutReferencedTextures);
+	mu::NodePtr Export(UCustomizableObject* Object, const FCompilationOptions& Options, TArray<TSoftObjectPtr<UTexture>>& OutRuntimeReferencedTextures, TArray<TSoftObjectPtr<UTexture>>& OutCompilerReferencedTextures);
 
 	void CompilerLog(const FText& Message, const TArray<const UObject*>& UObject, const EMessageSeverity::Type MessageSeverity = EMessageSeverity::Warning, const bool bAddBaseObjectInfo = true, const ELoggerSpamBin SpamBin = ELoggerSpamBin::ShowAll);
 	void CompilerLog(const FText& Message, const UObject* Context = nullptr, const EMessageSeverity::Type MessageSeverity = EMessageSeverity::Warning, const bool bAddBaseObjectInfo = true, const ELoggerSpamBin SpamBin = ELoggerSpamBin::ShowAll);
diff --git a/Engine/Plugins/Experimental/Mutable/Source/CustomizableObjectEditor/Private/MuCOE/GenerateMutableSource/GenerateMutableSource.cpp b/Engine/Plugins/Experimental/Mutable/Source/CustomizableObjectEditor/Private/MuCOE/GenerateMutableSource/GenerateMutableSource.cpp
--- a/Engine/Plugins/Experimental/Mutable/Source/CustomizableObjectEditor/Private/MuCOE/GenerateMutableSource/GenerateMutableSource.cpp
+++ b/Engine/Plugins/Experimental/Mutable/Source/CustomizableObjectEditor/Private/MuCOE/GenerateMutableSource/GenerateMutableSource.cpp
@@ -1712,9 +1712,10 @@ mu::Ptr<mu::Image> GenerateImageConstant(UTexture* Texture, FMutableGraphGenerat
 		Entry.ID = TextureMap.Num()-1;
 	}
 
-	// Create a descriptor for the image but fill it only if it is not a true passthrough image. Otherwise we want it empty.
+	// Create a descriptor for the image.
+	// \TODO: If passthrough (bIsReference) we should apply lod bias, and max texture size to this desc.
+	// For now it is not a problem because passthrough textures shouldn't mix with any other operations.
 	mu::FImageDesc ImageDesc;
-	if (bForceLoad)
 	{
 		ImageDesc.m_size[0] = Texture->Source.GetSizeX();
 		ImageDesc.m_size[1] = Texture->Source.GetSizeY();
diff --git a/Engine/Plugins/Experimental/Mutable/Source/CustomizableObjectEditor/Private/MuCOE/SMutableGraphViewer.cpp b/Engine/Plugins/Experimental/Mutable/Source/CustomizableObjectEditor/Private/MuCOE/SMutableGraphViewer.cpp
--- a/Engine/Plugins/Experimental/Mutable/Source/CustomizableObjectEditor/Private/MuCOE/SMutableGraphViewer.cpp
+++ b/Engine/Plugins/Experimental/Mutable/Source/CustomizableObjectEditor/Private/MuCOE/SMutableGraphViewer.cpp
@@ -131,14 +131,14 @@ FString SMutableGraphViewer::GetReferencerName() const
 
 
 void SMutableGraphViewer::Construct(const FArguments& InArgs, const mu::NodePtr& InRootNode,
-	const TArray<TSoftObjectPtr<UTexture>>& InReferencedTextures,
 	const FCompilationOptions& InCompileOptions,
 	TWeakPtr<FTabManager> InParentTabManager, const FName& InParentNewTabId)
 {
 	DataTag = InArgs._DataTag;
+	ReferencedRuntimeTextures = InArgs._ReferencedRuntimeTextures;
+	ReferencedCompileTextures = InArgs._ReferencedCompileTextures;
 	RootNode = InRootNode;
-	ReferencedTextures = InReferencedTextures;
-	CompileOptions = InCompileOptions;	
+	CompileOptions = InCompileOptions;
 	ParentTabManager = InParentTabManager;
 	ParentNewTabId = InParentNewTabId;
 
@@ -273,6 +273,7 @@ void SMutableGraphViewer::CompileMutableCodePressed()
 	// Do the compilation to Mutable Code synchronously.
 	TSharedPtr<FCustomizableObjectCompileRunnable> CompileTask = MakeShareable(new FCustomizableObjectCompileRunnable(RootNode));
 	CompileTask->Options = CompileOptions;
+	CompileTask->ReferencedTextures = ReferencedCompileTextures;
 	CompileTask->Init();
 	CompileTask->Run();
 
@@ -281,7 +282,7 @@ void SMutableGraphViewer::CompileMutableCodePressed()
 	TSharedPtr<SDockTab> NewMutableCodeTab = SNew(SDockTab)
 		.Label(LOCTEXT("MutableCode", "Mutable Code"))
 		[
-			SNew(SMutableCodeViewer, CompileTask->Model, ReferencedTextures)
+			SNew(SMutableCodeViewer, CompileTask->Model, ReferencedRuntimeTextures)
 			.DataTag(NewDataTag)
 		];
 
diff --git a/Engine/Plugins/Experimental/Mutable/Source/CustomizableObjectEditor/Private/MuCOE/SMutableGraphViewer.h b/Engine/Plugins/Experimental/Mutable/Source/CustomizableObjectEditor/Private/MuCOE/SMutableGraphViewer.h
--- a/Engine/Plugins/Experimental/Mutable/Source/CustomizableObjectEditor/Private/MuCOE/SMutableGraphViewer.h
+++ b/Engine/Plugins/Experimental/Mutable/Source/CustomizableObjectEditor/Private/MuCOE/SMutableGraphViewer.h
@@ -32,12 +32,16 @@ public:
 	// SWidget interface
 	SLATE_BEGIN_ARGS(SMutableGraphViewer) {}
 
-	/** User-visible tag to indentify the source of the data shwon. */
+	/** User-visible tag to indentify the source of the data shown. */
 	SLATE_ARGUMENT(FString, DataTag)
 
+	SLATE_ARGUMENT(TArray<TSoftObjectPtr<UTexture>>, ReferencedRuntimeTextures)
+	SLATE_ARGUMENT(TArray<TSoftObjectPtr<UTexture>>, ReferencedCompileTextures)
+
 	SLATE_END_ARGS()
 
-	void Construct(const FArguments& InArgs, const mu::NodePtr& InRootNode, const TArray<TSoftObjectPtr<UTexture>>& ReferencedTextures, const FCompilationOptions& InCompileOptions,
+	void Construct(const FArguments& InArgs, const mu::NodePtr& InRootNode, 
+		const FCompilationOptions& InCompileOptions,
 		TWeakPtr<FTabManager> InParentTabManager, const FName& InParentNewTabId);
 
 	// SWidget interface
@@ -57,7 +61,8 @@ private:
 	mu::NodePtr RootNode;
 
 	/** Array of external referenced textures in MutableModel, indexed by id. */
-	TArray<TSoftObjectPtr<UTexture>> ReferencedTextures;
+	TArray<TSoftObjectPtr<UTexture>> ReferencedRuntimeTextures;
+	TArray<TSoftObjectPtr<UTexture>> ReferencedCompileTextures;
 
 	/** Compilation options to use in the debugger operations. */
 	FCompilationOptions CompileOptions;
diff --git a/Engine/Plugins/Experimental/Mutable/Source/CustomizableObjectEditor/Private/MuCOE/SMutableMeshViewer.cpp b/Engine/Plugins/Experimental/Mutable/Source/CustomizableObjectEditor/Private/MuCOE/SMutableMeshViewer.cpp
--- a/Engine/Plugins/Experimental/Mutable/Source/CustomizableObjectEditor/Private/MuCOE/SMutableMeshViewer.cpp
+++ b/Engine/Plugins/Experimental/Mutable/Source/CustomizableObjectEditor/Private/MuCOE/SMutableMeshViewer.cpp
@@ -203,7 +203,6 @@ void SMutableMeshViewer::SetMesh(const mu::MeshPtrConst& InMesh)
 	if (InMesh != MutableMesh)
 	{
 		MutableMesh = InMesh;
-		bIsPendingUpdate = true;
 
 		if (MutableMesh)
 		{
diff --git a/Engine/Plugins/Experimental/Mutable/Source/CustomizableObjectEditor/Private/MuCOE/SMutableMeshViewer.h b/Engine/Plugins/Experimental/Mutable/Source/CustomizableObjectEditor/Private/MuCOE/SMutableMeshViewer.h
--- a/Engine/Plugins/Experimental/Mutable/Source/CustomizableObjectEditor/Private/MuCOE/SMutableMeshViewer.h
+++ b/Engine/Plugins/Experimental/Mutable/Source/CustomizableObjectEditor/Private/MuCOE/SMutableMeshViewer.h
@@ -78,9 +78,6 @@ private:
 	/** Data backend for the widget. It represents the mesh that is being "displayed" */
 	mu::MeshPtrConst MutableMesh = nullptr;
 
-	/** Is true, mesh has changed and we need to update. */
-	bool bIsPendingUpdate = false;
-
 	/** Splitter used to separate the two sides of the slate (tables and viewport) */
 	TSharedPtr<SSplitter> SpaceSplitter;
 
diff --git a/Engine/Plugins/Experimental/Mutable/Source/CustomizableObjectEditor/Private/MuCOE/SMutableObjectViewer.cpp b/Engine/Plugins/Experimental/Mutable/Source/CustomizableObjectEditor/Private/MuCOE/SMutableObjectViewer.cpp
--- a/Engine/Plugins/Experimental/Mutable/Source/CustomizableObjectEditor/Private/MuCOE/SMutableObjectViewer.cpp
+++ b/Engine/Plugins/Experimental/Mutable/Source/CustomizableObjectEditor/Private/MuCOE/SMutableObjectViewer.cpp
@@ -215,8 +215,9 @@ void SMutableObjectViewer::Construct(const FArguments& InArgs, UCustomizableObje
 void SMutableObjectViewer::GenerateMutableGraphPressed()
 {
 	// Convert from Unreal graph to Mutable graph.
-	TArray<TSoftObjectPtr<UTexture>> ReferencedTextures;
-	mu::Ptr<mu::Node> RootNode = Compiler.Export(CustomizableObject, CompileOptions, ReferencedTextures);
+	TArray<TSoftObjectPtr<UTexture>> RuntimeTextures;
+	TArray<TSoftObjectPtr<UTexture>> CompilerTextures;
+	mu::Ptr<mu::Node> RootNode = Compiler.Export(CustomizableObject, CompileOptions, RuntimeTextures, CompilerTextures);
 	if (!RootNode)
 	{
 		// TODO: Show errors
@@ -228,8 +229,10 @@ void SMutableObjectViewer::GenerateMutableGraphPressed()
 	TSharedPtr<SDockTab> NewMutableGraphTab = SNew(SDockTab)
 		.Label(LOCTEXT("MutableGraph", "Mutable Graph"))
 		[
-			SNew(SMutableGraphViewer, RootNode, ReferencedTextures, CompileOptions, ParentTabManager, ParentNewTabId)
+			SNew(SMutableGraphViewer, RootNode, CompileOptions, ParentTabManager, ParentNewTabId)
 			.DataTag(DataTag)
+			.ReferencedRuntimeTextures(RuntimeTextures)
+			.ReferencedCompileTextures(CompilerTextures)
 		];
 
 	TSharedPtr<FTabManager> TabManager = ParentTabManager.Pin();
@@ -240,7 +243,8 @@ void SMutableObjectViewer::GenerateMutableGraphPressed()
 
 void SMutableObjectViewer::CompileMutableCodePressed()
 {
-	TArray<TSoftObjectPtr<UTexture>> ReferencedTextures;
+	TArray<TSoftObjectPtr<UTexture>> RuntimeTextures;
+	TArray<TSoftObjectPtr<UTexture>> CompilerTextures;
 	if (CompileOptions.bForceLargeLODBias)
 	{
 		// Debug compile with many different biasses
@@ -249,7 +253,9 @@ void SMutableObjectViewer::CompileMutableCodePressed()
 		{
 			CompileOptions.DebugBias = Bias;
 
-			mu::NodePtr RootNode = Compiler.Export(CustomizableObject, CompileOptions, ReferencedTextures);
+			RuntimeTextures.Empty();
+			CompilerTextures.Empty();
+			mu::NodePtr RootNode = Compiler.Export(CustomizableObject, CompileOptions, RuntimeTextures, CompilerTextures);
 			if (!RootNode)
 			{
 				// TODO: Show errors
@@ -260,6 +266,7 @@ void SMutableObjectViewer::CompileMutableCodePressed()
 			// Do the compilation to Mutable Code synchronously.
 			TSharedPtr<FCustomizableObjectCompileRunnable> CompileTask = MakeShareable(new FCustomizableObjectCompileRunnable(RootNode));
 			CompileTask->Options = CompileOptions;
+			CompileTask->ReferencedTextures = CompilerTextures;
 			CompileTask->Init();
 			CompileTask->Run();
 		}
@@ -267,7 +274,9 @@ void SMutableObjectViewer::CompileMutableCodePressed()
 	}
 
 	// Convert from Unreal graph to Mutable graph.
-	mu::NodePtr RootNode = Compiler.Export(CustomizableObject, CompileOptions, ReferencedTextures);
+	RuntimeTextures.Empty();
+	CompilerTextures.Empty();
+	mu::NodePtr RootNode = Compiler.Export(CustomizableObject, CompileOptions, RuntimeTextures, CompilerTextures);
 	if (!RootNode)
 	{
 		// TODO: Show errors
@@ -278,7 +287,7 @@ void SMutableObjectViewer::CompileMutableCodePressed()
 	// Do the compilation to Mutable Code synchronously.
 	TSharedPtr<FCustomizableObjectCompileRunnable> CompileTask = MakeShareable(new FCustomizableObjectCompileRunnable(RootNode));
 	CompileTask->Options = CompileOptions;
-	CompileTask->ReferencedTextures = ReferencedTextures;
+	CompileTask->ReferencedTextures = CompilerTextures;
 	CompileTask->Init();
 	CompileTask->Run();
 
@@ -287,7 +296,7 @@ void SMutableObjectViewer::CompileMutableCodePressed()
 	TSharedPtr<SDockTab> NewMutableCodeTab = SNew(SDockTab)
 		.Label(LOCTEXT("MutableCode", "Mutable Code"))
 		[
-			SNew(SMutableCodeViewer, CompileTask->Model, ReferencedTextures)
+			SNew(SMutableCodeViewer, CompileTask->Model, RuntimeTextures)
 			.DataTag(DataTag)
 		];
 
diff --git a/Engine/Plugins/Experimental/Mutable/Source/MutableRuntime/Private/MuR/CodeRunnerTasked.cpp b/Engine/Plugins/Experimental/Mutable/Source/MutableRuntime/Private/MuR/CodeRunnerTasked.cpp
--- a/Engine/Plugins/Experimental/Mutable/Source/MutableRuntime/Private/MuR/CodeRunnerTasked.cpp
+++ b/Engine/Plugins/Experimental/Mutable/Source/MutableRuntime/Private/MuR/CodeRunnerTasked.cpp
@@ -1798,7 +1798,7 @@ namespace mu
 	bool FImageInvertTask::Prepare(CodeRunner* Runner, bool& bOutFailed)
 	{
 		// This runs in the mutable Runner thread
-		MUTABLE_CPUPROFILER_SCOPE(FImageLayerTask_Prepare);
+		MUTABLE_CPUPROFILER_SCOPE(FImageInvertTask_Prepare);
 		bOutFailed = false;
 
 		Ptr<const Image> Source = Runner->LoadImage({ Args.base, Op.ExecutionIndex, Op.ExecutionOptions });
@@ -1858,7 +1858,7 @@ namespace mu
 	bool FImageComposeTask::Prepare(CodeRunner* Runner, bool& bOutFailed)
 	{
 		// This runs in the mutable Runner thread
-		MUTABLE_CPUPROFILER_SCOPE(FImageLayerTask_Prepare);
+		MUTABLE_CPUPROFILER_SCOPE(FImageComposeTask_Prepare);
 		bOutFailed = false;
 
 		ImageCompressionQuality = Runner->m_pSettings->ImageCompressionQuality;
diff --git a/Engine/Plugins/Experimental/Mutable/Source/MutableTools/Private/MuT/CodeGenerator.cpp b/Engine/Plugins/Experimental/Mutable/Source/MutableTools/Private/MuT/CodeGenerator.cpp
--- a/Engine/Plugins/Experimental/Mutable/Source/MutableTools/Private/MuT/CodeGenerator.cpp
+++ b/Engine/Plugins/Experimental/Mutable/Source/MutableTools/Private/MuT/CodeGenerator.cpp
@@ -1274,7 +1274,7 @@ namespace mu
 					if (desc.m_size[0] == 0 || desc.m_size[1] == 0)
 					{
 						int currentLOD = m_currentParents.Last().m_lod;
-						FString Msg = FString::Printf( TEXT("An image block for [%s] [%s] [%s] at lod [%d] has zero size and will not be generated. "),
+						FString Msg = FString::Printf( TEXT("An image for [%s] [%s] [%s] at lod [%d] has zero size and will not be generated. "),
 							*node.m_images[t].m_name,
 							*node.m_images[t].m_materialName,
 							*node.m_images[t].m_materialParameterName,
diff --git a/Engine/Plugins/Experimental/Mutable/Source/MutableTools/Private/MuT/ImageDescGenerator.cpp b/Engine/Plugins/Experimental/Mutable/Source/MutableTools/Private/MuT/ImageDescGenerator.cpp
--- a/Engine/Plugins/Experimental/Mutable/Source/MutableTools/Private/MuT/ImageDescGenerator.cpp
+++ b/Engine/Plugins/Experimental/Mutable/Source/MutableTools/Private/MuT/ImageDescGenerator.cpp
@@ -119,21 +119,19 @@ namespace mu
             return 0;
         }
 
-        Ptr<const Image> pImage;
-
-        std::size_t i = 0;
-        while ( !pImage && i<node.m_pTable->GetPrivate()->Rows.Num() )
+        int32 i = 0;
+        while ( i<node.m_pTable->GetPrivate()->Rows.Num() )
         {
-            pImage = node.m_pTable->GetPrivate()->Rows[i].Values[ colIndex ].ProxyImage->Get();
+			Ptr<const Image> pImage = node.m_pTable->GetPrivate()->Rows[i].Values[ colIndex ].ProxyImage->Get();
             ++i;
-        }
 
-        if (pImage)
-        {
-            m_desc.m_size[0] = pImage->GetSizeX();
-            m_desc.m_size[1] = pImage->GetSizeY();
-            m_desc.m_format = pImage->GetFormat();
-            m_desc.m_lods = (uint8_t)pImage->GetLODCount();
+			if (pImage)
+			{
+				m_desc.m_size[0] = FMath::Max(m_desc.m_size[0], pImage->GetSizeX());
+				m_desc.m_size[1] = FMath::Max(m_desc.m_size[1], pImage->GetSizeX());
+				m_desc.m_lods = FMath::Max(m_desc.m_lods, uint8(pImage->GetLODCount()));
+				m_desc.m_format = mu::GetMostGenericFormat(m_desc.m_format, pImage->GetFormat());
+			}
         }
 
         return 0;
diff --git a/Engine/Plugins/Experimental/Mutable/Source/MutableTools/Private/MuT/Node.cpp b/Engine/Plugins/Experimental/Mutable/Source/MutableTools/Private/MuT/Node.cpp
--- a/Engine/Plugins/Experimental/Mutable/Source/MutableTools/Private/MuT/Node.cpp
+++ b/Engine/Plugins/Experimental/Mutable/Source/MutableTools/Private/MuT/Node.cpp
@@ -23,10 +23,6 @@
 #include "MuT/NodeScalar.h"
 #include "MuT/NodeString.h"
 #include "MuT/NodeSurface.h"
-#include "map"
-
-#include <stdint.h>
-#include <utility>
 
 
 namespace mu
