[Mutable] Fix ImageTransform operations generating small images with None optimization.

- The reference image descriptor was not propagated form the node to the ASTOp.

#rnx

diff --git a/Engine/Plugins/Experimental/Mutable/Source/MutableTools/Private/MuT/ASTOpReferenceResource.cpp b/Engine/Plugins/Experimental/Mutable/Source/MutableTools/Private/MuT/ASTOpReferenceResource.cpp
--- a/Engine/Plugins/Experimental/Mutable/Source/MutableTools/Private/MuT/ASTOpReferenceResource.cpp
+++ b/Engine/Plugins/Experimental/Mutable/Source/MutableTools/Private/MuT/ASTOpReferenceResource.cpp
@@ -72,8 +72,7 @@ namespace mu
 	//-------------------------------------------------------------------------------------------------
 	FImageDesc ASTOpReferenceResource::GetImageDesc(bool, class FGetImageDescContext*) const
 	{
-		FImageDesc res;
-		return res;
+		return ImageDesc;
 	}
 
 
diff --git a/Engine/Plugins/Experimental/Mutable/Source/MutableTools/Private/MuT/CodeGenerator_Image.cpp b/Engine/Plugins/Experimental/Mutable/Source/MutableTools/Private/MuT/CodeGenerator_Image.cpp
--- a/Engine/Plugins/Experimental/Mutable/Source/MutableTools/Private/MuT/CodeGenerator_Image.cpp
+++ b/Engine/Plugins/Experimental/Mutable/Source/MutableTools/Private/MuT/CodeGenerator_Image.cpp
@@ -197,7 +197,9 @@ namespace mu
 			Ptr<ASTOpReferenceResource> ReferenceOp = new ASTOpReferenceResource();
 			ReferenceOp->type = OP_TYPE::IM_REFERENCE;
 			ReferenceOp->ID = pImage->GetReferencedTexture();
-			ReferenceOp->bForceLoad = pImage->IsForceLoad();
+			ReferenceOp->bForceLoad = pImage->IsForceLoad(); 
+			ReferenceOp->ImageDesc = FImageDesc(
+					FImageSize(pImage->GetSizeX(), pImage->GetSizeY()), pImage->GetFormat(), pImage->GetLODCount());
 			Result.op = ReferenceOp;
 		}
 		else
