[mutable] Updated variable holding byte data for RomTotalBytes counter.
- Used instead variable with bigger capacity since all ints get converted to int64 at the MongoDB level

diff --git a/Engine/Plugins/Experimental/Mutable/Source/MutableValidation/Private/MuV/CustomizableObjectValidationCommandlet.cpp b/Engine/Plugins/Experimental/Mutable/Source/MutableValidation/Private/MuV/CustomizableObjectValidationCommandlet.cpp
--- a/Engine/Plugins/Experimental/Mutable/Source/MutableValidation/Private/MuV/CustomizableObjectValidationCommandlet.cpp
+++ b/Engine/Plugins/Experimental/Mutable/Source/MutableValidation/Private/MuV/CustomizableObjectValidationCommandlet.cpp
@@ -83,15 +83,15 @@ int32 UCustomizableObjectValidationCommandlet::Main(const FString& Params)
 		if (const TSharedPtr<const mu::Model> MutableModel = ToTestCustomizableObject->GetModel())
 		{
 			const int32 RomCount =  MutableModel->GetRomCount();
-			uint32 TotalRomSizeBits = 0;
+			int64 TotalRomSizeBytes = 0;
 			for (int32 RomIndex = 0; RomIndex < RomCount; RomIndex++)
 			{
-				const uint32 RomBitSize = MutableModel->GetRomSize(RomIndex);
-				TotalRomSizeBits += RomBitSize;
+				const uint32 RomByteSize = MutableModel->GetRomSize(RomIndex);
+				TotalRomSizeBytes += RomByteSize;
 			}
 			
 			UE_LOG(LogMutable, Log,TEXT("(int) model_rom_count : %d "), RomCount);
-			UE_LOG(LogMutable, Log,TEXT("(uint32) model_roms_size : %u "), TotalRomSizeBits);
+			UE_LOG(LogMutable, Log,TEXT("(int) model_roms_size : %lld "), TotalRomSizeBytes);
 		}
 		
 		// Generate target random instances to be tested ------------------------------------------------------------ //
