[mutable] Moved animation data stored in Table Node pins. Now this information is stored in the node and can be shared between all mesh pins of the same skeletal mesh.

#rnx

diff --git a/Engine/Plugins/Experimental/Mutable/Source/CustomizableObject/Public/MuCO/CustomizableObjectCustomVersion.h b/Engine/Plugins/Experimental/Mutable/Source/CustomizableObject/Public/MuCO/CustomizableObjectCustomVersion.h
--- a/Engine/Plugins/Experimental/Mutable/Source/CustomizableObject/Public/MuCO/CustomizableObjectCustomVersion.h
+++ b/Engine/Plugins/Experimental/Mutable/Source/CustomizableObject/Public/MuCO/CustomizableObjectCustomVersion.h
@@ -108,6 +108,8 @@ struct CUSTOMIZABLEOBJECT_API FCustomizableObjectCustomVersion
 
 		CustomizableObjectNodeHasSeparateNeverStreamFlag,
 
+		AddedColumnIdDataToTableNodePins,
+
 		// -----<new versions can be added above this line>-------------------------------------------------
 		VersionPlusOne,
 		LatestVersion = VersionPlusOne - 1
diff --git a/Engine/Plugins/Experimental/Mutable/Source/CustomizableObjectEditor/Classes/MuCOE/Nodes/CustomizableObjectNodeTable.h b/Engine/Plugins/Experimental/Mutable/Source/CustomizableObjectEditor/Classes/MuCOE/Nodes/CustomizableObjectNodeTable.h
--- a/Engine/Plugins/Experimental/Mutable/Source/CustomizableObjectEditor/Classes/MuCOE/Nodes/CustomizableObjectNodeTable.h
+++ b/Engine/Plugins/Experimental/Mutable/Source/CustomizableObjectEditor/Classes/MuCOE/Nodes/CustomizableObjectNodeTable.h
@@ -52,6 +52,25 @@ enum class ETableDataGatheringSource : uint8
 };
 
 
+USTRUCT()
+struct FTableNodeColumnData
+{
+	GENERATED_USTRUCT_BODY()
+
+	/** Anim Instance Column name related to this Mesh pin */
+	UPROPERTY()
+	FString AnimInstanceColumnName = "";
+
+	/** Anim Slot Column name related to this Mesh pin */
+	UPROPERTY()
+	FString AnimSlotColumnName = "";
+
+	/** Anim Tag Column name related to this Mesh pin */
+	UPROPERTY()
+	FString AnimTagColumnName = "";
+};
+
+
 /** Base class for all Table Pins. */
 UCLASS()
 class CUSTOMIZABLEOBJECTEDITOR_API UCustomizableObjectNodeTableObjectPinData : public UCustomizableObjectNodePinData
@@ -60,6 +79,10 @@ class CUSTOMIZABLEOBJECTEDITOR_API UCustomizableObjectNodeTableObjectPinData : p
 
 public:
 
+	/** Id of the property associated to a struct column */
+	UPROPERTY()
+	FGuid StructColumnId;
+
 	/** Name of the data table column related to the pin */
 	UPROPERTY()
 	FString ColumnName = "";
@@ -103,18 +126,18 @@ class CUSTOMIZABLEOBJECTEDITOR_API UCustomizableObjectNodeTableMeshPinData : pub
 	GENERATED_BODY()
 
 public:
-	
+
 	/** Anim Instance Column name related to this Mesh pin */
 	UPROPERTY()
-	FString AnimInstanceColumnName = "";
+	FString AnimInstanceColumnName_DEPRECATED = "";
 
 	/** Anim Slot Column name related to this Mesh pin */
 	UPROPERTY()
-	FString AnimSlotColumnName = "";
+	FString AnimSlotColumnName_DEPRECATED = "";
 
 	/** Anim Tag Column name related to this Mesh pin */
 	UPROPERTY()
-	FString AnimTagColumnName = "";
+	FString AnimTagColumnName_DEPRECATED = "";
 
 	UPROPERTY()
 	FString MutableColumnName = "";
@@ -190,6 +213,12 @@ public:
 	UPROPERTY(EditAnywhere, Category = UI, meta = (DisplayName = "Parameter UI Metadata"))
 	FMutableParamUIMetadata ParamUIMetadata;
 
+	/** Map to relate a Structure Column with its Data */
+	UPROPERTY()
+	TMap<FGuid, FTableNodeColumnData> ColumnDataMap;
+
+public:
+
 	// UObject interface
 	virtual void PreEditChange(FProperty* PropertyAboutToChange) override;
 	virtual void PostEditChangeChainProperty(FPropertyChangedChainEvent& PropertyChangedEvent) override;
@@ -236,7 +265,7 @@ public:
 	void GetPinLODAndSection(const UEdGraphPin* Pin, int32& LODIndex, int32& SectionIndex) const;
 
 	// Get the anim blueprint and anim slot columns related to a mesh
-	void GetAnimationColumns(const FString& ColumnName, FString& AnimBPColumnName, FString& AnimSlotColumnName, FString& AnimTagColumnName) const;
+	void GetAnimationColumns(const FGuid& ColumnId, FString& AnimBPColumnName, FString& AnimSlotColumnName, FString& AnimTagColumnName) const;
 
 	/** Callback called when the Table property or its contents has changed. */
 	void OnTableChanged();
@@ -356,6 +385,10 @@ public:
 	/** Returns the column property where PropertyName matches the name of the column property. Returns nullptr if no match is found or the match is not a supported table property */
 	FProperty* FindTableProperty(const UScriptStruct* ScriptStruct, const FName& PropertyName) const;
 
+	/** Returns the column id */
+	FGuid GetColumnIdByName(const FName& ColumnName) const;
+
+
 private:
 
 	/** Number of properties to know when the node needs an update */
@@ -365,10 +398,9 @@ private:
 	FDelegateHandle OnTableChangedDelegateHandle;
 	
 	// Generates a mesh pin for each LOD and Material Surface of the reference SkeletalMesh
-	void GenerateMeshPins(UObject* Mesh, const FString& Name);
+	void GenerateMeshPins(UObject* Mesh, const FString& ColumnName, const FGuid& ColumnId);
 
 	// Checks if a pin already exists and if it has the same type as before the node refresh
 	bool CheckPinUpdated(const FString& PinName, const FName& PinType) const;
 
-
 };
diff --git a/Engine/Plugins/Experimental/Mutable/Source/CustomizableObjectEditor/Private/MuCOE/GenerateMutableSource/GenerateMutableSourceTable.cpp b/Engine/Plugins/Experimental/Mutable/Source/CustomizableObjectEditor/Private/MuCOE/GenerateMutableSource/GenerateMutableSourceTable.cpp
--- a/Engine/Plugins/Experimental/Mutable/Source/CustomizableObjectEditor/Private/MuCOE/GenerateMutableSource/GenerateMutableSourceTable.cpp
+++ b/Engine/Plugins/Experimental/Mutable/Source/CustomizableObjectEditor/Private/MuCOE/GenerateMutableSource/GenerateMutableSourceTable.cpp
@@ -9,6 +9,7 @@
 #include "Engine/StaticMesh.h"
 #include "Engine/Texture2DArray.h"
 #include "GameplayTagContainer.h"
+#include "Kismet2/StructureEditorUtils.h"
 #include "Materials/Material.h"
 #include "Materials/MaterialInstance.h"
 #include "MuCO/CustomizableObjectSystem.h"
@@ -60,8 +61,9 @@ bool FillTableColumn(const UCustomizableObjectNodeTable* TableNode,	mu::TablePtr
 			// Getting Animation Blueprint and Animation Slot
 			FString AnimBP, AnimSlot, GameplayTag, AnimBPAssetTag;
 			TArray<FGameplayTag> GameplayTags;
+			FGuid ColumnPropertyId = FStructureEditorUtils::GetGuidForProperty(ColumnProperty);
 
-			TableNode->GetAnimationColumns(ColumnName, AnimBP, AnimSlot, GameplayTag);
+			TableNode->GetAnimationColumns(ColumnPropertyId, AnimBP, AnimSlot, GameplayTag);
 
 			if (!AnimBP.IsEmpty())
 			{
diff --git a/Engine/Plugins/Experimental/Mutable/Source/CustomizableObjectEditor/Private/MuCOE/Nodes/CustomizableObjectNodeTable.cpp b/Engine/Plugins/Experimental/Mutable/Source/CustomizableObjectEditor/Private/MuCOE/Nodes/CustomizableObjectNodeTable.cpp
--- a/Engine/Plugins/Experimental/Mutable/Source/CustomizableObjectEditor/Private/MuCOE/Nodes/CustomizableObjectNodeTable.cpp
+++ b/Engine/Plugins/Experimental/Mutable/Source/CustomizableObjectEditor/Private/MuCOE/Nodes/CustomizableObjectNodeTable.cpp
@@ -29,17 +29,39 @@ class UCustomizableObjectNodeRemapPins;
 
 bool UCustomizableObjectNodeTableRemapPins::Equal(const UCustomizableObjectNode& Node, const UEdGraphPin& OldPin, const UEdGraphPin& NewPin) const
 {
+	const UCustomizableObjectNodeTableObjectPinData& OldPinData = Node.GetPinData<UCustomizableObjectNodeTableObjectPinData>(OldPin);
+	const UCustomizableObjectNodeTableObjectPinData& NewPinData = Node.GetPinData<UCustomizableObjectNodeTableObjectPinData>(NewPin);
+
+	// If one of these two option fails, pins are different
+	if (OldPinData.StructColumnId != NewPinData.StructColumnId || OldPin.Direction != NewPin.Direction)
+	{
+		return false;
+	}
+
+	// In this case pin type may have changed but we consider them the same type
 	if (OldPin.PinType.PinCategory == UEdGraphSchema_CustomizableObject::PC_Image ||
 		OldPin.PinType.PinCategory == UEdGraphSchema_CustomizableObject::PC_PassThroughImage)
 	{
-		// In this case pin type may have changed but we consider them the same type
-		return Helper_GetPinName(&OldPin) == Helper_GetPinName(&NewPin) &&
-			OldPin.Direction == NewPin.Direction;
+		return true;
+	}
+
+	// Non image pins must remain the same pin type
+	if (OldPin.PinType != NewPin.PinType)
+	{
+		return false;
 	}
 
-	return Helper_GetPinName(&OldPin) == Helper_GetPinName(&NewPin) &&
-		OldPin.PinType == NewPin.PinType &&
-		OldPin.Direction == NewPin.Direction;
+	if (OldPin.PinType.PinCategory == UEdGraphSchema_CustomizableObject::PC_Mesh)
+	{
+		const UCustomizableObjectNodeTableMeshPinData& OldMeshPinData = Node.GetPinData<UCustomizableObjectNodeTableMeshPinData>(OldPin);
+		const UCustomizableObjectNodeTableMeshPinData& NewMeshPinData = Node.GetPinData<UCustomizableObjectNodeTableMeshPinData>(NewPin);
+
+		// LOD and Section must remain the same
+		return 	OldMeshPinData.LOD == NewMeshPinData.LOD &&
+			OldMeshPinData.Material == NewMeshPinData.Material;
+	}
+
+	return true;
 }
 
 
@@ -95,6 +117,53 @@ void UCustomizableObjectNodeTable::BackwardsCompatibleFixup()
 			}
 		}
 	}
+	
+	// Adding StructColumnID
+	if(CustomizableObjectCustomVersion < FCustomizableObjectCustomVersion::AddedColumnIdDataToTableNodePins)
+	{
+		const UScriptStruct* TableStruct = GetTableNodeStruct();
+
+		if (TableStruct)
+		{
+			for (UEdGraphPin* Pin : Pins)
+			{
+				UCustomizableObjectNodeTableObjectPinData* PinData = Cast<UCustomizableObjectNodeTableObjectPinData>(GetPinData(*(Pin)));
+
+				// Adding pindata to float and colors
+				if (!PinData)
+				{
+					PinData = NewObject<UCustomizableObjectNodeTableObjectPinData>(this);
+					PinData->ColumnName = Pin->PinFriendlyName.ToString();
+
+					AddPinData(*Pin, *PinData);
+				}
+
+				FProperty* ColumnProperty = FindTableProperty(TableStruct, FName(*PinData->ColumnName));
+
+				if (!ColumnProperty)
+				{
+					continue;
+				}
+
+				FGuid ColumnPropertyId = FStructureEditorUtils::GetGuidForProperty(ColumnProperty);
+				PinData->StructColumnId = ColumnPropertyId;
+
+				// Store anim columns in the node instead of the pin
+				if (UCustomizableObjectNodeTableMeshPinData* MeshPinData = Cast<UCustomizableObjectNodeTableMeshPinData>(GetPinData(*(Pin))))
+				{
+					if (!MeshPinData->AnimInstanceColumnName_DEPRECATED.IsEmpty())
+					{
+						FTableNodeColumnData NewColumnData;
+						NewColumnData.AnimInstanceColumnName = MeshPinData->AnimInstanceColumnName_DEPRECATED;
+						NewColumnData.AnimSlotColumnName = MeshPinData->AnimSlotColumnName_DEPRECATED;
+						NewColumnData.AnimTagColumnName = MeshPinData->AnimTagColumnName_DEPRECATED;
+
+						ColumnDataMap.Add(ColumnPropertyId, NewColumnData);
+					}
+				}
+			}
+		}
+	}
 }
 
 
@@ -246,6 +315,10 @@ void UCustomizableObjectNodeTable::PostEditChangeProperty(FPropertyChangedEvent&
 
 void UCustomizableObjectNodeTable::AllocateDefaultPins(UCustomizableObjectNodeRemapPins* RemapPins)
 {
+	// Reset of the column data map
+	TMap<FGuid, FTableNodeColumnData> AuxOldColumnData = ColumnDataMap;
+	ColumnDataMap.Reset();
+
 	// Getting Struct Pointer
 	const UScriptStruct* TableStruct = GetTableNodeStruct();
 
@@ -272,13 +345,14 @@ void UCustomizableObjectNodeTable::AllocateDefaultPins(UCustomizableObjectNodeRe
 		{
 			continue;
 		}
-		
+
 		const UEdGraphSchema_CustomizableObject* Schema = GetDefault<UEdGraphSchema_CustomizableObject>();
 
 		UEdGraphPin* OutPin = nullptr;
 
 		FString ColumnName = DataTableUtils::GetPropertyExportName(ColumnProperty);
 		FString PinName = ColumnName;
+		FGuid ColumnPropertyId = FStructureEditorUtils::GetGuidForProperty(ColumnProperty);
 
 		if (const FSoftObjectProperty* SoftObjectProperty = CastField<FSoftObjectProperty>(ColumnProperty))
 		{
@@ -296,13 +370,19 @@ void UCustomizableObjectNodeTable::AllocateDefaultPins(UCustomizableObjectNodeRe
 			{
 				if (Object->IsA(USkeletalMesh::StaticClass()) || Object->IsA(UStaticMesh::StaticClass()))
 				{
-					GenerateMeshPins(Object, ColumnName);
+					GenerateMeshPins(Object, ColumnName, ColumnPropertyId);
+
+					if (FTableNodeColumnData* ColumnData = AuxOldColumnData.Find(ColumnPropertyId))
+					{
+						ColumnDataMap.Add(ColumnPropertyId, *ColumnData);
+					}
 				}
 
 				else if (Object->IsA(UTexture2D::StaticClass()))
 				{
 					UCustomizableObjectNodeTableImagePinData* PinData = NewObject<UCustomizableObjectNodeTableImagePinData>(this);
 					PinData->ColumnName = ColumnName;
+					PinData->StructColumnId = ColumnPropertyId;
 					PinData->SetIsArrayTexture(false);
 
 					FName PinCategory = DefaultImageMode == ETableTextureType::PASSTHROUGH_TEXTURE ? Schema->PC_PassThroughImage : Schema->PC_Image;
@@ -327,6 +407,7 @@ void UCustomizableObjectNodeTable::AllocateDefaultPins(UCustomizableObjectNodeRe
 				{
 					UCustomizableObjectNodeTableImagePinData* PinData = NewObject<UCustomizableObjectNodeTableImagePinData>(this);
 					PinData->ColumnName = ColumnName;
+					PinData->StructColumnId = ColumnPropertyId;
 					PinData->ImageMode = ETableTextureType::PASSTHROUGH_TEXTURE;
 					PinData->SetIsArrayTexture(true);
 
@@ -337,6 +418,7 @@ void UCustomizableObjectNodeTable::AllocateDefaultPins(UCustomizableObjectNodeRe
 				{
 					UCustomizableObjectNodeTableObjectPinData* PinData = NewObject<UCustomizableObjectNodeTableObjectPinData>(this);
 					PinData->ColumnName = ColumnName;
+					PinData->StructColumnId = ColumnPropertyId;
 
 					OutPin = CustomCreatePin(EGPD_Output, Schema->PC_MaterialAsset, FName(*PinName), PinData);
 				}
@@ -357,18 +439,30 @@ void UCustomizableObjectNodeTable::AllocateDefaultPins(UCustomizableObjectNodeRe
 		{
 			if (StructProperty->Struct == TBaseStructure<FLinearColor>::Get())
 			{
-				OutPin = CustomCreatePin(EGPD_Output, Schema->PC_Color, FName(*PinName));
+				UCustomizableObjectNodeTableObjectPinData* PinData = NewObject<UCustomizableObjectNodeTableObjectPinData>(this);
+				PinData->ColumnName = ColumnName;
+				PinData->StructColumnId = ColumnPropertyId;
+
+				OutPin = CustomCreatePin(EGPD_Output, Schema->PC_Color, FName(*PinName), PinData);
 			}
 		}
 
 		else if (const FNumericProperty* NumFloatProperty = CastField<FFloatProperty>(ColumnProperty))
 		{
-			OutPin = CustomCreatePin(EGPD_Output, Schema->PC_Float, FName(*PinName));
+			UCustomizableObjectNodeTableObjectPinData* PinData = NewObject<UCustomizableObjectNodeTableObjectPinData>(this);
+			PinData->ColumnName = ColumnName;
+			PinData->StructColumnId = ColumnPropertyId;
+
+			OutPin = CustomCreatePin(EGPD_Output, Schema->PC_Float, FName(*PinName), PinData);
 		}
 		
 		else if (const FNumericProperty* NumDoubleProperty = CastField<FDoubleProperty>(ColumnProperty))
 		{
-			OutPin = CustomCreatePin(EGPD_Output, Schema->PC_Float, FName(*PinName));
+			UCustomizableObjectNodeTableObjectPinData* PinData = NewObject<UCustomizableObjectNodeTableObjectPinData>(this);
+			PinData->ColumnName = ColumnName;
+			PinData->StructColumnId = ColumnPropertyId;
+
+			OutPin = CustomCreatePin(EGPD_Output, Schema->PC_Float, FName(*PinName), PinData);
 		}
 
 		else if (const FObjectProperty* ObjectProperty = CastField<FObjectProperty>(ColumnProperty))
@@ -394,7 +488,7 @@ void UCustomizableObjectNodeTable::AllocateDefaultPins(UCustomizableObjectNodeRe
 }
 
 
-void UCustomizableObjectNodeTable::GenerateMeshPins(UObject* Mesh, const FString& Name)
+void UCustomizableObjectNodeTable::GenerateMeshPins(UObject* Mesh, const FString& ColumnName, const FGuid& ColumnId)
 {
 	if (USkeletalMesh* SkeletalMesh = Cast<USkeletalMesh>(Mesh))
 	{
@@ -407,10 +501,14 @@ void UCustomizableObjectNodeTable::GenerateMeshPins(UObject* Mesh, const FString
 
 			for (int32 MatIndex = 0; MatIndex < NumMaterials; ++MatIndex)
 			{
-				FString TableMeshPinName = GenerateSkeletalMeshMutableColumName(Name, LODIndex, MatIndex);
+				FString TableMeshPinName = GenerateSkeletalMeshMutableColumName(ColumnName, LODIndex, MatIndex);
 
+				// Pin Data
 				UCustomizableObjectNodeTableMeshPinData* PinData = NewObject<UCustomizableObjectNodeTableMeshPinData>(this);
-				PinData->ColumnName = Name;
+				PinData->ColumnName = ColumnName;
+				PinData->StructColumnId = ColumnId;
+
+				// Mesh Data
 				PinData->MutableColumnName = TableMeshPinName;
 				PinData->LOD = LODIndex;
 				PinData->Material = MatIndex;
@@ -457,10 +555,13 @@ void UCustomizableObjectNodeTable::GenerateMeshPins(UObject* Mesh, const FString
 
 			for (int32 MatIndex = 0; MatIndex < NumMaterials; ++MatIndex)
 			{
-				FString TableMeshPinName = GenerateStaticMeshMutableColumName(Name, MatIndex);
+				FString TableMeshPinName = GenerateStaticMeshMutableColumName(ColumnName, MatIndex);
 
 				UCustomizableObjectNodeTableMeshPinData* PinData = NewObject<UCustomizableObjectNodeTableMeshPinData>(this);
-				PinData->ColumnName = Name;
+				PinData->ColumnName = ColumnName;
+				PinData->StructColumnId = ColumnId;
+
+				// Mesh Data
 				PinData->MutableColumnName = TableMeshPinName;
 				PinData->LOD = 0;
 				PinData->Material = MatIndex;
@@ -675,40 +776,22 @@ void UCustomizableObjectNodeTable::RemapPinsData(const TMap<UEdGraphPin*, UEdGra
 
 			const UScriptStruct* ScriptStruct = GetTableNodeStruct();
 
-			if (PinDataOldPin && PinDataNewPin && ScriptStruct)
+			if (PinDataOldPin && PinDataNewPin && ScriptStruct && GetPinMeshType(Pair.Value) == ETableMeshPinType::SKELETAL_MESH)
 			{
-				if (FindTableProperty(ScriptStruct, FName(*PinDataOldPin->AnimInstanceColumnName)))
-				{
-					PinDataNewPin->AnimInstanceColumnName = PinDataOldPin->AnimInstanceColumnName;
-				}
-
-				if (FindTableProperty(ScriptStruct, FName(*PinDataOldPin->AnimSlotColumnName)))
+				// Keeping information added in layout editor if the layout is the same
+				for (TObjectPtr<UCustomizableObjectLayout>& NewLayout : PinDataNewPin->Layouts)
 				{
-					PinDataNewPin->AnimSlotColumnName = PinDataOldPin->AnimSlotColumnName;
-				}
-
-				if (FindTableProperty(ScriptStruct, FName(*PinDataOldPin->AnimTagColumnName)))
-				{
-					PinDataNewPin->AnimTagColumnName = PinDataOldPin->AnimTagColumnName;
-				}
-
-				if (GetPinMeshType(Pair.Value) == ETableMeshPinType::SKELETAL_MESH)
-				{
-					// Keeping information added in layout editor if the layout is the same
-					for (TObjectPtr<UCustomizableObjectLayout>& NewLayout : PinDataNewPin->Layouts)
+					for (TObjectPtr<UCustomizableObjectLayout>& OldLayout : PinDataOldPin->Layouts)
 					{
-						for (TObjectPtr<UCustomizableObjectLayout>& OldLayout : PinDataOldPin->Layouts)
+						if (NewLayout->GetLayoutName() == OldLayout->GetLayoutName())
 						{
-							if (NewLayout->GetLayoutName() == OldLayout->GetLayoutName())
-							{
-								NewLayout->Blocks = OldLayout->Blocks;
-								NewLayout->SetGridSize(OldLayout->GetGridSize());
-								NewLayout->SetMaxGridSize(OldLayout->GetMaxGridSize());
-								NewLayout->SetPackingStrategy(OldLayout->GetPackingStrategy());
-								NewLayout->SetBlockReductionMethod(OldLayout->GetBlockReductionMethod());
+							NewLayout->Blocks = OldLayout->Blocks;
+							NewLayout->SetGridSize(OldLayout->GetGridSize());
+							NewLayout->SetMaxGridSize(OldLayout->GetMaxGridSize());
+							NewLayout->SetPackingStrategy(OldLayout->GetPackingStrategy());
+							NewLayout->SetBlockReductionMethod(OldLayout->GetBlockReductionMethod());
 
-								break;
-							}
+							break;
 						}
 					}
 				}
@@ -890,22 +973,14 @@ FString UCustomizableObjectNodeTable::GetColumnNameByPin(const UEdGraphPin* Pin)
 
 FString UCustomizableObjectNodeTable::GetMutableColumnName(const UEdGraphPin* Pin, const int32& LOD) const
 {
-	UCustomizableObjectNodeTableMeshPinData* PinData = Cast<UCustomizableObjectNodeTableMeshPinData >(GetPinData(*Pin));
+	UCustomizableObjectNodeTableMeshPinData* MeshData = Cast<UCustomizableObjectNodeTableMeshPinData >(GetPinData(*Pin));
 
-	if (PinData)
+	if (MeshData)
 	{
-		FString ColumnName = PinData->ColumnName;
-		int32 MaterialIndex = PinData->Material;
-
-		for (const UEdGraphPin* NodePin : Pins)
-		{
-			const UCustomizableObjectNodeTableMeshPinData* MeshPinData = Cast<UCustomizableObjectNodeTableMeshPinData >(GetPinData(*NodePin));
+		FString ColumnName = MeshData->ColumnName;
+		int32 MaterialIndex = MeshData->Material;
 
-			if (MeshPinData && MeshPinData->ColumnName == ColumnName && MeshPinData->LOD == LOD && MeshPinData->Material == MaterialIndex)
-			{
-				return MeshPinData->MutableColumnName;
-			}
-		}
+		return GenerateSkeletalMeshMutableColumName(MeshData->ColumnName, LOD, MeshData->Material);
 	}
 
 	return FString();
@@ -927,21 +1002,14 @@ void UCustomizableObjectNodeTable::GetPinLODAndSection(const UEdGraphPin* Pin, i
 }
 
 
-void UCustomizableObjectNodeTable::GetAnimationColumns(const FString& ColumnName, FString& AnimBPColumnName, FString& AnimSlotColumnName, FString& AnimTagColumnName) const
+void UCustomizableObjectNodeTable::GetAnimationColumns(const FGuid& ColumnId, FString& AnimBPColumnName, FString& AnimSlotColumnName, FString& AnimTagColumnName) const
 {
-	for (const UEdGraphPin* Pin : Pins)
+	if (ColumnDataMap.Contains(ColumnId))
 	{
-		const UCustomizableObjectNodeTableMeshPinData* PinData = Cast<UCustomizableObjectNodeTableMeshPinData >(GetPinData(*Pin));
-
-		if (PinData)
-		{
-			if (PinData->ColumnName == ColumnName)
-			{
-				AnimBPColumnName = PinData->AnimInstanceColumnName;
-				AnimSlotColumnName = PinData->AnimSlotColumnName;
-				AnimTagColumnName = PinData->AnimTagColumnName;
-			}
-		}
+		const FTableNodeColumnData* ColumnData = ColumnDataMap.Find(ColumnId);
+		AnimBPColumnName = ColumnData->AnimInstanceColumnName;
+		AnimSlotColumnName = ColumnData->AnimSlotColumnName;
+		AnimTagColumnName = ColumnData->AnimTagColumnName;
 	}
 }
 
@@ -993,27 +1061,26 @@ USkeletalMesh* UCustomizableObjectNodeTable::GetSkeletalMeshAt(const UEdGraphPin
 
 	FString ColumnName = GetColumnNameByPin(Pin);
 
-	for (TFieldIterator<FProperty> It(TableStruct); It; ++It)
-	{
-		FProperty* ColumnProperty = *It;
+	// Here we are using the DataTable function because we are only calling GetSkeletalMeshAt() at Generation Time
+	// If we want to use it elswhere, we can use our FindTablreProperty method
+	FProperty* ColumnProperty = DataTable->FindTableProperty(FName(*ColumnName));
 
-		if (!ColumnProperty || ColumnName != DataTableUtils::GetPropertyExportName(ColumnProperty))
-		{
-			continue;
-		}
+	if (!ColumnProperty)
+	{
+		return nullptr;
+	}
 
-		if (const FSoftObjectProperty* SoftObjectProperty = CastField<FSoftObjectProperty>(ColumnProperty))
+	if (const FSoftObjectProperty* SoftObjectProperty = CastField<FSoftObjectProperty>(ColumnProperty))
+	{
+		if (uint8* RowData = DataTable->FindRowUnchecked(RowName))
 		{
-			if (uint8* RowData = DataTable->FindRowUnchecked(RowName))
+			if (uint8* CellData = ColumnProperty->ContainerPtrToValuePtr<uint8>(RowData, 0))
 			{
-				if (uint8* CellData = ColumnProperty->ContainerPtrToValuePtr<uint8>(RowData, 0))
+				if (UObject* Object = SoftObjectProperty->GetPropertyValue(CellData).LoadSynchronous())
 				{
-					if (UObject* Object = SoftObjectProperty->GetPropertyValue(CellData).LoadSynchronous())
+					if (USkeletalMesh* SkeletalMesh = Cast<USkeletalMesh>(Object))
 					{
-						if (USkeletalMesh* SkeletalMesh = Cast<USkeletalMesh>(Object))
-						{
-							return SkeletalMesh;
-						}
+						return SkeletalMesh;
 					}
 				}
 			}
@@ -1031,17 +1098,24 @@ TSoftClassPtr<UAnimInstance> UCustomizableObjectNodeTable::GetAnimInstanceAt(con
 		return TSoftClassPtr<UAnimInstance>();
 	}
 
-	const UScriptStruct* TableStruct = DataTable->GetRowStruct();
+	const UCustomizableObjectNodeTableObjectPinData* PinData = Cast<UCustomizableObjectNodeTableObjectPinData>(GetPinData(*Pin));
 
-	FString ColumnName = GetColumnNameByPin(Pin);
+	if (!PinData)
+	{
+		return TSoftClassPtr<UAnimInstance>();
+	}
 
-	for (TFieldIterator<FProperty> It(TableStruct); It; ++It)
+	if (const FTableNodeColumnData* ColumnData = ColumnDataMap.Find(PinData->StructColumnId))
 	{
-		FProperty* ColumnProperty = *It;
+		FString AnimColumn = ColumnData->AnimInstanceColumnName;
 
-		if (!ColumnProperty || ColumnName != DataTableUtils::GetPropertyExportName(ColumnProperty))
+		// Here we are using the DataTable function because we are only calling GetAnimInstanceAt() at Generation Time
+		// If we want to use it elswhere, we can use our FindTablreProperty method
+		FProperty* ColumnProperty = DataTable->FindTableProperty(FName(*AnimColumn));
+
+		if (!ColumnProperty)
 		{
-			continue;
+			return TSoftClassPtr<UAnimInstance>();
 		}
 
 		if (const FSoftClassProperty* SoftClassProperty = CastField<FSoftClassProperty>(ColumnProperty))
@@ -1309,5 +1383,34 @@ FProperty* UCustomizableObjectNodeTable::FindTableProperty(const UScriptStruct*
 }
 
 
+FGuid UCustomizableObjectNodeTable::GetColumnIdByName(const FName& ColumnName) const
+{
+	if (const UScriptStruct* TableStruct = GetTableNodeStruct())
+	{
+		if(FProperty* ColumnProperty = FindTableProperty(TableStruct, ColumnName))
+		{
+			return FStructureEditorUtils::GetGuidForProperty(ColumnProperty);
+		}
+		else
+		{
+			// We could not find the property by name
+			for (const UEdGraphPin* Pin : GetAllNonOrphanPins())
+			{
+				if (UCustomizableObjectNodeTableMeshPinData* PinData = Cast< UCustomizableObjectNodeTableMeshPinData>(GetPinData(*Pin)))
+				{
+					if (PinData->ColumnName != ColumnName)
+					{
+						continue;
+					}
+
+					return PinData->StructColumnId;
+				}
+			}
+		}
+	}
+
+	return FGuid();
+}
+
 
 #undef LOCTEXT_NAMESPACE
diff --git a/Engine/Plugins/Experimental/Mutable/Source/CustomizableObjectEditor/Private/MuCOE/Nodes/CustomizableObjectNodeTableDetails.cpp b/Engine/Plugins/Experimental/Mutable/Source/CustomizableObjectEditor/Private/MuCOE/Nodes/CustomizableObjectNodeTableDetails.cpp
--- a/Engine/Plugins/Experimental/Mutable/Source/CustomizableObjectEditor/Private/MuCOE/Nodes/CustomizableObjectNodeTableDetails.cpp
+++ b/Engine/Plugins/Experimental/Mutable/Source/CustomizableObjectEditor/Private/MuCOE/Nodes/CustomizableObjectNodeTableDetails.cpp
@@ -340,16 +340,19 @@ void FCustomizableObjectNodeTableDetails::GenerateAnimInstanceComboBoxOptions()
 	AnimSlotComboBox->ClearSelection();
 	AnimTagsComboBox->ClearSelection();
 
+	const UScriptStruct* TableStruct = Node->GetTableNodeStruct();
+
 	FString ColumnName;
+	FTableNodeColumnData* MeshColumnData = nullptr;
 
-	if (AnimMeshColumnComboBox.IsValid())
+	if (TableStruct && AnimMeshColumnComboBox.IsValid())
 	{
 		ColumnName = *AnimMeshColumnComboBox->GetSelectedItem();
-	}
-
-	const UScriptStruct* TableStruct = Node->GetTableNodeStruct();
+		FGuid ColumnId = Node->GetColumnIdByName(FName(*ColumnName));
 
-	if(!TableStruct)
+		MeshColumnData = Node->ColumnDataMap.Find(ColumnId);
+	}
+	else
 	{
 		return;
 	}
@@ -366,34 +369,24 @@ void FCustomizableObjectNodeTableDetails::GenerateAnimInstanceComboBoxOptions()
 					TSharedPtr<FString> Option = MakeShareable(new FString(DataTableUtils::GetPropertyExportName(ColumnProperty)));
 					AnimOptionNames.Add(Option);
 
-					for (const UEdGraphPin* Pin : Node->GetAllNonOrphanPins())
+					if (MeshColumnData && MeshColumnData->AnimInstanceColumnName == *Option)
 					{
-						const UCustomizableObjectNodeTableMeshPinData* PinData = Cast<UCustomizableObjectNodeTableMeshPinData >(Node->GetPinData(*Pin));
-
-						if (PinData && PinData->ColumnName == ColumnName && PinData->AnimInstanceColumnName == *Option)
-						{
-							AnimComboBox->SetSelectedItem(Option);
-							break;
-						}
+						AnimComboBox->SetSelectedItem(Option);
 					}
 				}
 			}
+
 			else if (CastField<FIntProperty>(ColumnProperty) || CastField<FNameProperty>(ColumnProperty))
 			{
 				TSharedPtr<FString> Option = MakeShareable(new FString(DataTableUtils::GetPropertyExportName(ColumnProperty)));
 				AnimSlotOptionNames.Add(Option);
 
-				for (const UEdGraphPin* Pin : Node->GetAllNonOrphanPins())
+				if (MeshColumnData && MeshColumnData->AnimSlotColumnName == *Option)
 				{
-					const UCustomizableObjectNodeTableMeshPinData* PinData = Cast<UCustomizableObjectNodeTableMeshPinData >(Node->GetPinData(*Pin));
-
-					if (PinData && PinData->ColumnName == ColumnName && PinData->AnimSlotColumnName == *Option)
-					{
-						AnimSlotComboBox->SetSelectedItem(Option);
-						break;
-					}
+					AnimSlotComboBox->SetSelectedItem(Option);
 				}
 			}
+
 			else if (const FStructProperty* StructProperty = CastField<FStructProperty>(ColumnProperty))
 			{
 				if (StructProperty->Struct == TBaseStructure<FGameplayTagContainer>::Get())
@@ -401,15 +394,9 @@ void FCustomizableObjectNodeTableDetails::GenerateAnimInstanceComboBoxOptions()
 					TSharedPtr<FString> Option = MakeShareable(new FString(DataTableUtils::GetPropertyExportName(ColumnProperty)));
 					AnimTagsOptionNames.Add(Option);
 
-					for (const UEdGraphPin* Pin : Node->GetAllNonOrphanPins())
+					if (MeshColumnData && MeshColumnData->AnimTagColumnName == *Option)
 					{
-						const UCustomizableObjectNodeTableMeshPinData* PinData = Cast<UCustomizableObjectNodeTableMeshPinData >(Node->GetPinData(*Pin));
-
-						if (PinData && PinData->ColumnName == ColumnName && PinData->AnimTagColumnName == *Option)
-						{
-							AnimTagsComboBox->SetSelectedItem(Option);
-							break;
-						}
+						AnimTagsComboBox->SetSelectedItem(Option);
 					}
 				}
 			}
@@ -463,18 +450,22 @@ void FCustomizableObjectNodeTableDetails::OnLayoutMeshColumnComboBoxSelectionCha
 
 void FCustomizableObjectNodeTableDetails::OnAnimInstanceComboBoxSelectionChanged(TSharedPtr<FString> Selection, ESelectInfo::Type SelectInfo)
 {
-	if (Selection.IsValid() && AnimMeshColumnComboBox->GetSelectedItem().IsValid())
+	if (Selection.IsValid() && AnimMeshColumnComboBox->GetSelectedItem().IsValid() && SelectInfo != ESelectInfo::Direct)
 	{
 		FString ColumnName = *AnimMeshColumnComboBox->GetSelectedItem();
-
-		for (const UEdGraphPin* Pin : Node->GetAllNonOrphanPins())
+		FGuid ColumnId = Node->GetColumnIdByName(FName(*ColumnName));
+		FTableNodeColumnData* MeshColumnData = Node->ColumnDataMap.Find(ColumnId);
+		
+		if (MeshColumnData)
+		{
+			MeshColumnData->AnimInstanceColumnName = *Selection;
+		}
+		else if(ColumnId.IsValid())
 		{
-			UCustomizableObjectNodeTableMeshPinData* PinData = Cast< UCustomizableObjectNodeTableMeshPinData>(Node->GetPinData(*Pin));
+			FTableNodeColumnData NewMeshColumnData;
+			NewMeshColumnData.AnimInstanceColumnName = *Selection;
 
-			if (PinData && PinData->ColumnName == ColumnName)
-			{
-				PinData->AnimInstanceColumnName = *Selection;
-			}
+			Node->ColumnDataMap.Add(ColumnId, NewMeshColumnData);
 		}
 
 		Node->MarkPackageDirty();
@@ -484,18 +475,22 @@ void FCustomizableObjectNodeTableDetails::OnAnimInstanceComboBoxSelectionChanged
 
 void FCustomizableObjectNodeTableDetails::OnAnimSlotComboBoxSelectionChanged(TSharedPtr<FString> Selection, ESelectInfo::Type SelectInfo)
 {
-	if (Selection.IsValid() && AnimMeshColumnComboBox->GetSelectedItem().IsValid())
+	if (Selection.IsValid() && AnimMeshColumnComboBox->GetSelectedItem().IsValid() && SelectInfo != ESelectInfo::Direct)
 	{
 		FString ColumnName = *AnimMeshColumnComboBox->GetSelectedItem();
+		FGuid ColumnId = Node->GetColumnIdByName(FName(*ColumnName));
+		FTableNodeColumnData* MeshColumnData = Node->ColumnDataMap.Find(ColumnId);
 
-		for (const UEdGraphPin* Pin : Node->GetAllNonOrphanPins())
+		if (MeshColumnData)
 		{
-			UCustomizableObjectNodeTableMeshPinData* PinData = Cast< UCustomizableObjectNodeTableMeshPinData>(Node->GetPinData(*Pin));
+			MeshColumnData->AnimSlotColumnName = *Selection;
+		}
+		else if (ColumnId.IsValid())
+		{
+			FTableNodeColumnData NewMeshColumnData;
+			NewMeshColumnData.AnimSlotColumnName = *Selection;
 
-			if (PinData && PinData->ColumnName == ColumnName)
-			{
-				PinData->AnimSlotColumnName = *Selection;
-			}
+			Node->ColumnDataMap.Add(ColumnId, NewMeshColumnData);
 		}
 
 		Node->MarkPackageDirty();
@@ -505,18 +500,24 @@ void FCustomizableObjectNodeTableDetails::OnAnimSlotComboBoxSelectionChanged(TSh
 
 void FCustomizableObjectNodeTableDetails::OnAnimTagsComboBoxSelectionChanged(TSharedPtr<FString> Selection, ESelectInfo::Type SelectInfo)
 {
-	if (Selection.IsValid() && AnimMeshColumnComboBox->GetSelectedItem().IsValid())
+	if (Selection.IsValid() && AnimMeshColumnComboBox->GetSelectedItem().IsValid() && SelectInfo != ESelectInfo::Direct)
 	{
 		FString ColumnName = *AnimMeshColumnComboBox->GetSelectedItem();
+		FGuid ColumnId = Node->GetColumnIdByName(FName(*ColumnName));
+		FTableNodeColumnData* MeshColumnData = Node->ColumnDataMap.Find(ColumnId);
 
-		for (const UEdGraphPin* Pin : Node->GetAllNonOrphanPins())
+		MeshColumnData = Node->ColumnDataMap.Find(ColumnId);
+
+		if (MeshColumnData)
 		{
-			UCustomizableObjectNodeTableMeshPinData* PinData = Cast< UCustomizableObjectNodeTableMeshPinData>(Node->GetPinData(*Pin));
+			MeshColumnData->AnimTagColumnName = *Selection;
+		}
+		else if (ColumnId.IsValid())
+		{
+			FTableNodeColumnData NewMeshColumnData;
+			NewMeshColumnData.AnimTagColumnName = *Selection;
 
-			if (PinData && PinData->ColumnName == ColumnName)
-			{
-				PinData->AnimTagColumnName = *Selection;
-			}
+			Node->ColumnDataMap.Add(ColumnId, NewMeshColumnData);
 		}
 
 		Node->MarkPackageDirty();
@@ -535,46 +536,41 @@ void FCustomizableObjectNodeTableDetails::OnNodePinValueChanged()
 
 FReply FCustomizableObjectNodeTableDetails::OnClearButtonPressed()
 {
-	bool bCleared = false;
+	if (!AnimMeshColumnComboBox->GetSelectedItem().IsValid())
+	{
+		return FReply::Unhandled();
+	}
+		
+	FString ColumnName = *AnimMeshColumnComboBox->GetSelectedItem();
+	FGuid ColumnId = Node->GetColumnIdByName(FName(*ColumnName));
+	FTableNodeColumnData* MeshColumnData = Node->ColumnDataMap.Find(ColumnId);
 
-	if (AnimMeshColumnComboBox->GetSelectedItem().IsValid())
+	if (MeshColumnData)
 	{
-		FString ColumnName = *AnimMeshColumnComboBox->GetSelectedItem();
+		MeshColumnData->AnimInstanceColumnName.Reset();
+		MeshColumnData->AnimSlotColumnName.Reset();
+		MeshColumnData->AnimTagColumnName.Reset();
 
-		for (const UEdGraphPin* Pin : Node->GetAllNonOrphanPins())
+		if (AnimComboBox.IsValid())
 		{
-			UCustomizableObjectNodeTableMeshPinData* PinData = Cast< UCustomizableObjectNodeTableMeshPinData>(Node->GetPinData(*Pin));
-
-			if (PinData && PinData->ColumnName == ColumnName)
-			{
-				PinData->AnimInstanceColumnName.Reset();
-				PinData->AnimSlotColumnName.Reset();
-				PinData->AnimTagColumnName.Reset();
+			AnimComboBox->ClearSelection();
+		}
 
-				bCleared = true;
-			}
+		if (AnimSlotComboBox.IsValid())
+		{
+			AnimSlotComboBox->ClearSelection();
 		}
 
-		if (bCleared)
+		if (AnimTagsComboBox.IsValid())
 		{
-			if (AnimComboBox.IsValid())
-			{
-				AnimComboBox->ClearSelection();
-			}
-			
-			if (AnimSlotComboBox.IsValid())
-			{
-				AnimSlotComboBox->ClearSelection();
-			}
-			
-			if (AnimTagsComboBox.IsValid())
-			{
-				AnimTagsComboBox->ClearSelection();
-			}
+			AnimTagsComboBox->ClearSelection();
 		}
+
+		Node->MarkPackageDirty();
+		return FReply::Handled();
 	}
 
-	return bCleared ? FReply::Handled() : FReply::Unhandled();
+	return FReply::Unhandled();
 }
 
 
