[mutable] Fix freeze bug when shutting down editor while compiling a CustomizableObject.
#rnx

diff --git a/Engine/Plugins/Experimental/Mutable/Source/CustomizableObjectEditor/Private/MuCOE/CustomizableObjectCompiler.cpp b/Engine/Plugins/Experimental/Mutable/Source/CustomizableObjectEditor/Private/MuCOE/CustomizableObjectCompiler.cpp
--- a/Engine/Plugins/Experimental/Mutable/Source/CustomizableObjectEditor/Private/MuCOE/CustomizableObjectCompiler.cpp
+++ b/Engine/Plugins/Experimental/Mutable/Source/CustomizableObjectEditor/Private/MuCOE/CustomizableObjectCompiler.cpp
@@ -1426,9 +1426,12 @@ void FCustomizableObjectCompiler::ForceFinishCompilation()
 {
 	if (CompileTask.IsValid())
 	{
-		if (CompileThread.IsValid())
+		// Compilation needs game thread tasks every now and then. Wait for compilation to finish while
+		// giving execution time for these tasks.
+		// TODO: interruptible compilations?
+		while (!CompileTask->IsCompleted())
 		{
-			CompileThread->WaitForCompletion();
+			FTaskGraphInterface::Get().ProcessThreadUntilIdle(ENamedThreads::GameThread);
 		}
 
 		CleanCachedReferencers();
diff --git a/Engine/Plugins/Experimental/Mutable/Source/CustomizableObjectEditor/Private/MuCOE/Nodes/CustomizableObjectNodeObject.cpp b/Engine/Plugins/Experimental/Mutable/Source/CustomizableObjectEditor/Private/MuCOE/Nodes/CustomizableObjectNodeObject.cpp
--- a/Engine/Plugins/Experimental/Mutable/Source/CustomizableObjectEditor/Private/MuCOE/Nodes/CustomizableObjectNodeObject.cpp
+++ b/Engine/Plugins/Experimental/Mutable/Source/CustomizableObjectEditor/Private/MuCOE/Nodes/CustomizableObjectNodeObject.cpp
@@ -61,7 +61,7 @@ void UCustomizableObjectNodeObject::BackwardsCompatibleFixup()
 	
 	if (CustomizableObjectCustomVersion < FCustomizableObjectCustomVersion::RegenerateNodeObjectsIds)
 	{
-		// This will regenerate all the Node Objets Guids to finally remove the duplicated Guids warning.
+		// This will regenerate all the Node Object Guids to finally remove the duplicated Guids warning.
 		// It is safe to do this here as Node Object do not use its node guid to link themeselves to other nodes.
 		CreateNewGuid();
 	}
