#rnx

diff --git a/Engine/Plugins/Experimental/Mutable/Source/CustomizableObject/Private/MuCO/CustomizableObject.cpp b/Engine/Plugins/Experimental/Mutable/Source/CustomizableObject/Private/MuCO/CustomizableObject.cpp
--- a/Engine/Plugins/Experimental/Mutable/Source/CustomizableObject/Private/MuCO/CustomizableObject.cpp
+++ b/Engine/Plugins/Experimental/Mutable/Source/CustomizableObject/Private/MuCO/CustomizableObject.cpp
@@ -1649,7 +1649,7 @@ float UCustomizableObject::GetFloatParameterDefaultValue(const FString& InParame
 	const int32 ParameterIndex = FindParameter(InParameterName);
 	if (ParameterIndex == INDEX_NONE)
 	{
-		checkNoEntry();
+		UE_LOG(LogMutable, Error, TEXT("Tried to access the default value of the nonexistent float parameter [%s] in the CustomizableObject [%s]."), *InParameterName, *GetName());
 		return FCustomizableObjectFloatParameterValue::DEFAULT_PARAMETER_VALUE;
 	}
 
@@ -1669,7 +1669,7 @@ int32 UCustomizableObject::GetIntParameterDefaultValue(const FString& InParamete
 	const int32 ParameterIndex = FindParameter(InParameterName);
 	if (ParameterIndex == INDEX_NONE)
 	{
-		checkNoEntry();
+		UE_LOG(LogMutable, Error, TEXT("Tried to access the default value of the nonexistent integer parameter [%s] in the CustomizableObject [%s]."), *InParameterName, *GetName());
 		return FCustomizableObjectIntParameterValue::DEFAULT_PARAMETER_VALUE;
 	}
 
@@ -1689,7 +1689,7 @@ bool UCustomizableObject::GetBoolParameterDefaultValue(const FString& InParamete
 	const int32 ParameterIndex = FindParameter(InParameterName);
 	if (ParameterIndex == INDEX_NONE)
 	{
-		checkNoEntry();
+		UE_LOG(LogMutable, Error, TEXT("Tried to access the default value of the nonexistent boolean parameter [%s] in the CustomizableObject [%s]."), *InParameterName, *GetName());
 		return FCustomizableObjectBoolParameterValue::DEFAULT_PARAMETER_VALUE;
 	}
 
@@ -1709,7 +1709,7 @@ FLinearColor UCustomizableObject::GetColorParameterDefaultValue(const FString& I
 	const int32 ParameterIndex = FindParameter(InParameterName);
 	if (ParameterIndex == INDEX_NONE)
 	{
-		checkNoEntry();
+		UE_LOG(LogMutable, Error, TEXT("Tried to access the default value of the nonexistent color parameter [%s] in the CustomizableObject [%s]."), *InParameterName, *GetName());
 		return FCustomizableObjectVectorParameterValue::DEFAULT_PARAMETER_VALUE;;
 	}
 
@@ -1747,7 +1747,7 @@ FCustomizableObjectProjector UCustomizableObject::GetProjectorParameterDefaultVa
 	const int32 ParameterIndex = FindParameter(InParameterName);
 	if (ParameterIndex == INDEX_NONE)
 	{
-		UE_LOG(LogMutable, Error, TEXT("Tried to access the default value of the unexistent projector [%s] in the CustomizableObject [%s]."), *InParameterName, *GetName());
+		UE_LOG(LogMutable, Error, TEXT("Tried to access the default value of the nonexistent projector [%s] in the CustomizableObject [%s]."), *InParameterName, *GetName());
 		return FCustomizableObjectProjectorParameterValue::DEFAULT_PARAMETER_VALUE;
 	}
 
@@ -1772,15 +1772,15 @@ FName UCustomizableObject::GetTextureParameterDefaultValue(const FString& InPara
 	const int32 ParameterIndex = FindParameter(InParameterName);
 	if (ParameterIndex == INDEX_NONE)
 	{
-		checkNoEntry();
-		return FName();
+		UE_LOG(LogMutable, Error, TEXT("Tried to access the default value of the nonexistent texture parameter [%s] in the CustomizableObject [%s]."), *InParameterName, *GetName());
+		return FCustomizableObjectTextureParameterValue::DEFAULT_PARAMETER_VALUE;
 	}
 
 	const TSharedPtr<mu::Model> Model = GetModel();
 	if (!Model)
 	{
 		checkNoEntry();
-		return FName();
+		return FCustomizableObjectTextureParameterValue::DEFAULT_PARAMETER_VALUE;
 	}
 	
 	return Model->GetImageDefaultValue(ParameterIndex);
@@ -1792,8 +1792,8 @@ bool UCustomizableObject::IsParameterMultidimensional(const FString& InParameter
 	const int32 ParameterIndex = FindParameter(InParameterName);
 	if (ParameterIndex == INDEX_NONE)
 	{
-		checkNoEntry();
-		return {};
+		UE_LOG(LogMutable, Error, TEXT("Tried to access the default value of the nonexistent parameter [%s] in the CustomizableObject [%s]."), *InParameterName, *GetName());
+		return false;
 	}
 
 	return IsParameterMultidimensional(ParameterIndex);
diff --git a/Engine/Plugins/Experimental/Mutable/Source/CustomizableObject/Public/MuCO/CustomizableObjectParameterTypeDefinitions.h b/Engine/Plugins/Experimental/Mutable/Source/CustomizableObject/Public/MuCO/CustomizableObjectParameterTypeDefinitions.h
--- a/Engine/Plugins/Experimental/Mutable/Source/CustomizableObject/Public/MuCO/CustomizableObjectParameterTypeDefinitions.h
+++ b/Engine/Plugins/Experimental/Mutable/Source/CustomizableObject/Public/MuCO/CustomizableObjectParameterTypeDefinitions.h
@@ -158,6 +158,8 @@ USTRUCT(BlueprintType)
 struct FCustomizableObjectTextureParameterValue
 {
 	GENERATED_USTRUCT_BODY()
+
+	inline static const FName DEFAULT_PARAMETER_VALUE = FName();
 	
 	UPROPERTY(Category = CustomizableObjectTextureParameterValue, VisibleAnywhere)
 	FString ParameterName;
