[Mutable] Fix image interpolate crash caused by mismatch formats.

- Enforce at execution time a reformat if formats are not compatible. This fixes the crash but it can introduce a performance issue.
- It looks like there is a code generation bug worth investigating since the case should not happen.

#rnx

diff --git a/Engine/Plugins/Experimental/Mutable/Source/MutableRuntime/Private/MuR/CodeRunner.cpp b/Engine/Plugins/Experimental/Mutable/Source/MutableRuntime/Private/MuR/CodeRunner.cpp
--- a/Engine/Plugins/Experimental/Mutable/Source/MutableRuntime/Private/MuR/CodeRunner.cpp
+++ b/Engine/Plugins/Experimental/Mutable/Source/MutableRuntime/Private/MuR/CodeRunner.cpp
@@ -3811,6 +3811,21 @@ namespace mu
 							pMax = Resized;
 						}
 
+						// Be defensive: ensure format matches.
+						if (pNew->GetFormat() != pMax->GetFormat())
+						{
+							MUTABLE_CPUPROFILER_SCOPE(Format_ForInterpolate);
+
+							Ptr<Image> Formatted = CreateImage(pMax->GetSizeX(), pMax->GetSizeY(), pMax->GetLODCount(), pNew->GetFormat(), EInitializationType::NotInitialized);
+							
+							bool bSuccess = false;
+							ImOp.ImagePixelFormat(bSuccess, m_pSettings->ImageCompressionQuality, Formatted.get(), pMax.get());
+							check(bSuccess);
+							
+							Release(pMax);
+							pMax = Formatted;
+						}
+
 						if (pNew->GetLODCount() != LevelCount)
 						{
 							MUTABLE_CPUPROFILER_SCOPE(Mipmap_ForInterpolate);
