[Mutable] Add absolute memory getters to GlobalMemoryCounter.

#rnx

diff --git a/Engine/Plugins/Experimental/Mutable/Source/MutableRuntime/Private/MuR/MemoryTrackingUtils.cpp b/Engine/Plugins/Experimental/Mutable/Source/MutableRuntime/Private/MuR/MemoryTrackingUtils.cpp
--- a/Engine/Plugins/Experimental/Mutable/Source/MutableRuntime/Private/MuR/MemoryTrackingUtils.cpp
+++ b/Engine/Plugins/Experimental/Mutable/Source/MutableRuntime/Private/MuR/MemoryTrackingUtils.cpp
@@ -122,6 +122,33 @@ namespace mu
 		const volatile SSIZE_T Result = Counter;
 		return Result;
 	}
+
+	SSIZE_T FGlobalMemoryCounter::GetAbsolutePeak()
+	{
+		if (LIKELY(!FMemoryTrackingConsoleFlags::bEnableGlobalMemoryTracking))
+		{
+			return 0;
+		}
+
+		FScopeLock Lock(&Mutex);
+
+		const volatile SSIZE_T Result = AbsolutePeakValue;
+		return Result;
+	}
+
+	SSIZE_T FGlobalMemoryCounter::GetAbsoluteCounter()
+	{
+		if (LIKELY(!FMemoryTrackingConsoleFlags::bEnableGlobalMemoryTracking))
+		{
+			return 0;
+		}
+
+		FScopeLock Lock(&Mutex);
+
+		const volatile SSIZE_T Result = AbsoluteCounter;
+		return Result;
+	}
+
 #else
 	void FGlobalMemoryCounter::Zero()
 	{
@@ -140,6 +167,17 @@ namespace mu
 	{
 		return 0;
 	}
+
+	SSIZE_T FGlobalMemoryCounter::GetAbsolutePeak()
+	{
+		return 0;
+	}
+
+	SSIZE_T FGlobalMemoryCounter::GetAbsoluteCounter()
+	{
+		return 0;
+	}
+
 #endif
 
 
@@ -171,6 +209,33 @@ namespace mu
 			})
 		);
 
+		static inline FAutoConsoleCommand CmdDumpMutableGlobalMemoryAbsoluteCounter = FAutoConsoleCommand(
+			TEXT("mutable.DumpMutableGlobalMemoryAbsoluteCounter"),
+			TEXT("Dump Mutable plugin global tracked memory absolute counter."),
+			FConsoleCommandDelegate::CreateStatic([]()
+			{
+				if (!UE_MUTABLE_TRACK_ALLOCATOR_MEMORY_PEAK)
+				{
+					UE_LOG(LogMutableCore, Display, TEXT("The build does not have UE_MUTABLE_TRACK_ALLOCATOR_MEMORY_PEAK enabled, no counting performed."));
+					return;
+				}
+
+				if (!FMemoryTrackingConsoleFlags::bEnableGlobalMemoryTracking)
+				{
+					UE_LOG(LogMutableCore, Display, TEXT("mutable.EnableGlobalMemoryTracking is set to false, no counting performed."));
+					return;
+				}
+
+				UE_LOG(LogMutableCore, 
+					   Display, 
+					   TEXT("Mutable Memory: AbsoluteCurrent %2.f MiB, AbsolutePeak %2.f MiB"), 
+					   double(mu::FGlobalMemoryCounter::GetAbsoluteCounter()) / double(1024*1024),
+					   double(mu::FGlobalMemoryCounter::GetAbsolutePeak())    / double(1024*1024)
+				);
+			})
+		);
+
+
 		static inline FAutoConsoleCommand CmdZeroMutableGlobalMemoryCounter = FAutoConsoleCommand(
 			TEXT("mutable.ZeroMutableGlobalMemoryCounter"),
 			TEXT("Zero Mutable plugin global tracked memory counter."),
diff --git a/Engine/Plugins/Experimental/Mutable/Source/MutableRuntime/Public/MuR/MemoryTrackingUtils.h b/Engine/Plugins/Experimental/Mutable/Source/MutableRuntime/Public/MuR/MemoryTrackingUtils.h
--- a/Engine/Plugins/Experimental/Mutable/Source/MutableRuntime/Public/MuR/MemoryTrackingUtils.h
+++ b/Engine/Plugins/Experimental/Mutable/Source/MutableRuntime/Public/MuR/MemoryTrackingUtils.h
@@ -33,6 +33,8 @@ namespace mu
 		 static void Restore();
 		 static SSIZE_T GetPeak();
 		 static SSIZE_T GetCounter();
+		 static SSIZE_T GetAbsolutePeak();
+		 static SSIZE_T GetAbsoluteCounter();
 	};
 #else
 	struct MUTABLERUNTIME_API FGlobalMemoryCounter
@@ -42,6 +44,8 @@ namespace mu
 		 static void Restore();
 		 static SSIZE_T GetPeak();
 		 static SSIZE_T GetCounter();
+		 static SSIZE_T GetAbsolutePeak();
+		 static SSIZE_T GetAbsoluteCounter();
 	};
 #endif
 }
