[Mutable] Remove COE private dependencies 2/3.

- Moved required files to the Internal directory.
- Moved private structs to UCustomizableObjectSystemPrivate or its own source files.
- Moved public static functions outside of the private.

#rnx

diff --git a/Engine/Plugins/Experimental/Mutable/Source/CustomizableObject/Classes/MuCO/CustomizableObjectInstance.h b/Engine/Plugins/Experimental/Mutable/Source/CustomizableObject/Classes/MuCO/CustomizableObjectInstance.h
--- a/Engine/Plugins/Experimental/Mutable/Source/CustomizableObject/Classes/MuCO/CustomizableObjectInstance.h
+++ b/Engine/Plugins/Experimental/Mutable/Source/CustomizableObject/Classes/MuCO/CustomizableObjectInstance.h
@@ -6,15 +6,22 @@
 #include "MuCO/CustomizableObjectInstanceDescriptor.h"
 #include "Templates/SubclassOf.h"
 #include "Math/RandomStream.h"
+#include "MuR/Ptr.h"
 
 #include "CustomizableObjectInstance.generated.h"
 
+struct FMutableModelImageProperties;
+
+namespace mu
+{
+	class Image;
+}
+
 class FCustomizableObjectSystemPrivate;
 class USkeletalMesh;
 class AActor;
 class FProperty;
 class UAnimInstance;
-class UCustomizableInstancePrivateData; // This is used to hide Mutable SDK members in the public headers.
 class UCustomizableObject;
 class UTexture2D;
 class FUpdateContextPrivate;
@@ -23,6 +30,7 @@ class UMaterialInterface;
 struct FFrame;
 struct FGameplayTagContainer;
 struct FPropertyChangedEvent;
+struct FTexturePlatformData;
 
 
 // Priority for the mutable update queue, Low is the normal distance-based priority, High is normally used for discards and Mid for LOD downgrades
@@ -776,6 +784,12 @@ private:
 };
 
 
+/** 
+ * \param OnlyLOD: If not 0, extract and convert only one single LOD from the source image.
+ * \param ExtractChannel: If different than -1, extract a single-channel image with the specified source channel data. */
+CUSTOMIZABLEOBJECT_API void ConvertImage(UTexture2D* Texture, mu::Ptr<const mu::Image> MutableImage, const FMutableModelImageProperties& Props, int32 OnlyLOD = -1, int32 ExtractChannel = -1);
+
+
 #if UE_ENABLE_INCLUDE_ORDER_DEPRECATED_IN_5_2
 #include "Async/TaskGraphInterfaces.h"
 #include "MuCO/CustomizableObjectParameterTypeDefinitions.h"
diff --git a/Engine/Plugins/Experimental/Mutable/Source/CustomizableObject/Internal/MuCO/CustomizableInstancePrivateData.h b/Engine/Plugins/Experimental/Mutable/Source/CustomizableObject/Internal/MuCO/CustomizableInstancePrivateData.h
new file mode 100644
--- /dev/null
+++ b/Engine/Plugins/Experimental/Mutable/Source/CustomizableObject/Internal/MuCO/CustomizableInstancePrivateData.h
@@ -0,0 +1,339 @@
+// Copyright Epic Games, Inc. All Rights Reserved.
+
+#pragma once
+
+
+#include "Materials/MaterialInterface.h"
+#include "Rendering/SkeletalMeshLODModel.h"
+#include "Rendering/SkeletalMeshModel.h"
+#include "MuCO/CustomizableObjectInstance.h"
+#include "MuR/Instance.h"
+#include "GameplayTagContainer.h"
+#include "UObject/Package.h"
+
+#include "CustomizableInstancePrivateData.generated.h"
+
+namespace mu 
+{
+	class PhysicsBody;
+	class Mesh;
+	typedef uint64 FResourceID;
+}
+
+struct FMutableModelImageProperties;
+struct FMutableRefSkeletalMeshData;
+struct FMutableImageCacheKey;
+struct FStreamableHandle;
+struct FGeneratedMaterial;
+struct FGeneratedTexture;
+class UPhysicsAsset;
+class USkeleton;
+
+
+/** CustomizableObject Instance flags for internal use  */
+enum ECOInstanceFlags
+{
+	ECONone							= 0,  // Should not use the name None here.. it collides with other enum in global namespace
+
+	// Update process
+	CreatingSkeletalMesh			= 1 << 1,	//
+	Generated						= 1 << 2,	//
+	ReuseTextures					= 1 << 3, 	// 
+	ReplacePhysicsAssets			= 1 << 4,	// Merge active PhysicsAssets and replace the base physics asset
+
+	// Update priorities
+	UsedByComponent					= 1 << 5,	// If any components are using this instance, they will set flag every frame
+	UsedByComponentInPlay			= 1 << 6,	// If any components are using this instance in play, they will set flag every frame
+	UsedByPlayerOrNearIt			= 1 << 7,	// The instance is used by the player or is near the player, used to give more priority to its updates
+	DiscardedByNumInstancesLimit	= 1 << 8,	// The instance is descarded because we exceeded the limit of instances generated 
+
+	// Types of updates
+	PendingLODsUpdate				= 1 << 9,	// Used to queue an update due to a change in LODs required by the instance
+	PendingLODsDowngrade			= 1 << 10,	// Used to queue a downgrade update to reduce the number of LODs. LOD update goes from a high res level to a low res one, ex: 0 to 1 or 1 to 2
+	
+	// Generation
+	ForceGenerateMipTail			= 1 << 13,	// If set, SkipGenerateResidentMips will be ignored and the mip tail will be generated
+};
+
+ENUM_CLASS_FLAGS(ECOInstanceFlags);
+
+
+USTRUCT()
+struct FReferencedPhysicsAssets
+{
+	GENERATED_USTRUCT_BODY();
+	
+	TArray<FString> PhysicsAssetToLoad;
+	
+	UPROPERTY(Transient)
+	TArray< TObjectPtr<UPhysicsAsset> > PhysicsAssetsToMerge;
+
+	TArray<int32> AdditionalPhysicsAssetsToLoad;
+
+	UPROPERTY(Transient)
+	TArray<TObjectPtr<UPhysicsAsset>> AdditionalPhysicsAssets;
+};
+
+
+USTRUCT()
+struct FReferencedSkeletons
+{
+	GENERATED_USTRUCT_BODY();
+
+	// Merged skeleton if found in the cache
+	UPROPERTY()
+	TObjectPtr<USkeleton> Skeleton;
+
+	UPROPERTY()
+	TArray<uint16> SkeletonIds;
+
+	UPROPERTY()
+	TArray< TObjectPtr<USkeleton> > SkeletonsToMerge;
+};
+
+
+USTRUCT()
+struct FCustomizableInstanceComponentData
+{
+	GENERATED_USTRUCT_BODY();
+
+	// AnimBP data gathered for a component from its constituent meshes
+	UPROPERTY(Transient, Category = CustomizableObjectInstance, editfixedsize, VisibleAnywhere)
+	TMap<FName, TSoftClassPtr<UAnimInstance>> AnimSlotToBP;
+
+	// AssetUserData gathered for a component from its constituent meshes
+	UPROPERTY(Transient, Category = CustomizableObjectInstance, editfixedsize, VisibleAnywhere)
+	TSet<TObjectPtr<UAssetUserData>> AssetUserDataArray;
+
+	// Index of the resource in the StreamedResourceData array of the CustomizableObject.
+	TArray<int32> StreamedResourceIndex;
+
+#if WITH_EDITORONLY_DATA
+	// Just used for mutable.EnableMutableAnimInfoDebugging command
+	TArray<FString> MeshPartPaths;
+#endif
+
+	/** Skeletons required by the current generated instance. Skeletons to be loaded and merged.*/
+	UPROPERTY(Transient)
+	FReferencedSkeletons Skeletons;
+	
+	/** PhysicsAssets required by the current generated instance. PhysicsAssets to be loaded and merged.*/
+	UPROPERTY(Transient)
+	FReferencedPhysicsAssets PhysicsAssets;
+
+	/** Clothing PhysicsAssets required by the current generated instance. PhysicsAssets to be loaded and merged.*/
+	TArray<TPair<int32, FString>> ClothingPhysicsAssetsToStream;
+
+	/** Array of generated MeshIds per each LOD, used to decide if the mesh should be updated or not.
+	 *  Size == NumLODsAvailable
+	 *  LODs without mesh will be set to the maximum value of FResourceID (Max_uint64). */
+	TArray<mu::FResourceID> LastMeshIdPerLOD;
+
+	UPROPERTY(Transient)
+	TArray<TObjectPtr<UMaterialInterface>> OverrideMaterials;
+};
+
+USTRUCT()
+struct FAnimInstanceOverridePhysicsAsset
+{
+	GENERATED_BODY()
+
+	UPROPERTY(Transient)
+	int32 PropertyIndex = 0;
+
+	UPROPERTY(Transient)
+	TObjectPtr<UPhysicsAsset> PhysicsAsset;
+};
+
+USTRUCT()
+struct FAnimBpGeneratedPhysicsAssets
+{
+	GENERATED_BODY()
+
+	UPROPERTY(Transient)
+	TArray<FAnimInstanceOverridePhysicsAsset> AnimInstancePropertyIndexAndPhysicsAssets;
+};
+
+UCLASS()
+class UCustomizableInstancePrivateData : public UObject
+{
+public:
+	GENERATED_BODY()
+
+	UCustomizableInstancePrivateData();
+
+	UPROPERTY( Transient )
+	TArray<FGeneratedMaterial> GeneratedMaterials;
+
+	UPROPERTY( Transient )
+	TArray<FGeneratedTexture> GeneratedTextures;
+
+	// Indices of the parameters that are relevant for the given parameter values.
+	// This only gets updated if parameter decorations are generated.
+	TArray<int> RelevantParameters;
+
+	// If Texture reuse is enabled, stores which texture is being used in a particular <LODIndex, ComponentIndex, MeshSurfaceIndex, image>
+	// \TODO: Create a key based on a struct instead of generating strings dynamically.
+	UPROPERTY(Transient)
+	TMap<FString, TWeakObjectPtr<UTexture2D>> TextureReuseCache;
+
+	// Handle used to store a streaming request operation if one is ongoing.
+	TSharedPtr<FStreamableHandle> StreamingHandle;
+
+	// Only used in LiveUpdateMode to reuse core instances between updates and their temp data to speed up updates, but spend way more memory
+	mu::Instance::ID LiveUpdateModeInstanceID = 0;
+
+	/** Cached Texture Parameters from the last update.
+	 * This cache is required since the Instance can have a LOD Update at any time.
+	 * So we need to make sure that the initially provided Texture Parameters by the user will be available until the user decides to change them. */
+	TArray<FName> UpdateTextureParameters;
+
+#if WITH_EDITOR
+
+	virtual void PostDuplicate(bool bDuplicateForPIE) override;
+
+	void OnPostCompile();
+
+#endif
+
+	/** Invalidates the previously generated data and retrieves information from the CObject after specific actions.
+	 *  It'll be called in the PostLoad, after Compiling the CO, and after changing the CO of the Instance. */
+	void InitCustomizableObjectData(const UCustomizableObject* InCustomizableObject);
+	
+	FCustomizableInstanceComponentData* GetComponentData(int32 ComponentIndex);
+	const FCustomizableInstanceComponentData* GetComponentData(int32 ComponentIndex) const;
+
+	ECOInstanceFlags GetCOInstanceFlags() const { return InstanceFlagsPrivate; }
+	void SetCOInstanceFlags(ECOInstanceFlags FlagsToSet) { InstanceFlagsPrivate = (ECOInstanceFlags)(InstanceFlagsPrivate | FlagsToSet); }
+	void ClearCOInstanceFlags(ECOInstanceFlags FlagsToClear) { InstanceFlagsPrivate = (ECOInstanceFlags)(InstanceFlagsPrivate & ~FlagsToClear); }
+	bool HasCOInstanceFlags(ECOInstanceFlags FlagsToCheck) const { return (InstanceFlagsPrivate & FlagsToCheck) != 0; }
+
+	void BuildMaterials(const TSharedRef<FUpdateContextPrivate>& OperationData, UCustomizableObjectInstance* Public);
+
+	void ReuseTexture(UTexture2D* Texture, TSharedRef<FTexturePlatformData, ESPMode::ThreadSafe>& PlatformData);
+
+	// Return an event that will be fired when the assets  have been loaded. It returns null if no asset needs loading.
+	FGraphEventRef LoadAdditionalAssetsAsync(const TSharedRef<FUpdateContextPrivate>& OperationData, UCustomizableObjectInstance* Public, struct FStreamableManager &StreamableManager);
+	void AdditionalAssetsAsyncLoaded(UCustomizableObjectInstance* Public);
+
+	void TickUpdateCloseCustomizableObjects(UCustomizableObjectInstance& Publics, FMutableInstanceUpdateMap& InOutRequestedUpdates);
+	void UpdateInstanceIfNotGenerated(UCustomizableObjectInstance& Public, FMutableInstanceUpdateMap& InOutRequestedUpdates);
+
+	// Returns true if success (?)
+	bool UpdateSkeletalMesh_PostBeginUpdate0(UCustomizableObjectInstance* Public, const TSharedRef<FUpdateContextPrivate>& OperationData);
+
+	static void ReleaseMutableTexture(const FMutableImageCacheKey& MutableTextureKey, UTexture2D* Texture, struct FMutableResourceCache& Cache);
+
+	// Copy data generated in the mutable thread over to the instance and initializes additional data required during the update
+	void PrepareForUpdate(const TSharedRef<FUpdateContextPrivate>& OperationData);
+
+	int32 GetNumLODsAvailable() const { return NumLODsAvailable; }
+	
+	// The following method is basically copied from PostEditChangeProperty and/or SkeletalMesh.cpp to be able to replicate PostEditChangeProperty without the editor
+	void PostEditChangePropertyWithoutEditor(USkeletalMesh* SkeletalMesh);
+	
+	void DiscardResourcesAndSetReferenceSkeletalMesh(UCustomizableObjectInstance* Instance);
+
+	// Releases all the mutable resources this instance holds, should only be called when it is not going to be used any more.
+	void ReleaseMutableResources(bool bCalledFromBeginDestroy, const UCustomizableObjectInstance& Instance);
+
+	const TArray<FAnimInstanceOverridePhysicsAsset>* GetGeneratedPhysicsAssetsForAnimInstance(TSubclassOf<UAnimInstance> AnimInstance) const;
+
+	/** */
+#if WITH_EDITORONLY_DATA
+	static void RegenerateImportedModel(USkeletalMesh* SkeletalMesh);
+#endif
+
+private:
+
+	void InitSkeletalMeshData(const TSharedRef<FUpdateContextPrivate>& OperationData, USkeletalMesh* SkeletalMesh, const FMutableRefSkeletalMeshData* RefSkeletalMeshData, const UCustomizableObject& CustomizableObject, int32 ComponentIndex);
+
+	bool BuildSkeletonData(const TSharedRef<FUpdateContextPrivate>& OperationData, USkeletalMesh& SkeletalMesh, const FMutableRefSkeletalMeshData& RefSkeletalMeshData, UCustomizableObject& CustomizableObject, int32 ComponentIndex);
+	void BuildMeshSockets(const TSharedRef<FUpdateContextPrivate>& OperationData, USkeletalMesh* SkeletalMesh, const FMutableRefSkeletalMeshData* RefSkeletalMeshData, UCustomizableObjectInstance* CustomizableObjectInstance, mu::Ptr<const mu::Mesh> MutableMesh);
+	void BuildOrCopyElementData(const TSharedRef<FUpdateContextPrivate>& OperationData, USkeletalMesh* SkeletalMesh, UCustomizableObjectInstance* CustomizableObjectInstance, int32 ComponentIndex);
+	void BuildOrCopyMorphTargetsData(const TSharedRef<FUpdateContextPrivate>& OperationData, USkeletalMesh* SkeletalMesh, const USkeletalMesh* SrcSkeletalMesh, UCustomizableObjectInstance* CustomizableObjectInstance, int32 ComponentIndex);
+	bool BuildOrCopyRenderData(const TSharedRef<FUpdateContextPrivate>& OperationData, USkeletalMesh* SkeletalMesh, const USkeletalMesh* SrcSkeletalMesh, UCustomizableObjectInstance* CustomizableObjectInstance, int32 ComponentIndex);
+	void BuildOrCopyClothingData(const TSharedRef<FUpdateContextPrivate>& OperationData, USkeletalMesh* SkeletalMesh, const USkeletalMesh* SrcSkeletalMesh, UCustomizableObjectInstance* CustomizableObjectInstance, int32 ComponentIndex);
+	
+	//
+	USkeleton* MergeSkeletons(UCustomizableObject& CustomizableObject, const FMutableRefSkeletalMeshData& RefSkeletalMeshData, int32 ComponentIndex);
+
+	//
+	UPhysicsAsset* GetOrBuildMainPhysicsAsset(TObjectPtr<class UPhysicsAsset> TamplateAsset, const mu::PhysicsBody* PhysicsBody, const UCustomizableObject& CustomizableObject, int32 ComponentIndex, bool bDisableCollisionBetweenAssets);
+	
+	// Create a transient texture and add it to the TextureTrackerArray
+	UTexture2D* CreateTexture();
+
+	void InvalidateGeneratedData();
+
+	bool DoComponentsNeedUpdate(UCustomizableObjectInstance* CustomizableObjectInstance, const TSharedRef<FUpdateContextPrivate>& OperationData, TArray<bool>& OutComponentNeedsUpdate, bool& bOutEmptyMesh);
+
+	mu::FResourceID GetLastMeshId(int32 ComponentIndex, int32 LODIndex) const;
+	void SetLastMeshId(int32 ComponentIndex, int32 LODIndex, mu::FResourceID MeshId);
+
+public:
+	CUSTOMIZABLEOBJECT_API ESkeletalMeshStatus GetSkeletalMeshStatus() const;
+
+	void SetSkeletalMeshStatus(ESkeletalMeshStatus Status);
+	
+	// If any components are using this instance, they will store the min of their distances to the player here every frame for LOD purposes
+	float MinSquareDistFromComponentToPlayer;
+	float LastMinSquareDistFromComponentToPlayer; // The same as the previous dist for last frame
+												
+	// This is the LODs that the Customizable Object has
+	int32 NumLODsAvailable;
+
+	// First SkeletalMesh LOD we can generate on the running platform
+	uint8 FirstLODAvailable;
+
+	// Maximum number of SkeletalMesh LODs to stream
+	uint8 NumMaxLODsToStream;
+	
+	UPROPERTY(Transient)
+	TArray<FCustomizableInstanceComponentData> ComponentsData;
+
+	UPROPERTY(Transient)
+	TArray< TObjectPtr<UMaterialInterface> > ReferencedMaterials;
+
+	// Converts a ReferencedMaterials index from the CustomizableObject to an index in the ReferencedMaterials in the Instance
+	TMap<uint32, uint32> ObjectToInstanceIndexMap;
+
+	TArray<FGeneratedTexture> TexturesToRelease;
+
+	UPROPERTY(Transient)
+	TArray< TObjectPtr<UPhysicsAsset> > ClothingPhysicsAssets;
+
+	// To keep loaded AnimBPs referenced and prevent GC
+	UPROPERTY(Transient, Category = CustomizableObjectInstance, editfixedsize, VisibleAnywhere)
+	TArray<TSubclassOf<UAnimInstance>> GatheredAnimBPs;
+
+	UPROPERTY(Transient, Category = CustomizableObjectInstance, editfixedsize, VisibleAnywhere)
+	FGameplayTagContainer AnimBPGameplayTags;
+
+	UPROPERTY(Transient, Category = CustomizableObjectInstance, editfixedsize, VisibleAnywhere)
+	TMap<TSubclassOf<UAnimInstance>, FAnimBpGeneratedPhysicsAssets> AnimBpPhysicsAssets;
+
+	// The pass-through textures that will be loaded during an update
+	TArray<TSoftObjectPtr<UTexture>> PassThroughTexturesToLoad;
+
+	// Used during an update to prevent the pass-through textures loaded by LoadAdditionalAssetsAsync() from being unloaded by GC
+	// between AdditionalAssetsAsyncLoaded() and their setting into the generated materials in BuildMaterials()
+	UPROPERTY(Transient)
+	TArray<TObjectPtr<UTexture>> LoadedPassThroughTexturesPendingSetMaterial;
+
+private:
+	
+	ECOInstanceFlags InstanceFlagsPrivate = ECOInstanceFlags::ECONone;
+
+public:
+	/** Hash of the UCustomizableObjectInstance::Descriptor on the last update request. */
+	FDescriptorRuntimeHash UpdateDescriptorRuntimeHash;
+	
+	/** Hash of the UCustomizableObjectInstance::Descriptor on the last successful update. */
+	FDescriptorRuntimeHash DescriptorRuntimeHash;
+
+private:
+	/** Status of the generated Skeletal Mesh. Not to be confused with the Update Result. */
+	ESkeletalMeshStatus SkeletalMeshStatus = ESkeletalMeshStatus::NotGenerated;
+};
+
diff --git a/Engine/Plugins/Experimental/Mutable/Source/CustomizableObject/Internal/MuCO/CustomizableObjectPrivate.h b/Engine/Plugins/Experimental/Mutable/Source/CustomizableObject/Internal/MuCO/CustomizableObjectPrivate.h
new file mode 100644
--- /dev/null
+++ b/Engine/Plugins/Experimental/Mutable/Source/CustomizableObject/Internal/MuCO/CustomizableObjectPrivate.h
@@ -0,0 +1,57 @@
+// Copyright Epic Games, Inc. All Rights Reserved.
+
+#pragma once
+
+#include "Templates/SharedPointer.h"
+#include "UObject/WeakObjectPtr.h"
+#include "MuR/Types.h"
+
+#if WITH_EDITOR
+#include "Misc/Guid.h"
+#endif
+
+namespace mu { class Model; }
+class UCustomizableObject;
+class USkeletalMesh;
+
+
+class FMeshCache
+{
+public:
+	USkeletalMesh* Get(const TArray<mu::FResourceID>& Key);
+
+	void Add(const TArray<mu::FResourceID>& Key, USkeletalMesh* Value);
+
+private:
+	TMap<TArray<mu::FResourceID>, TWeakObjectPtr<USkeletalMesh>> GeneratedMeshes;
+};
+
+
+class FCustomizableObjectPrivateData
+{
+	TSharedPtr<mu::Model, ESPMode::ThreadSafe> MutableModel;
+
+public:
+	void SetModel(const TSharedPtr<mu::Model, ESPMode::ThreadSafe>& Model, const FGuid Identifier);
+	const TSharedPtr<mu::Model, ESPMode::ThreadSafe>& GetModel();
+	TSharedPtr<const mu::Model, ESPMode::ThreadSafe> GetModel() const;
+
+#if WITH_EDITORONLY_DATA
+	/** See UCustomizableObject::ParticipatingObjects. */
+	CUSTOMIZABLEOBJECT_API TMap<TObjectPtr<const UObject>, FGuid>& GetParticipatingObjects(UCustomizableObject& Public);
+#endif
+	
+	// See UCustomizableObjectSystem::LockObject. Must only be modified from the game thread
+	bool bLocked = false;
+
+#if WITH_EDITOR
+	FGuid Identifier;
+
+	bool bModelCompiledForCook = false;
+	TArray<FString> CachedPlatformNames;
+#endif
+
+	/** Cache of generated SkeletalMeshes */
+	FMeshCache MeshCache;
+};
+
diff --git a/Engine/Plugins/Experimental/Mutable/Source/CustomizableObject/Internal/MuCO/UnrealToMutableTextureConversionUtils.h b/Engine/Plugins/Experimental/Mutable/Source/CustomizableObject/Internal/MuCO/UnrealToMutableTextureConversionUtils.h
new file mode 100644
--- /dev/null
+++ b/Engine/Plugins/Experimental/Mutable/Source/CustomizableObject/Internal/MuCO/UnrealToMutableTextureConversionUtils.h
@@ -0,0 +1,29 @@
+// Copyright Epic Games, Inc. All Rights Reserved.
+
+#pragma once
+
+#if WITH_EDITOR
+
+#include "HAL/Platform.h"
+
+namespace mu
+{
+	class Image;
+}
+
+// Forward declarations
+class UTexture2D;
+
+enum class EUnrealToMutableConversionError 
+{
+    Success,
+    UnsupportedFormat,
+    CompositeImageDimensionMismatch,
+    CompositeUnsupportedFormat,
+    Unknown
+};
+
+
+CUSTOMIZABLEOBJECT_API EUnrealToMutableConversionError ConvertTextureUnrealSourceToMutable(mu::Image* OutResult, UTexture2D* Texture, bool bIsNormalComposite, uint8 MipmapsToSkip);
+
+#endif
diff --git a/Engine/Plugins/Experimental/Mutable/Source/CustomizableObject/Private/MuCO/CustomizableInstancePrivateData.h b/Engine/Plugins/Experimental/Mutable/Source/CustomizableObject/Private/MuCO/CustomizableInstancePrivateData.h
deleted file mode 100644
diff --git a/Engine/Plugins/Experimental/Mutable/Source/CustomizableObject/Private/MuCO/CustomizableObject.cpp b/Engine/Plugins/Experimental/Mutable/Source/CustomizableObject/Private/MuCO/CustomizableObject.cpp
--- a/Engine/Plugins/Experimental/Mutable/Source/CustomizableObject/Private/MuCO/CustomizableObject.cpp
+++ b/Engine/Plugins/Experimental/Mutable/Source/CustomizableObject/Private/MuCO/CustomizableObject.cpp
@@ -50,7 +50,7 @@ static bool bUsesOnCookStart = false;
 UCustomizableObject::UCustomizableObject()
 	: UObject()
 {
-	PrivateData = TSharedPtr<FCustomizableObjectPrivateData>( new FCustomizableObjectPrivateData() );
+	PrivateData = MakeShared<FCustomizableObjectPrivateData>();
 
 #if WITH_EDITORONLY_DATA
 	const FString CVarName = TEXT("r.SkeletalMesh.MinLodQualityLevel");
@@ -1095,7 +1095,7 @@ void UCustomizableObject::LoadEmbeddedData(FArchive& Ar)
 
 FCustomizableObjectPrivateData* UCustomizableObject::GetPrivate() const
 {
-	check(PrivateData)
+	check(PrivateData);
 	return PrivateData.Get();
 }
 
diff --git a/Engine/Plugins/Experimental/Mutable/Source/CustomizableObject/Private/MuCO/CustomizableObjectInstance.cpp b/Engine/Plugins/Experimental/Mutable/Source/CustomizableObject/Private/MuCO/CustomizableObjectInstance.cpp
--- a/Engine/Plugins/Experimental/Mutable/Source/CustomizableObject/Private/MuCO/CustomizableObjectInstance.cpp
+++ b/Engine/Plugins/Experimental/Mutable/Source/CustomizableObject/Private/MuCO/CustomizableObjectInstance.cpp
@@ -19,6 +19,7 @@
 #include "Materials/MaterialInstanceDynamic.h"
 #include "Modules/ModuleManager.h"
 
+#include "MuCO/CustomizableObjectSystemPrivate.h"
 #include "MuCO/CustomizableInstanceLODManagement.h"
 #include "MuCO/CustomizableInstancePrivateData.h"
 #include "MuCO/CustomizableObjectExtension.h"
@@ -36,6 +37,7 @@
 #include "Rendering/Texture2DResource.h"
 #include "RenderingThread.h"
 #include "SkeletalMergingLibrary.h"
+#include "UnrealMutableImageProvider.h"
 #include "UObject/UObjectIterator.h"
 #include "Widgets/Notifications/SNotificationList.h"
 #include "PhysicsEngine/PhysicsConstraintTemplate.h"
@@ -47,6 +49,80 @@
 #include "MessageLogModule.h"
 #endif
 
+
+// Struct used by BuildMaterials() to identify common materials between LODs
+struct FMutableMaterialPlaceholder
+{
+	enum class EPlaceHolderParamType { Vector, Scalar, Texture };
+
+	struct FMutableMaterialPlaceHolderParam
+	{
+		FName ParamName;
+		int32 LayerIndex; // Set to -1 for non-multilayer params
+		FLinearColor Vector;
+		float Scalar;
+		FGeneratedTexture Texture;
+		EPlaceHolderParamType Type;
+
+		FMutableMaterialPlaceHolderParam(const FName& InParamName, const int32 InLayerIndex, const FLinearColor& InVector)
+			: ParamName(InParamName), LayerIndex(InLayerIndex), Vector(InVector), Type(EPlaceHolderParamType::Vector) {}
+
+		FMutableMaterialPlaceHolderParam(const FName& InParamName, const int32 InLayerIndex, const float InScalar)
+			: ParamName(InParamName), LayerIndex(InLayerIndex), Scalar(InScalar), Type(EPlaceHolderParamType::Scalar) {}
+
+		FMutableMaterialPlaceHolderParam(const FName& InParamName, const int32 InLayerIndex, const FGeneratedTexture& InTexture)
+			: ParamName(InParamName), LayerIndex(InLayerIndex), Texture(InTexture), Type(EPlaceHolderParamType::Texture) {}
+
+		bool operator<(const FMutableMaterialPlaceHolderParam& Other) const
+		{
+			return Type < Other.Type || ParamName.CompareIndexes(Other.ParamName);
+		}
+	};
+
+	UMaterialInterface* ParentMaterial;
+	TArray<FMutableMaterialPlaceHolderParam> Params;
+	int32 MatIndex = -1;
+
+	void AddParam(const FMutableMaterialPlaceHolderParam& NewParam) { Params.Add(NewParam); }
+
+	// Return a hash of the material and its parameters
+	uint32 GetHash()
+	{
+
+		uint32 Hash = ParentMaterial ? ParentMaterial->GetUniqueID() : 0;
+
+		// Sort parameters before building the hash.
+		Params.Sort();
+
+		for (const FMutableMaterialPlaceHolderParam& Param : Params)
+		{
+			uint32 ParamHash = GetTypeHash(Param.ParamName);
+			ParamHash = HashCombineFast(ParamHash, (uint32)Param.LayerIndex);
+			ParamHash = HashCombineFast(ParamHash, (uint32)Param.Type);
+
+			switch (Param.Type)
+			{
+			case EPlaceHolderParamType::Vector:
+				ParamHash = HashCombineFast(ParamHash, GetTypeHash(Param.Vector));
+				break;
+
+			case EPlaceHolderParamType::Scalar:
+				ParamHash = HashCombineFast(ParamHash, GetTypeHash(Param.Scalar));
+				break;
+
+			case EPlaceHolderParamType::Texture:
+				ParamHash = HashCombineFast(ParamHash, Param.Texture.Texture->GetUniqueID());
+				break;
+			}
+
+			Hash = HashCombineFast(Hash, ParamHash);
+		}
+
+		return Hash;
+	}
+};
+
+
 UTexture2D* UCustomizableInstancePrivateData::CreateTexture()
 {
 	UTexture2D* NewTexture = NewObject<UTexture2D>(
@@ -2886,7 +2962,7 @@ void UCustomizableObjectInstance::CommitMinMaxLOD()
 
 // The memory allocated in the function and pointed by the returned pointer is owned by the caller and must be freed. 
 // If assigned to a UTexture2D, it will be freed by that UTexture2D
-FTexturePlatformData* UCustomizableInstancePrivateData::MutableCreateImagePlatformData(mu::Ptr<const mu::Image> MutableImage, int32 OnlyLOD, uint16 FullSizeX, uint16 FullSizeY)
+FTexturePlatformData* MutableCreateImagePlatformData(mu::Ptr<const mu::Image> MutableImage, int32 OnlyLOD, uint16 FullSizeX, uint16 FullSizeY)
 {
 	int32 SizeX = FMath::Max(MutableImage->GetSize()[0], FullSizeX);
 	int32 SizeY = FMath::Max(MutableImage->GetSize()[1], FullSizeY);
@@ -3161,7 +3237,7 @@ FTexturePlatformData* UCustomizableInstancePrivateData::MutableCreateImagePlatfo
 
 
 
-void UCustomizableInstancePrivateData::ConvertImage(UTexture2D* Texture, mu::ImagePtrConst MutableImage, const FMutableModelImageProperties& Props, int OnlyLOD, int32 ExtractChannel )
+void ConvertImage(UTexture2D* Texture, mu::Ptr<const mu::Image> MutableImage, const FMutableModelImageProperties& Props, int OnlyLOD, int32 ExtractChannel)
 {
 	MUTABLE_CPUPROFILER_SCOPE(UCustomizableInstancePrivateData::ConvertImage);
 
@@ -6574,41 +6650,6 @@ TSet<UAssetUserData*> UCustomizableObjectInstance::GetMergedAssetUserData(int32
 }
 
 
-uint32 UCustomizableInstancePrivateData::FMutableMaterialPlaceholder::GetHash()
-{
-
-	uint32 Hash = ParentMaterial ? ParentMaterial->GetUniqueID() : 0;
-
-	// Sort parameters before building the hash.
-	Params.Sort();
-
-	for (const FMutableMaterialPlaceHolderParam& Param : Params)
-	{
-		uint32 ParamHash = GetTypeHash(Param.ParamName);
-		ParamHash = HashCombineFast(ParamHash, (uint32)Param.LayerIndex);
-		ParamHash = HashCombineFast(ParamHash, (uint32)Param.Type);
-
-		switch (Param.Type)
-		{
-		case EPlaceHolderParamType::Vector:
-			ParamHash = HashCombineFast(ParamHash, GetTypeHash(Param.Vector));
-			break;
-
-		case EPlaceHolderParamType::Scalar:
-			ParamHash = HashCombineFast(ParamHash, GetTypeHash(Param.Scalar));
-			break;
-
-		case EPlaceHolderParamType::Texture:
-			ParamHash = HashCombineFast(ParamHash, Param.Texture.Texture->GetUniqueID());
-			break;
-		}
-
-		Hash = HashCombineFast(Hash, ParamHash);
-	}
-
-	return Hash;
-}
-
 #if WITH_EDITORONLY_DATA
 
 void UCustomizableInstancePrivateData::RegenerateImportedModel(USkeletalMesh* SkeletalMesh)
diff --git a/Engine/Plugins/Experimental/Mutable/Source/CustomizableObject/Private/MuCO/CustomizableObjectInstanceUsage.cpp b/Engine/Plugins/Experimental/Mutable/Source/CustomizableObject/Private/MuCO/CustomizableObjectInstanceUsage.cpp
--- a/Engine/Plugins/Experimental/Mutable/Source/CustomizableObject/Private/MuCO/CustomizableObjectInstanceUsage.cpp
+++ b/Engine/Plugins/Experimental/Mutable/Source/CustomizableObject/Private/MuCO/CustomizableObjectInstanceUsage.cpp
@@ -13,6 +13,7 @@
 #include "Engine/World.h"
 #include "GameFramework/Pawn.h"
 #include "Materials/MaterialInstanceDynamic.h"
+#include "MuCO/CustomizableObject.h"
 #include "Rendering/SkeletalMeshRenderData.h"
 #include "UObject/ObjectSaveContext.h"
 #include "Stats/Stats.h"
diff --git a/Engine/Plugins/Experimental/Mutable/Source/CustomizableObject/Private/MuCO/CustomizableObjectMipDataProvider.cpp b/Engine/Plugins/Experimental/Mutable/Source/CustomizableObject/Private/MuCO/CustomizableObjectMipDataProvider.cpp
--- a/Engine/Plugins/Experimental/Mutable/Source/CustomizableObject/Private/MuCO/CustomizableObjectMipDataProvider.cpp
+++ b/Engine/Plugins/Experimental/Mutable/Source/CustomizableObject/Private/MuCO/CustomizableObjectMipDataProvider.cpp
@@ -6,6 +6,7 @@
 #include "MuCO/CustomizableObjectSystemPrivate.h"
 #include "MuR/Model.h"
 #include "TextureResource.h"
+#include "UnrealMutableImageProvider.h"
 
 #include UE_INLINE_GENERATED_CPP_BY_NAME(CustomizableObjectMipDataProvider)
 
diff --git a/Engine/Plugins/Experimental/Mutable/Source/CustomizableObject/Private/MuCO/CustomizableObjectPrivate.h b/Engine/Plugins/Experimental/Mutable/Source/CustomizableObject/Private/MuCO/CustomizableObjectPrivate.h
deleted file mode 100644
diff --git a/Engine/Plugins/Experimental/Mutable/Source/CustomizableObject/Private/MuCO/CustomizableObjectSystem.cpp b/Engine/Plugins/Experimental/Mutable/Source/CustomizableObject/Private/MuCO/CustomizableObjectSystem.cpp
--- a/Engine/Plugins/Experimental/Mutable/Source/CustomizableObject/Private/MuCO/CustomizableObjectSystem.cpp
+++ b/Engine/Plugins/Experimental/Mutable/Source/CustomizableObject/Private/MuCO/CustomizableObjectSystem.cpp
@@ -30,6 +30,7 @@
 #include "ContentStreaming.h"
 #include "MuCO/EditorImageProvider.h"
 #include "Components/SkeletalMeshComponent.h"
+#include "CustomizableObjectSystemPrivate.h"
 
 #if WITH_EDITOR
 #include "Editor.h"
@@ -2031,7 +2032,7 @@ namespace impl
 					{
 						if (!OperationData->ImageToPlatformDataMap.Contains(Image.ImageID))
 						{
-							FTexturePlatformData* PlatformData = UCustomizableInstancePrivateData::MutableCreateImagePlatformData(MutableImage, -1, Image.FullImageSizeX, Image.FullImageSizeY);
+							FTexturePlatformData* PlatformData = MutableCreateImagePlatformData(MutableImage, -1, Image.FullImageSizeX, Image.FullImageSizeY);
 							OperationData->ImageToPlatformDataMap.Add(Image.ImageID, PlatformData);
 						}
 						else
diff --git a/Engine/Plugins/Experimental/Mutable/Source/CustomizableObject/Private/MuCO/CustomizableObjectSystemPrivate.h b/Engine/Plugins/Experimental/Mutable/Source/CustomizableObject/Private/MuCO/CustomizableObjectSystemPrivate.h
--- a/Engine/Plugins/Experimental/Mutable/Source/CustomizableObject/Private/MuCO/CustomizableObjectSystemPrivate.h
+++ b/Engine/Plugins/Experimental/Mutable/Source/CustomizableObject/Private/MuCO/CustomizableObjectSystemPrivate.h
@@ -25,6 +25,8 @@
 	#include "Async/TaskGraphInterfaces.h"
 #endif
 
+#include "CustomizableObjectSystemPrivate.generated.h"
+
 #if WITH_EDITORONLY_DATA
 class UEditorImageProvider;
 #endif
@@ -188,6 +190,42 @@ struct FMutableResourceCache
 };
 
 
+USTRUCT()
+struct FGeneratedTexture
+{
+	GENERATED_USTRUCT_BODY();
+
+	FMutableImageCacheKey Key;
+
+	UPROPERTY(Category = CustomizableObjectInstance, VisibleAnywhere)
+	FString Name;
+
+	UPROPERTY(Category = CustomizableObjectInstance, VisibleAnywhere)
+	TObjectPtr<UTexture> Texture = nullptr;
+};
+
+
+USTRUCT()
+struct FGeneratedMaterial
+{
+	GENERATED_USTRUCT_BODY();
+
+	UPROPERTY()
+	TObjectPtr<UMaterialInterface> MaterialInterface;
+
+	UPROPERTY(Category = CustomizedMaterial, VisibleAnywhere)
+	TArray< FGeneratedTexture > Textures;
+
+	// Surface or SharedSurface Id
+	uint32 SurfaceId = 0;
+
+	// Index of the material to instantiate (UCustomizableObject::ReferencedMaterials)
+	uint32 MaterialIndex = 0;
+
+	bool operator==(const FGeneratedMaterial& Other) const { return SurfaceId == Other.SurfaceId && MaterialIndex == Other.MaterialIndex; };
+};
+
+
 DECLARE_DELEGATE(FMutableTaskDelegate);
 struct FMutableTask
 {
@@ -614,3 +652,6 @@ public:
 	int32 NumSkeletalMeshes = 0;
 };
 
+
+/** Set OnlyLOD to -1 to generate all mips */
+CUSTOMIZABLEOBJECT_API FTexturePlatformData* MutableCreateImagePlatformData(mu::Ptr<const mu::Image> MutableImage, int32 OnlyLOD, uint16 FullSizeX, uint16 FullSizeY);
diff --git a/Engine/Plugins/Experimental/Mutable/Source/CustomizableObject/Private/MuCO/LogInformationUtil.cpp b/Engine/Plugins/Experimental/Mutable/Source/CustomizableObject/Private/MuCO/LogInformationUtil.cpp
--- a/Engine/Plugins/Experimental/Mutable/Source/CustomizableObject/Private/MuCO/LogInformationUtil.cpp
+++ b/Engine/Plugins/Experimental/Mutable/Source/CustomizableObject/Private/MuCO/LogInformationUtil.cpp
@@ -2,17 +2,19 @@
 
 #include "MuCO/LogInformationUtil.h"
 
+#include "CustomizableObjectSystemPrivate.h"
 #include "Engine/SkeletalMesh.h"
 #include "Engine/World.h"
 #include "GameFramework/PlayerController.h"
 #include "MuCO/CustomizableInstancePrivateData.h"
 #include "TextureResource.h"
+#include "MuCO/CustomizableObject.h"
 
 int LogInformationUtil::CountLOD0 = 0;
 int LogInformationUtil::CountLOD1 = 0;
 int LogInformationUtil::CountLOD2 = 0;
 
-void LogInformationUtil::PrintGeneratedTextures(const TArray<struct FGeneratedTexture> GeneratedTextures, FString& Log, bool DoPrintInitialMessage)
+void LogInformationUtil::PrintGeneratedTextures(const TArray<FGeneratedTexture> GeneratedTextures, FString& Log, bool DoPrintInitialMessage)
 {
 	if (GeneratedTextures.Num() == 0)
 	{
diff --git a/Engine/Plugins/Experimental/Mutable/Source/CustomizableObject/Private/MuCO/UnrealConversionUtils.cpp b/Engine/Plugins/Experimental/Mutable/Source/CustomizableObject/Private/MuCO/UnrealConversionUtils.cpp
--- a/Engine/Plugins/Experimental/Mutable/Source/CustomizableObject/Private/MuCO/UnrealConversionUtils.cpp
+++ b/Engine/Plugins/Experimental/Mutable/Source/CustomizableObject/Private/MuCO/UnrealConversionUtils.cpp
@@ -6,6 +6,10 @@
 #include "MuCO/UnrealPortabilityHelpers.h"
 #include "Engine/SkeletalMesh.h"
 #include "GPUSkinVertexFactory.h"
+#include "MuCO/CustomizableObject.h"
+#include "MuR/Mesh.h"
+#include "MuR/MeshBufferSet.h"
+#include "MuR/MutableTrace.h"
 
 class USkeleton;
 
diff --git a/Engine/Plugins/Experimental/Mutable/Source/CustomizableObject/Private/MuCO/UnrealMutableModelDiskStreamer.cpp b/Engine/Plugins/Experimental/Mutable/Source/CustomizableObject/Private/MuCO/UnrealMutableModelDiskStreamer.cpp
--- a/Engine/Plugins/Experimental/Mutable/Source/CustomizableObject/Private/MuCO/UnrealMutableModelDiskStreamer.cpp
+++ b/Engine/Plugins/Experimental/Mutable/Source/CustomizableObject/Private/MuCO/UnrealMutableModelDiskStreamer.cpp
@@ -7,12 +7,14 @@
 #include "MuCO/CustomizableObjectSystem.h"
 #include "MuR/Model.h"
 #include "MuR/MutableTrace.h"
+#include "MuCO/LogBenchmarkUtil.h"
 
 #if WITH_EDITOR
 #include "HAL/PlatformFileManager.h"
 #endif
 
-DEFINE_STAT(STAT_MutableStreamingOps);
+DECLARE_DWORD_ACCUMULATOR_STAT(TEXT("Streaming Ops"), STAT_MutableStreamingOps, STATGROUP_Mutable);
+
 
 //-------------------------------------------------------------------------------------------------
 #if WITH_EDITOR
diff --git a/Engine/Plugins/Experimental/Mutable/Source/CustomizableObject/Private/MuCO/UnrealToMutableTextureConversionUtils.h b/Engine/Plugins/Experimental/Mutable/Source/CustomizableObject/Private/MuCO/UnrealToMutableTextureConversionUtils.h
deleted file mode 100644
diff --git a/Engine/Plugins/Experimental/Mutable/Source/CustomizableObject/Public/MuCO/CustomizableObjectMipDataProvider.h b/Engine/Plugins/Experimental/Mutable/Source/CustomizableObject/Public/MuCO/CustomizableObjectMipDataProvider.h
--- a/Engine/Plugins/Experimental/Mutable/Source/CustomizableObject/Public/MuCO/CustomizableObjectMipDataProvider.h
+++ b/Engine/Plugins/Experimental/Mutable/Source/CustomizableObject/Public/MuCO/CustomizableObjectMipDataProvider.h
@@ -10,8 +10,6 @@
 #include "Streaming/TextureMipDataProvider.h"
 
 #include "DefaultImageProvider.h"
-#include "MuCO/CustomizableObjectSystemPrivate.h"
-#include "MuCO/UnrealMutableImageProvider.h"
 #include "MuR/System.h"
 
 #include "CustomizableObjectMipDataProvider.generated.h"
diff --git a/Engine/Plugins/Experimental/Mutable/Source/CustomizableObject/Public/MuCO/UnrealConversionUtils.h b/Engine/Plugins/Experimental/Mutable/Source/CustomizableObject/Public/MuCO/UnrealConversionUtils.h
--- a/Engine/Plugins/Experimental/Mutable/Source/CustomizableObject/Public/MuCO/UnrealConversionUtils.h
+++ b/Engine/Plugins/Experimental/Mutable/Source/CustomizableObject/Public/MuCO/UnrealConversionUtils.h
@@ -2,14 +2,19 @@
 
 #pragma once
 
-#include "MuCO/CustomizableObjectSystemPrivate.h"
-
 //! Order of the unreal vertex buffers when in mutable data
 #define MUTABLE_VERTEXBUFFER_POSITION	0
 #define MUTABLE_VERTEXBUFFER_TANGENT	1
 #define MUTABLE_VERTEXBUFFER_TEXCOORDS	2
 
-namespace mu { class FMeshBufferSet; }
+#include "MuR/Ptr.h"
+#include "Containers/Array.h"
+
+namespace mu
+{
+	class Mesh;
+	class FMeshBufferSet;
+}
 struct FReferenceSkeleton;
 
 class FSkeletalMeshLODRenderData;
@@ -34,7 +39,7 @@ namespace UnrealConversionUtils
 	 */
 	CUSTOMIZABLEOBJECT_API void SetupRenderSections(
 		const USkeletalMesh* OutSkeletalMesh,
-		const mu::MeshPtrConst InMutableMesh,
+		const mu::Ptr<const mu::Mesh> InMutableMesh,
 		const int32 InMeshLODIndex,
 		const TArray<uint16>& InBoneMap,
 		const int32 InFirstBoneMapIndex);
@@ -47,7 +52,7 @@ namespace UnrealConversionUtils
 	 */
 	CUSTOMIZABLEOBJECT_API void CopyMutableVertexBuffers(
 		USkeletalMesh* OutSkeletalMesh,
-		const mu::MeshPtrConst InMutableMesh,
+		const mu::Ptr<const mu::Mesh> InMutableMesh,
 		const int32 InMeshLODIndex);
 
 	
@@ -57,7 +62,7 @@ namespace UnrealConversionUtils
 	 * @param OutLODModel - The LOD model to be updated.
 	 * @return True if the operation could be performed successfully, false if not.
 	 */
-	CUSTOMIZABLEOBJECT_API bool CopyMutableIndexBuffers(mu::MeshPtrConst InMutableMesh,
+	CUSTOMIZABLEOBJECT_API bool CopyMutableIndexBuffers(mu::Ptr<const mu::Mesh> InMutableMesh,
 	                                                    FSkeletalMeshLODRenderData& OutLODModel);
 	
 
diff --git a/Engine/Plugins/Experimental/Mutable/Source/CustomizableObject/Public/MuCO/UnrealMutableModelDiskStreamer.h b/Engine/Plugins/Experimental/Mutable/Source/CustomizableObject/Public/MuCO/UnrealMutableModelDiskStreamer.h
--- a/Engine/Plugins/Experimental/Mutable/Source/CustomizableObject/Public/MuCO/UnrealMutableModelDiskStreamer.h
+++ b/Engine/Plugins/Experimental/Mutable/Source/CustomizableObject/Public/MuCO/UnrealMutableModelDiskStreamer.h
@@ -3,7 +3,6 @@
 #pragma once
 
 #include "MuCO/CustomizableObject.h"
-#include "MuCO/LogBenchmarkUtil.h"
 #include "MuR/Serialisation.h"
 
 class FArchive;
@@ -13,9 +12,6 @@ class UCustomizableObject;
 namespace mu { class Model; }
 struct FMutableStreamableBlock;
 
-DECLARE_DWORD_ACCUMULATOR_STAT_EXTERN(TEXT("Streaming Ops"), STAT_MutableStreamingOps, STATGROUP_Mutable, );
-
-
 class UnrealMutableInputStream : public mu::InputStream
 {
 public:
diff --git a/Engine/Plugins/Experimental/Mutable/Source/CustomizableObjectEditor/CustomizableObjectEditor.Build.cs b/Engine/Plugins/Experimental/Mutable/Source/CustomizableObjectEditor/CustomizableObjectEditor.Build.cs
--- a/Engine/Plugins/Experimental/Mutable/Source/CustomizableObjectEditor/CustomizableObjectEditor.Build.cs
+++ b/Engine/Plugins/Experimental/Mutable/Source/CustomizableObjectEditor/CustomizableObjectEditor.Build.cs
@@ -101,7 +101,6 @@ public class CustomizableObjectEditor : ModuleRules
 			);
 
 		PrivateIncludePaths.AddRange(new string[] {
-				"CustomizableObject/Private",
 				"MutableRuntime/Private",
 				"MutableTools/Private",
 			});
diff --git a/Engine/Plugins/Experimental/Mutable/Source/CustomizableObjectEditor/Private/MuCOE/GenerateMutableSource/GenerateMutableSource.cpp b/Engine/Plugins/Experimental/Mutable/Source/CustomizableObjectEditor/Private/MuCOE/GenerateMutableSource/GenerateMutableSource.cpp
--- a/Engine/Plugins/Experimental/Mutable/Source/CustomizableObjectEditor/Private/MuCOE/GenerateMutableSource/GenerateMutableSource.cpp
+++ b/Engine/Plugins/Experimental/Mutable/Source/CustomizableObjectEditor/Private/MuCOE/GenerateMutableSource/GenerateMutableSource.cpp
@@ -49,7 +49,6 @@
 #include "MuT/NodeSurfaceEdit.h"
 #include "PhysicsEngine/PhysicsAsset.h"
 #include "PlatformInfo.h"
-#include "MuCO/CustomizableObjectPrivate.h"
 #include "Math/NumericLimits.h"
 
 #define LOCTEXT_NAMESPACE "CustomizableObjectEditor"
diff --git a/Engine/Plugins/Experimental/Mutable/Source/CustomizableObjectEditor/Private/MuCOE/MutableMeshPreviewUtils.cpp b/Engine/Plugins/Experimental/Mutable/Source/CustomizableObjectEditor/Private/MuCOE/MutableMeshPreviewUtils.cpp
--- a/Engine/Plugins/Experimental/Mutable/Source/CustomizableObjectEditor/Private/MuCOE/MutableMeshPreviewUtils.cpp
+++ b/Engine/Plugins/Experimental/Mutable/Source/CustomizableObjectEditor/Private/MuCOE/MutableMeshPreviewUtils.cpp
@@ -11,6 +11,7 @@
 #include "MuCO/UnrealPortabilityHelpers.h"
 #include "MuR/OpMeshFormat.h"
 #include "Engine/SkeletalMesh.h"
+#include "MuCO/CustomizableObject.h"
 #include "UObject/Package.h"
 
 class UMaterialInterface;
diff --git a/Engine/Plugins/Experimental/Mutable/Source/CustomizableObjectEditor/Private/MuCOE/SMutableImageViewer.cpp b/Engine/Plugins/Experimental/Mutable/Source/CustomizableObjectEditor/Private/MuCOE/SMutableImageViewer.cpp
--- a/Engine/Plugins/Experimental/Mutable/Source/CustomizableObjectEditor/Private/MuCOE/SMutableImageViewer.cpp
+++ b/Engine/Plugins/Experimental/Mutable/Source/CustomizableObjectEditor/Private/MuCOE/SMutableImageViewer.cpp
@@ -3,9 +3,10 @@
 #include "MuCOE/SMutableImageViewer.h"
 
 #include "Brushes/SlateImageBrush.h"
-#include "MuCO/CustomizableInstancePrivateData.h"
 #include "MuT/TypeInfo.h"
 #include "TextureResource.h"
+#include "MuCO/CustomizableObject.h"
+#include "MuCO/CustomizableObjectInstance.h"
 #include "Widgets/Input/SNumericEntryBox.h"
 #include "Widgets/Input/SSegmentedControl.h"
 
@@ -288,7 +289,7 @@ void SMutableImageViewer::Tick(const FGeometry& AllottedGeometry, const double I
 			Props.Filter = TF_Nearest;
 			Props.SRGB = true;
 			Props.LODBias = 0;
-			UCustomizableInstancePrivateData::ConvertImage(UnrealImage, ImageToConvert, Props, CurrentVisibleLOD, ExtractChannel);
+			ConvertImage(UnrealImage, ImageToConvert, Props, CurrentVisibleLOD, ExtractChannel);
 			UnrealImage->NeverStream = true;
 			UnrealImage->UpdateResource();
 		}
