#rnx
#tests Editor, Bake

diff --git a/Engine/Plugins/Experimental/Mutable/Source/CustomizableObject/Private/MuCO/CustomizableObjectInstance.cpp b/Engine/Plugins/Experimental/Mutable/Source/CustomizableObject/Private/MuCO/CustomizableObjectInstance.cpp
--- a/Engine/Plugins/Experimental/Mutable/Source/CustomizableObject/Private/MuCO/CustomizableObjectInstance.cpp
+++ b/Engine/Plugins/Experimental/Mutable/Source/CustomizableObject/Private/MuCO/CustomizableObjectInstance.cpp
@@ -6798,7 +6798,29 @@ void UCustomizableInstancePrivateData::RegenerateImportedModel(USkeletalMesh* Sk
 			ImportedSection.BaseIndex = RenderSection.BaseIndex;
 			ImportedSection.BaseVertexIndex = RenderSection.BaseVertexIndex;
 			ImportedSection.BoneMap = RenderSection.BoneMap;
-			ImportedSection.MaterialIndex = RenderSection.MaterialIndex;
+
+			const TArray<int32>& LODMaterialMap = SkeletalMesh->GetLODInfoArray()[LODIndex].LODMaterialMap;
+
+			if (LODMaterialMap.IsValidIndex(RenderSection.MaterialIndex))
+			{
+				ImportedSection.MaterialIndex = LODMaterialMap[RenderSection.MaterialIndex];
+			}
+			else
+			{
+				// The material should have been in the LODMaterialMap
+				ensureMsgf(false, TEXT("Unexpected material index in UCustomizableInstancePrivateData::RegenerateImportedModel"));
+
+				// Fallback index, may shift materials around sections
+				if (SkeletalMesh->GetMaterials().IsValidIndex(RenderSection.MaterialIndex))
+				{
+					ImportedSection.MaterialIndex = RenderSection.MaterialIndex;
+				}
+				else
+				{
+					ImportedSection.MaterialIndex = 0;
+				}
+			}
+
 			ImportedSection.MaxBoneInfluences = RenderSection.MaxBoneInfluences;
 			ImportedSection.OriginalDataSectionIndex = OriginalIndex++;
 
@@ -6821,7 +6843,6 @@ void UCustomizableInstancePrivateData::RegenerateImportedModel(USkeletalMesh* Sk
 		LODInfo->BuildGUID = LODInfo->ComputeDeriveDataCacheKey(SkeletalMeshLODGroupSettings);
 
 		ImportedModel->LODModels[LODIndex].BuildStringID = ImportedModel->LODModels[LODIndex].GetLODModelDeriveDataKey();
-
 	}
 
 }
