[mutable] Fixed mesh update not calling callbacks when changing graphic settings.

#rnx

diff --git a/Engine/Plugins/Experimental/Mutable/Source/CustomizableObject/Private/MuCO/CustomizableObjectSystem.cpp b/Engine/Plugins/Experimental/Mutable/Source/CustomizableObject/Private/MuCO/CustomizableObjectSystem.cpp
--- a/Engine/Plugins/Experimental/Mutable/Source/CustomizableObject/Private/MuCO/CustomizableObjectSystem.cpp
+++ b/Engine/Plugins/Experimental/Mutable/Source/CustomizableObject/Private/MuCO/CustomizableObjectSystem.cpp
@@ -863,7 +863,7 @@ void FinishUpdateGlobal(const TSharedRef<FUpdateContextPrivate>& Context)
 }
 
 
-/** Update the given Instance Skeletal Meshes and call its callbacks. */
+/** Update the given Instance Skeletal Meshes */
 void UpdateSkeletalMesh(const TSharedRef<FUpdateContextPrivate>& Context)
 {
 	MUTABLE_CPUPROFILER_SCOPE(UpdateSkeletalMesh);
@@ -2754,7 +2754,7 @@ namespace impl
 		{
 			System->ClearCurrentMutableOperation();
 
-			Operation->UpdateResult = EUpdateResult::ErrorOptimized;
+			Operation->UpdateResult = EUpdateResult::Success;
 			UpdateSkeletalMesh(Operation);
 			FinishUpdateGlobal(Operation);
 			return;
