[Mutable] Remove image format from ASTOpReferenceResource ImageDesc.

- This fixes problems where the optimizer would see a format with alpha, while constant resolution will remove the alpha.

#rnx

diff --git a/Engine/Plugins/Experimental/Mutable/Source/MutableTools/Private/MuT/CodeGenerator_Image.cpp b/Engine/Plugins/Experimental/Mutable/Source/MutableTools/Private/MuT/CodeGenerator_Image.cpp
--- a/Engine/Plugins/Experimental/Mutable/Source/MutableTools/Private/MuT/CodeGenerator_Image.cpp
+++ b/Engine/Plugins/Experimental/Mutable/Source/MutableTools/Private/MuT/CodeGenerator_Image.cpp
@@ -197,9 +197,15 @@ namespace mu
 			Ptr<ASTOpReferenceResource> ReferenceOp = new ASTOpReferenceResource();
 			ReferenceOp->type = OP_TYPE::IM_REFERENCE;
 			ReferenceOp->ID = pImage->GetReferencedTexture();
-			ReferenceOp->bForceLoad = pImage->IsForceLoad(); 
-			ReferenceOp->ImageDesc = FImageDesc(
-					FImageSize(pImage->GetSizeX(), pImage->GetSizeY()), pImage->GetFormat(), pImage->GetLODCount());
+			ReferenceOp->bForceLoad = pImage->IsForceLoad();
+
+			// Don't store the format. Format can vary between loaded constant image and reference and cause
+			// code optimization bugs.
+			// As it is now, reference will always have alpha channel but constant resolution can remove the 
+			// channel if not used.
+			// TODO: review this, probably the reference descriptor generation needs to check for alpha channels 
+			// as well.
+			ReferenceOp->ImageDesc = FImageDesc(pImage->GetSize(), EImageFormat::IF_NONE, pImage->GetLODCount());
 			Result.op = ReferenceOp;
 		}
 		else
