[mutable] Separate mesh tag strings from mesh data for more opportunities of deduplicating mesh data.
A new operation has been added, and several optimization steps for the generated code have been implemented.
Added basic widget to show for tags in the mutable debugger.
#rnx

diff --git a/Engine/Plugins/Experimental/Mutable/Source/CustomizableObject/Classes/MuCO/CustomizableObject.h b/Engine/Plugins/Experimental/Mutable/Source/CustomizableObject/Classes/MuCO/CustomizableObject.h
--- a/Engine/Plugins/Experimental/Mutable/Source/CustomizableObject/Classes/MuCO/CustomizableObject.h
+++ b/Engine/Plugins/Experimental/Mutable/Source/CustomizableObject/Classes/MuCO/CustomizableObject.h
@@ -1407,7 +1407,7 @@ private:
 	// This is a manual version number for the binary blobs in this asset.
 	// Increasing it invalidates all the previously compiled models.
 	// Warning: If while merging code both versions have changed, take the highest+1.
-	static const int32 CurrentSupportedVersion = 413;
+	static const int32 CurrentSupportedVersion = 414;
 
 public:
 
diff --git a/Engine/Plugins/Experimental/Mutable/Source/CustomizableObjectEditor/Private/MuCOE/SMutableCodeViewer.cpp b/Engine/Plugins/Experimental/Mutable/Source/CustomizableObjectEditor/Private/MuCOE/SMutableCodeViewer.cpp
--- a/Engine/Plugins/Experimental/Mutable/Source/CustomizableObjectEditor/Private/MuCOE/SMutableCodeViewer.cpp
+++ b/Engine/Plugins/Experimental/Mutable/Source/CustomizableObjectEditor/Private/MuCOE/SMutableCodeViewer.cpp
@@ -1564,6 +1564,24 @@ void SMutableCodeViewer::GenerateElementRecursive(const int32& InStateIndex, mu:
 		break;
 	}
 
+	case mu::OP_TYPE::ME_ADDTAGS:
+	{
+		const uint8* OpData = InProgram.GetOpArgsPointer(InParentAddress);
+
+		mu::OP::ADDRESS SourceAddress;
+		FMemory::Memcpy(&SourceAddress, OpData, sizeof(mu::OP::ADDRESS));
+		OpData += sizeof(mu::OP::ADDRESS);
+
+		uint16 TagCount;
+		FMemory::Memcpy(&TagCount, OpData, sizeof(uint16));
+		OpData += sizeof(uint16);
+
+		FString Caption = FString::Printf(TEXT("add %d tags to "), TagCount);
+		AddOpFunc(SourceAddress, Caption);
+
+		break;
+	}
+
 	default:
 	{
 		// Generic list of child operations
diff --git a/Engine/Plugins/Experimental/Mutable/Source/CustomizableObjectEditor/Private/MuCOE/SMutableCodeViewer.h b/Engine/Plugins/Experimental/Mutable/Source/CustomizableObjectEditor/Private/MuCOE/SMutableCodeViewer.h
--- a/Engine/Plugins/Experimental/Mutable/Source/CustomizableObjectEditor/Private/MuCOE/SMutableCodeViewer.h
+++ b/Engine/Plugins/Experimental/Mutable/Source/CustomizableObjectEditor/Private/MuCOE/SMutableCodeViewer.h
@@ -218,7 +218,9 @@ private:
 		mu::OP_TYPE::ME_MASKCLIPMESH,
 		mu::OP_TYPE::ME_MASKCLIPUVMASK,
 		mu::OP_TYPE::ME_FORMAT,
+		mu::OP_TYPE::ME_MASKDIFF,
 		mu::OP_TYPE::ME_DIFFERENCE,
+		mu::OP_TYPE::ME_EXTRACTLAYOUTBLOCK,
 		mu::OP_TYPE::IM_MAKEGROWMAP,
 	};
 
diff --git a/Engine/Plugins/Experimental/Mutable/Source/CustomizableObjectEditor/Private/MuCOE/SMutableMeshViewer.cpp b/Engine/Plugins/Experimental/Mutable/Source/CustomizableObjectEditor/Private/MuCOE/SMutableMeshViewer.cpp
--- a/Engine/Plugins/Experimental/Mutable/Source/CustomizableObjectEditor/Private/MuCOE/SMutableMeshViewer.cpp
+++ b/Engine/Plugins/Experimental/Mutable/Source/CustomizableObjectEditor/Private/MuCOE/SMutableMeshViewer.cpp
@@ -204,6 +204,14 @@ void SMutableMeshViewer::SetMesh(const mu::MeshPtrConst& InMesh)
 	{
 		MutableMesh = InMesh;
 
+		// Extract a copy of the tags
+		int32 TagCount = MutableMesh ? MutableMesh->GetTagCount() : 0;
+		MeshTagList.SetNum(TagCount, false);
+		for (int32 TagIndex=0; TagIndex<TagCount; ++TagIndex)
+		{
+			MeshTagList[TagIndex] = MakeShared<FString>(MutableMesh->GetTag(TagIndex));
+		}
+
 		if (MutableMesh)
 		{
 			OnMeshChanged();
@@ -432,13 +440,40 @@ TSharedRef<SWidget> SMutableMeshViewer::GenerateDataTableSlates()
 			]
 
 			// Bones data ----------------------------------------------------------------
-			+ SVerticalBox::Slot().
-			  Padding(0, EndOfSectionSpacing).
-			  AutoHeight()
+			+ SVerticalBox::Slot()
+			  .Padding(0, EndOfSectionSpacing)
+			  .AutoHeight()
 			[
 				SAssignNew(MutableSkeletonViewer, SMutableSkeletonViewer)
 			]
 			// ---------------------------------
+
+			// Tags
+			+ SVerticalBox::Slot()
+				.AutoHeight()
+				[
+					SNew(STextBlock).
+					Text(LOCTEXT("Tags", "Tags"))
+				]
+
+			+ SVerticalBox::Slot()
+				.Padding(IndentationSpace, AfterTitleSpacing)
+				.AutoHeight()
+				[
+					SNew(SListView< TSharedPtr<FString> >)
+					.ListItemsSource(&MeshTagList)
+					.OnGenerateRow(this, &SMutableMeshViewer::GenerateTagRow)
+				]
+			// ---------------------------------
+		];
+}
+
+
+TSharedRef<ITableRow> SMutableMeshViewer::GenerateTagRow(TSharedPtr<FString> InItem, const TSharedRef<STableViewBase>& OwnerTable)
+{
+	return SNew(STableRow<TSharedPtr<FName>>, OwnerTable)
+		[
+			SNew(STextBlock).Text(FText::FromString(*InItem))
 		];
 }
 
diff --git a/Engine/Plugins/Experimental/Mutable/Source/CustomizableObjectEditor/Private/MuCOE/SMutableMeshViewer.h b/Engine/Plugins/Experimental/Mutable/Source/CustomizableObjectEditor/Private/MuCOE/SMutableMeshViewer.h
--- a/Engine/Plugins/Experimental/Mutable/Source/CustomizableObjectEditor/Private/MuCOE/SMutableMeshViewer.h
+++ b/Engine/Plugins/Experimental/Mutable/Source/CustomizableObjectEditor/Private/MuCOE/SMutableMeshViewer.h
@@ -71,6 +71,10 @@ private:
 	/** Slate whose task is to display the skeleton found on this mesh as a slate tree view */
 	TSharedPtr<SMutableSkeletonViewer> MutableSkeletonViewer;
 
+	/** Widget-side copy of the tags in the mesh.  */
+	TArray<TSharedPtr<FString>> MeshTagList;
+	TSharedRef<ITableRow> GenerateTagRow(TSharedPtr<FString> InItem, const TSharedRef<STableViewBase>& OwnerTable);
+
 	/** Thumbnail objects used to render a view of the selected skeletal mesh for the mesh viewport*/
 	TSharedPtr<FAssetThumbnailPool> AssetThumbnailPool;
 	TSharedPtr<FAssetThumbnail> AssetThumbnail;
diff --git a/Engine/Plugins/Experimental/Mutable/Source/MutableRuntime/Private/MuR/CodeRunner.cpp b/Engine/Plugins/Experimental/Mutable/Source/MutableRuntime/Private/MuR/CodeRunner.cpp
--- a/Engine/Plugins/Experimental/Mutable/Source/MutableRuntime/Private/MuR/CodeRunner.cpp
+++ b/Engine/Plugins/Experimental/Mutable/Source/MutableRuntime/Private/MuR/CodeRunner.cpp
@@ -2814,6 +2814,80 @@ namespace mu
             break;
         }
 
+        case OP_TYPE::ME_ADDTAGS:
+		{
+			MUTABLE_CPUPROFILER_SCOPE(ME_ADDTAGS)
+
+			// Decode op
+			// TODO: Partial decode for each stage
+			const uint8* Data = pModel->GetPrivate()->m_program.GetOpArgsPointer(item.At);
+
+			OP::ADDRESS Source;
+			FMemory::Memcpy(&Source, Data, sizeof(OP::ADDRESS));
+			Data += sizeof(OP::ADDRESS);
+
+			// Schedule next stages
+			switch (item.Stage)
+			{
+			case 0:
+			{
+				if (Source)
+				{
+					// Request the source
+					AddOp(FScheduledOp(item.At, item, 1), FScheduledOp(Source, item));
+				}
+				else
+				{
+					StoreMesh(item, nullptr);
+				}
+				break;
+			}
+
+			case 1:
+			{
+				MUTABLE_CPUPROFILER_SCOPE(ME_ADDTAGS_2)
+
+				Ptr<const Mesh> SourceMesh = LoadMesh(FCacheAddress(Source, item));
+
+				if (!SourceMesh)
+				{
+					StoreMesh(item, nullptr);
+				}
+				else
+				{
+					Ptr<Mesh> Result = CloneOrTakeOver(SourceMesh);
+
+					// Decode the tags
+					uint16 TagCount;
+					FMemory::Memcpy(&TagCount, Data, sizeof(uint16));
+					Data += sizeof(uint16);
+
+					int32 FirstMeshTagIndex = Result->m_tags.Num();
+					Result->m_tags.SetNum(FirstMeshTagIndex+TagCount);
+					for (uint16 TagIndex = 0; TagIndex < TagCount; ++TagIndex)
+					{
+						OP::ADDRESS TagConstant;
+						FMemory::Memcpy(&TagConstant, Data, sizeof(OP::ADDRESS));
+						Data += sizeof(OP::ADDRESS);
+
+						check(TagConstant < (uint32)pModel->GetPrivate()->m_program.m_constantStrings.Num());
+						const FString& Name = pModel->GetPrivate()->m_program.m_constantStrings[TagConstant];
+						Result->m_tags[FirstMeshTagIndex+TagIndex] = Name;
+					}
+
+					StoreMesh(item, Result);
+				}
+
+				break;
+			}
+
+			default:
+				check(false);
+			}
+
+			break;
+		}
+
         case OP_TYPE::ME_PROJECT:
         {
 			OP::MeshProjectArgs args = pModel->GetPrivate()->m_program.GetOpArgs<OP::MeshProjectArgs>(item.At);
diff --git a/Engine/Plugins/Experimental/Mutable/Source/MutableRuntime/Private/MuR/CodeRunnerTasked.cpp b/Engine/Plugins/Experimental/Mutable/Source/MutableRuntime/Private/MuR/CodeRunnerTasked.cpp
--- a/Engine/Plugins/Experimental/Mutable/Source/MutableRuntime/Private/MuR/CodeRunnerTasked.cpp
+++ b/Engine/Plugins/Experimental/Mutable/Source/MutableRuntime/Private/MuR/CodeRunnerTasked.cpp
@@ -190,7 +190,8 @@ namespace mu
 		{   0, -10,   0,   0 },	// ME_CLIPDEFORM
 		{   0, -10,   0,   0 },	// ME_MORPHRESHAPE
 		{   0,   0,   0,   0 },	// ME_OPTIMIZESKINNING
-		    
+		{	0,   0,	  0,   0 },	// ME_ADDTAGS
+
 		{   0,   0,   0,   0 },	// IN_ADDMESH
 		{   0,   0,   0,   0 },	// IN_ADDIMAGE
 		{   0,   0,   0,   0 },	// IN_ADDVECTOR
diff --git a/Engine/Plugins/Experimental/Mutable/Source/MutableRuntime/Private/MuR/Operations.cpp b/Engine/Plugins/Experimental/Mutable/Source/MutableRuntime/Private/MuR/Operations.cpp
--- a/Engine/Plugins/Experimental/Mutable/Source/MutableRuntime/Private/MuR/Operations.cpp
+++ b/Engine/Plugins/Experimental/Mutable/Source/MutableRuntime/Private/MuR/Operations.cpp
@@ -105,7 +105,7 @@ namespace mu
 
 		{ DT_MESH,			true,	{ 0, 0, 0, 0, 0, 0, 0, 0, 0, 0, 0, 0, 0, 0, 0, 0, 0, 0, 0, 0, 0, 0, 0, 0, 0, 0, 0, 0, 0 }		},	// ME_APPLYLAYOUT
 		{ DT_MESH,			true,	{ 0, 0, 0, 0, 0, 0, 0, 0, 0, 0, 0, 0, 0, 0, 0, 0, 0, 0, 0, 0, 0, 0, 0, 0, 0, 0, 0, 0, 0 }		},	// ME_DIFFERENCE
-		{ DT_MESH,			true,	{ 0, 0, 0, 0, 0, 0, 0, 0, 0, 0, 0, 0, 0, 0, 0, 0, 0, 0, 0, 0, 0, 0, 0, 0, 0, 0, 0, 0, 0 }		},	// ME_MORPH2
+		{ DT_MESH,			true,	{ 0, 0, 0, 0, 0, 0, 0, 0, 0, 0, 0, 0, 0, 0, 0, 0, 0, 0, 0, 0, 0, 0, 0, 0, 0, 0, 0, 0, 0 }		},	// ME_MORPH
 		{ DT_MESH,			true,	{ 0, 0, 0, 0, 0, 0, 0, 0, 0, 0, 0, 0, 0, 0, 0, 0, 0, 0, 0, 0, 0, 0, 0, 0, 0, 0, 0, 0, 0 }		},	// ME_MERGE
 		{ DT_MESH,			true,	{ 0, 0, 0, 0, 0, 0, 0, 0, 0, 0, 0, 0, 0, 0, 0, 0, 0, 0, 0, 0, 0, 0, 0, 0, 0, 0, 0, 0, 0 }		},	// ME_INTERPOLATE
 		{ DT_MESH,			true,	{ 0, 0, 0, 0, 0, 0, 0, 0, 0, 0, 0, 0, 0, 0, 0, 0, 0, 0, 0, 0, 0, 0, 0, 0, 0, 0, 0, 0, 0 }		},	// ME_MASKCLIPMESH
@@ -126,6 +126,7 @@ namespace mu
 		{ DT_MESH,			true,	{ 0, 0, 0, 0, 0, 0, 0, 0, 0, 0, 0, 0, 0, 0, 0, 0, 0, 0, 0, 0, 0, 0, 0, 0, 0, 0, 0, 0, 0 }		},	// ME_CLIPDEFORM
 		{ DT_MESH,			true,	{ 0, 0, 0, 0, 0, 0, 0, 0, 0, 0, 0, 0, 0, 0, 0, 0, 0, 0, 0, 0, 0, 0, 0, 0, 0, 0, 0, 0, 0 }		},	// ME_MORPHRESHAPE
 		{ DT_MESH,			true,	{ 0, 0, 0, 0, 0, 0, 0, 0, 0, 0, 0, 0, 0, 0, 0, 0, 0, 0, 0, 0, 0, 0, 0, 0, 0, 0, 0, 0, 0 }		},	// ME_OPTIMIZESKINNING
+		{ DT_MESH,			true,	{ 0, 0, 0, 0, 0, 0, 0, 0, 0, 0, 0, 0, 0, 0, 0, 0, 0, 0, 0, 0, 0, 0, 0, 0, 0, 0, 0, 0, 0 }		},	// ME_ADDTAGS
 
 		{ DT_INSTANCE,		false,	{ 0, 0, 0, 0, 0, 0, 0, 0, 0, 0, 0, 0, 0, 0, 0, 0, 0, 0, 0, 0, 0, 0, 0, 0, 0, 0, 0, 0, 0 }		},	// IN_ADDMESH
 		{ DT_INSTANCE,		false,	{ 0, 0, 0, 0, 0, 0, 0, 0, 0, 0, 0, 0, 0, 0, 0, 0, 0, 0, 0, 0, 0, 0, 0, 0, 0, 0, 0, 0, 0 }		},	// IN_ADDIMAGE
@@ -863,24 +864,33 @@ namespace mu
         {
             const uint8_t* data = program.GetOpArgsPointer(at);
             mu::OP::ADDRESS source;
-            memcpy( &source, data, sizeof(OP::ADDRESS) ); data+=sizeof(OP::ADDRESS);
+            FMemory::Memcpy( &source, data, sizeof(OP::ADDRESS) ); data+=sizeof(OP::ADDRESS);
             f(source);
 
             uint16 removes = 0;
-            memcpy( &removes, data, sizeof(uint16) ); data+=sizeof(uint16);
+			FMemory::Memcpy( &removes, data, sizeof(uint16) ); data+=sizeof(uint16);
             for (uint16 r=0; r<removes; ++r)
             {
                 mu::OP::ADDRESS condition;
-                memcpy( &condition, data, sizeof(OP::ADDRESS) ); data+=sizeof(OP::ADDRESS);
+				FMemory::Memcpy( &condition, data, sizeof(OP::ADDRESS) ); data+=sizeof(OP::ADDRESS);
                 f(condition);
 
                 mu::OP::ADDRESS mask;
-                memcpy( &mask, data, sizeof(OP::ADDRESS) ); data+=sizeof(OP::ADDRESS);
+				FMemory::Memcpy( &mask, data, sizeof(OP::ADDRESS) ); data+=sizeof(OP::ADDRESS);
                 f(mask);
             }
             break;
         }
 
+		case OP_TYPE::ME_ADDTAGS:
+		{
+			const uint8_t* data = program.GetOpArgsPointer(at);
+			mu::OP::ADDRESS source;
+			FMemory::Memcpy(&source, data, sizeof(OP::ADDRESS)); data += sizeof(OP::ADDRESS);
+			f(source);
+			break;
+		}
+
         case OP_TYPE::ME_FORMAT:
         {
 			OP::MeshFormatArgs args = program.GetOpArgs<OP::MeshFormatArgs>(at);
diff --git a/Engine/Plugins/Experimental/Mutable/Source/MutableRuntime/Private/MuR/Operations.h b/Engine/Plugins/Experimental/Mutable/Source/MutableRuntime/Private/MuR/Operations.h
--- a/Engine/Plugins/Experimental/Mutable/Source/MutableRuntime/Private/MuR/Operations.h
+++ b/Engine/Plugins/Experimental/Mutable/Source/MutableRuntime/Private/MuR/Operations.h
@@ -290,6 +290,9 @@ namespace mu
 		//! Optimize skinning before adding a mesh to the component
 		ME_OPTIMIZESKINNING,
 
+		//! Add a set of tags to a mesh
+		ME_ADDTAGS,
+
         //-----------------------------------------------------------------------------------------
         // Instance operations
         //-----------------------------------------------------------------------------------------
diff --git a/Engine/Plugins/Experimental/Mutable/Source/MutableRuntime/Public/MuR/Model.h b/Engine/Plugins/Experimental/Mutable/Source/MutableRuntime/Public/MuR/Model.h
--- a/Engine/Plugins/Experimental/Mutable/Source/MutableRuntime/Public/MuR/Model.h
+++ b/Engine/Plugins/Experimental/Mutable/Source/MutableRuntime/Public/MuR/Model.h
@@ -13,7 +13,7 @@
 //! data structures. Compiled models are not necessarily compatible when the runtime is updated,
 //! so this version number can be used externally to verify this. It is not used internally, and
 //! serializing models from different versions than this runtime will probably result in a crash.
-#define MUTABLE_COMPILED_MODEL_CODE_VERSION		uint32( 81 )
+#define MUTABLE_COMPILED_MODEL_CODE_VERSION		uint32( 82 )
 #define MUTABLE_PARAMETERS_VERSION              uint32( 2 )
 
 
diff --git a/Engine/Plugins/Experimental/Mutable/Source/MutableTools/Private/MuT/ASTOpImageRasterMesh.cpp b/Engine/Plugins/Experimental/Mutable/Source/MutableTools/Private/MuT/ASTOpImageRasterMesh.cpp
--- a/Engine/Plugins/Experimental/Mutable/Source/MutableTools/Private/MuT/ASTOpImageRasterMesh.cpp
+++ b/Engine/Plugins/Experimental/Mutable/Source/MutableTools/Private/MuT/ASTOpImageRasterMesh.cpp
@@ -7,12 +7,13 @@
 #include "MuT/ASTOpMeshMorph.h"
 #include "MuT/ASTOpConditional.h"
 #include "MuT/ASTOpSwitch.h"
-#include "Containers/Map.h"
-#include "HAL/PlatformMath.h"
+#include "MuT/ASTOpMeshAddTags.h"
+#include "MuT/StreamsPrivate.h"
 #include "MuR/ModelPrivate.h"
 #include "MuR/RefCounted.h"
 #include "MuR/Types.h"
-#include "MuT/StreamsPrivate.h"
+#include "Containers/Map.h"
+#include "HAL/PlatformMath.h"
 
 
 namespace mu
@@ -261,6 +262,7 @@ namespace mu
 
 		case OP_TYPE::ME_INTERPOLATE:
 		{
+			// TODO: should be sink only if no imageAt?
 			auto typedSource = dynamic_cast<const ASTOpFixed*>(sourceAt.get());
 			Ptr<ASTOpImageRasterMesh> rasterOp = mu::Clone<ASTOpImageRasterMesh>(this);
 			rasterOp->mesh = typedSource->children[typedSource->op.args.MeshInterpolate.base].child();
@@ -270,6 +272,7 @@ namespace mu
 
 		case OP_TYPE::ME_MORPH:
 		{
+			// TODO: should be sink only if no imageAt?
 			const ASTOpMeshMorph* typedSource = dynamic_cast<const ASTOpMeshMorph*>(sourceAt.get());
 			Ptr<ASTOpImageRasterMesh> rasterOp = mu::Clone<ASTOpImageRasterMesh>(this);
 			rasterOp->mesh = typedSource->Base.child();
@@ -277,6 +280,16 @@ namespace mu
 			break;
 		}
 
+		case OP_TYPE::ME_ADDTAGS:
+		{
+			// Ignore tags
+			const ASTOpMeshAddTags* typedSource = dynamic_cast<const ASTOpMeshAddTags*>(sourceAt.get());
+			Ptr<ASTOpImageRasterMesh> rasterOp = mu::Clone<ASTOpImageRasterMesh>(this);
+			rasterOp->mesh = typedSource->Source.child();
+			at = rasterOp;
+			break;
+		}
+
 		case OP_TYPE::ME_CONDITIONAL:
 		{
 			auto nop = mu::Clone<ASTOpConditional>(sourceAt.get());
diff --git a/Engine/Plugins/Experimental/Mutable/Source/MutableTools/Private/MuT/ASTOpLayoutFromMesh.cpp b/Engine/Plugins/Experimental/Mutable/Source/MutableTools/Private/MuT/ASTOpLayoutFromMesh.cpp
--- a/Engine/Plugins/Experimental/Mutable/Source/MutableTools/Private/MuT/ASTOpLayoutFromMesh.cpp
+++ b/Engine/Plugins/Experimental/Mutable/Source/MutableTools/Private/MuT/ASTOpLayoutFromMesh.cpp
@@ -19,6 +19,7 @@
 #include "MuT/ASTOpMeshFormat.h"
 #include "MuT/ASTOpMeshApplyShape.h"
 #include "MuT/ASTOpMeshBindShape.h"
+#include "MuT/ASTOpMeshAddTags.h"
 #include "MuT/ASTOpSwitch.h"
 #include "MuT/ASTOpLayoutMerge.h"
 #include "MuT/StreamsPrivate.h"
@@ -191,6 +192,15 @@ namespace mu
 					break;
 				}
 
+				case OP_TYPE::ME_ADDTAGS:
+				{
+					// Sink, ignoring the op
+					const ASTOpMeshAddTags* Typed = dynamic_cast<const ASTOpMeshAddTags*>(at.get());
+					check(Typed);
+					newAt = Visit(Typed->Source.child(), CurrentSinkingOp);
+					break;
+				}
+
 				case OP_TYPE::ME_INTERPOLATE:
 				{
 					// Sink, ignoring the op
diff --git a/Engine/Plugins/Experimental/Mutable/Source/MutableTools/Private/MuT/ASTOpMeshAddTags.cpp b/Engine/Plugins/Experimental/Mutable/Source/MutableTools/Private/MuT/ASTOpMeshAddTags.cpp
new file mode 100644
--- /dev/null
+++ b/Engine/Plugins/Experimental/Mutable/Source/MutableTools/Private/MuT/ASTOpMeshAddTags.cpp
@@ -0,0 +1,81 @@
+// Copyright Epic Games, Inc. All Rights Reserved.
+
+#include "MuT/ASTOpMeshAddTags.h"
+
+#include "MuT/ASTOpMeshMorph.h"
+#include "HAL/PlatformMath.h"
+#include "MuR/ModelPrivate.h"
+#include "MuR/RefCounted.h"
+#include "MuR/Types.h"
+#include "MuT/StreamsPrivate.h"
+
+
+namespace mu
+{
+
+	ASTOpMeshAddTags::ASTOpMeshAddTags()
+		: Source(this)
+	{
+	}
+
+
+	ASTOpMeshAddTags::~ASTOpMeshAddTags()
+	{
+		// Explicit call needed to avoid recursive destruction
+		ASTOp::RemoveChildren();
+	}
+
+
+	bool ASTOpMeshAddTags::IsEqual(const ASTOp& otherUntyped) const
+	{
+		if (const ASTOpMeshAddTags* other = dynamic_cast<const ASTOpMeshAddTags*>(&otherUntyped))
+		{
+			return Source == other->Source && Tags == other->Tags;
+		}
+		return false;
+	}
+
+
+	mu::Ptr<ASTOp> ASTOpMeshAddTags::Clone(MapChildFuncRef mapChild) const
+	{
+		Ptr<ASTOpMeshAddTags> n = new ASTOpMeshAddTags();
+		n->Source = mapChild(Source.child());
+		n->Tags = Tags;
+		return n;
+	}
+
+
+	void ASTOpMeshAddTags::ForEachChild(const TFunctionRef<void(ASTChild&)> f)
+	{
+		f(Source);
+	}
+
+
+	uint64 ASTOpMeshAddTags::Hash() const
+	{
+		uint64 res = std::hash<ASTOp*>()(Source.child().get());
+		return res;
+	}
+
+
+	void ASTOpMeshAddTags::Link(FProgram& Program, FLinkerOptions*)
+	{
+		// Already linked?
+		if (!linkedAddress)
+		{
+			linkedAddress = (OP::ADDRESS)Program.m_opAddress.Num();
+
+			Program.m_opAddress.Add((uint32)Program.m_byteCode.Num());
+			AppendCode(Program.m_byteCode, OP_TYPE::ME_ADDTAGS);
+			OP::ADDRESS SourceAt = Source ? Source->linkedAddress : 0;
+			AppendCode(Program.m_byteCode, SourceAt);
+			AppendCode(Program.m_byteCode, (uint16)Tags.Num());
+			for (const FString& Tag : Tags)
+			{
+				OP::ADDRESS TagConstantAddress = Program.AddConstant(Tag);
+				AppendCode(Program.m_byteCode, TagConstantAddress);
+			}
+		}
+	}
+
+}
diff --git a/Engine/Plugins/Experimental/Mutable/Source/MutableTools/Private/MuT/ASTOpMeshAddTags.h b/Engine/Plugins/Experimental/Mutable/Source/MutableTools/Private/MuT/ASTOpMeshAddTags.h
new file mode 100644
--- /dev/null
+++ b/Engine/Plugins/Experimental/Mutable/Source/MutableTools/Private/MuT/ASTOpMeshAddTags.h
@@ -0,0 +1,42 @@
+// Copyright Epic Games, Inc. All Rights Reserved.
+
+#pragma once
+
+#include "HAL/Platform.h"
+#include "MuR/Operations.h"
+#include "MuR/Ptr.h"
+#include "MuT/AST.h"
+
+
+namespace mu
+{
+	struct FProgram;
+
+	/** Add tags to a mesh. */
+	class ASTOpMeshAddTags : public ASTOp
+	{
+	public:
+
+		/** Source mesh to add tags to. */
+		ASTChild Source;
+
+		/** Tags to add. */
+		TArray<FString> Tags;
+
+	public:
+
+		ASTOpMeshAddTags();
+		ASTOpMeshAddTags(const ASTOpMeshAddTags&) = delete;
+		virtual ~ASTOpMeshAddTags() override;
+
+		// ASTOp interface
+		virtual OP_TYPE GetOpType() const override { return OP_TYPE::ME_ADDTAGS; }
+		virtual uint64 Hash() const override;
+		virtual void ForEachChild(const TFunctionRef<void(ASTChild&)>) override;
+		virtual bool IsEqual(const ASTOp& otherUntyped) const override;
+		virtual Ptr<ASTOp> Clone(MapChildFuncRef mapChild) const override;
+		virtual void Link(FProgram& program, FLinkerOptions*) override;
+	};
+
+}
+
diff --git a/Engine/Plugins/Experimental/Mutable/Source/MutableTools/Private/MuT/ASTOpMeshBindShape.cpp b/Engine/Plugins/Experimental/Mutable/Source/MutableTools/Private/MuT/ASTOpMeshBindShape.cpp
--- a/Engine/Plugins/Experimental/Mutable/Source/MutableTools/Private/MuT/ASTOpMeshBindShape.cpp
+++ b/Engine/Plugins/Experimental/Mutable/Source/MutableTools/Private/MuT/ASTOpMeshBindShape.cpp
@@ -10,6 +10,7 @@
 #include "MuT/ASTOpConditional.h"
 #include "MuT/ASTOpSwitch.h"
 #include "MuT/ASTOpMeshRemoveMask.h"
+#include "MuT/ASTOpMeshAddTags.h"
 #include "MuT/StreamsPrivate.h"
 
 #include <memory>
@@ -370,6 +371,20 @@ namespace mu
 				break;
 			}
 
+			case OP_TYPE::ME_ADDTAGS:
+			{
+				Ptr<ASTOpMeshAddTags> New = mu::Clone<ASTOpMeshAddTags>(MeshAt);
+				if (New->Source)
+				{
+					Ptr<ASTOpMeshBindShape> NewBind = mu::Clone<ASTOpMeshBindShape>(this);
+					NewBind->Mesh = New->Source.child();
+					New->Source = NewBind;
+				}
+
+				NewOp = New;
+				break;
+			}
+
 			default:
 				break;
 
@@ -432,6 +447,16 @@ namespace mu
 				break;
 			}
 
+			case OP_TYPE::ME_ADDTAGS:
+			{
+				// Ignore the tags in the shape
+				Ptr<ASTOpMeshBindShape> NewBind = mu::Clone<ASTOpMeshBindShape>(this);
+				const ASTOpMeshAddTags* New = dynamic_cast<const ASTOpMeshAddTags*>(ShapeAt.get());
+				NewBind->Shape = New->Source.child();
+				NewOp = NewBind;
+				break;
+			}
+
 			default:
 				break;
 
diff --git a/Engine/Plugins/Experimental/Mutable/Source/MutableTools/Private/MuT/ASTOpMeshClipMorphPlane.cpp b/Engine/Plugins/Experimental/Mutable/Source/MutableTools/Private/MuT/ASTOpMeshClipMorphPlane.cpp
--- a/Engine/Plugins/Experimental/Mutable/Source/MutableTools/Private/MuT/ASTOpMeshClipMorphPlane.cpp
+++ b/Engine/Plugins/Experimental/Mutable/Source/MutableTools/Private/MuT/ASTOpMeshClipMorphPlane.cpp
@@ -2,14 +2,12 @@
 
 #include "MuT/ASTOpMeshClipMorphPlane.h"
 
-#include "HAL/PlatformMath.h"
+#include "MuT/ASTOpMeshAddTags.h"
+#include "MuT/StreamsPrivate.h"
 #include "MuR/ModelPrivate.h"
 #include "MuR/RefCounted.h"
 #include "MuR/Types.h"
-#include "MuT/StreamsPrivate.h"
-
-#include <memory>
-#include <utility>
+#include "HAL/PlatformMath.h"
 
 
 namespace mu
@@ -113,4 +111,42 @@ namespace mu
 
 	}
 
+	mu::Ptr<ASTOp> ASTOpMeshClipMorphPlane::OptimiseSink(const FModelOptimizationOptions&, FOptimizeSinkContext&) const
+	{
+		Ptr<ASTOp> NewOp;
+
+		if (!source.child())
+		{
+			return nullptr;
+		}
+
+		OP_TYPE SourceType = source.child()->GetOpType();
+
+		// Optimize only the mesh parameter
+		switch (SourceType)
+		{
+
+		case OP_TYPE::ME_ADDTAGS:
+		{
+			Ptr<ASTOpMeshAddTags> NewAddTags = mu::Clone<ASTOpMeshAddTags>(source.child());
+
+			if (NewAddTags->Source)
+			{
+				Ptr<ASTOpMeshClipMorphPlane> New = mu::Clone<ASTOpMeshClipMorphPlane>(this);
+				New->source = NewAddTags->Source.child();
+				NewAddTags->Source = New;
+			}
+
+			NewOp = NewAddTags;
+			break;
+		}
+
+		default:
+			break;
+
+		}
+
+		return NewOp;
+	}
+
 }
diff --git a/Engine/Plugins/Experimental/Mutable/Source/MutableTools/Private/MuT/ASTOpMeshClipMorphPlane.h b/Engine/Plugins/Experimental/Mutable/Source/MutableTools/Private/MuT/ASTOpMeshClipMorphPlane.h
--- a/Engine/Plugins/Experimental/Mutable/Source/MutableTools/Private/MuT/ASTOpMeshClipMorphPlane.h
+++ b/Engine/Plugins/Experimental/Mutable/Source/MutableTools/Private/MuT/ASTOpMeshClipMorphPlane.h
@@ -34,14 +34,15 @@ struct FProgram;
 
 		ASTOpMeshClipMorphPlane();
 		ASTOpMeshClipMorphPlane(const ASTOpMeshClipMorphPlane&) = delete;
-		~ASTOpMeshClipMorphPlane();
-
-		OP_TYPE GetOpType() const override { return OP_TYPE::ME_CLIPMORPHPLANE; }
-		uint64 Hash() const override;
-		bool IsEqual(const ASTOp& otherUntyped) const override;
-		Ptr<ASTOp> Clone(MapChildFuncRef mapChild) const override;
-		void ForEachChild(const TFunctionRef<void(ASTChild&)>) override;
-		void Link(FProgram& program, FLinkerOptions* Options) override;
+		virtual ~ASTOpMeshClipMorphPlane();
+
+		virtual OP_TYPE GetOpType() const override { return OP_TYPE::ME_CLIPMORPHPLANE; }
+		virtual uint64 Hash() const override;
+		virtual bool IsEqual(const ASTOp& otherUntyped) const override;
+		virtual Ptr<ASTOp> Clone(MapChildFuncRef mapChild) const override;
+		virtual void ForEachChild(const TFunctionRef<void(ASTChild&)>) override;
+		virtual void Link(FProgram& program, FLinkerOptions* Options) override;
+		virtual mu::Ptr<ASTOp> OptimiseSink(const FModelOptimizationOptions&, FOptimizeSinkContext&) const override;
 	};
 
 
diff --git a/Engine/Plugins/Experimental/Mutable/Source/MutableTools/Private/MuT/ASTOpMeshDifference.cpp b/Engine/Plugins/Experimental/Mutable/Source/MutableTools/Private/MuT/ASTOpMeshDifference.cpp
--- a/Engine/Plugins/Experimental/Mutable/Source/MutableTools/Private/MuT/ASTOpMeshDifference.cpp
+++ b/Engine/Plugins/Experimental/Mutable/Source/MutableTools/Private/MuT/ASTOpMeshDifference.cpp
@@ -10,6 +10,7 @@
 #include "MuT/ASTOpConditional.h"
 #include "MuT/ASTOpSwitch.h"
 #include "MuT/ASTOpMeshRemoveMask.h"
+#include "MuT/ASTOpMeshAddTags.h"
 
 #include <memory>
 
@@ -263,6 +264,21 @@ namespace mu
 				break;
 			}
 
+			case OP_TYPE::ME_ADDTAGS:
+			{
+				Ptr<ASTOpMeshAddTags> NewAdd = mu::Clone<ASTOpMeshAddTags>(BaseAt);
+
+				if (NewAdd->Source)
+				{
+					Ptr<ASTOpMeshDifference> NewDiff = mu::Clone<ASTOpMeshDifference>(this);
+					NewDiff->Base = NewAdd->Source.child();
+					NewAdd->Source = NewDiff;
+				}
+
+				NewOp = NewAdd;
+				break;
+			}
+
 			default:
 				break;
 
@@ -325,6 +341,17 @@ namespace mu
 				break;
 			}
 
+			case OP_TYPE::ME_ADDTAGS:
+			{
+				// Ignore tags in this branch
+				const ASTOpMeshAddTags* Add = dynamic_cast<const ASTOpMeshAddTags*>(TargetAt.get());
+
+				Ptr<ASTOpMeshDifference> NewDiff = mu::Clone<ASTOpMeshDifference>(this);
+				NewDiff->Target = Add->Source.child();
+				NewOp = NewDiff;
+				break;
+			}
+
 			default:
 				break;
 
diff --git a/Engine/Plugins/Experimental/Mutable/Source/MutableTools/Private/MuT/ASTOpMeshExtractLayoutBlocks.cpp b/Engine/Plugins/Experimental/Mutable/Source/MutableTools/Private/MuT/ASTOpMeshExtractLayoutBlocks.cpp
--- a/Engine/Plugins/Experimental/Mutable/Source/MutableTools/Private/MuT/ASTOpMeshExtractLayoutBlocks.cpp
+++ b/Engine/Plugins/Experimental/Mutable/Source/MutableTools/Private/MuT/ASTOpMeshExtractLayoutBlocks.cpp
@@ -4,13 +4,14 @@
 
 #include "MuT/ASTOpSwitch.h"
 #include "MuT/ASTOpConditional.h"
-#include "Misc/AssertionMacros.h"
+#include "MuT/ASTOpMeshAddTags.h"
+#include "MuT/ASTOpMeshMorph.h"
+#include "MuT/ASTOpMeshRemoveMask.h"
+#include "MuT/ASTOpMeshClipMorphPlane.h"
 #include "MuR/ModelPrivate.h"
 #include "MuR/RefCounted.h"
 #include "MuR/Types.h"
-
-#include <memory>
-#include <utility>
+#include "Misc/AssertionMacros.h"
 
 namespace mu
 {
@@ -157,6 +158,77 @@ namespace mu
 			break;
 		}
 
+		case OP_TYPE::ME_ADDTAGS:
+		{
+			Ptr<ASTOpMeshAddTags> NewAddTags = mu::Clone<ASTOpMeshAddTags>(Source.child());
+
+			if (NewAddTags->Source)
+			{
+				Ptr<ASTOpMeshExtractLayoutBlocks> New = mu::Clone<ASTOpMeshExtractLayoutBlocks>(this);
+				New->Source = NewAddTags->Source.child();
+				NewAddTags->Source = New;
+			}
+
+			NewOp = NewAddTags;
+			break;
+		}
+
+		case OP_TYPE::ME_MORPH:
+		{
+			// Move the operation down the base and the target
+			Ptr<ASTOpMeshMorph> NewMorph = mu::Clone<ASTOpMeshMorph>(Source.child());
+
+			if (NewMorph->Base)
+			{
+				Ptr<ASTOpMeshExtractLayoutBlocks> New = mu::Clone<ASTOpMeshExtractLayoutBlocks>(this);
+				New->Source = NewMorph->Base.child();
+				NewMorph->Base = New;
+			}
+
+			if (NewMorph->Target)
+			{
+				Ptr<ASTOpMeshExtractLayoutBlocks> New = mu::Clone<ASTOpMeshExtractLayoutBlocks>(this);
+				New->Source = NewMorph->Target.child();
+				NewMorph->Target = New;
+			}
+
+			NewOp = NewMorph;
+			break;
+		}
+
+		case OP_TYPE::ME_REMOVEMASK:
+		{
+			// Move the operation down the base
+			// \TODO: mask too to try to make them smaller?
+			Ptr<ASTOpMeshRemoveMask> NewRemove = mu::Clone<ASTOpMeshRemoveMask>(Source.child());
+
+			if (NewRemove->source)
+			{
+				Ptr<ASTOpMeshExtractLayoutBlocks> New = mu::Clone<ASTOpMeshExtractLayoutBlocks>(this);
+				New->Source = NewRemove->source.child();
+				NewRemove->source = New;
+			}
+
+			NewOp = NewRemove;
+			break;
+		}
+
+		case OP_TYPE::ME_CLIPMORPHPLANE:
+		{
+			// Move the operation down the source
+			Ptr<ASTOpMeshClipMorphPlane> NewMorph = mu::Clone<ASTOpMeshClipMorphPlane>(Source.child());
+
+			if (NewMorph->source)
+			{
+				Ptr<ASTOpMeshExtractLayoutBlocks> New = mu::Clone<ASTOpMeshExtractLayoutBlocks>(this);
+				New->Source = NewMorph->source.child();
+				NewMorph->source = New;
+			}
+
+			NewOp = NewMorph;
+			break;
+		}
+
 		default:
 			break;
 
diff --git a/Engine/Plugins/Experimental/Mutable/Source/MutableTools/Private/MuT/ASTOpMeshFormat.cpp b/Engine/Plugins/Experimental/Mutable/Source/MutableTools/Private/MuT/ASTOpMeshFormat.cpp
--- a/Engine/Plugins/Experimental/Mutable/Source/MutableTools/Private/MuT/ASTOpMeshFormat.cpp
+++ b/Engine/Plugins/Experimental/Mutable/Source/MutableTools/Private/MuT/ASTOpMeshFormat.cpp
@@ -16,6 +16,7 @@
 #include "MuT/ASTOpMeshClipMorphPlane.h"
 #include "MuT/ASTOpMeshRemoveMask.h"
 #include "MuT/ASTOpMeshMorph.h"
+#include "MuT/ASTOpMeshAddTags.h"
 #include "MuT/ASTOpSwitch.h"
 #include "MuT/StreamsPrivate.h"
 
@@ -256,6 +257,14 @@ mu::Ptr<ASTOp> Sink_MeshFormatAST::Visit(const mu::Ptr<ASTOp>& at, const ASTOpMe
 		break;
 	}
 
+	case OP_TYPE::ME_ADDTAGS:
+	{
+		Ptr<ASTOpMeshAddTags> newOp = mu::Clone<ASTOpMeshAddTags>(at);
+		newOp->Source = Visit(newOp->Source.child(), currentFormatOp);
+		newAt = newOp;
+		break;
+	}
+
 	case OP_TYPE::ME_CLIPMORPHPLANE:
 	{
 		auto newOp = mu::Clone<ASTOpMeshClipMorphPlane>(at);
diff --git a/Engine/Plugins/Experimental/Mutable/Source/MutableTools/Private/MuT/ASTOpMeshMaskClipMesh.cpp b/Engine/Plugins/Experimental/Mutable/Source/MutableTools/Private/MuT/ASTOpMeshMaskClipMesh.cpp
--- a/Engine/Plugins/Experimental/Mutable/Source/MutableTools/Private/MuT/ASTOpMeshMaskClipMesh.cpp
+++ b/Engine/Plugins/Experimental/Mutable/Source/MutableTools/Private/MuT/ASTOpMeshMaskClipMesh.cpp
@@ -8,6 +8,7 @@
 #include "MuR/Types.h"
 #include "MuT/ASTOpConditional.h"
 #include "MuT/ASTOpMeshRemoveMask.h"
+#include "MuT/ASTOpMeshAddTags.h"
 #include "MuT/ASTOpSwitch.h"
 #include "MuT/StreamsPrivate.h"
 
@@ -152,7 +153,7 @@ namespace mu
 				// This cannot be sunk since the result is different. Since the clipping is now correctly
 				// generated at the end of the chain when really necessary, this wrong optimisation is no 
 				// longer needed.
-				//case OP_TYPE::ME_MORPH2:
+				//case OP_TYPE::ME_MORPH:
 		        //{
 		        //    break;
 		        //}
@@ -169,6 +170,14 @@ namespace mu
 					break;
 				}
 
+				case OP_TYPE::ME_ADDTAGS:
+				{
+					Ptr<ASTOpMeshAddTags> newOp = mu::Clone<ASTOpMeshAddTags>(at);
+					newOp->Source = Visit(newOp->Source.child());
+					newAt = newOp;
+					break;
+				}
+
 				case OP_TYPE::ME_CONDITIONAL:
 				{
 					// We move the mask creation down the two paths
@@ -304,6 +313,13 @@ namespace mu
 					break;
 				}
 
+				case OP_TYPE::ME_ADDTAGS:
+				{
+					// Ignore tags in the clip mesh branch
+					const ASTOpMeshAddTags* AddOp = dynamic_cast<const ASTOpMeshAddTags*>(at.get());
+					newAt = Visit(AddOp->Source.child());
+					break;
+				}
 
 				default:
 				{
diff --git a/Engine/Plugins/Experimental/Mutable/Source/MutableTools/Private/MuT/ASTOpMeshMaskClipUVMask.cpp b/Engine/Plugins/Experimental/Mutable/Source/MutableTools/Private/MuT/ASTOpMeshMaskClipUVMask.cpp
--- a/Engine/Plugins/Experimental/Mutable/Source/MutableTools/Private/MuT/ASTOpMeshMaskClipUVMask.cpp
+++ b/Engine/Plugins/Experimental/Mutable/Source/MutableTools/Private/MuT/ASTOpMeshMaskClipUVMask.cpp
@@ -8,6 +8,7 @@
 #include "MuR/Types.h"
 #include "MuT/ASTOpConditional.h"
 #include "MuT/ASTOpMeshRemoveMask.h"
+#include "MuT/ASTOpMeshAddTags.h"
 #include "MuT/ASTOpSwitch.h"
 #include "MuT/StreamsPrivate.h"
 
@@ -153,7 +154,7 @@ namespace mu
 				// This cannot be sunk since the result is different. Since the clipping is now correctly
 				// generated at the end of the chain when really necessary, this wrong optimisation is no 
 				// longer needed.
-				//case OP_TYPE::ME_MORPH2:
+				//case OP_TYPE::ME_MORPH:
 		        //{
 		        //    break;
 		        //}
@@ -170,6 +171,14 @@ namespace mu
 					break;
 				}
 
+				case OP_TYPE::ME_ADDTAGS:
+				{
+					Ptr<ASTOpMeshAddTags> newOp = mu::Clone<ASTOpMeshAddTags>(at);
+					newOp->Source = Visit(newOp->Source.child());
+					newAt = newOp;
+					break;
+				}
+
 				case OP_TYPE::ME_CONDITIONAL:
 				{
 					// We move the mask creation down the two paths
diff --git a/Engine/Plugins/Experimental/Mutable/Source/MutableTools/Private/MuT/ASTOpMeshMorph.cpp b/Engine/Plugins/Experimental/Mutable/Source/MutableTools/Private/MuT/ASTOpMeshMorph.cpp
--- a/Engine/Plugins/Experimental/Mutable/Source/MutableTools/Private/MuT/ASTOpMeshMorph.cpp
+++ b/Engine/Plugins/Experimental/Mutable/Source/MutableTools/Private/MuT/ASTOpMeshMorph.cpp
@@ -2,13 +2,13 @@
 
 #include "MuT/ASTOpMeshMorph.h"
 
-#include "HAL/PlatformMath.h"
+#include "MuT/StreamsPrivate.h"
+#include "MuT/ASTOpMeshAddTags.h"
 #include "MuR/ModelPrivate.h"
 #include "MuR/RefCounted.h"
 #include "MuR/Types.h"
-#include "MuT/StreamsPrivate.h"
+#include "HAL/PlatformMath.h"
 
-#include <memory>
 
 namespace mu
 {
@@ -95,7 +95,66 @@ namespace mu
 	//-------------------------------------------------------------------------------------------------
 	mu::Ptr<ASTOp> ASTOpMeshMorph::OptimiseSink(const FModelOptimizationOptions&, FOptimizeSinkContext&) const
 	{
-		return nullptr;
+		Ptr<ASTOp> NewOp;
+
+		if (!Base.child())
+		{
+			return nullptr;
+		}
+
+		// Base optimizations
+		OP_TYPE BaseType = Base.child()->GetOpType();
+		switch (BaseType)
+		{
+
+		case OP_TYPE::ME_ADDTAGS:
+		{
+			// Add the base tags after morphing
+			Ptr<ASTOpMeshAddTags> NewAddTags = mu::Clone<ASTOpMeshAddTags>(Base.child());
+
+			if (NewAddTags->Source)
+			{
+				Ptr<ASTOpMeshMorph> New = mu::Clone<ASTOpMeshMorph>(this);
+				New->Base = NewAddTags->Source.child();
+				NewAddTags->Source = New;
+			}
+
+			NewOp = NewAddTags;
+			break;
+		}
+
+		default:
+			break;
+
+		}
+
+		// If not optimized yet
+		if (!NewOp)
+		{
+			// Target optimizations
+			OP_TYPE MorphType = Target.child()->GetOpType();
+			switch (MorphType)
+			{
+
+			case OP_TYPE::ME_ADDTAGS:
+			{
+				// Ignore the morph target tags
+				const ASTOpMeshAddTags* AddTags = dynamic_cast<const ASTOpMeshAddTags*>(Target.child().get());
+
+				Ptr<ASTOpMeshMorph> New = mu::Clone<ASTOpMeshMorph>(this);
+				New->Target = AddTags->Source.child();
+				NewOp = New;
+
+				break;
+			}
+
+			default:
+				break;
+
+			}
+		}
+
+		return NewOp;
 	}
 
 }
diff --git a/Engine/Plugins/Experimental/Mutable/Source/MutableTools/Private/MuT/ASTOpMeshRemoveMask.cpp b/Engine/Plugins/Experimental/Mutable/Source/MutableTools/Private/MuT/ASTOpMeshRemoveMask.cpp
--- a/Engine/Plugins/Experimental/Mutable/Source/MutableTools/Private/MuT/ASTOpMeshRemoveMask.cpp
+++ b/Engine/Plugins/Experimental/Mutable/Source/MutableTools/Private/MuT/ASTOpMeshRemoveMask.cpp
@@ -3,14 +3,13 @@
 #include "MuT/ASTOpMeshRemoveMask.h"
 
 #include "MuT/ASTOpMeshMorph.h"
+#include "MuT/ASTOpMeshAddTags.h"
 #include "HAL/PlatformMath.h"
 #include "MuR/ModelPrivate.h"
 #include "MuR/RefCounted.h"
 #include "MuR/Types.h"
 #include "MuT/StreamsPrivate.h"
 
-#include <memory>
-
 
 namespace mu
 {
@@ -176,6 +175,14 @@ namespace mu
 					break;
 				}
 
+				case OP_TYPE::ME_ADDTAGS:
+				{
+					mu::Ptr<ASTOpMeshAddTags> newOp = mu::Clone<ASTOpMeshAddTags>(at);
+					newOp->Source = Visit(newOp->Source.child());
+					newAt = newOp;
+					break;
+				}
+
 				// disabled to avoid code explosion (or bug?) TODO
 //            case OP_TYPE::ME_CONDITIONAL:
 //            {
@@ -226,6 +233,43 @@ namespace mu
 		Sink_MeshRemoveMaskAST sinker;
 		mu::Ptr<ASTOp> at = sinker.Apply(this);
 
+		// If not optimized already, see if we can optimize the "remove" branches
+		if (!at)
+		{
+			Ptr<ASTOpMeshRemoveMask> NewOp;
+
+			for (int32 RemoveIndex=0; RemoveIndex<removes.Num(); ++RemoveIndex)
+			{
+				if (!removes[RemoveIndex].Value)
+				{
+					continue;
+				}
+
+				OP_TYPE RemoveType = removes[RemoveIndex].Value->GetOpType();
+				switch (RemoveType)
+				{
+				case OP_TYPE::ME_ADDTAGS:
+				{
+					// It can be ignored.
+					if (!NewOp)
+					{
+						NewOp = mu::Clone<ASTOpMeshRemoveMask>(this);
+					}
+
+					const ASTOpMeshAddTags* Add = dynamic_cast<const ASTOpMeshAddTags*>(removes[RemoveIndex].Value.child().get());
+					NewOp->removes[RemoveIndex].Value = Add->Source.child();
+
+					break;
+				}
+
+				default:
+					break;
+				}
+			}
+
+			at = NewOp;
+		}
+
 		return at;
 	}
 
diff --git a/Engine/Plugins/Experimental/Mutable/Source/MutableTools/Private/MuT/CodeGenerator_Mesh.cpp b/Engine/Plugins/Experimental/Mutable/Source/MutableTools/Private/MuT/CodeGenerator_Mesh.cpp
--- a/Engine/Plugins/Experimental/Mutable/Source/MutableTools/Private/MuT/CodeGenerator_Mesh.cpp
+++ b/Engine/Plugins/Experimental/Mutable/Source/MutableTools/Private/MuT/CodeGenerator_Mesh.cpp
@@ -29,6 +29,7 @@
 #include "MuT/ASTOpMeshTransform.h"
 #include "MuT/ASTOpMeshDifference.h"
 #include "MuT/ASTOpMeshMorph.h"
+#include "MuT/ASTOpMeshAddTags.h"
 #include "MuT/ASTOpSwitch.h"
 #include "MuT/CodeGenerator.h"
 #include "MuT/CodeGenerator_FirstPass.h"
@@ -77,9 +78,6 @@
 
 #include "Spatial/PointHashGrid3.h"
 
-#include <memory>
-#include <utility>
-
 
 namespace mu
 {
@@ -1278,6 +1276,15 @@ namespace mu
 			return;
 		}
 
+		// Separate the tags from the mesh
+		TArray<FString> Tags = pMesh->m_tags;
+		if (Tags.Num())
+		{
+			Ptr<Mesh> TaglessMesh = CloneOrTakeOver(pMesh.get());
+			TaglessMesh->m_tags.SetNum(0, false);
+			pMesh = TaglessMesh;
+		}
+
 		// Find out if we can (or have to) reuse a mesh that we have already generated.
 		MeshPtrConst DuplicateOf;
 		for (int32 i = 0; i < m_constantMeshes.Num(); ++i)
@@ -1331,6 +1338,7 @@ namespace mu
 			}
 		}
 
+		Ptr<const Mesh> FinalMesh;
 		if (DuplicateOf)
 		{
 			// Make sure the source layouts of the mesh are mapped to the layouts of the duplicated mesh.
@@ -1354,7 +1362,7 @@ namespace mu
 				}
 			}
 
-			op->SetValue(DuplicateOf, m_compilerOptions->OptimisationOptions.bUseDiskCache);
+			FinalMesh = DuplicateOf;
 		}
 		else
 		{
@@ -1425,27 +1433,39 @@ namespace mu
 					uint32* pIdData = (uint32*)pCloned->GetVertexBuffers().GetBufferData(newBuffer);
 					for (int i = 0; i < pMesh->GetVertexCount(); ++i)
 					{
-						check(m_freeVertexIndex < std::numeric_limits<uint32>::max());
+						check(m_freeVertexIndex < TNumericLimits<uint32>::Max());
 
 						(*pIdData++) = m_freeVertexIndex++;
-						check(m_freeVertexIndex < std::numeric_limits<uint32>::max());
+						check(m_freeVertexIndex < TNumericLimits<uint32>::Max());
 					}
 				}
 			}
 
 			// Add the constant data
 			m_constantMeshes.Add(pCloned);
-			op->SetValue(pCloned.get(), m_compilerOptions->OptimisationOptions.bUseDiskCache);
+			FinalMesh = pCloned;
+		}
+
+		op->SetValue(FinalMesh, m_compilerOptions->OptimisationOptions.bUseDiskCache);
+
+		Ptr<ASTOp> LastMeshOp = op;
+
+		// Add the tags operation
+		if (Tags.Num())
+		{
+			Ptr<ASTOpMeshAddTags> AddTagsOp = new ASTOpMeshAddTags;
+			AddTagsOp->Source = LastMeshOp;
+			AddTagsOp->Tags = Tags;
+			LastMeshOp = AddTagsOp;
 		}
 
- 
 		// Apply the modifier for the pre-normal operations stage.
-		FBottomUpState temp = m_currentBottomUpState;
+		FBottomUpState TempState = m_currentBottomUpState;
 
 		bool bModifiersForBeforeOperations = true;
-		OutResult.meshOp = ApplyMeshModifiers(op, InOptions.ActiveTags, bModifiersForBeforeOperations, node.m_errorContext);
+		OutResult.meshOp = ApplyMeshModifiers(LastMeshOp, InOptions.ActiveTags, bModifiersForBeforeOperations, node.m_errorContext);
 
-		m_currentBottomUpState = temp;
+		m_currentBottomUpState = TempState;
 
     }
 
diff --git a/Engine/Plugins/Experimental/Mutable/Source/MutableTools/Private/MuT/CodeOptimiser.cpp b/Engine/Plugins/Experimental/Mutable/Source/MutableTools/Private/MuT/CodeOptimiser.cpp
--- a/Engine/Plugins/Experimental/Mutable/Source/MutableTools/Private/MuT/CodeOptimiser.cpp
+++ b/Engine/Plugins/Experimental/Mutable/Source/MutableTools/Private/MuT/CodeOptimiser.cpp
@@ -674,6 +674,7 @@ namespace mu
 				case OP_TYPE::ME_CLIPMORPHPLANE:
 				case OP_TYPE::ME_APPLYPOSE:
 				case OP_TYPE::ME_REMOVEMASK:
+				case OP_TYPE::ME_ADDTAGS:
 				case OP_TYPE::IM_PLAINCOLOUR:
 					hasSpecialOpInSubtree = true;
 					break;
diff --git a/Engine/Plugins/Experimental/Mutable/Source/MutableTools/Private/MuT/CodeOptimiser_Semantic.cpp b/Engine/Plugins/Experimental/Mutable/Source/MutableTools/Private/MuT/CodeOptimiser_Semantic.cpp
--- a/Engine/Plugins/Experimental/Mutable/Source/MutableTools/Private/MuT/CodeOptimiser_Semantic.cpp
+++ b/Engine/Plugins/Experimental/Mutable/Source/MutableTools/Private/MuT/CodeOptimiser_Semantic.cpp
@@ -25,13 +25,11 @@
 #include "MuT/ASTOpImageTransform.h"
 #include "MuT/ASTOpImageRasterMesh.h"
 #include "MuT/ASTOpMeshMorph.h"
+#include "MuT/ASTOpMeshAddTags.h"
 #include "MuT/ASTOpSwitch.h"
 #include "MuT/CodeOptimiser.h"
 #include "MuT/Table.h"
 
-#include <memory>
-#include <utility>
-
 
 namespace mu
 {
@@ -393,6 +391,31 @@ namespace mu
 			break;
 		}
 
+		//-------------------------------------------------------------------------------------
+		case OP_TYPE::ME_MASKDIFF:
+		{
+			Ptr<ASTOp> Fragment = children[op.args.MeshMaskDiff.fragment].child();
+			OP_TYPE FragmentType = Fragment->GetOpType();
+			switch (FragmentType)
+			{
+			case OP_TYPE::ME_ADDTAGS:
+			{
+				// Tags in the fragment can be ignored.
+				const ASTOpMeshAddTags* Add = dynamic_cast<const ASTOpMeshAddTags*>(Fragment.get());
+
+				Ptr<ASTOpFixed> NewAt = mu::Clone<ASTOpFixed>(this);
+				NewAt->SetChild(NewAt->op.args.MeshMaskDiff.fragment, Add->Source);
+				at = NewAt;
+				break;
+			}
+
+			default:
+				break;
+			}
+
+			break;
+		}
+
         default:
             break;
         }
@@ -505,6 +528,17 @@ namespace mu
                 break;
             }
 
+			case OP_TYPE::ME_ADDTAGS:
+			{
+				// Ignore tags for projection
+				const ASTOpMeshAddTags* Add = dynamic_cast<const ASTOpMeshAddTags*>(sourceAt.get());
+
+				Ptr<ASTOpFixed> NewProject = mu::Clone<ASTOpFixed>(this);
+				NewProject->SetChild(NewProject->op.args.MeshProject.mesh, Add->Source);
+				at = NewProject;
+				break;
+			}
+
             default:
                 break;
             }
diff --git a/Engine/Plugins/Experimental/Mutable/Source/MutableTools/Private/MuT/ErrorLog.cpp b/Engine/Plugins/Experimental/Mutable/Source/MutableTools/Private/MuT/ErrorLog.cpp
--- a/Engine/Plugins/Experimental/Mutable/Source/MutableTools/Private/MuT/ErrorLog.cpp
+++ b/Engine/Plugins/Experimental/Mutable/Source/MutableTools/Private/MuT/ErrorLog.cpp
@@ -289,7 +289,7 @@ namespace mu
 
 		TEXT("ME_APPLYLAYOUT   "),
 		TEXT("ME_DIFFERENCE    "),
-		TEXT("ME_MORPH2        "),
+		TEXT("ME_MORPH         "),
 		TEXT("ME_MERGE         "),
 		TEXT("ME_INTERPOLATE   "),
 		TEXT("ME_MASKCLIPMESH  "),
@@ -310,6 +310,7 @@ namespace mu
 		TEXT("ME_CLIPDEFORM	   "),
 		TEXT("ME_MORPHRESHAPE  "),
 		TEXT("ME_OPTIMIZESKIN  "),
+		TEXT("ME_ADDTAGS       "),
 
 		TEXT("IN_ADDMESH       "),
 		TEXT("IN_ADDIMAGE      "),
