[Mutable] Add Skeleton to the participating object system.

- Previously changes to it where ignored and not forcing a recompilation.

#rnx

diff --git a/Engine/Plugins/Experimental/Mutable/Source/CustomizableObjectEditor/Private/MuCOE/GenerateMutableSource/GenerateMutableSourceMesh.cpp b/Engine/Plugins/Experimental/Mutable/Source/CustomizableObjectEditor/Private/MuCOE/GenerateMutableSource/GenerateMutableSourceMesh.cpp
--- a/Engine/Plugins/Experimental/Mutable/Source/CustomizableObjectEditor/Private/MuCOE/GenerateMutableSource/GenerateMutableSourceMesh.cpp
+++ b/Engine/Plugins/Experimental/Mutable/Source/CustomizableObjectEditor/Private/MuCOE/GenerateMutableSource/GenerateMutableSourceMesh.cpp
@@ -693,6 +693,8 @@ mu::MeshPtr ConvertSkeletalMeshToMutable(const USkeletalMesh* InSkeletalMesh, co
 			return nullptr;
 		}
 
+		GenerationContext.AddParticipatingObject(*InSkeleton);
+
 		FMutableComponentInfo& MutComponentInfo = GenerationContext.GetCurrentComponentInfo();
 		USkeletalMesh* ComponentRefSkeletalMesh = MutComponentInfo.RefSkeletalMesh;
 		USkeleton* ComponentRefSkeleton = MutComponentInfo.RefSkeleton;
