[Mutable] Avoid moving a UObject into the same outer when splitting the streamed resources into separate packages. An assert failed when checking for other objects with the same name within the outer and colliding with the object itself.

#rnx
#tests PIE and Frontend

diff --git a/Engine/Plugins/Experimental/Mutable/Source/CustomizableObject/Private/MuCO/CustomizableObject.cpp b/Engine/Plugins/Experimental/Mutable/Source/CustomizableObject/Private/MuCO/CustomizableObject.cpp
--- a/Engine/Plugins/Experimental/Mutable/Source/CustomizableObject/Private/MuCO/CustomizableObject.cpp
+++ b/Engine/Plugins/Experimental/Mutable/Source/CustomizableObject/Private/MuCO/CustomizableObject.cpp
@@ -458,30 +458,20 @@ void SerializeStreamedResources(FArchive& Ar, UObject* Object, TArray<FCustomiza
 		int32 NumStreamedResources = 0;
 		Ar << NumStreamedResources;
 
-		// Initialize if not cooking. Otherwise, resources will be initialized at this point, and only their data will be updated.
-		if (!bIsCooking) 
-		{
-			StreamedResources.SetNum(NumStreamedResources);
-		}			
-		
-		check(NumStreamedResources == StreamedResources.Num())
+		StreamedResources.SetNum(NumStreamedResources);
 
 		for (int32 ResourceIndex = 0; ResourceIndex < NumStreamedResources; ++ResourceIndex)
 		{
-			UCustomizableObjectResourceDataContainer* Container = nullptr;
-
-			if (bIsCooking)
-			{
-				// Override existing containers
-				Container = StreamedResources[ResourceIndex].GetPath().Get();
-			}
+			// Override existing containers
+			UCustomizableObjectResourceDataContainer* Container = StreamedResources[ResourceIndex].GetPath().Get();
 
-			// Generate a deterministic name to help with deterministic cooking
+			// Create a new container if none.
 			if (!Container)
 			{
+				// Generate a deterministic name to help with deterministic cooking
 				const FString ContainerName = CustomizableObjectName + FString::Printf(TEXT("SR_%d"), ResourceIndex);
 
-				UCustomizableObjectResourceDataContainer* ExistingContainer = Cast<UCustomizableObjectResourceDataContainer>(FindObject<UObject>(Object, *ContainerName));
+				UCustomizableObjectResourceDataContainer* ExistingContainer = FindObject<UCustomizableObjectResourceDataContainer>(Object, *ContainerName);
 				Container = ExistingContainer ? ExistingContainer : NewObject<UCustomizableObjectResourceDataContainer>(
 					Object,
 					FName(*ContainerName),
diff --git a/Engine/Plugins/Experimental/Mutable/Source/CustomizableObjectEditor/Private/MuCOE/CustomizableObjectCookPackageSplitter.cpp b/Engine/Plugins/Experimental/Mutable/Source/CustomizableObjectEditor/Private/MuCOE/CustomizableObjectCookPackageSplitter.cpp
--- a/Engine/Plugins/Experimental/Mutable/Source/CustomizableObjectEditor/Private/MuCOE/CustomizableObjectCookPackageSplitter.cpp
+++ b/Engine/Plugins/Experimental/Mutable/Source/CustomizableObjectEditor/Private/MuCOE/CustomizableObjectCookPackageSplitter.cpp
@@ -57,16 +57,19 @@ bool MoveContainerToNewOuter(
 		return false;
 	}
 
-	// Ensure the target object doesn't exist
-	check(!FindObject<UObject>(NewOuter, *Container->GetName()));
-
-	// The Rename function moves the object into the given package
-	if (!Container->Rename(nullptr, NewOuter, REN_DontCreateRedirectors))
+	if (Container->GetOuter() != NewOuter)
 	{
-		UE_LOG(LogMutable, Error, TEXT("Failed to move streamed Resource Data container %s into Outer %s"),
-			*Container->GetPathName(), *NewOuter->GetPathName());
+		// Ensure the target object doesn't exist
+		check(!FindObject<UObject>(NewOuter, *Container->GetName()));
 
-		return false;
+		// The Rename function moves the object into the given package
+		if (!Container->Rename(nullptr, NewOuter, REN_DontCreateRedirectors))
+		{
+			UE_LOG(LogMutable, Error, TEXT("Failed to move streamed Resource Data container %s into Outer %s"),
+				*Container->GetPathName(), *NewOuter->GetPathName());
+
+			return false;
+		}
 	}
 
 	OutContainer = Container;
diff --git a/Engine/Plugins/Experimental/Mutable/Source/CustomizableObjectEditor/Private/MuCOE/GenerateMutableSource/GenerateMutableSource.cpp b/Engine/Plugins/Experimental/Mutable/Source/CustomizableObjectEditor/Private/MuCOE/GenerateMutableSource/GenerateMutableSource.cpp
--- a/Engine/Plugins/Experimental/Mutable/Source/CustomizableObjectEditor/Private/MuCOE/GenerateMutableSource/GenerateMutableSource.cpp
+++ b/Engine/Plugins/Experimental/Mutable/Source/CustomizableObjectEditor/Private/MuCOE/GenerateMutableSource/GenerateMutableSource.cpp
@@ -168,7 +168,7 @@ int32 FMutableGraphGenerationContext::AddStreamedResource(uint32 InResourceHash,
 	int32 NewResourceIndex = StreamedResourceData.Num();
 	const FString ContainerName = GetNameSafe(Object) + FString::Printf(TEXT("_SR_%d"), NewResourceIndex);
 
-	UCustomizableObjectResourceDataContainer* ExistingContainer = Cast<UCustomizableObjectResourceDataContainer>(FindObject<UObject>(Object, *ContainerName));
+	UCustomizableObjectResourceDataContainer* ExistingContainer = FindObject<UCustomizableObjectResourceDataContainer>(Object, *ContainerName);
 	OutNewResource = ExistingContainer ? ExistingContainer : NewObject<UCustomizableObjectResourceDataContainer>(
 		Object,
 		FName(*ContainerName),
