[Mutable] Removed invalid TODOs.

#rnx

diff --git a/Engine/Plugins/Experimental/Mutable/Source/CustomizableObject/Private/MuCO/CustomizableObjectSystem.cpp b/Engine/Plugins/Experimental/Mutable/Source/CustomizableObject/Private/MuCO/CustomizableObjectSystem.cpp
--- a/Engine/Plugins/Experimental/Mutable/Source/CustomizableObject/Private/MuCO/CustomizableObjectSystem.cpp
+++ b/Engine/Plugins/Experimental/Mutable/Source/CustomizableObject/Private/MuCO/CustomizableObjectSystem.cpp
@@ -1625,7 +1625,7 @@ namespace impl
 			UE_LOG(LogMutable, Verbose, TEXT("Creating Mutable instance with id [%d] "), Operation->InstanceID);
 		}
 
-		Operation->MutableInstance = MutableSystem->BeginUpdate(Operation->InstanceID, Operation->Parameters, Operation->State, mu::System::AllLODs); // TODO GMT possible data race between GameThread and MutableThread
+		Operation->MutableInstance = MutableSystem->BeginUpdate(Operation->InstanceID, Operation->Parameters, Operation->State, mu::System::AllLODs);
 	}
 
 	
@@ -2931,7 +2931,7 @@ namespace impl
 					impl::Task_Mutable_GetMeshID(Operation);
 				});
 
-			SystemPrivateData->AddGameThreadTask( // TODO GMTFuture GameThread task instead of a tick task
+			SystemPrivateData->AddGameThreadTask(
 				{
 				FMutableTaskDelegate::CreateLambda(
 					[Operation]()
