[Mutable] Replaced Generated flag with the SkeletalMeshStatus.

- Removed getter and setter of the SkeletalMeshStatus sinde it is already in the private file.

#rnx

diff --git a/Engine/Plugins/Experimental/Mutable/Source/CustomizableObject/Classes/MuCO/CustomizableObjectInstance.h b/Engine/Plugins/Experimental/Mutable/Source/CustomizableObject/Classes/MuCO/CustomizableObjectInstance.h
--- a/Engine/Plugins/Experimental/Mutable/Source/CustomizableObject/Classes/MuCO/CustomizableObjectInstance.h
+++ b/Engine/Plugins/Experimental/Mutable/Source/CustomizableObject/Classes/MuCO/CustomizableObjectInstance.h
@@ -62,9 +62,9 @@ enum class EUpdateResult : uint8
 /** Indicates the status of the generated Skeletal Mesh. */
 enum class ESkeletalMeshStatus : uint8
 {
-	NotGenerated,
-	Success,
-	Error
+	NotGenerated, // Set only when loading the Instance for the first time or after compiling. Any generation, successful or not, can not end up in this state.
+	Success, // Generated successfully.
+	Error // Generated with errors.
 };
 
 
diff --git a/Engine/Plugins/Experimental/Mutable/Source/CustomizableObject/Internal/MuCO/CustomizableInstancePrivateData.h b/Engine/Plugins/Experimental/Mutable/Source/CustomizableObject/Internal/MuCO/CustomizableInstancePrivateData.h
--- a/Engine/Plugins/Experimental/Mutable/Source/CustomizableObject/Internal/MuCO/CustomizableInstancePrivateData.h
+++ b/Engine/Plugins/Experimental/Mutable/Source/CustomizableObject/Internal/MuCO/CustomizableInstancePrivateData.h
@@ -37,7 +37,6 @@ enum ECOInstanceFlags
 
 	// Update process
 	CreatingSkeletalMesh			= 1 << 1,	//
-	Generated						= 1 << 2,	//
 	ReuseTextures					= 1 << 3, 	// 
 	ReplacePhysicsAssets			= 1 << 4,	// Merge active PhysicsAssets and replace the base physics asset
 
@@ -272,10 +271,6 @@ private:
 	void SetLastMeshId(int32 ComponentIndex, int32 LODIndex, mu::FResourceID MeshId);
 
 public:
-	CUSTOMIZABLEOBJECT_API ESkeletalMeshStatus GetSkeletalMeshStatus() const;
-
-	void SetSkeletalMeshStatus(ESkeletalMeshStatus Status);
-	
 	// If any components are using this instance, they will store the min of their distances to the player here every frame for LOD purposes
 	float MinSquareDistFromComponentToPlayer;
 	float LastMinSquareDistFromComponentToPlayer; // The same as the previous dist for last frame
@@ -332,7 +327,6 @@ public:
 	/** Hash of the UCustomizableObjectInstance::Descriptor on the last successful update. */
 	FDescriptorRuntimeHash DescriptorRuntimeHash;
 
-private:
 	/** Status of the generated Skeletal Mesh. Not to be confused with the Update Result. */
 	ESkeletalMeshStatus SkeletalMeshStatus = ESkeletalMeshStatus::NotGenerated;
 };
diff --git a/Engine/Plugins/Experimental/Mutable/Source/CustomizableObject/Private/MuCO/CustomizableObjectInstance.cpp b/Engine/Plugins/Experimental/Mutable/Source/CustomizableObject/Private/MuCO/CustomizableObjectInstance.cpp
--- a/Engine/Plugins/Experimental/Mutable/Source/CustomizableObject/Private/MuCO/CustomizableObjectInstance.cpp
+++ b/Engine/Plugins/Experimental/Mutable/Source/CustomizableObject/Private/MuCO/CustomizableObjectInstance.cpp
@@ -167,27 +167,16 @@ void UCustomizableInstancePrivateData::SetLastMeshId(int32 ComponentIndex, int32
 }
 
 
-ESkeletalMeshStatus UCustomizableInstancePrivateData::GetSkeletalMeshStatus() const
-{
-	return SkeletalMeshStatus;
-}
-
-
-void UCustomizableInstancePrivateData::SetSkeletalMeshStatus(ESkeletalMeshStatus Status)
-{
-	SkeletalMeshStatus = Status;
-}
-
-
 void UCustomizableInstancePrivateData::InvalidateGeneratedData()
 {
+	SkeletalMeshStatus = ESkeletalMeshStatus::NotGenerated;
+
 	// Init Component Data
 	FCustomizableInstanceComponentData TemplateComponentData;
 	TemplateComponentData.LastMeshIdPerLOD.Init(MAX_uint64, MAX_MESH_LOD_COUNT);
 	ComponentsData.Init(TemplateComponentData, ComponentsData.Num());
 
 	GeneratedMaterials.Empty();
-	ClearCOInstanceFlags(Generated);
 }
 
 
@@ -784,7 +773,7 @@ void UCustomizableInstancePrivateData::TickUpdateCloseCustomizableObjects(UCusto
 
 	const FCustomizableObjectSystemPrivate* SystemPrivate = UCustomizableObjectSystem::GetInstance()->GetPrivate();	
 
-	const EUpdateRequired UpdateRequired = SystemPrivate->IsUpdateRequired(Public, true, false);
+	const EUpdateRequired UpdateRequired = SystemPrivate->IsUpdateRequired(Public, true, true, false);
 	if (UpdateRequired != EUpdateRequired::NoUpdate) // Since this is done in the tick, avoid starting an update that we know for sure that would not be performed. Once started it has some performance implications that we want to avoid.
 	{
 		if (UpdateRequired == EUpdateRequired::Discard)
@@ -829,7 +818,7 @@ void UCustomizableInstancePrivateData::TickUpdateCloseCustomizableObjects(UCusto
 
 void UCustomizableInstancePrivateData::UpdateInstanceIfNotGenerated(UCustomizableObjectInstance& Public, FMutableInstanceUpdateMap& InOutRequestedUpdates)
 {
-	if (HasCOInstanceFlags(Generated))
+	if (SkeletalMeshStatus != ESkeletalMeshStatus::NotGenerated)
 	{
 		return;
 	}
@@ -1847,7 +1836,7 @@ bool UCustomizableInstancePrivateData::UpdateSkeletalMesh_PostBeginUpdate0(UCust
 
 		InvalidateGeneratedData();
 
-		SetCOInstanceFlags(Generated);
+		OperationData->UpdateResult = EUpdateResult::Error;
 
 		return false;
 	}
@@ -1855,7 +1844,6 @@ bool UCustomizableInstancePrivateData::UpdateSkeletalMesh_PostBeginUpdate0(UCust
 	// None of the current meshes requires a mesh update. Continue to BuildMaterials
 	if (!bUpdateMeshes)
 	{
-		SetCOInstanceFlags(Generated);
 		return true;
 	}
 
@@ -2153,7 +2141,6 @@ bool UCustomizableInstancePrivateData::UpdateSkeletalMesh_PostBeginUpdate0(UCust
 		}
 	}
 
-	SetCOInstanceFlags(Generated);
 	ClearCOInstanceFlags(CreatingSkeletalMesh); // TODO MTBL-391: Review
 
 	if (!bSuccess)
@@ -2805,7 +2792,7 @@ bool UCustomizableObjectInstance::IsSelectedParameterProfileDirty() const
 
 void UCustomizableInstancePrivateData::DiscardResourcesAndSetReferenceSkeletalMesh(UCustomizableObjectInstance* Instance)
 {
-	if (HasCOInstanceFlags(Generated))
+	if (SkeletalMeshStatus == ESkeletalMeshStatus::Success)
 	{
 		for (int32 Component = 0; Component < Instance->SkeletalMeshes.Num(); ++Component)
 		{
@@ -2818,9 +2805,6 @@ void UCustomizableInstancePrivateData::DiscardResourcesAndSetReferenceSkeletalMe
 
 		ReleaseMutableResources(false, *Instance);
 	}
-
-	ClearCOInstanceFlags(Generated);
-	SkeletalMeshStatus = ESkeletalMeshStatus::NotGenerated;
 	
 	InvalidateGeneratedData();
 	
@@ -6197,7 +6181,7 @@ void UCustomizableObjectInstance::SetRequestedLODs(int32 InMinLOD, int32 InMaxLO
 
 	if (CVarPreserveUserLODsOnFirstGeneration.GetValueOnGameThread() &&
 		GetCustomizableObject()->IsPreserveUserLODsOnFirstGeneration() &&
-		GetPrivate()->GetSkeletalMeshStatus() != ESkeletalMeshStatus::Success)
+		GetPrivate()->SkeletalMeshStatus != ESkeletalMeshStatus::Success)
 	{
 		return;
 	}
diff --git a/Engine/Plugins/Experimental/Mutable/Source/CustomizableObject/Private/MuCO/CustomizableObjectInstanceUsage.cpp b/Engine/Plugins/Experimental/Mutable/Source/CustomizableObject/Private/MuCO/CustomizableObjectInstanceUsage.cpp
--- a/Engine/Plugins/Experimental/Mutable/Source/CustomizableObject/Private/MuCO/CustomizableObjectInstanceUsage.cpp
+++ b/Engine/Plugins/Experimental/Mutable/Source/CustomizableObject/Private/MuCO/CustomizableObjectInstanceUsage.cpp
@@ -524,7 +524,7 @@ void UCustomizableObjectInstanceUsage::Tick(float DeltaTime)
 
 		USkeletalMesh* SkeletalMesh = nullptr;
 
-		const bool bInstanceGenerated = CustomizableObjectInstance->GetPrivate()->GetSkeletalMeshStatus() == ESkeletalMeshStatus::Success;		
+		const bool bInstanceGenerated = CustomizableObjectInstance->GetPrivate()->SkeletalMeshStatus == ESkeletalMeshStatus::Success;		
 		if (bInstanceGenerated)
 		{
 			// Generated SkeletalMesh to set, can be null if the component is empty
diff --git a/Engine/Plugins/Experimental/Mutable/Source/CustomizableObject/Private/MuCO/CustomizableObjectSystem.cpp b/Engine/Plugins/Experimental/Mutable/Source/CustomizableObject/Private/MuCO/CustomizableObjectSystem.cpp
--- a/Engine/Plugins/Experimental/Mutable/Source/CustomizableObject/Private/MuCO/CustomizableObjectSystem.cpp
+++ b/Engine/Plugins/Experimental/Mutable/Source/CustomizableObject/Private/MuCO/CustomizableObjectSystem.cpp
@@ -770,7 +770,7 @@ void FinishUpdateGlobal(const TSharedRef<FUpdateContextPrivate>& Context)
 		switch (Context->UpdateResult)
 		{
 		case EUpdateResult::Success:
-			PrivateInstance->SetSkeletalMeshStatus(ESkeletalMeshStatus::Success);
+			PrivateInstance->SkeletalMeshStatus = ESkeletalMeshStatus::Success;
 
 			PrivateInstance->DescriptorRuntimeHash = Context->InstanceDescriptorRuntimeHash;
 
@@ -787,7 +787,7 @@ void FinishUpdateGlobal(const TSharedRef<FUpdateContextPrivate>& Context)
 
 		case EUpdateResult::Error: 
 		case EUpdateResult::Error16BitBoneIndex:
-			PrivateInstance->SetSkeletalMeshStatus(ESkeletalMeshStatus::Error);
+			PrivateInstance->SkeletalMeshStatus = ESkeletalMeshStatus::Error;
 			break;
 			
 		case EUpdateResult::ErrorReplaced:
@@ -917,7 +917,6 @@ void UpdateSkeletalMesh(const TSharedRef<FUpdateContextPrivate>& Context)
 		}
 	}
 
-	CustomizableObjectInstancePrivateData->SetCOInstanceFlags(Generated);
 	CustomizableObjectInstancePrivateData->ClearCOInstanceFlags(CreatingSkeletalMesh);
 
 	CustomizableObjectInstance->bEditorPropertyChanged = false;
@@ -1027,18 +1026,17 @@ bool FCustomizableObjectSystemPrivate::TextureHasReferences(const FMutableImageC
 }
 
 
-EUpdateRequired FCustomizableObjectSystemPrivate::IsUpdateRequired(const UCustomizableObjectInstance& Instance, bool bOnlyUpdateIfNotGenerated, bool bIgnoreCloseDist) const
+EUpdateRequired FCustomizableObjectSystemPrivate::IsUpdateRequired(const UCustomizableObjectInstance& Instance, bool bOnlyUpdateIfNotGenerated, bool bOnlyUpdateIfLODs, bool bIgnoreCloseDist) const
 {
 	UCustomizableObjectSystem* System = UCustomizableObjectSystem::GetInstance();
 	const UCustomizableInstancePrivateData* const Private = Instance.GetPrivate();
 	
-	const UCustomizableObject* CustomizableObject = Instance.GetCustomizableObject();
 	if (!Instance.CanUpdateInstance())
 	{
 		return EUpdateRequired::NoUpdate;
 	}
 
-	const bool bIsGenerated = Private->HasCOInstanceFlags(Generated);
+	const bool bIsGenerated = Private->SkeletalMeshStatus != ESkeletalMeshStatus::NotGenerated;
 	const int32 NumGeneratedInstancesLimit = System->GetInstanceLODManagement()->GetNumGeneratedInstancesLimitFullLODs();
 	const int32 NumGeneratedInstancesLimitLOD1 = System->GetInstanceLODManagement()->GetNumGeneratedInstancesLimitLOD1();
 	const int32 NumGeneratedInstancesLimitLOD2 = System->GetInstanceLODManagement()->GetNumGeneratedInstancesLimitLOD2();
@@ -1050,8 +1048,6 @@ EUpdateRequired FCustomizableObjectSystemPrivate::IsUpdateRequired(const UCustom
 		return EUpdateRequired::NoUpdate;
 	}
 
-	const bool bShouldUpdateLODs = Private->HasCOInstanceFlags(PendingLODsUpdate);
-
 	const bool bDiscardByDistance = Private->LastMinSquareDistFromComponentToPlayer > FMath::Square(System->GetInstanceLODManagement()->GetOnlyUpdateCloseCustomizableObjectsDist());
 	const bool bLODManagementDiscard = System->GetInstanceLODManagement()->IsOnlyUpdateCloseCustomizableObjectsEnabled() &&
 			bDiscardByDistance &&
@@ -1070,9 +1066,15 @@ EUpdateRequired FCustomizableObjectSystemPrivate::IsUpdateRequired(const UCustom
 		}
 	}
 
-	if (bIsGenerated &&
-		!bShouldUpdateLODs &&
-		bOnlyUpdateIfNotGenerated)
+	if (bOnlyUpdateIfNotGenerated &&
+		bIsGenerated)
+	{
+		return EUpdateRequired::NoUpdate;
+	}
+
+	const bool bShouldUpdateLODs = Private->HasCOInstanceFlags(PendingLODsUpdate);
+	if (bOnlyUpdateIfLODs &&
+		!bShouldUpdateLODs)
 	{
 		return EUpdateRequired::NoUpdate;
 	}
@@ -1085,7 +1087,7 @@ EQueuePriorityType FCustomizableObjectSystemPrivate::GetUpdatePriority(const UCu
 {
 	const UCustomizableInstancePrivateData* InstancePrivate = Instance.GetPrivate();
 		
-	const bool bIsGenerated = InstancePrivate->HasCOInstanceFlags(Generated);
+	const bool bNotGenerated = InstancePrivate->SkeletalMeshStatus == ESkeletalMeshStatus::NotGenerated;
 	const bool bShouldUpdateLODs = InstancePrivate->HasCOInstanceFlags(PendingLODsUpdate);
 	const bool bIsDowngradeLODUpdate = InstancePrivate->HasCOInstanceFlags(PendingLODsDowngrade);
 	const bool bIsPlayerOrNearIt = InstancePrivate->HasCOInstanceFlags(UsedByPlayerOrNearIt);
@@ -1095,7 +1097,7 @@ EQueuePriorityType FCustomizableObjectSystemPrivate::GetUpdatePriority(const UCu
 	{
 		Priority = EQueuePriorityType::High;
 	}
-	else if (!bIsGenerated || !Instance.HasAnySkeletalMesh())
+	else if (bNotGenerated || !Instance.HasAnySkeletalMesh())
 	{
 		Priority = EQueuePriorityType::Med;
 	}
@@ -1145,7 +1147,7 @@ void FCustomizableObjectSystemPrivate::EnqueueUpdateSkeletalMesh(const TSharedRe
 
 	UCustomizableObjectSystem* System = UCustomizableObjectSystem::GetInstance();
 
-	const EUpdateRequired UpdateRequired = IsUpdateRequired(*Instance, Context->bOnlyUpdateIfNotGenerated, Context->bIgnoreCloseDist);
+	const EUpdateRequired UpdateRequired = IsUpdateRequired(*Instance, Context->bOnlyUpdateIfNotGenerated, false, Context->bIgnoreCloseDist);
 	switch (UpdateRequired)
 	{
 	case EUpdateRequired::NoUpdate:
diff --git a/Engine/Plugins/Experimental/Mutable/Source/CustomizableObject/Private/MuCO/CustomizableObjectSystemPrivate.h b/Engine/Plugins/Experimental/Mutable/Source/CustomizableObject/Private/MuCO/CustomizableObjectSystemPrivate.h
--- a/Engine/Plugins/Experimental/Mutable/Source/CustomizableObject/Private/MuCO/CustomizableObjectSystemPrivate.h
+++ b/Engine/Plugins/Experimental/Mutable/Source/CustomizableObject/Private/MuCO/CustomizableObjectSystemPrivate.h
@@ -584,7 +584,7 @@ public:
 
 	bool TextureHasReferences(const FMutableImageCacheKey& TextureId) const;
 
-	EUpdateRequired IsUpdateRequired(const UCustomizableObjectInstance& Instance, bool bOnlyUpdateIfNotGenerated, bool bIgnoreCloseDist) const;
+	EUpdateRequired IsUpdateRequired(const UCustomizableObjectInstance& Instance, bool bOnlyUpdateIfNotGenerated, bool bOnlyUpdateIfLOD, bool bIgnoreCloseDist) const;
 
 	EQueuePriorityType GetUpdatePriority(const UCustomizableObjectInstance& Instance, bool bForceHighPriority) const;
 
diff --git a/Engine/Plugins/Experimental/Mutable/Source/CustomizableObject/Private/MuCO/LogBenchmarkUtil.cpp b/Engine/Plugins/Experimental/Mutable/Source/CustomizableObject/Private/MuCO/LogBenchmarkUtil.cpp
--- a/Engine/Plugins/Experimental/Mutable/Source/CustomizableObject/Private/MuCO/LogBenchmarkUtil.cpp
+++ b/Engine/Plugins/Experimental/Mutable/Source/CustomizableObject/Private/MuCO/LogBenchmarkUtil.cpp
@@ -94,7 +94,7 @@ void FLogBenchmarkUtil::GetInstancesStats(int32& OutNumInstances, int32& OutNumB
 		OutNumInstancesLOD1 += Instance->GetCurrentMinLOD() <= 1 && Instance->GetCurrentMaxLOD() >= 1;
 		OutNumInstancesLOD2 += Instance->GetCurrentMinLOD() >= 2;
 
-		if (Instance->GetPrivate()->GetSkeletalMeshStatus() == ESkeletalMeshStatus::Success)
+		if (Instance->GetPrivate()->SkeletalMeshStatus == ESkeletalMeshStatus::Success)
 		{
 			++OutNumBuiltInstances;
 		}
diff --git a/Engine/Plugins/Experimental/Mutable/Source/CustomizableObject/Private/MuCO/LogInformationUtil.cpp b/Engine/Plugins/Experimental/Mutable/Source/CustomizableObject/Private/MuCO/LogInformationUtil.cpp
--- a/Engine/Plugins/Experimental/Mutable/Source/CustomizableObject/Private/MuCO/LogInformationUtil.cpp
+++ b/Engine/Plugins/Experimental/Mutable/Source/CustomizableObject/Private/MuCO/LogInformationUtil.cpp
@@ -168,7 +168,6 @@ void LogInformationUtil::LogShowInstanceDataFull(const UCustomizableObjectInstan
 
 	MessageChunk += FString::Printf(TEXT("        bIsDowngradeLODUpdate = %d\n"), InstanceFlags & ECOInstanceFlags::PendingLODsDowngrade);
 	MessageChunk += FString::Printf(TEXT("        bIsCreatingSkeletalMesh = %d\n"), InstanceFlags & ECOInstanceFlags::CreatingSkeletalMesh);
-	MessageChunk += FString::Printf(TEXT("        bIsGenerated = %d\n"), InstanceFlags & ECOInstanceFlags::Generated);
 	LogData += MessageChunk;
 
 	for (int32 ComponentIndex = 0; ComponentIndex < CustomizableObjectInstance->SkeletalMeshes.Num(); ++ComponentIndex)
diff --git a/Engine/Plugins/Experimental/Mutable/Source/CustomizableObjectEditor/Private/MuCOE/CustomizableObjectEditor.cpp b/Engine/Plugins/Experimental/Mutable/Source/CustomizableObjectEditor/Private/MuCOE/CustomizableObjectEditor.cpp
--- a/Engine/Plugins/Experimental/Mutable/Source/CustomizableObjectEditor/Private/MuCOE/CustomizableObjectEditor.cpp
+++ b/Engine/Plugins/Experimental/Mutable/Source/CustomizableObjectEditor/Private/MuCOE/CustomizableObjectEditor.cpp
@@ -2492,7 +2492,7 @@ void FCustomizableObjectEditor::UpdatePreviewVisibility()
 	{
 		if (PreviewSkeletalMeshComponent)
 		{
-			if (PreviewInstance->GetPrivate()->GetSkeletalMeshStatus() == ESkeletalMeshStatus::Success)
+			if (PreviewInstance->GetPrivate()->SkeletalMeshStatus == ESkeletalMeshStatus::Success)
 			{
 				PreviewSkeletalMeshComponent->SetVisibility(true, true);
 			}
diff --git a/Engine/Plugins/Experimental/Mutable/Source/CustomizableObjectEditor/Private/MuCOE/CustomizableObjectInstanceEditor.cpp b/Engine/Plugins/Experimental/Mutable/Source/CustomizableObjectEditor/Private/MuCOE/CustomizableObjectInstanceEditor.cpp
--- a/Engine/Plugins/Experimental/Mutable/Source/CustomizableObjectEditor/Private/MuCOE/CustomizableObjectInstanceEditor.cpp
+++ b/Engine/Plugins/Experimental/Mutable/Source/CustomizableObjectEditor/Private/MuCOE/CustomizableObjectInstanceEditor.cpp
@@ -437,7 +437,7 @@ void FCustomizableObjectInstanceEditor::OnAssetRegistryLoadComplete()
 
 void FCustomizableObjectInstanceEditor::UpdatePreviewVisibility()
 {
-	const bool bEnableVisibility = CustomizableObjectInstance->GetPrivate()->GetSkeletalMeshStatus() == ESkeletalMeshStatus::Success;
+	const bool bEnableVisibility = CustomizableObjectInstance->GetPrivate()->SkeletalMeshStatus == ESkeletalMeshStatus::Success;
 	for (UDebugSkelMeshComponent* PreviewSkeletalMeshComponent : PreviewSkeletalMeshComponents)
 	{
 		PreviewSkeletalMeshComponent->SetVisibility(bEnableVisibility, true);
diff --git a/Engine/Plugins/Experimental/Mutable/Source/CustomizableObjectEditor/Private/MuCOE/SCustomizableObjectEditorViewportToolBar.cpp b/Engine/Plugins/Experimental/Mutable/Source/CustomizableObjectEditor/Private/MuCOE/SCustomizableObjectEditorViewportToolBar.cpp
--- a/Engine/Plugins/Experimental/Mutable/Source/CustomizableObjectEditor/Private/MuCOE/SCustomizableObjectEditorViewportToolBar.cpp
+++ b/Engine/Plugins/Experimental/Mutable/Source/CustomizableObjectEditor/Private/MuCOE/SCustomizableObjectEditorViewportToolBar.cpp
@@ -210,7 +210,7 @@ FText SCustomizableObjectEditorViewportToolBar::GetCompileErrorOverlayText() con
 
 	const UCustomizableInstancePrivateData* PrivateInstance = Instance->GetPrivate();
 	
-	switch (PrivateInstance->GetSkeletalMeshStatus())
+	switch (PrivateInstance->SkeletalMeshStatus)
 	{
 	case ESkeletalMeshStatus::NotGenerated:
 		return LOCTEXT("NoSkeletalMeshGenerated", "No Skeletal Mesh Generated");
