[Mutable] Fix LOD updates not working.

Broken in 30074696.

#rnx

diff --git a/Engine/Plugins/Experimental/Mutable/Source/CustomizableObject/Private/MuCO/CustomizableObjectSystem.cpp b/Engine/Plugins/Experimental/Mutable/Source/CustomizableObject/Private/MuCO/CustomizableObjectSystem.cpp
--- a/Engine/Plugins/Experimental/Mutable/Source/CustomizableObject/Private/MuCO/CustomizableObjectSystem.cpp
+++ b/Engine/Plugins/Experimental/Mutable/Source/CustomizableObject/Private/MuCO/CustomizableObjectSystem.cpp
@@ -1066,15 +1066,13 @@ EUpdateRequired FCustomizableObjectSystemPrivate::IsUpdateRequired(const UCustom
 		}
 	}
 
-	if (bOnlyUpdateIfNotGenerated &&
-		bIsGenerated)
-	{
-		return EUpdateRequired::NoUpdate;
-	}
-
 	const bool bShouldUpdateLODs = Private->HasCOInstanceFlags(PendingLODsUpdate);
-	if (bOnlyUpdateIfLODs &&
-		!bShouldUpdateLODs)
+
+	const bool bNoUpdateLODs = bOnlyUpdateIfLODs && !bShouldUpdateLODs;
+	const bool bNoInitialUpdate = bOnlyUpdateIfNotGenerated && bIsGenerated;
+
+	if (bNoUpdateLODs &&
+		bNoInitialUpdate)
 	{
 		return EUpdateRequired::NoUpdate;
 	}
