[mutable] Move debugger from the toolbar to the Compile options menu.
#rnx

diff --git a/Engine/Plugins/Experimental/Mutable/Source/CustomizableObjectEditor/Private/MuCOE/CustomizableObjectEditor.cpp b/Engine/Plugins/Experimental/Mutable/Source/CustomizableObjectEditor/Private/MuCOE/CustomizableObjectEditor.cpp
--- a/Engine/Plugins/Experimental/Mutable/Source/CustomizableObjectEditor/Private/MuCOE/CustomizableObjectEditor.cpp
+++ b/Engine/Plugins/Experimental/Mutable/Source/CustomizableObjectEditor/Private/MuCOE/CustomizableObjectEditor.cpp
@@ -1270,12 +1270,6 @@ void FCustomizableObjectEditor::ExtendToolbar()
 				true);
 
 			ToolbarBuilder.EndSection();
-
-			
-			ToolbarBuilder.BeginSection("Debug");
-			ToolbarBuilder.AddToolBarButton(FCustomizableObjectEditorCommands::Get().Debug);
-			ToolbarBuilder.EndSection();
-
 			
 			ToolbarBuilder.BeginSection("Information");
 			ToolbarBuilder.AddToolBarButton(FCustomizableObjectEditorCommands::Get().TextureAnalyzer);
@@ -1384,6 +1378,13 @@ TSharedRef<SWidget> FCustomizableObjectEditor::GenerateCompileOptionsMenuContent
 	}
 	MenuBuilder.EndSection();
 
+	// Debugging options
+	MenuBuilder.BeginSection("Debugger");
+	{
+		MenuBuilder.AddMenuEntry(FCustomizableObjectEditorCommands::Get().Debug);
+	}
+	MenuBuilder.EndSection();
+
 	return MenuBuilder.MakeWidget();
 }
 
diff --git a/Engine/Plugins/Experimental/Mutable/Source/CustomizableObjectEditor/Private/MuCOE/CustomizableObjectEditorActions.cpp b/Engine/Plugins/Experimental/Mutable/Source/CustomizableObjectEditor/Private/MuCOE/CustomizableObjectEditorActions.cpp
--- a/Engine/Plugins/Experimental/Mutable/Source/CustomizableObjectEditor/Private/MuCOE/CustomizableObjectEditorActions.cpp
+++ b/Engine/Plugins/Experimental/Mutable/Source/CustomizableObjectEditor/Private/MuCOE/CustomizableObjectEditorActions.cpp
@@ -26,10 +26,7 @@ void FCustomizableObjectEditorCommands::RegisterCommands()
 	UI_COMMAND(ResetCompileOptions, "Reset Compilation Options", "Set reasonable defaults for the compilation options.", EUserInterfaceActionType::Button, FInputChord());
 	UI_COMMAND(CompileOptions_EnableTextureCompression, "Enable texture compression.", "Only for debug. Do not use.", EUserInterfaceActionType::ToggleButton, FInputChord());
 	UI_COMMAND(CompileOptions_UseDiskCompilation, "Enable compiling using the disk as memory.", "This is very slow but supports compiling huge objects. It requires a lot of free space in the OS disk.", EUserInterfaceActionType::ToggleButton, FInputChord());
-
-	UI_COMMAND(Debug, "Debug", "Debug the object.", EUserInterfaceActionType::Button, FInputChord());
-	UI_COMMAND(DebugOptions_OnlySelected, "Debug Only Selected", "Debug only for the selected options in the preview. The rest of options are discarded.", EUserInterfaceActionType::Button, FInputChord());
-	UI_COMMAND(DebugOptions_EnableTextureCompression, "Enable texture compression.", "Only for debug. Do not use.", EUserInterfaceActionType::ToggleButton, FInputChord());
+	UI_COMMAND(Debug, "Debug", "Open the CustomizableObject debugger tab for this object.", EUserInterfaceActionType::Button, FInputChord());
 
 	UI_COMMAND(PerformanceReport, "Performance Report", "Open the Performance Report window to set up and perform all tests relevant to Customizable Objects and access worst cases data.", EUserInterfaceActionType::Button, FInputChord());
 	UI_COMMAND(TextureAnalyzer, "Texture Memory Analyzer", "Open the Texture Analyzer window to check all the information of the textures created by Mutable.", EUserInterfaceActionType::Button, FInputChord());
diff --git a/Engine/Plugins/Experimental/Mutable/Source/CustomizableObjectEditor/Private/MuCOE/CustomizableObjectEditorStyle.cpp b/Engine/Plugins/Experimental/Mutable/Source/CustomizableObjectEditor/Private/MuCOE/CustomizableObjectEditorStyle.cpp
--- a/Engine/Plugins/Experimental/Mutable/Source/CustomizableObjectEditor/Private/MuCOE/CustomizableObjectEditorStyle.cpp
+++ b/Engine/Plugins/Experimental/Mutable/Source/CustomizableObjectEditor/Private/MuCOE/CustomizableObjectEditorStyle.cpp
@@ -83,8 +83,8 @@ TSharedRef<FSlateStyleSet> FCustomizableObjectEditorStyle::Create()
 	Style->Set("CustomizableObjectEditor.Debug.Small", new IMAGE_BRUSH_SVG(Style, "Starship/Common/Debug", Icon20x20));
 	Style->Set("CustomizableObjectEditor.TextureAnalyzer", new IMAGE_BRUSH_SVG(Style, "Starship/Common/IssueTracker", Icon40x40));
 	Style->Set("CustomizableObjectEditor.TextureAnalyzer.Small", new IMAGE_BRUSH_SVG(Style, "Starship/Common/IssueTracker", Icon20x20));
-	Style->Set("CustomizableObjectEditor.PerofrmanceReport", new IMAGE_BRUSH_SVG(Style, "Starship/MainToolbar/simulate", Icon40x40));
-	Style->Set("CustomizableObjectEditor.PerofrmanceReport.Small", new IMAGE_BRUSH_SVG(Style, "Starship/MainToolbar/simulate", Icon20x20));
+	Style->Set("CustomizableObjectEditor.PerformanceReport", new IMAGE_BRUSH_SVG(Style, "Starship/MainToolbar/simulate", Icon40x40));
+	Style->Set("CustomizableObjectEditor.PerformanceReport.Small", new IMAGE_BRUSH_SVG(Style, "Starship/MainToolbar/simulate", Icon20x20));
 
 	Style->Set("CustomizableObjectEditor.Tabs.Preview", new IMAGE_BRUSH_SVG(Style, "Starship/Persona/PersonaPreviewAnimation", Icon20x20));
 	Style->Set("CustomizableObjectEditor.Tabs.PreviewSettings", new IMAGE_BRUSH_SVG(Style, "Starship/AssetIcons/CameraActor_64", Icon20x20));
diff --git a/Engine/Plugins/Experimental/Mutable/Source/CustomizableObjectEditor/Public/MuCOE/CustomizableObjectEditorActions.h b/Engine/Plugins/Experimental/Mutable/Source/CustomizableObjectEditor/Public/MuCOE/CustomizableObjectEditorActions.h
--- a/Engine/Plugins/Experimental/Mutable/Source/CustomizableObjectEditor/Public/MuCOE/CustomizableObjectEditorActions.h
+++ b/Engine/Plugins/Experimental/Mutable/Source/CustomizableObjectEditor/Public/MuCOE/CustomizableObjectEditorActions.h
@@ -20,10 +20,7 @@ public:
 	TSharedPtr< FUICommandInfo > ResetCompileOptions;
 	TSharedPtr< FUICommandInfo > CompileOptions_EnableTextureCompression;
 	TSharedPtr< FUICommandInfo > CompileOptions_UseDiskCompilation;
-
 	TSharedPtr< FUICommandInfo > Debug;
-	TSharedPtr< FUICommandInfo > DebugOptions_OnlySelected;
-	TSharedPtr< FUICommandInfo > DebugOptions_EnableTextureCompression;
 
 	TSharedPtr< FUICommandInfo > PerformanceReport;
 	TSharedPtr< FUICommandInfo > ResetPerformanceReportOptions;
