[mutable] Memory peak counting for instance update and mip updates.
- Count and report update time for instance.
- Count and report update peak and real peak memory for instance and mip updates.
- The real memory used is the memory being used before the task itself + the memory peak of the task.

diff --git a/Engine/Plugins/Experimental/Mutable/Source/CustomizableObject/Private/MuCO/CustomizableObjectMipDataProvider.cpp b/Engine/Plugins/Experimental/Mutable/Source/CustomizableObject/Private/MuCO/CustomizableObjectMipDataProvider.cpp
--- a/Engine/Plugins/Experimental/Mutable/Source/CustomizableObject/Private/MuCO/CustomizableObjectMipDataProvider.cpp
+++ b/Engine/Plugins/Experimental/Mutable/Source/CustomizableObject/Private/MuCO/CustomizableObjectMipDataProvider.cpp
@@ -131,6 +131,10 @@ namespace impl
 		MUTABLE_CPUPROFILER_SCOPE(Task_Mutable_UpdateImage);
 		const double StartTime = FPlatformTime::Seconds();
 		
+		// Cache memory used when starting the update of the image
+		OperationData->ImageUpdateStartBytes = mu::FGlobalMemoryCounter::GetCounter();
+		mu::FGlobalMemoryCounter::Zero();
+		
 		// Any external texture that may be needed for this update will be requested from Mutable Core's GetImage
 		// which will safely access the GlobalExternalImages map, and then just get the cached image or issue a disk read
 
@@ -233,12 +237,16 @@ namespace impl
 		if (CVarEnableBenchmark.GetValueOnAnyThread())
 		{
 			double Time = FPlatformTime::Seconds() - StartTime;
+			// Report the peak memory used by the operation
+			const int64 PeakMemory = mu::FGlobalMemoryCounter::GetPeak();
+			// Report the peak memory used during the operation (operation + baseline)
+			const int64 RealMemoryPeak = PeakMemory + OperationData->ImageUpdateStartBytes;
 
 			const FString& CustomizableObjectPathName = OperationData->UpdateContext->GetCustomizableObjectPathName();
 			const FString& InstancePathName = OperationData->UpdateContext->GetInstancePathName();
 			
 			FFunctionGraphTask::CreateAndDispatchWhenReady(
-			[CustomizableObjectPathName, InstancePathName, Time]()
+			[CustomizableObjectPathName, InstancePathName, Time, PeakMemory, RealMemoryPeak]()
 			{
 				UCustomizableObjectSystem* System = UCustomizableObjectSystem::GetInstance();
 				if (!System)
@@ -246,7 +254,7 @@ namespace impl
 					return;
 				}
 
-				System->GetPrivateChecked()->LogBenchmarkUtil.FinishUpdateImage(CustomizableObjectPathName, InstancePathName, Time);
+				System->GetPrivateChecked()->LogBenchmarkUtil.FinishUpdateImage(CustomizableObjectPathName, InstancePathName, Time, PeakMemory, RealMemoryPeak);
 			},
 			TStatId{},
 			nullptr,
diff --git a/Engine/Plugins/Experimental/Mutable/Source/CustomizableObject/Private/MuCO/CustomizableObjectSystem.cpp b/Engine/Plugins/Experimental/Mutable/Source/CustomizableObject/Private/MuCO/CustomizableObjectSystem.cpp
--- a/Engine/Plugins/Experimental/Mutable/Source/CustomizableObject/Private/MuCO/CustomizableObjectSystem.cpp
+++ b/Engine/Plugins/Experimental/Mutable/Source/CustomizableObject/Private/MuCO/CustomizableObjectSystem.cpp
@@ -2190,6 +2190,9 @@ namespace impl
 	{
 		MUTABLE_CPUPROFILER_SCOPE(Task_Mutable_Update_GetMesh)
 		FMutableScopeTimer Timer(OperationData->TaskGetMeshTime);
+		
+		OperationData->UpdateStartBytes = mu::FGlobalMemoryCounter::GetCounter();
+		mu::FGlobalMemoryCounter::Zero();
 
 #if WITH_EDITOR
 		const uint32 StartCycles = FPlatformTime::Cycles();
@@ -2266,6 +2269,11 @@ namespace impl
 			MutableSystem->ClearWorkingMemory();
 		}
 
+		// Memory used in the context of this the update of mesh
+		OperationData->UpdateEndPeakBytes = mu::FGlobalMemoryCounter::GetPeak();
+		// Memory used in the context of the mesh update + the baseline memory already in use by mutable
+		OperationData->UpdateEndRealPeakBytes = OperationData->UpdateEndPeakBytes + OperationData->UpdateStartBytes;
+
 		UCustomizableObjectSystem::GetInstance()->GetPrivate()->MutableTaskGraph.AllowLaunchingMutableTaskLowPriority(true, true);
 	}
 
diff --git a/Engine/Plugins/Experimental/Mutable/Source/CustomizableObject/Private/MuCO/CustomizableObjectSystemPrivate.h b/Engine/Plugins/Experimental/Mutable/Source/CustomizableObject/Private/MuCO/CustomizableObjectSystemPrivate.h
--- a/Engine/Plugins/Experimental/Mutable/Source/CustomizableObject/Private/MuCO/CustomizableObjectSystemPrivate.h
+++ b/Engine/Plugins/Experimental/Mutable/Source/CustomizableObject/Private/MuCO/CustomizableObjectSystemPrivate.h
@@ -495,6 +495,11 @@ public:
 	double TaskGetImagesTime = 0.0;
 	double TaskConvertResourcesTime = 0.0f;
 	double TaskCallbacksTime = 0.0;
+
+	// Update Memory stats
+	int64 UpdateStartBytes = 0;
+	int64 UpdateEndPeakBytes = 0;
+	int64 UpdateEndRealPeakBytes = 0;
 	
 #if WITH_EDITOR
 	/** Used for profiling in the editor. */
diff --git a/Engine/Plugins/Experimental/Mutable/Source/CustomizableObject/Private/MuCO/LogBenchmarkUtil.cpp b/Engine/Plugins/Experimental/Mutable/Source/CustomizableObject/Private/MuCO/LogBenchmarkUtil.cpp
--- a/Engine/Plugins/Experimental/Mutable/Source/CustomizableObject/Private/MuCO/LogBenchmarkUtil.cpp
+++ b/Engine/Plugins/Experimental/Mutable/Source/CustomizableObject/Private/MuCO/LogBenchmarkUtil.cpp
@@ -55,7 +55,7 @@ TSharedPtr<FArchive> CreateFile()
 	TSharedPtr<FArchive> Archive = MakeShareable(IFileManager::Get().CreateFileWriter(*Filename, FILEWRITE_AllowRead | FILEWRITE_NoFail));
 	check(Archive);
 
-	const FString HeaderRow = TEXT("ID_CO,ID_COI,ID_UpdateType,ID_UpdateResult,Time_Queue,Time_Update,Time_TaskGetMesh,Time_TaskLockCache,Time_TaskGetImages,Time_TaskConvertResources,Time_TaskCallbacks,Time_TaskUpdateImages");
+	const FString HeaderRow = TEXT("ID_CO,ID_COI,ID_UpdateType,ID_UpdateResult,Time_Queue,Time_Update,Time_TaskGetMesh,Time_TaskLockCache,Time_TaskGetImages,Time_TaskConvertResources,Time_TaskCallbacks,Memory_Update,Memory_Update_Real,Time_TaskUpdateImage,Memory_TaskUpdateImage,Memory_TaskUpdateImage_Real");
 	LogBenchmarkUtil::Write(*Archive, HeaderRow);
 
 	return Archive;
@@ -308,13 +308,16 @@ void FLogBenchmarkUtil::FinishUpdateMesh(const TSharedRef<FUpdateContextPrivate>
 	const double Time_TaskConvertResources = Context->TaskConvertResourcesTime * 1000;
 	const double Time_TaskCallbacks =  Context->TaskCallbacksTime * 1000;
 
-	const FString UpdateString = FString::Printf(TEXT("%s,%s,%s,%s,%f,%f,%f,%f,%f,%f,%f"), *ID_CO, *ID_COI, *ID_UpdateType, *ID_UpdateResult, Time_Queue, Time_Update, Time_TaskGetMesh, Time_TaskLockCache, Time_TaskGetImages, Time_TaskConvertResources, Time_TaskCallbacks);
+	const double Memory_UpdateEndPeakMB = (Context->UpdateEndPeakBytes / 1024.0) / 1024.0;
+	const double Memory_UpdateEndRealPeakMB = (Context->UpdateEndRealPeakBytes / 1024.0) / 1024.0;
+	
+	const FString UpdateString = FString::Printf(TEXT("%s,%s,%s,%s,%f,%f,%f,%f,%f,%f,%f,%f,%f"), *ID_CO, *ID_COI, *ID_UpdateType, *ID_UpdateResult, Time_Queue, Time_Update, Time_TaskGetMesh, Time_TaskLockCache, Time_TaskGetImages, Time_TaskConvertResources, Time_TaskCallbacks, Memory_UpdateEndPeakMB, Memory_UpdateEndRealPeakMB);
 	LogBenchmarkUtil::Write(*Archive, UpdateString);
 	Archive->Flush();
 }
 
 
-void FLogBenchmarkUtil::FinishUpdateImage(const FString& CustomizableObjectPathName, const FString& InstancePathName, const double TaskUpdateImageTime) const
+void FLogBenchmarkUtil::FinishUpdateImage(const FString& CustomizableObjectPathName, const FString& InstancePathName, const double TaskUpdateImageTime, const int64 TaskUpdateImageMemoryPeak, const int64 TaskUpdateImageRealMemoryPeak) const
 {
 	check(IsInGameThread());
 
@@ -327,8 +330,11 @@ void FLogBenchmarkUtil::FinishUpdateImage(const FString& CustomizableObjectPathN
 	const FString& ID_COI = InstancePathName;
 	const FString ID_UpdateType = TEXT("Image");
 	const double Time_TaskUpdateImage = TaskUpdateImageTime * 1000;
-		
-	const FString UpdateString = FString::Printf(TEXT("%s,%s,%s,,,,,,,,,%f"), *ID_CO, *ID_COI, *ID_UpdateType, Time_TaskUpdateImage);
+
+	const double Memory_TaskUpdateImagePeakMB = (TaskUpdateImageMemoryPeak / 1024.0) / 1024.0;
+	const double Memory_TaskUpdateImageRealPeakMB = (TaskUpdateImageRealMemoryPeak / 1024.0) / 1024.0;
+
+	const FString UpdateString = FString::Printf(TEXT("%s,%s,%s,,,,,,,,,,,%f,%f,%f"), *ID_CO, *ID_COI, *ID_UpdateType, Time_TaskUpdateImage, Memory_TaskUpdateImagePeakMB,Memory_TaskUpdateImageRealPeakMB);
 	LogBenchmarkUtil::Write(*Archive, UpdateString);
 	Archive->Flush();	
 }
diff --git a/Engine/Plugins/Experimental/Mutable/Source/CustomizableObject/Private/MuCO/LogBenchmarkUtil.h b/Engine/Plugins/Experimental/Mutable/Source/CustomizableObject/Private/MuCO/LogBenchmarkUtil.h
--- a/Engine/Plugins/Experimental/Mutable/Source/CustomizableObject/Private/MuCO/LogBenchmarkUtil.h
+++ b/Engine/Plugins/Experimental/Mutable/Source/CustomizableObject/Private/MuCO/LogBenchmarkUtil.h
@@ -86,7 +86,7 @@ public:
 	/** Gathers update stats when it has finished. */
 	void FinishUpdateMesh(const TSharedRef<FUpdateContextPrivate>& Context);
 
-	void FinishUpdateImage(const FString& CustomizableObjectPathName, const FString& InstancePathName, double TaskUpdateImageTime) const;
+	void FinishUpdateImage(const FString& CustomizableObjectPathName, const FString& InstancePathName, double TaskUpdateImageTime, const int64 TaskUpdateImageMemoryPeak, const int64 TaskUpdateImageRealMemoryPeak) const;
 	
 	DECLARE_BENCHMARK_STAT(NumAllocatedTextures, uint32);
 	DECLARE_BENCHMARK_STAT(NumAllocatedSkeletalMeshes, int32);
diff --git a/Engine/Plugins/Experimental/Mutable/Source/CustomizableObject/Public/MuCO/CustomizableObjectMipDataProvider.h b/Engine/Plugins/Experimental/Mutable/Source/CustomizableObject/Public/MuCO/CustomizableObjectMipDataProvider.h
--- a/Engine/Plugins/Experimental/Mutable/Source/CustomizableObject/Public/MuCO/CustomizableObjectMipDataProvider.h
+++ b/Engine/Plugins/Experimental/Mutable/Source/CustomizableObject/Public/MuCO/CustomizableObjectMipDataProvider.h
@@ -96,6 +96,9 @@ struct FMutableImageOperationData
 
 	/** Access to the Counter must be protected with this because it may be accessed from another thread to null it. */
 	FCriticalSection CounterTaskLock;
+
+	// Image Update Memory stats
+	int64 ImageUpdateStartBytes = 0;
 };
 
 
