Lower issues moving objects back to original packages to warnings until better fixes can be made.

diff --git a/Engine/Plugins/Experimental/Mutable/Source/CustomizableObjectEditor/Private/MuCOE/CustomizableObjectCookPackageSplitter.cpp b/Engine/Plugins/Experimental/Mutable/Source/CustomizableObjectEditor/Private/MuCOE/CustomizableObjectCookPackageSplitter.cpp
--- a/Engine/Plugins/Experimental/Mutable/Source/CustomizableObjectEditor/Private/MuCOE/CustomizableObjectCookPackageSplitter.cpp
+++ b/Engine/Plugins/Experimental/Mutable/Source/CustomizableObjectEditor/Private/MuCOE/CustomizableObjectCookPackageSplitter.cpp
@@ -37,8 +37,28 @@ FCustomizableObjectStreamedResourceData* FindStreamedResourceData(
 			});
 }
 
+enum class EMoveContainerError
+{
+	None,
+	FailedToLoadContainer,
+	NameCollision, // Object with that name already exists in the new outer
+	RenameFailed,
+};
+
+const TCHAR* LexToString(EMoveContainerError Error)
+{
+	switch (Error)
+	{
+		case EMoveContainerError::None: return TEXT("None");
+		case EMoveContainerError::FailedToLoadContainer: return TEXT("FailedToLoadContainer");
+		case EMoveContainerError::NameCollision: return TEXT("NameCollision");
+		case EMoveContainerError::RenameFailed: return TEXT("RenameFailed");
+		default: return TEXT("Unknown");
+	}
+}
+
 // Moves the StreamedResourceData's data container to the given Outer.
-bool MoveContainerToNewOuter(
+EMoveContainerError MoveContainerToNewOuter(
 	UObject* NewOuter,
 	const FCustomizableObjectStreamedResourceData* StreamedResourceData,
 	UCustomizableObjectResourceDataContainer*& OutContainer
@@ -51,29 +71,26 @@ bool MoveContainerToNewOuter(
 	UCustomizableObjectResourceDataContainer* Container = StreamedResourceData->GetPath().LoadSynchronous();
 	if (!Container)
 	{
-		UE_LOG(LogMutable, Error, TEXT("Failed to load streamed Resource Data container %s"),
-			*StreamedResourceData->GetPath().ToString());
-
-		return false;
+		return EMoveContainerError::FailedToLoadContainer;
 	}
 
 	if (Container->GetOuter() != NewOuter)
 	{
 		// Ensure the target object doesn't exist
-		check(!FindObject<UObject>(NewOuter, *Container->GetName()));
+		if(FindObject<UObject>(NewOuter, *Container->GetName()))
+		{
+			return EMoveContainerError::NameCollision;
+		}
 
 		// The Rename function moves the object into the given package
 		if (!Container->Rename(nullptr, NewOuter, REN_DontCreateRedirectors))
 		{
-			UE_LOG(LogMutable, Error, TEXT("Failed to move streamed Resource Data container %s into Outer %s"),
-				*Container->GetPathName(), *NewOuter->GetPathName());
-
-			return false;
+			return EMoveContainerError::RenameFailed;
 		}
 	}
 
 	OutContainer = Container;
-	return true;
+	return EMoveContainerError::None;
 }
 
 
@@ -195,8 +212,10 @@ bool FCustomizableObjectCookPackageSplitter::PreSaveGeneratorPackage(
 
 		// Move the streamed data to the generated package
 		UCustomizableObjectResourceDataContainer* Container = nullptr;
-		if (!MoveContainerToNewOuter(GeneratedPackage.Package, FoundData, Container))
+		EMoveContainerError Error = MoveContainerToNewOuter(GeneratedPackage.Package, FoundData, Container);
+		if (Error != EMoveContainerError::None)
 		{
+			UE_LOG(LogMutable, Error, TEXT("Failed to move container %s to new outer %s - %s"), *FoundData->GetPath().ToSoftObjectPath().ToString(), *GetPathNameSafe(GeneratedPackage.Package), LexToString(Error));
 			return false;
 		}
 
@@ -268,7 +287,8 @@ void FCustomizableObjectCookPackageSplitter::PostSaveGeneratorPackage(UPackage*
 		}
 
 		UCustomizableObjectResourceDataContainer* Container = nullptr;
-		MoveContainerToNewOuter(OwnerObject, ResourceData, Container);
+		EMoveContainerError Error = MoveContainerToNewOuter(OwnerObject, ResourceData, Container);
+		UE_CLOG(Error != EMoveContainerError::None, LogMutable, Warning, TEXT("Failed to move container %s back to %s - %s"), *ContainerName, *GetPathNameSafe(OwnerObject), LexToString(Error));
 
 		NewArray.Emplace(Container);
 	}
@@ -289,7 +309,8 @@ void FCustomizableObjectCookPackageSplitter::PostSaveGeneratorPackage(UPackage*
 		}
 
 		UCustomizableObjectResourceDataContainer* Container = nullptr;
-		MoveContainerToNewOuter(OwnerObject, ResourceData, Container);
+		EMoveContainerError Error = MoveContainerToNewOuter(OwnerObject, ResourceData, Container);
+		UE_CLOG(Error != EMoveContainerError::None, LogMutable, Warning, TEXT("Failed to move container %s back to %s - %s"), *ContainerName, *GetPathNameSafe(OwnerObject), LexToString(Error));
 
 		NewArray.Emplace(Container);
 	}
@@ -323,8 +344,10 @@ bool FCustomizableObjectCookPackageSplitter::PopulateGeneratedPackage(
 	}
 
 	UCustomizableObjectResourceDataContainer* Container = nullptr;
-	if (!MoveContainerToNewOuter(GeneratedPackage.Package, ResourceData, Container))
+	EMoveContainerError Error = MoveContainerToNewOuter(GeneratedPackage.Package, ResourceData, Container);
+	if (Error != EMoveContainerError::None)
 	{
+		UE_LOG(LogMutable, Error, TEXT("Failed to move container %s to new outer %s - %s"), *ResourceData->GetPath().ToSoftObjectPath().ToString(), *GetPathNameSafe(GeneratedPackage.Package), LexToString(Error));
 		return false;
 	}
 
@@ -358,8 +381,9 @@ void FCustomizableObjectCookPackageSplitter::PostSaveGeneratedPackage(
 	}
 
 	UCustomizableObjectResourceDataContainer* Container = nullptr;
-	MoveContainerToNewOuter(OwnerObject, ResourceData, Container);
-
+	EMoveContainerError Error = MoveContainerToNewOuter(OwnerObject, ResourceData, Container);
+	UE_CLOG(Error != EMoveContainerError::None, LogMutable, Warning, 
+		TEXT("Failed to move container %s back to %s - %s"), *ResourceData->GetPath().ToSoftObjectPath().ToString(), *GetPathNameSafe(OwnerObject), LexToString(Error));
 }
 
 void FCustomizableObjectCookPackageSplitter::Teardown(ETeardown Status)
