[mutable] Made MutableValidationCommandlet compile always at max optimization.
- Also forced texture compression to Fast (to override the user config saved in the CO)

diff --git a/Engine/Plugins/Experimental/Mutable/Source/MutableValidation/Private/MuV/CustomizableObjectValidationCommandlet.cpp b/Engine/Plugins/Experimental/Mutable/Source/MutableValidation/Private/MuV/CustomizableObjectValidationCommandlet.cpp
--- a/Engine/Plugins/Experimental/Mutable/Source/MutableValidation/Private/MuV/CustomizableObjectValidationCommandlet.cpp
+++ b/Engine/Plugins/Experimental/Mutable/Source/MutableValidation/Private/MuV/CustomizableObjectValidationCommandlet.cpp
@@ -60,10 +60,22 @@ int32 UCustomizableObjectValidationCommandlet::Main(const FString& Params)
     	// Request a compiler to be able to locate the root and to compile it
     	const TUniquePtr<FCustomizableObjectCompilerBase> Compiler =
     		TUniquePtr<FCustomizableObjectCompilerBase>(UCustomizableObjectSystem::GetInstance()->GetNewCompiler());
-    		
+
+		// Override some configurations that may have been changed by the user
+		FCompilationOptions CompilationOptions = ToTestCustomizableObject->CompileOptions;
+		CompilationOptions.bSilentCompilation = false;
+		CompilationOptions.OptimizationLevel = 2;			// Set the optimization level to the max
+		CompilationOptions.TextureCompression = ECustomizableObjectTextureCompression::Fast;
+		
+		// TODO: Add logs for the other relevant configs of the model being compiled
+		// Print MTU parseable logs
+		UE_LOG(LogMutable, Log, TEXT("(int) model_optimization_level : %d "), CompilationOptions.OptimizationLevel);
+		UE_LOG(LogMutable, Log, TEXT("(string) model_texture_compression : %s "), *UEnum::GetValueAsString(CompilationOptions.TextureCompression));
+		UE_LOG(LogMutable, Log, TEXT("(string) model_disk_compilation : %s "), CompilationOptions.bUseDiskCompilation ? TEXT("true") : TEXT("false"));
+
     	// Compile the CO with the provided compilation options
     	// Run Sync compilation -> Warning : Potentially long operation -------------
-    	Compiler->Compile(*ToTestCustomizableObject, ToTestCustomizableObject->CompileOptions, false);
+    	Compiler->Compile(*ToTestCustomizableObject, CompilationOptions, false);
     	// --------------------------------------------------------------------------
 		
     	// Get the compilation result
@@ -89,7 +101,8 @@ int32 UCustomizableObjectValidationCommandlet::Main(const FString& Params)
 				const uint32 RomByteSize = MutableModel->GetRomSize(RomIndex);
 				TotalRomSizeBytes += RomByteSize;
 			}
-			
+
+			// Print MTU parseable logs
 			UE_LOG(LogMutable, Log,TEXT("(int) model_rom_count : %d "), RomCount);
 			UE_LOG(LogMutable, Log,TEXT("(int) model_roms_size : %lld "), TotalRomSizeBytes);
 		}
