[Mutable] CO package splitter now duplicates the streamed resources instead of moving them to a new outer to avoid a check that fails due to an engine issue.

#rnx

diff --git a/Engine/Plugins/Experimental/Mutable/Source/CustomizableObjectEditor/Private/MuCOE/CustomizableObjectCookPackageSplitter.cpp b/Engine/Plugins/Experimental/Mutable/Source/CustomizableObjectEditor/Private/MuCOE/CustomizableObjectCookPackageSplitter.cpp
--- a/Engine/Plugins/Experimental/Mutable/Source/CustomizableObjectEditor/Private/MuCOE/CustomizableObjectCookPackageSplitter.cpp
+++ b/Engine/Plugins/Experimental/Mutable/Source/CustomizableObjectEditor/Private/MuCOE/CustomizableObjectCookPackageSplitter.cpp
@@ -343,8 +343,18 @@ bool FCustomizableObjectCookPackageSplitter::PopulateGeneratedPackage(
 		return false;
 	}
 
-	UCustomizableObjectResourceDataContainer* Container = nullptr;
-	EMoveContainerError Error = MoveContainerToNewOuter(GeneratedPackage.Package, ResourceData, Container);
+	// [TEMP] Loading a package referencing the CO before PostSaveGeneratedPackage is called causes a name collision.
+	// Duplicate the object with the new outer instead of moving it until it is fixed.
+	UObject* Container = ResourceData->GetPath().LoadSynchronous();
+	EMoveContainerError Error = Container ? EMoveContainerError::None : EMoveContainerError::FailedToLoadContainer;
+	if (Container)
+	{
+		Container = StaticDuplicateObject(Container, GeneratedPackage.Package);
+	}
+
+	//UCustomizableObjectResourceDataContainer* Container = nullptr;
+	//EMoveContainerError Error = MoveContainerToNewOuter(GeneratedPackage.Package, ResourceData, Container);
+
 	if (Error != EMoveContainerError::None)
 	{
 		UE_LOG(LogMutable, Error, TEXT("Failed to move container %s to new outer %s - %s"), *ResourceData->GetPath().ToSoftObjectPath().ToString(), *GetPathNameSafe(GeneratedPackage.Package), LexToString(Error));
