[Mutable] Added image descriptor to table nodes. It is generate using the data of the reference texture of each column.

#rnx

diff --git a/Engine/Plugins/Experimental/Mutable/Source/CustomizableObjectEditor/Private/MuCOE/GenerateMutableSource/GenerateMutableSource.cpp b/Engine/Plugins/Experimental/Mutable/Source/CustomizableObjectEditor/Private/MuCOE/GenerateMutableSource/GenerateMutableSource.cpp
--- a/Engine/Plugins/Experimental/Mutable/Source/CustomizableObjectEditor/Private/MuCOE/GenerateMutableSource/GenerateMutableSource.cpp
+++ b/Engine/Plugins/Experimental/Mutable/Source/CustomizableObjectEditor/Private/MuCOE/GenerateMutableSource/GenerateMutableSource.cpp
@@ -1676,6 +1676,36 @@ int32 GetMaxTextureSize(const UTexture2D* ReferenceTexture, const FMutableGraphG
 }
 
 
+mu::FImageDesc GenerateImageDescriptor(UTexture* Texture)
+{
+	check(Texture);
+	mu::FImageDesc ImageDesc;
+
+	ImageDesc.m_size[0] = Texture->Source.GetSizeX();
+	ImageDesc.m_size[1] = Texture->Source.GetSizeY();
+	ImageDesc.m_lods = Texture->Source.GetNumMips();
+
+	mu::EImageFormat MutableFormat = mu::EImageFormat::IF_RGBA_UBYTE;
+	ETextureSourceFormat SourceFormat = Texture->Source.GetFormat();
+	switch (SourceFormat)
+	{
+	case ETextureSourceFormat::TSF_G8:
+	case ETextureSourceFormat::TSF_G16:
+	case ETextureSourceFormat::TSF_R16F:
+	case ETextureSourceFormat::TSF_R32F:
+		MutableFormat = mu::EImageFormat::IF_L_UBYTE;
+		break;
+
+	default:
+		break;
+	}
+
+	ImageDesc.m_format = MutableFormat;
+
+	return ImageDesc;
+}
+
+
 mu::Ptr<mu::Image> GenerateImageConstant(UTexture* Texture, FMutableGraphGenerationContext& GenerationContext, bool bIsReference)
 {
 	if (!Texture)
@@ -1714,29 +1744,7 @@ mu::Ptr<mu::Image> GenerateImageConstant(UTexture* Texture, FMutableGraphGenerat
 	// Create a descriptor for the image.
 	// \TODO: If passthrough (bIsReference) we should apply lod bias, and max texture size to this desc.
 	// For now it is not a problem because passthrough textures shouldn't mix with any other operations.
-	mu::FImageDesc ImageDesc;
-	{
-		ImageDesc.m_size[0] = Texture->Source.GetSizeX();
-		ImageDesc.m_size[1] = Texture->Source.GetSizeY();
-		ImageDesc.m_lods = Texture->Source.GetNumMips();
-
-		mu::EImageFormat MutableFormat = mu::EImageFormat::IF_RGBA_UBYTE;
-		ETextureSourceFormat SourceFormat = Texture->Source.GetFormat();
-		switch (SourceFormat)
-		{
-		case ETextureSourceFormat::TSF_G8:
-		case ETextureSourceFormat::TSF_G16:
-		case ETextureSourceFormat::TSF_R16F:
-		case ETextureSourceFormat::TSF_R32F:
-			MutableFormat = mu::EImageFormat::IF_L_UBYTE;
-			break;
-
-		default:
-			break;
-		}
-
-		ImageDesc.m_format = MutableFormat;
-	}
+	mu::FImageDesc ImageDesc = GenerateImageDescriptor(Texture);
 
 	// Compile-time references that are left should be resolved immediately (should only happen in editor).
 	mu::Ptr<mu::Image> Result = mu::Image::CreateAsReference(Entry.ID, ImageDesc, bForceLoad);
diff --git a/Engine/Plugins/Experimental/Mutable/Source/CustomizableObjectEditor/Private/MuCOE/GenerateMutableSource/GenerateMutableSource.h b/Engine/Plugins/Experimental/Mutable/Source/CustomizableObjectEditor/Private/MuCOE/GenerateMutableSource/GenerateMutableSource.h
--- a/Engine/Plugins/Experimental/Mutable/Source/CustomizableObjectEditor/Private/MuCOE/GenerateMutableSource/GenerateMutableSource.h
+++ b/Engine/Plugins/Experimental/Mutable/Source/CustomizableObjectEditor/Private/MuCOE/GenerateMutableSource/GenerateMutableSource.h
@@ -847,3 +847,7 @@ int32 ComputeLODBias(const FMutableGraphGenerationContext& GenerationContext, co
 int32 GetMaxTextureSize(const UTexture2D* ReferenceTexture, const FMutableGraphGenerationContext& GenerationContext);
 
 mu::Ptr<mu::Image> GenerateImageConstant( UTexture*, FMutableGraphGenerationContext&, bool bIsReference);
+
+/** Generates a mutable image descriptor from an unreal engine texture */
+mu::FImageDesc GenerateImageDescriptor(UTexture* Texture);
+
diff --git a/Engine/Plugins/Experimental/Mutable/Source/CustomizableObjectEditor/Private/MuCOE/GenerateMutableSource/GenerateMutableSourceImage.cpp b/Engine/Plugins/Experimental/Mutable/Source/CustomizableObjectEditor/Private/MuCOE/GenerateMutableSource/GenerateMutableSourceImage.cpp
--- a/Engine/Plugins/Experimental/Mutable/Source/CustomizableObjectEditor/Private/MuCOE/GenerateMutableSource/GenerateMutableSourceImage.cpp
+++ b/Engine/Plugins/Experimental/Mutable/Source/CustomizableObjectEditor/Private/MuCOE/GenerateMutableSource/GenerateMutableSourceImage.cpp
@@ -1077,6 +1077,7 @@ mu::NodeImagePtr GenerateMutableSourceImage(const UEdGraphPin* Pin, FMutableGrap
 							{
 								int32 DefaultMaxTextureSize = GetMaxTextureSize(DefaultTexture2D, GenerationContext);
 								ImageTableNode->SetMaxTextureSize(FMath::Max(DefaultTexture2D->Source.GetSizeX(), DefaultTexture2D->Source.GetSizeY()));
+								ImageTableNode->SetReferenceImageDescriptor(GenerateImageDescriptor(DefaultTexture2D));
 							}
 							
 							GenerationContext.AddParameterNameUnique(Node, TypedNodeTable->ParameterName);
diff --git a/Engine/Plugins/Experimental/Mutable/Source/MutableTools/Private/MuT/CodeGenerator_Image.cpp b/Engine/Plugins/Experimental/Mutable/Source/MutableTools/Private/MuT/CodeGenerator_Image.cpp
--- a/Engine/Plugins/Experimental/Mutable/Source/MutableTools/Private/MuT/CodeGenerator_Image.cpp
+++ b/Engine/Plugins/Experimental/Mutable/Source/MutableTools/Private/MuT/CodeGenerator_Image.cpp
@@ -1721,9 +1721,11 @@ namespace mu
 					ImageOp = Result.op;
 				}
 
-				if (node.MaxTextureSize > 0 && (node.MaxTextureSize < pImage->GetSizeX() || node.MaxTextureSize < pImage->GetSizeY()))
+				int32 MaxTextureSize = FMath::Max(node.ReferenceImageDesc.m_size[0], node.ReferenceImageDesc.m_size[1]);
+
+				if (MaxTextureSize > 0 && (MaxTextureSize < pImage->GetSizeX() || MaxTextureSize < pImage->GetSizeY()))
 				{
-					float Factor = FMath::Min(node.MaxTextureSize / (float)(pImage->GetSizeX()), node.MaxTextureSize / (float)(pImage->GetSizeY()));
+					float Factor = FMath::Min(MaxTextureSize / (float)(pImage->GetSizeX()), MaxTextureSize / (float)(pImage->GetSizeY()));
 					Ptr<ASTOpFixed> op = new ASTOpFixed();
 					op->op.type = OP_TYPE::IM_RESIZE;
 					op->op.args.ImageResize.size[0] = (uint16)pImage->GetSizeX() * Factor;
diff --git a/Engine/Plugins/Experimental/Mutable/Source/MutableTools/Private/MuT/ImageDescGenerator.cpp b/Engine/Plugins/Experimental/Mutable/Source/MutableTools/Private/MuT/ImageDescGenerator.cpp
--- a/Engine/Plugins/Experimental/Mutable/Source/MutableTools/Private/MuT/ImageDescGenerator.cpp
+++ b/Engine/Plugins/Experimental/Mutable/Source/MutableTools/Private/MuT/ImageDescGenerator.cpp
@@ -119,20 +119,7 @@ namespace mu
             return 0;
         }
 
-        int32 i = 0;
-        while ( i<node.m_pTable->GetPrivate()->Rows.Num() )
-        {
-			Ptr<const Image> pImage = node.m_pTable->GetPrivate()->Rows[i].Values[ colIndex ].ProxyImage->Get();
-            ++i;
-
-			if (pImage)
-			{
-				m_desc.m_size[0] = FMath::Max(m_desc.m_size[0], pImage->GetSizeX());
-				m_desc.m_size[1] = FMath::Max(m_desc.m_size[1], pImage->GetSizeX());
-				m_desc.m_lods = FMath::Max(m_desc.m_lods, uint8(pImage->GetLODCount()));
-				m_desc.m_format = mu::GetMostGenericFormat(m_desc.m_format, pImage->GetFormat());
-			}
-        }
+		m_desc = node.ReferenceImageDesc;
 
         return 0;
     }
diff --git a/Engine/Plugins/Experimental/Mutable/Source/MutableTools/Private/MuT/NodeImageTable.cpp b/Engine/Plugins/Experimental/Mutable/Source/MutableTools/Private/MuT/NodeImageTable.cpp
--- a/Engine/Plugins/Experimental/Mutable/Source/MutableTools/Private/MuT/NodeImageTable.cpp
+++ b/Engine/Plugins/Experimental/Mutable/Source/MutableTools/Private/MuT/NodeImageTable.cpp
@@ -70,6 +70,13 @@ namespace mu
 		return m_pD->MaxTextureSize;
 	}
 
+
+	//---------------------------------------------------------------------------------------------
+	void NodeImageTable::SetReferenceImageDescriptor(const FImageDesc& ImageDesc)
+	{
+		m_pD->ReferenceImageDesc = ImageDesc;
+	}
+
 }
 
 
diff --git a/Engine/Plugins/Experimental/Mutable/Source/MutableTools/Private/MuT/NodeImageTablePrivate.h b/Engine/Plugins/Experimental/Mutable/Source/MutableTools/Private/MuT/NodeImageTablePrivate.h
--- a/Engine/Plugins/Experimental/Mutable/Source/MutableTools/Private/MuT/NodeImageTablePrivate.h
+++ b/Engine/Plugins/Experimental/Mutable/Source/MutableTools/Private/MuT/NodeImageTablePrivate.h
@@ -26,17 +26,25 @@ namespace mu
 		TablePtr m_pTable;
 		FString m_columnName;
 		uint16 MaxTextureSize = 0;
+		FImageDesc ReferenceImageDesc;
 
 		//!
 		void Serialise( OutputArchive& arch ) const
 		{
-            uint32_t ver = 3;
+            uint32_t ver = 4;
 			arch << ver;
 
 			arch << m_parameterName;
 			arch << m_pTable;
 			arch << m_columnName;
 			arch << MaxTextureSize;
+			
+			// FImageDesc
+			{
+				arch << ReferenceImageDesc.m_size;
+				arch << uint8(ReferenceImageDesc.m_format);
+				arch << ReferenceImageDesc.m_lods;
+			}
 		}
 
 		//!
@@ -44,7 +52,7 @@ namespace mu
 		{
             uint32_t ver;
 			arch >> ver;
-			check(ver>=1 && ver<= 3);
+			check(ver>=1 && ver<= 4);
 
 			if (ver == 1)
 			{
@@ -74,6 +82,17 @@ namespace mu
 			{
 				arch >> MaxTextureSize;
 			}
+
+			if (ver >= 4)
+			{
+				arch >> ReferenceImageDesc.m_size;
+				
+				uint8 AuxFromat = 0;
+				arch >> AuxFromat;
+				ReferenceImageDesc.m_format = EImageFormat(AuxFromat);
+				
+				arch >> ReferenceImageDesc.m_lods;
+			}
 		}
 	};
 
diff --git a/Engine/Plugins/Experimental/Mutable/Source/MutableTools/Public/MuT/NodeImageTable.h b/Engine/Plugins/Experimental/Mutable/Source/MutableTools/Public/MuT/NodeImageTable.h
--- a/Engine/Plugins/Experimental/Mutable/Source/MutableTools/Public/MuT/NodeImageTable.h
+++ b/Engine/Plugins/Experimental/Mutable/Source/MutableTools/Public/MuT/NodeImageTable.h
@@ -22,6 +22,7 @@ namespace mu
 
 	class InputArchive;
 	class OutputArchive;
+	struct FImageDesc;
 
 
 	//! This node provides the meshes stored in the column of a table.
@@ -66,6 +67,9 @@ namespace mu
 		void SetMaxTextureSize(uint16 Size);
 		uint16 GetMaxTextureSize();
 
+		//!
+		void SetReferenceImageDescriptor(const FImageDesc& Descriptor);
+
 		//-----------------------------------------------------------------------------------------
 		// Interface pattern
 		//-----------------------------------------------------------------------------------------
