#rnx
#tests PIE

diff --git a/Engine/Plugins/Experimental/Mutable/Source/CustomizableObjectEditor/Classes/MuCOE/CustomizableObjectLayout.h b/Engine/Plugins/Experimental/Mutable/Source/CustomizableObjectEditor/Classes/MuCOE/CustomizableObjectLayout.h
--- a/Engine/Plugins/Experimental/Mutable/Source/CustomizableObjectEditor/Classes/MuCOE/CustomizableObjectLayout.h
+++ b/Engine/Plugins/Experimental/Mutable/Source/CustomizableObjectEditor/Classes/MuCOE/CustomizableObjectLayout.h
@@ -3,17 +3,23 @@
 #pragma once
 
 
+#include "MuT/NodeLayout.h"
 #include "CustomizableObjectLayout.generated.h"
 
+
 UENUM()
 enum class ECustomizableObjectTextureLayoutPackingStrategy : uint8
 {
 	// The layout increases its size to fit all the blocks.
 	Resizable = 0 UMETA(DisplayName = "Resizable Layout"),
 	// The layout resizes the blocks to keep its size.
-	Fixed = 1 UMETA(DisplayName = "Fixed Layout")
+	Fixed = 1 UMETA(DisplayName = "Fixed Layout"),
+	// The layout is not modified and blocks are ignored. Extend material nodes just add their layouts on top of the base one.
+	Overlay = 2 UMETA(DisplayName = "Overlay Layout")
 };
 
+mu::EPackStrategy ConvertLayoutStrategy(const ECustomizableObjectTextureLayoutPackingStrategy LayoutPackStrategy);
+
 // Fixed Layout reduction methods
 UENUM()
 enum class ECustomizableObjectLayoutBlockReductionMethod : uint8
diff --git a/Engine/Plugins/Experimental/Mutable/Source/CustomizableObjectEditor/Private/MuCOE/CustomizableObjectLayout.cpp b/Engine/Plugins/Experimental/Mutable/Source/CustomizableObjectEditor/Private/MuCOE/CustomizableObjectLayout.cpp
--- a/Engine/Plugins/Experimental/Mutable/Source/CustomizableObjectEditor/Private/MuCOE/CustomizableObjectLayout.cpp
+++ b/Engine/Plugins/Experimental/Mutable/Source/CustomizableObjectEditor/Private/MuCOE/CustomizableObjectLayout.cpp
@@ -10,7 +10,6 @@
 #include "MuCOE/Nodes/CustomizableObjectNodeLayoutBlocks.h"
 #include "MuCOE/Nodes/CustomizableObjectNodeMesh.h"
 #include "MuCOE/Nodes/CustomizableObjectNodeTable.h"
-#include "MuT/NodeLayout.h"
 
 #define LOCTEXT_NAMESPACE "CustomizableObjectEditor"
 
@@ -42,6 +41,32 @@ void UCustomizableObjectLayout::SetPackingStrategy(ECustomizableObjectTextureLay
 }
 
 
+mu::EPackStrategy ConvertLayoutStrategy(const ECustomizableObjectTextureLayoutPackingStrategy LayoutPackStrategy)
+{
+	mu::EPackStrategy PackStrategy = mu::EPackStrategy::FIXED_LAYOUT;
+
+	switch (LayoutPackStrategy)
+	{
+	case ECustomizableObjectTextureLayoutPackingStrategy::Fixed:
+		PackStrategy = mu::EPackStrategy::FIXED_LAYOUT;
+		break;
+
+	case ECustomizableObjectTextureLayoutPackingStrategy::Resizable:
+		PackStrategy = mu::EPackStrategy::RESIZABLE_LAYOUT;
+		break;
+
+	case ECustomizableObjectTextureLayoutPackingStrategy::Overlay:
+		PackStrategy = mu::EPackStrategy::OVERLAY_LAYOUT;
+		break;
+
+	default:
+		checkNoEntry();
+	}
+
+	return PackStrategy;
+}
+
+
 void UCustomizableObjectLayout::SetGridSize(FIntPoint Size)
 {
 	GridSize = Size;
diff --git a/Engine/Plugins/Experimental/Mutable/Source/CustomizableObjectEditor/Private/MuCOE/GenerateMutableSource/GenerateMutableSourceLayout.cpp b/Engine/Plugins/Experimental/Mutable/Source/CustomizableObjectEditor/Private/MuCOE/GenerateMutableSource/GenerateMutableSourceLayout.cpp
--- a/Engine/Plugins/Experimental/Mutable/Source/CustomizableObjectEditor/Private/MuCOE/GenerateMutableSource/GenerateMutableSourceLayout.cpp
+++ b/Engine/Plugins/Experimental/Mutable/Source/CustomizableObjectEditor/Private/MuCOE/GenerateMutableSource/GenerateMutableSourceLayout.cpp
@@ -58,7 +58,10 @@ mu::NodeLayoutPtr GenerateMutableSourceLayout(const UEdGraphPin * Pin, FMutableG
 		LayoutNode->SetGridSize(TypedNodeBlocks->Layout->GetGridSize().X, TypedNodeBlocks->Layout->GetGridSize().Y);
 		LayoutNode->SetMaxGridSize(TypedNodeBlocks->Layout->GetMaxGridSize().X, TypedNodeBlocks->Layout->GetMaxGridSize().Y);
 		LayoutNode->SetBlockCount(TypedNodeBlocks->Layout->Blocks.Num() ? TypedNodeBlocks->Layout->Blocks.Num() : 1);
-		LayoutNode->SetLayoutPackingStrategy(TypedNodeBlocks->Layout->GetPackingStrategy() == ECustomizableObjectTextureLayoutPackingStrategy::Fixed ? mu::EPackStrategy::FIXED_LAYOUT : mu::EPackStrategy::RESIZABLE_LAYOUT);
+
+		mu::EPackStrategy PackStrategy = ConvertLayoutStrategy(TypedNodeBlocks->Layout->GetPackingStrategy());
+		LayoutNode->SetLayoutPackingStrategy(PackStrategy);
+
 		LayoutNode->SetBlockReductionMethod(TypedNodeBlocks->Layout->GetBlockReductionMethod() == ECustomizableObjectLayoutBlockReductionMethod::Halve ? mu::EReductionMethod::HALVE_REDUCTION : mu::EReductionMethod::UNITARY_REDUCTION);
 
 		if (bLinkedToExtendMaterial)
diff --git a/Engine/Plugins/Experimental/Mutable/Source/CustomizableObjectEditor/Private/MuCOE/GenerateMutableSource/GenerateMutableSourceMesh.cpp b/Engine/Plugins/Experimental/Mutable/Source/CustomizableObjectEditor/Private/MuCOE/GenerateMutableSource/GenerateMutableSourceMesh.cpp
--- a/Engine/Plugins/Experimental/Mutable/Source/CustomizableObjectEditor/Private/MuCOE/GenerateMutableSource/GenerateMutableSourceMesh.cpp
+++ b/Engine/Plugins/Experimental/Mutable/Source/CustomizableObjectEditor/Private/MuCOE/GenerateMutableSource/GenerateMutableSourceMesh.cpp
@@ -3725,7 +3725,10 @@ mu::NodeMeshPtr GenerateMutableSourceMesh(const UEdGraphPin* Pin,
 								LayoutNode->SetGridSize(Layouts[i]->GetGridSize().X, Layouts[i]->GetGridSize().Y);
 								LayoutNode->SetMaxGridSize(Layouts[i]->GetMaxGridSize().X, Layouts[i]->GetMaxGridSize().Y);
 								LayoutNode->SetBlockCount(Layouts[i]->Blocks.Num() ? Layouts[i]->Blocks.Num() : 1);
-								LayoutNode->SetLayoutPackingStrategy(Layouts[i]->GetPackingStrategy() == ECustomizableObjectTextureLayoutPackingStrategy::Fixed ? mu::EPackStrategy::FIXED_LAYOUT : mu::EPackStrategy::RESIZABLE_LAYOUT);
+
+								mu::EPackStrategy PackStrategy = ConvertLayoutStrategy(Layouts[i]->GetPackingStrategy());
+								LayoutNode->SetLayoutPackingStrategy(PackStrategy);
+								
 								LayoutNode->SetBlockReductionMethod(Layouts[i]->GetBlockReductionMethod() == ECustomizableObjectLayoutBlockReductionMethod::Halve ? mu::EReductionMethod::HALVE_REDUCTION : mu::EReductionMethod::UNITARY_REDUCTION);
 
 								if (bLinkedToExtendMaterial)
diff --git a/Engine/Plugins/Experimental/Mutable/Source/CustomizableObjectEditor/Private/MuCOE/SCustomizableObjectNodeLayoutBlocksEditor.cpp b/Engine/Plugins/Experimental/Mutable/Source/CustomizableObjectEditor/Private/MuCOE/SCustomizableObjectNodeLayoutBlocksEditor.cpp
--- a/Engine/Plugins/Experimental/Mutable/Source/CustomizableObjectEditor/Private/MuCOE/SCustomizableObjectNodeLayoutBlocksEditor.cpp
+++ b/Engine/Plugins/Experimental/Mutable/Source/CustomizableObjectEditor/Private/MuCOE/SCustomizableObjectNodeLayoutBlocksEditor.cpp
@@ -249,6 +249,7 @@ TSharedRef<SWidget> SCustomizableObjectNodeLayoutBlocksEditor::BuildLayoutStrate
 	LayoutPackingStrategies.Empty();
 	LayoutPackingStrategies.Add(MakeShareable(new FString("Resizable")));
 	LayoutPackingStrategies.Add(MakeShareable(new FString("Fixed")));
+	LayoutPackingStrategies.Add(MakeShareable(new FString("Overlay")));
 
 	BlockReductionMethods.Empty();
 	BlockReductionMethods.Add(MakeShareable(new FString("Halve")));
diff --git a/Engine/Plugins/Experimental/Mutable/Source/MutableRuntime/Public/MuR/Layout.h b/Engine/Plugins/Experimental/Mutable/Source/MutableRuntime/Public/MuR/Layout.h
--- a/Engine/Plugins/Experimental/Mutable/Source/MutableRuntime/Public/MuR/Layout.h
+++ b/Engine/Plugins/Experimental/Mutable/Source/MutableRuntime/Public/MuR/Layout.h
@@ -22,7 +22,8 @@ namespace mu
 	enum class EPackStrategy : uint32
 	{
 		RESIZABLE_LAYOUT,
-		FIXED_LAYOUT
+		FIXED_LAYOUT,
+		OVERLAY_LAYOUT
 	};
 
 	//! Types of layout reduction methods 
diff --git a/Engine/Plugins/Experimental/Mutable/Source/MutableTools/Private/MuT/CodeGenerator.cpp b/Engine/Plugins/Experimental/Mutable/Source/MutableTools/Private/MuT/CodeGenerator.cpp
--- a/Engine/Plugins/Experimental/Mutable/Source/MutableTools/Private/MuT/CodeGenerator.cpp
+++ b/Engine/Plugins/Experimental/Mutable/Source/MutableTools/Private/MuT/CodeGenerator.cpp
@@ -1158,7 +1158,9 @@ namespace mu
 						}
 					}
 
-                    if ( layoutOp )
+					bool bIsOverlayLayout = pLayout->GetLayoutPackingStrategy() == mu::EPackStrategy::OVERLAY_LAYOUT;
+
+                    if (!bIsOverlayLayout && layoutOp)
                     {
                         // Add layout packing instructions
 						if (!SharedMeshResults)
