[Mutable] Add unique material slots per surface ID. Materials may be reused between slots.

#rnx

diff --git a/Engine/Plugins/Experimental/Mutable/Source/CustomizableObject/Private/MuCO/CustomizableObjectInstance.cpp b/Engine/Plugins/Experimental/Mutable/Source/CustomizableObject/Private/MuCO/CustomizableObjectInstance.cpp
--- a/Engine/Plugins/Experimental/Mutable/Source/CustomizableObject/Private/MuCO/CustomizableObjectInstance.cpp
+++ b/Engine/Plugins/Experimental/Mutable/Source/CustomizableObject/Private/MuCO/CustomizableObjectInstance.cpp
@@ -79,7 +79,7 @@ struct FMutableMaterialPlaceholder
 		}
 	};
 
-	UMaterialInterface* ParentMaterial;
+	uint32 ParentMaterialID = 0;
 	TArray<FMutableMaterialPlaceHolderParam> Params;
 	int32 MatIndex = -1;
 
@@ -88,8 +88,7 @@ struct FMutableMaterialPlaceholder
 	// Return a hash of the material and its parameters
 	uint32 GetHash()
 	{
-
-		uint32 Hash = ParentMaterial ? ParentMaterial->GetUniqueID() : 0;
+		uint32 Hash = ParentMaterialID;
 
 		// Sort parameters before building the hash.
 		Params.Sort();
@@ -5467,8 +5466,8 @@ void UCustomizableInstancePrivateData::BuildMaterials(const TSharedRef<FUpdateCo
 		// Maps serializations of FMutableMaterialPlaceholder to Created Dynamic Material instances, used to reuse materials across LODs
 		TMap<uint32, TSharedPtr<FMutableMaterialPlaceholder>> ReuseMaterialCache;
 
-		// Map of SharedSurfaceId to FMutableMaterialPlaceholder serialization key 
-		TMap<int32, uint32> SharedSurfacesCache;
+		// SurfaceId per MaterialSlotIndex
+		TArray<int32> SurfaceIdToMaterialIndex;
 
 		MUTABLE_CPUPROFILER_SCOPE(BuildMaterials_LODLoop);
 
@@ -5497,6 +5496,14 @@ void UCustomizableInstancePrivateData::BuildMaterials(const TSharedRef<FUpdateCo
 			{
 				const FInstanceUpdateData::FSurface& Surface = OperationData->InstanceUpdateData.Surfaces[Component.FirstSurface + SurfaceIndex];
 
+				// Reuse MaterialSlot from the previous LOD.
+				if (const int32 MaterialIndex = SurfaceIdToMaterialIndex.Find(Surface.SurfaceId); MaterialIndex != INDEX_NONE)
+				{
+					const int32 LODMaterialIndex = SkeletalMesh->GetLODInfoArray()[LODIndex].LODMaterialMap.Add(MaterialIndex);
+					SkeletalMesh->GetResourceForRendering()->LODRenderData[LODIndex].RenderSections[SurfaceIndex].MaterialIndex = LODMaterialIndex;
+					continue;
+				}
+
 				const uint32 ReferencedMaterialIndex = ObjectToInstanceIndexMap[Surface.MaterialIndex];
 				UMaterialInterface* MaterialTemplate = ReferencedMaterials[ReferencedMaterialIndex];
 				if (!MaterialTemplate)
@@ -5505,22 +5512,23 @@ void UCustomizableInstancePrivateData::BuildMaterials(const TSharedRef<FUpdateCo
 					continue;
 				}
 
-				// Reuse surface between LODs when using AutomaticLODs from mesh.
-				if (const uint32* SharedSurfaceSerialization = SharedSurfacesCache.Find(Surface.SurfaceId))
-				{
-					TSharedPtr<FMutableMaterialPlaceholder>* FoundMaterialPlaceholder = ReuseMaterialCache.Find(*SharedSurfaceSerialization);
-					check(FoundMaterialPlaceholder);
+				// This section will require a new slot
+				SurfaceIdToMaterialIndex.Add(Surface.SurfaceId);
 
-					const int32 MaterialIndex = (*FoundMaterialPlaceholder)->MatIndex;
-					check(MaterialIndex >= 0);
+				// Add and set up the material data for this slot
+				const int32 MaterialSlotIndex = Materials.Num();
+				FSkeletalMaterial& MaterialSlot = Materials.AddDefaulted_GetRef();
+				MaterialSlot.MaterialInterface = MaterialTemplate;
+				MaterialSlot.MaterialSlotName = CustomizableObject->ReferencedMaterialSlotNames[Surface.MaterialIndex];
+				SetMeshUVChannelDensity(MaterialSlot.UVChannelData, RefSkeletalMeshData->Settings.DefaultUVChannelDensity);
 
-					int32 LODMaterialIndex = SkeletalMesh->GetLODInfoArray()[LODIndex].LODMaterialMap.Add(MaterialIndex);
-					SkeletalMesh->GetResourceForRendering()->LODRenderData[LODIndex].RenderSections[SurfaceIndex].MaterialIndex = LODMaterialIndex;
-					continue;
-				}
+				const int32 LODMaterialIndex = SkeletalMesh->GetLODInfoArray()[LODIndex].LODMaterialMap.Add(MaterialSlotIndex);
+				SkeletalMesh->GetResourceForRendering()->LODRenderData[LODIndex].RenderSections[SurfaceIndex].MaterialIndex = LODMaterialIndex;
 
 				TSharedPtr<FMutableMaterialPlaceholder> MutableMaterialPlaceholderPtr(new FMutableMaterialPlaceholder);
 				FMutableMaterialPlaceholder& MutableMaterialPlaceholder = *MutableMaterialPlaceholderPtr;
+				MutableMaterialPlaceholder.ParentMaterialID = MaterialTemplate->GetUniqueID();
+				MutableMaterialPlaceholder.MatIndex = MaterialSlotIndex;
 
 				{
 					MUTABLE_CPUPROFILER_SCOPE(ParamLoop);
@@ -5867,15 +5875,13 @@ void UCustomizableInstancePrivateData::BuildMaterials(const TSharedRef<FUpdateCo
 					}
 				}
 
-				MutableMaterialPlaceholder.ParentMaterial = MaterialTemplate;
+				// Find or create the material for this slot
+				UMaterialInterface* MaterialInterface = MaterialSlot.MaterialInterface;
 
 				const uint32 MaterialParameterHash = MutableMaterialPlaceholder.GetHash();
-
-				int32 MatIndex = INDEX_NONE;
-				
 				if (TSharedPtr<FMutableMaterialPlaceholder>* FoundMaterialPlaceholder = ReuseMaterialCache.Find(MaterialParameterHash))
 				{
-					MatIndex = (*FoundMaterialPlaceholder)->MatIndex;
+					MaterialInterface = Materials[(*FoundMaterialPlaceholder)->MatIndex].MaterialInterface;
 				}
 				else // Material not cached, create a new one
 				{
@@ -5883,12 +5889,10 @@ void UCustomizableInstancePrivateData::BuildMaterials(const TSharedRef<FUpdateCo
 
 					ReuseMaterialCache.Add(MaterialParameterHash, MutableMaterialPlaceholderPtr);
 
-					SharedSurfacesCache.Add(Surface.SurfaceId, MaterialParameterHash);
-
-					FGeneratedMaterial Material;
+					FGeneratedMaterial& Material = GeneratedMaterials.AddDefaulted_GetRef();
 					Material.SurfaceId = Surface.SurfaceId;
 					Material.MaterialIndex = Surface.MaterialIndex;
-					Material.MaterialInterface = MaterialTemplate;
+					Material.MaterialInterface = MaterialInterface;
 
 					UMaterialInstanceDynamic* MaterialInstance = nullptr;
 					
@@ -5905,17 +5909,6 @@ void UCustomizableInstancePrivateData::BuildMaterials(const TSharedRef<FUpdateCo
 						Material.MaterialInterface = MaterialInstance;
 					}
 
-					if (SkeletalMesh)
-					{
-						MatIndex = Materials.Num();
-						MutableMaterialPlaceholder.MatIndex = MatIndex;
-
-						// Set up SkeletalMaterial data
-						FSkeletalMaterial& SkeletalMaterial = Materials.Add_GetRef(Material.MaterialInterface.Get());
-						SkeletalMaterial.MaterialSlotName = CustomizableObject->ReferencedMaterialSlotNames[Surface.MaterialIndex];
-						SetMeshUVChannelDensity(SkeletalMaterial.UVChannelData, RefSkeletalMeshData->Settings.DefaultUVChannelDensity);
-					}
-
 					if (MaterialInstance)
 					{
 						for (const FMutableMaterialPlaceholder::FMutableMaterialPlaceHolderParam& Param : MutableMaterialPlaceholder.Params)
@@ -5974,13 +5967,12 @@ void UCustomizableInstancePrivateData::BuildMaterials(const TSharedRef<FUpdateCo
 						}
 					}
 
-					GeneratedMaterials.Add(Material);
-					ComponentsData[ComponentIndex].OverrideMaterials.Add(Material.MaterialInterface);
+					MaterialInterface = Material.MaterialInterface;
 				}
 
-				int32 LODMaterialIndex = SkeletalMesh->GetLODInfoArray()[LODIndex].LODMaterialMap.Add(MatIndex);
-				SkeletalMesh->GetResourceForRendering()->LODRenderData[LODIndex].RenderSections[SurfaceIndex].MaterialIndex = LODMaterialIndex;
-
+				// Assign the material to the slot, and add it to the  OverrideMaterials
+				MaterialSlot.MaterialInterface = MaterialInterface;
+				ComponentsData[ComponentIndex].OverrideMaterials.Add(MaterialInterface);
 			}
 		}
 
