[Mutable] Remove COE private dependencies 3/3.

Removed MuR runtime private dependency. Required private dependencies (because of the Debugger) have been moved to Internal folder.

#rnx

diff --git a/Engine/Plugins/Experimental/Mutable/Source/CustomizableObjectEditor/CustomizableObjectEditor.Build.cs b/Engine/Plugins/Experimental/Mutable/Source/CustomizableObjectEditor/CustomizableObjectEditor.Build.cs
--- a/Engine/Plugins/Experimental/Mutable/Source/CustomizableObjectEditor/CustomizableObjectEditor.Build.cs
+++ b/Engine/Plugins/Experimental/Mutable/Source/CustomizableObjectEditor/CustomizableObjectEditor.Build.cs
@@ -101,7 +101,6 @@ public class CustomizableObjectEditor : ModuleRules
 			);
 
 		PrivateIncludePaths.AddRange(new string[] {
-				"MutableRuntime/Private",
 				"MutableTools/Private",
 			});
 	}
diff --git a/Engine/Plugins/Experimental/Mutable/Source/CustomizableObjectEditor/Private/MuCOE/SMutableCodeViewer.cpp b/Engine/Plugins/Experimental/Mutable/Source/CustomizableObjectEditor/Private/MuCOE/SMutableCodeViewer.cpp
--- a/Engine/Plugins/Experimental/Mutable/Source/CustomizableObjectEditor/Private/MuCOE/SMutableCodeViewer.cpp
+++ b/Engine/Plugins/Experimental/Mutable/Source/CustomizableObjectEditor/Private/MuCOE/SMutableCodeViewer.cpp
@@ -25,7 +25,7 @@
 #include "MuCOE/SMutableStringViewer.h"
 #include "MuCOE/UnrealEditorPortabilityHelpers.h"
 #include "MuR/SystemPrivate.h"
-#include "MuT/ErrorLogPrivate.h"
+#include "MuT/ErrorLog.h"
 #include "MuT/Streams.h"
 #include "MuT/TypeInfo.h"
 #include "Widgets/SNullWidget.h"
diff --git a/Engine/Plugins/Experimental/Mutable/Source/CustomizableObjectEditor/Private/MuCOE/SMutableCodeViewer.h b/Engine/Plugins/Experimental/Mutable/Source/CustomizableObjectEditor/Private/MuCOE/SMutableCodeViewer.h
--- a/Engine/Plugins/Experimental/Mutable/Source/CustomizableObjectEditor/Private/MuCOE/SMutableCodeViewer.h
+++ b/Engine/Plugins/Experimental/Mutable/Source/CustomizableObjectEditor/Private/MuCOE/SMutableCodeViewer.h
@@ -9,6 +9,7 @@
 #include "Widgets/SCompoundWidget.h"
 #include "Widgets/Input/SSearchBox.h"
 #include "Engine/Texture.h"
+#include "MuR/Mesh.h"
 #include "UObject/SoftObjectPtr.h"
 
 class STableViewBase;
diff --git a/Engine/Plugins/Experimental/Mutable/Source/CustomizableObjectEditor/Private/MuCOE/SMutableConstantsWidget.h b/Engine/Plugins/Experimental/Mutable/Source/CustomizableObjectEditor/Private/MuCOE/SMutableConstantsWidget.h
--- a/Engine/Plugins/Experimental/Mutable/Source/CustomizableObjectEditor/Private/MuCOE/SMutableConstantsWidget.h
+++ b/Engine/Plugins/Experimental/Mutable/Source/CustomizableObjectEditor/Private/MuCOE/SMutableConstantsWidget.h
@@ -5,6 +5,7 @@
 #include "MuR/Image.h"
 #include "MuR/Mesh.h"
 #include "MuR/Operations.h"
+#include "MuR/ModelPrivate.h"
 #include "Widgets/Views/SHeaderRow.h"
 #include "Widgets/Views/STileView.h"
 
diff --git a/Engine/Plugins/Experimental/Mutable/Source/MutableRuntime/Internal/MuR/ModelPrivate.h b/Engine/Plugins/Experimental/Mutable/Source/MutableRuntime/Internal/MuR/ModelPrivate.h
new file mode 100644
--- /dev/null
+++ b/Engine/Plugins/Experimental/Mutable/Source/MutableRuntime/Internal/MuR/ModelPrivate.h
@@ -0,0 +1,719 @@
+// Copyright Epic Games, Inc. All Rights Reserved.
+
+#pragma once
+
+#include "MuR/Model.h"
+#include "MuR/System.h"
+
+#include "MuR/SerialisationPrivate.h"
+#include "MuR/Operations.h"
+#include "MuR/ExtensionData.h"
+#include "MuR/ExtensionDataStreamer.h"
+#include "MuR/Image.h"
+#include "MuR/Mesh.h"
+#include "MuR/ParametersPrivate.h"
+#include "MuR/MutableRuntimeModule.h"
+
+#define MUTABLE_MAX_RUNTIME_PARAMETERS_PER_STATE	64
+#define MUTABLE_GROW_BORDER_VALUE					2
+
+namespace mu
+{
+
+	/** Used to debug and log. */
+	constexpr bool DebugRom = false;
+	constexpr bool DebugRomAll = false;
+	constexpr int32 DebugRomIndex = 44;
+	constexpr int32 DebugImageIndex = 9;
+
+    //! Data stored for a rom even if it is not loaded
+    struct FRomData
+    {
+        //! This is used to identify a ROM file. It is usually a hash from its data.
+        uint32 Id;
+
+		//! Size of the rom
+		uint32 Size;
+
+		//! Index of the resource in its type-sepcific array
+		uint32 ResourceIndex;
+
+		//! Index of the resource in its type-sepcific array (one of DATATYPE values)
+		uint32 ResourceType;        
+    };
+
+    MUTABLE_DEFINE_POD_SERIALISABLE(FRomData);
+
+    //!
+    template<typename DATA>
+    inline void AppendCode(TArray<uint8>& code, const DATA& data )
+    {
+        int32 pos = code.Num();
+        code.SetNum( pos+sizeof(DATA), false );
+		FMemory::Memcpy (&code[pos], &data, sizeof(DATA));
+    }
+
+
+	//!
+	struct FImageLODRange
+	{
+		int32 FirstIndex = 0;
+		uint16 ImageSizeX = 0;
+		uint16 ImageSizeY = 0;
+		uint16 _Padding = 0;
+		uint8 LODCount = 0;
+		EImageFormat ImageFormat = EImageFormat::IF_NONE;
+	};
+	MUTABLE_DEFINE_POD_SERIALISABLE(FImageLODRange);
+
+	struct FExtensionDataConstant
+	{
+		// This should always be valid, but if the state is Unloaded it won't be usable.
+		//
+		// Avoid storing references to this Data in Memory while the state is Unloaded.
+		ExtensionDataPtrConst Data;
+
+		enum class ELoadState : uint8
+		{
+			Invalid,
+			Unloaded,
+			FailedToLoad,
+			CurrentlyLoaded,
+			AlwaysLoaded
+		};
+
+		// This should be initialized to a valid load state when Data is set
+		ELoadState LoadState = ELoadState::Invalid;
+
+		inline void Serialise(OutputArchive& arch) const
+		{
+			arch << Data;
+		}
+		
+		inline void Unserialise(InputArchive& arch)
+		{
+			arch >> Data;
+
+			check(Data.get());
+			check(Data->Origin == ExtensionData::EOrigin::ConstantAlwaysLoaded
+				|| Data->Origin == ExtensionData::EOrigin::ConstantStreamed);
+
+			if (Data->Origin == ExtensionData::EOrigin::ConstantAlwaysLoaded)
+			{
+				LoadState = ELoadState::AlwaysLoaded;
+			}
+			else
+			{
+				// Streamed constants are assumed to be unloaded to start with
+				LoadState = ELoadState::Unloaded;
+			}
+		}
+	};
+
+    //!
+    struct FProgram
+    {
+        FProgram()
+        {
+            // Add the null instruction at address 0.
+            // TODO: Will do it in the linker
+            AppendCode( m_byteCode, OP_TYPE::NONE );
+            m_opAddress.Add(0);
+        }
+
+        struct FState
+        {
+            //! Name of the state
+            FString Name;
+
+            //! First instruction of the full build of an instance in this state
+            OP::ADDRESS m_root = 0;
+
+            //! List of parameters index (to FProgram::m_parameters) of the runtime parameters of
+            //! this state.
+			TArray<int> m_runtimeParameters;
+
+            //! List of instructions that need to be cached to efficiently update this state
+			TArray<OP::ADDRESS> m_updateCache;
+
+            //! List of root instructions for the dynamic resources that depend on the runtime
+            //! parameters of this state, with a mask of relevant runtime parameters.
+            //! The mask has a bit on for every runtime parameter in the m_runtimeParameters
+            //! vector.
+			//! The uint64 is linked to MUTABLE_MAX_RUNTIME_PARAMETERS_PER_STATE
+			TArray< TPair<OP::ADDRESS,uint64> > m_dynamicResources;
+
+            //!
+            inline void Serialise( OutputArchive& arch ) const
+            {
+                arch << Name;
+                arch << m_root;
+                arch << m_runtimeParameters;
+                arch << m_updateCache;
+                arch << m_dynamicResources;
+            }
+
+            //!
+            inline void Unserialise( InputArchive& arch )
+            {
+                arch >> Name;
+                arch >> m_root;
+                arch >> m_runtimeParameters;
+                arch >> m_updateCache;
+                arch >> m_dynamicResources;
+            }
+
+            /** Returns the mask of parameters (from the runtime parameter list of this state) including the parameters that 
+			* are relevant for the dynamic resource at the given address.
+			*/
+            uint64 IsDynamic( OP::ADDRESS at ) const
+            {
+                uint64 res = 0;
+
+                for ( int32 i=0; i<m_dynamicResources.Num(); ++i )
+                {
+                    if ( m_dynamicResources[i].Key==at )
+                    {
+                        res = m_dynamicResources[i].Value;
+						break;
+                    }
+                }
+
+                return res;
+            }
+
+            //!
+            bool IsUpdateCache( OP::ADDRESS at ) const
+            {
+                bool res = false;
+
+                for ( SIZE_T i=0; !res && i<m_updateCache.Num(); ++i )
+                {
+                    if ( m_updateCache[i]==at )
+                    {
+                        res = true;
+                    }
+                }
+
+                return res;
+            }
+
+            //!
+            void AddUpdateCache( OP::ADDRESS at )
+            {
+                if ( !IsUpdateCache(at) )
+                {
+                    m_updateCache.Add( at );
+                }
+            }
+        };
+
+        //! Location in the m_byteCode of the beginning of each operation
+		TArray<uint32> m_opAddress;
+
+        //! Byte-coded representation of the program, using variable-sized op data.
+		TArray<uint8> m_byteCode;
+
+        //!
+		TArray<FState> m_states;
+
+		//! Data for every rom.
+		TArray<FRomData> m_roms;
+
+		//! Constant image mip data: the first is the index in m_roms for each image mip or -1 if it is always loaded.
+		TArray<TPair<int32, Ptr<const Image>>> m_constantImageLODs;
+
+		//! Constant image mip chain indices: ranges in this array are defined in FImageLODRange and the indices here refer to m_constantImageLODs.
+		TArray<uint32> m_constantImageLODIndices;
+
+		//! Constant image data.
+		TArray<FImageLODRange> m_constantImages;
+
+        //! Constant mesh data: the first is the index in m_roms for each mesh or -1 if it is always loaded.
+		TArray<TPair<int32, Ptr<const Mesh>>> m_constantMeshes;
+
+		//! Constant ExtensionData
+		TArray<FExtensionDataConstant> m_constantExtensionData;
+
+        //! Constant string data
+		TArray<FString> m_constantStrings;
+
+        //! Constant layout data
+		TArray<Ptr<const Layout>> m_constantLayouts;
+
+        //! Constant projectors
+		TArray<FProjector> m_constantProjectors;
+
+        //! Constant matrices, usually used for transforms
+		TArray<mat4f> m_constantMatrices;
+
+		//! Constant shapes
+		TArray<FShape> m_constantShapes;
+
+        //! Constant curves
+		TArray<Curve> m_constantCurves;
+
+        //! Constant skeletons
+		TArray<Ptr<const Skeleton>> m_constantSkeletons;
+
+		//! Constant Physics Bodies
+		TArray<Ptr<const PhysicsBody>> m_constantPhysicsBodies;
+
+        //! Parameters of the model.
+        //! The value stored here is the default value.
+		TArray<FParameterDesc> m_parameters;
+
+        //! Ranges for iteration of the model operations.
+		TArray<FRangeDesc> m_ranges;
+
+        //! List of parameter lists. These are used in several places, like storing the
+        //! pregenerated list of parameters influencing a resource.
+        //! The parameter lists are sorted.
+		TArray<TArray<uint16>> m_parameterLists;
+
+#if WITH_EDITOR
+		//! State of the program. True unless the streamed resources were destroyed,
+		//! which could happen in the editor after recompiling the CO.
+		bool bIsValid = true;
+#endif
+        //!
+        void Serialise( OutputArchive& arch ) const
+        {
+            arch << m_opAddress;
+            arch << m_byteCode;
+            arch << m_states;
+			arch << m_roms;
+			arch << m_constantImageLODs;
+			arch << m_constantImageLODIndices;
+			arch << m_constantImages;
+			arch << m_constantMeshes;
+			arch << m_constantExtensionData;
+			arch << m_constantStrings;
+            arch << m_constantLayouts;
+            arch << m_constantProjectors;
+			arch << m_constantMatrices;
+			arch << m_constantShapes;
+            arch << m_constantCurves;
+            arch << m_constantSkeletons;
+			arch << m_constantPhysicsBodies;
+            arch << m_parameters;
+            arch << m_ranges;
+            arch << m_parameterLists;
+        }
+
+        //!
+        void Unserialise( InputArchive& arch )
+        {
+            arch >> m_opAddress;
+            arch >> m_byteCode;
+            arch >> m_states;
+			arch >> m_roms;
+			arch >> m_constantImageLODs;
+			arch >> m_constantImageLODIndices;
+			arch >> m_constantImages;
+			arch >> m_constantMeshes;
+			arch >> m_constantExtensionData;
+			arch >> m_constantStrings;
+            arch >> m_constantLayouts;
+            arch >> m_constantProjectors;
+			arch >> m_constantMatrices;
+			arch >> m_constantShapes;
+            arch >> m_constantCurves;
+            arch >> m_constantSkeletons;
+			arch >> m_constantPhysicsBodies;
+            arch >> m_parameters;
+            arch >> m_ranges;
+            arch >> m_parameterLists;
+        }
+
+        //! Debug method that sanity-checks the program with a variety of tests.
+        void Check();
+
+        //! Debug method that logs the top used instruction types.
+        void LogHistogram() const;
+
+		//! Return true if the given ROM is loaded.
+		inline bool IsRomLoaded( int32 RomIndex ) const
+		{
+			switch (m_roms[RomIndex].ResourceType)
+			{
+			case DT_IMAGE:
+				return m_constantImageLODs[m_roms[RomIndex].ResourceIndex].Value.get() != nullptr;
+			case DT_MESH:
+				return m_constantMeshes[m_roms[RomIndex].ResourceIndex].Value.get() != nullptr;
+			default:
+				check(false);
+				break;
+			}
+			return false;
+		}
+
+		/** Unload a rom resource. Return the size of the unloaded rom.*/
+		inline int32 UnloadRom(int32 RomIndex)
+		{
+			int32 RomSize = 0;
+
+			if (DebugRom && (DebugRomAll||RomIndex == DebugRomIndex))
+				UE_LOG(LogMutableCore, Log, TEXT("Unloading rom %d."), RomIndex);
+
+			switch (m_roms[RomIndex].ResourceType)
+			{
+			case DT_IMAGE:
+				if (m_constantImageLODs[m_roms[RomIndex].ResourceIndex].Value)
+				{
+					RomSize = m_constantImageLODs[m_roms[RomIndex].ResourceIndex].Value->GetDataSize();
+					m_constantImageLODs[m_roms[RomIndex].ResourceIndex].Value = nullptr;
+				}
+				break;
+			case DT_MESH:
+				if (m_constantMeshes[m_roms[RomIndex].ResourceIndex].Value)
+				{
+					RomSize = m_constantMeshes[m_roms[RomIndex].ResourceIndex].Value->GetDataSize();
+					m_constantMeshes[m_roms[RomIndex].ResourceIndex].Value = nullptr;
+				}
+				break;
+			default:
+				check(false);
+				break;
+			}
+
+			return RomSize;
+		}
+
+
+		int32 AddConstant(Ptr<const Mesh> pMesh)
+		{
+			// Uniques needs to be ensured outside
+			return m_constantMeshes.Add(TPair<int32, Ptr<const Mesh>>( -1, pMesh.get() ));
+		}
+
+		OP::ADDRESS AddConstant(Ptr<const ExtensionData> Data)
+		{
+			// Ensure unique
+			for (int32 Index = 0; Index < m_constantExtensionData.Num(); Index++)
+			{
+				const ExtensionData* Candidate = m_constantExtensionData[Index].Data.get();
+				if (*Candidate == *Data)
+				{
+					return Index;
+				}
+			}
+
+			FExtensionDataConstant& NewConstant = m_constantExtensionData.AddDefaulted_GetRef();
+			NewConstant.Data = Data;
+
+			// The data is assumed to be loaded during compilation
+			NewConstant.LoadState =
+				Data->Origin == ExtensionData::EOrigin::ConstantAlwaysLoaded
+				? FExtensionDataConstant::ELoadState::AlwaysLoaded
+				: FExtensionDataConstant::ELoadState::CurrentlyLoaded;
+
+			return m_constantExtensionData.Num() - 1;
+		}
+
+		OP::ADDRESS AddConstant( Ptr<const Layout> pLayout )
+        {
+            // Ensure unique
+            for ( SIZE_T i=0; i<m_constantLayouts.Num(); ++i)
+            {
+                if (m_constantLayouts[i]==pLayout)
+                {
+                    return (OP::ADDRESS)i;
+                }
+            }
+
+            OP::ADDRESS index = OP::ADDRESS( m_constantLayouts.Num() );
+            m_constantLayouts.Add( pLayout );
+            return index;
+        }
+
+        OP::ADDRESS AddConstant( Ptr<const Skeleton> pSkeleton )
+        {
+            // Ensure unique
+            for ( SIZE_T i=0; i<m_constantSkeletons.Num(); ++i)
+            {
+                if ( m_constantSkeletons[i]==pSkeleton
+                     ||
+                     *m_constantSkeletons[i]==*pSkeleton
+                     )
+                {
+                    return (OP::ADDRESS)i;
+                }
+            }
+
+            OP::ADDRESS index = OP::ADDRESS( m_constantSkeletons.Num() );
+            m_constantSkeletons.Add( pSkeleton );
+            return index;
+        }
+
+        OP::ADDRESS AddConstant( Ptr<const PhysicsBody> pPhysicsBody )
+        {
+            // Ensure unique
+            for ( SIZE_T i=0; i<m_constantPhysicsBodies.Num(); ++i)
+            {
+                if ( m_constantPhysicsBodies[i]==pPhysicsBody
+                     ||
+                     *m_constantPhysicsBodies[i]==*pPhysicsBody
+                     )
+                {
+                    return (OP::ADDRESS)i;
+                }
+            }
+
+            OP::ADDRESS index = OP::ADDRESS( m_constantPhysicsBodies.Num() );
+            m_constantPhysicsBodies.Add( pPhysicsBody );
+            return index;
+        }
+
+        OP::ADDRESS AddConstant( const FString& str )
+        {            
+            // Ensure unique
+            for ( SIZE_T i=0; i<m_constantStrings.Num(); ++i)
+            {
+                if (m_constantStrings[i]==str)
+                {
+                    return (OP::ADDRESS)i;
+                }
+            }
+
+            OP::ADDRESS index = OP::ADDRESS( m_constantStrings.Num() );
+            m_constantStrings.Add( str );
+            return index;
+        }
+
+        OP::ADDRESS AddConstant( const mat4f& m )
+        {
+            // Ensure unique
+            for ( SIZE_T i=0; i<m_constantMatrices.Num(); ++i)
+            {
+                if (m_constantMatrices[i]==m)
+                {
+                    return (OP::ADDRESS)i;
+                }
+            }
+
+            OP::ADDRESS index = OP::ADDRESS( m_constantMatrices.Num() );
+            m_constantMatrices.Add( m );
+            return index;
+        }
+
+        OP::ADDRESS AddConstant( const FShape& m )
+        {
+            // Ensure unique
+            for ( SIZE_T i=0; i<m_constantShapes.Num(); ++i)
+            {
+                if (m_constantShapes[i]==m)
+                {
+                    return (OP::ADDRESS)i;
+                }
+            }
+
+            OP::ADDRESS index = OP::ADDRESS( m_constantShapes.Num() );
+            m_constantShapes.Add( m );
+            return index;
+        }
+
+        OP::ADDRESS AddConstant( const FProjector& m )
+        {
+            // Ensure unique
+            for ( SIZE_T i=0; i<m_constantProjectors.Num(); ++i)
+            {
+                if (m_constantProjectors[i]==m)
+                {
+                    return (OP::ADDRESS)i;
+                }
+            }
+
+            OP::ADDRESS index = OP::ADDRESS( m_constantProjectors.Num() );
+            m_constantProjectors.Add( m );
+            return index;
+        }
+
+        OP::ADDRESS AddConstant( const Curve& m )
+        {
+            // Ensure unique
+//            for ( SIZE_T i=0; i<m_constantCurves.Num(); ++i)
+//            {
+//                if (m_constantCurves[i]==m)
+//                {
+//                    return (OP::ADDRESS)i;
+//                }
+//            }
+
+            OP::ADDRESS index = OP::ADDRESS( m_constantCurves.Num() );
+            m_constantCurves.Add( m );
+            return index;
+        }
+
+
+        //! Get a constant image, assuming it is fully loaded. The image constant will be composed with lodaded mips if necessary.
+		template <typename CreateImageFunc>
+        void GetConstant( int32 ConstantIndex, ImagePtrConst& res, int32 MipsToSkip, const CreateImageFunc& CreateImage) const
+        {
+			int32 ReallySkippedLODs = FMath::Min(m_constantImages[ConstantIndex].LODCount - 1, MipsToSkip);
+			int32 FirstLODIndexIndex = m_constantImages[ConstantIndex].FirstIndex;
+			int32 ResultLODIndexIndex = FirstLODIndexIndex + ReallySkippedLODs;
+			int32 FinalLODs = m_constantImages[ConstantIndex].LODCount - ReallySkippedLODs;
+			check(FinalLODs > 0);
+
+			// Get the first mip
+			int32 ResultLODIndex = m_constantImageLODIndices[ResultLODIndexIndex];
+			Ptr<const Image> CurrentMip = m_constantImageLODs[ResultLODIndex].Value;
+			check(CurrentMip);
+				
+			// Shortcut if we only want one mip
+			if (FinalLODs == 1)
+			{
+				res = CurrentMip;
+				return;
+			}
+
+			// Compose the result image
+			{
+				MUTABLE_CPUPROFILER_SCOPE(ComposeConstantImage);
+
+				Ptr<Image> Result = CreateImage( CurrentMip->GetSizeX(), CurrentMip->GetSizeY(), FinalLODs, CurrentMip->GetFormat(), EInitializationType::NotInitialized );
+				Result->m_flags = CurrentMip->m_flags;
+
+				// Some non-block pixel formats require separate memory size calculation
+				if (!Result->GetDataSize())
+				{
+					int32 TotalSize = 0;
+					for (int32 LOD = 0; LOD < FinalLODs; ++LOD)
+					{
+						int LODIndex = m_constantImageLODIndices[ResultLODIndexIndex + LOD];
+						int32 MipSizeBytes = m_constantImageLODs[LODIndex].Value->GetDataSize();
+						TotalSize += MipSizeBytes;
+					}
+					Result->m_data.SetNum(TotalSize);
+				}
+
+				int32 WrittenDataBytes = 0;
+				uint8* Data = Result->GetData();
+				for (int32 LOD=0; LOD<FinalLODs; ++LOD)
+				{
+					check(CurrentMip->GetLODCount() == 1);
+					check(CurrentMip->GetFormat() == Result->GetFormat());
+					int32 MipSizeBytes = CurrentMip->GetDataSize();
+					check( Result->GetDataSize()>=WrittenDataBytes+MipSizeBytes);
+
+					FMemory::Memcpy( Data, CurrentMip->GetData(), MipSizeBytes);
+					Data += MipSizeBytes;
+					WrittenDataBytes += MipSizeBytes;
+
+					if (LOD + 1 < FinalLODs)
+					{
+						ResultLODIndex = m_constantImageLODIndices[ResultLODIndexIndex + LOD + 1];
+						CurrentMip = m_constantImageLODs[ResultLODIndex].Value;
+						check(CurrentMip);
+					}
+				}
+
+				res = Result;
+			}
+		}
+
+
+        void GetConstant(int32 ConstantIndex, MeshPtrConst& res ) const
+        {
+			res = m_constantMeshes[ConstantIndex].Value;
+		}
+
+		void GetExtensionDataConstant(int32 ConstantIndex, ExtensionDataPtrConst& Result) const
+		{
+			const FExtensionDataConstant& Constant = m_constantExtensionData[ConstantIndex];
+
+			check(Constant.LoadState != FExtensionDataConstant::ELoadState::Unloaded);
+			check(Constant.Data.get());
+
+			Result = Constant.Data;
+		}
+
+        inline OP_TYPE GetOpType( OP::ADDRESS at ) const
+        {
+            if (at>=OP::ADDRESS(m_opAddress.Num())) return OP_TYPE::NONE;
+
+            OP_TYPE result;
+            uint64 byteCodeAddress = m_opAddress[at];
+			FMemory::Memcpy( &result, &m_byteCode[byteCodeAddress], sizeof(OP_TYPE) );
+            check( result<OP_TYPE::COUNT);
+            return result;
+        }
+
+        template<typename ARGS>
+        inline const ARGS GetOpArgs( OP::ADDRESS at ) const
+        {
+            ARGS result;
+            uint64 byteCodeAddress = m_opAddress[at];
+            byteCodeAddress += sizeof(OP_TYPE);
+			FMemory::Memcpy( &result, &m_byteCode[byteCodeAddress], sizeof(ARGS));
+            return result;
+        }
+
+        inline const uint8* GetOpArgsPointer( OP::ADDRESS at ) const
+        {
+            uint64 byteCodeAddress = m_opAddress[at];
+            byteCodeAddress += sizeof(OP_TYPE);
+            const uint8* pData = (const uint8*)&m_byteCode[byteCodeAddress];
+            return pData;
+        }
+
+        inline uint8* GetOpArgsPointer( OP::ADDRESS at )
+        {
+            uint64 byteCodeAddress = m_opAddress[at];
+            byteCodeAddress += sizeof(OP_TYPE);
+            uint8* pData = (uint8*)&m_byteCode[byteCodeAddress];
+            return pData;
+        }
+    };
+
+
+    //!
+    class MUTABLERUNTIME_API Model::Private : public Base
+    {
+    public:
+
+        //!
+        FProgram m_program;
+
+        //!
+        void Serialise( OutputArchive& arch ) const
+        {
+            uint32 ver = 0;
+            arch << ver;
+
+            arch << m_program;
+        }
+
+        //!
+        void Unserialise( InputArchive& arch )
+        {
+            uint32 ver;
+            arch >> ver;
+            check(ver==0);
+
+            arch >> m_program;
+        }
+
+    };
+
+
+    //!
+    class MUTABLERUNTIME_API ModelParametersGenerator::Private : public Base
+    {
+    public:
+
+        //!
+        TSharedPtr<const Model> m_pModel;
+        Ptr<System> m_pSystem;
+
+        //! Number of possible instances
+        int64 m_instanceCount;
+
+		/** Value used for scalar paramters that control multidemnsional sizes. 
+		* TODO: Make it an option.
+		*/
+		int32 DefaultRangeDimension = 8;
+    };
+
+}
diff --git a/Engine/Plugins/Experimental/Mutable/Source/MutableRuntime/Internal/MuR/OpMeshFormat.h b/Engine/Plugins/Experimental/Mutable/Source/MutableRuntime/Internal/MuR/OpMeshFormat.h
new file mode 100644
--- /dev/null
+++ b/Engine/Plugins/Experimental/Mutable/Source/MutableRuntime/Internal/MuR/OpMeshFormat.h
@@ -0,0 +1,43 @@
+// Copyright Epic Games, Inc. All Rights Reserved.
+
+#pragma once
+
+#include "MuR/Mesh.h"
+
+
+namespace mu
+{
+class FMeshBufferSet;
+
+	//---------------------------------------------------------------------------------------------
+	//! Convert a mesh format into another one.
+	//! Slow implementation, but it should never happen at run-time
+	//! \param keepSystemBuffers Will keep the internal system buffers even if they are not in the
+	//! original format. If they are, they will be duplicated, so be careful.
+	//---------------------------------------------------------------------------------------------
+	MUTABLERUNTIME_API extern void MeshFormat
+		(
+			Mesh* Result,
+			const Mesh* pSource,
+			const Mesh* pFormat,
+			bool keepSystemBuffers,
+			bool formatVertices,
+			bool formatIndices,
+			bool formatFaces,
+			bool ignoreMissingChannels,
+			bool& bOutSuccess
+		);
+
+
+    //---------------------------------------------------------------------------------------------
+    //! Fill a meshbuffer from result with the current data in source but keeping the format
+    //! already in the result.
+    //---------------------------------------------------------------------------------------------
+    void MeshFormatBuffer
+        (
+            const FMeshBufferSet& pSource,
+            FMeshBufferSet& pResult,
+            int bufferIndex
+        );
+
+}
\ No newline at end of file
diff --git a/Engine/Plugins/Experimental/Mutable/Source/MutableRuntime/Internal/MuR/Operations.h b/Engine/Plugins/Experimental/Mutable/Source/MutableRuntime/Internal/MuR/Operations.h
new file mode 100644
--- /dev/null
+++ b/Engine/Plugins/Experimental/Mutable/Source/MutableRuntime/Internal/MuR/Operations.h
@@ -0,0 +1,1248 @@
+// Copyright Epic Games, Inc. All Rights Reserved.
+
+#pragma once
+
+#include "MuR/Image.h"
+#include "MuR/SerialisationPrivate.h"
+#include "HAL/PlatformMath.h"
+#include "HAL/UnrealMemory.h"
+#include "Misc/AssertionMacros.h"
+#include "Misc/EnumClassFlags.h"
+#include "Templates/Function.h"
+
+#define MUTABLE_OP_MAX_INTERPOLATE_COUNT	6
+#define MUTABLE_OP_MAX_ADD_COUNT			7
+#define MUTABLE_OP_MAX_SWIZZLE_CHANNELS		4
+
+namespace mu
+{
+
+    //!
+    enum class OP_TYPE : uint16
+    {
+        //-----------------------------------------------------------------------------------------
+        // No operation
+        //-----------------------------------------------------------------------------------------
+        NONE,
+
+        //-----------------------------------------------------------------------------------------
+        // Generic operations
+        //-----------------------------------------------------------------------------------------
+
+        //! Constant value
+        BO_CONSTANT,
+        NU_CONSTANT,
+        SC_CONSTANT,
+        CO_CONSTANT,
+        IM_CONSTANT,
+        ME_CONSTANT,
+        LA_CONSTANT,
+        PR_CONSTANT,
+        ST_CONSTANT,
+		ED_CONSTANT,
+
+        //! User parameter
+        BO_PARAMETER,
+        NU_PARAMETER,
+        SC_PARAMETER,
+        CO_PARAMETER,
+        PR_PARAMETER,
+        IM_PARAMETER,
+        ST_PARAMETER,
+
+		//! A referenced, but opaque engine resource
+		IM_REFERENCE,
+
+        //! Select one value or the other depending on a boolean input
+        NU_CONDITIONAL,
+        SC_CONDITIONAL,
+        CO_CONDITIONAL,
+        IM_CONDITIONAL,
+        ME_CONDITIONAL,
+        LA_CONDITIONAL,
+        IN_CONDITIONAL,
+		ED_CONDITIONAL,
+
+        //! Select one of several values depending on an int input
+        NU_SWITCH,
+        SC_SWITCH,
+        CO_SWITCH,
+        IM_SWITCH,
+        ME_SWITCH,
+        LA_SWITCH,
+        IN_SWITCH,
+		ED_SWITCH,
+
+        //-----------------------------------------------------------------------------------------
+        // Boolean operations
+        //-----------------------------------------------------------------------------------------
+
+        //! Compare two scalars, return true if the first is less than the second.
+		//! \TODO: Deprecated?
+		BO_LESS,
+
+        //! Compare an integerexpression with an integer constant
+        BO_EQUAL_INT_CONST,
+
+        //! Logical and
+        BO_AND,
+
+        //! Logical or
+        BO_OR,
+
+        //! Left as an exercise to the reader to find out what this op does.
+        BO_NOT,
+
+        //-----------------------------------------------------------------------------------------
+        // Integer operations
+        //-----------------------------------------------------------------------------------------
+
+        //-----------------------------------------------------------------------------------------
+        // Scalar operations
+        //-----------------------------------------------------------------------------------------
+
+        //! Multiply a scalar value by another onw and add a third one to the result
+		//! \TODO: Deprecated?
+		SC_MULTIPLYADD,
+
+        //! Apply an arithmetic operation to two scalars
+        SC_ARITHMETIC,
+
+        //! Get a scalar value from a curve
+        SC_CURVE,
+
+        //-----------------------------------------------------------------------------------------
+        // Colour operations. Colours are sometimes used as generic vectors.
+        //-----------------------------------------------------------------------------------------
+
+        //! Sample an image to get its colour.
+        CO_SAMPLEIMAGE,
+
+        //! Make a color by shuffling channels from other colours.
+        CO_SWIZZLE,
+
+        //! Compose a vector from 4 scalars
+        CO_FROMSCALARS,
+
+        //! Apply component-wise arithmetic operations to two colours
+        CO_ARITHMETIC,
+
+        //-----------------------------------------------------------------------------------------
+        // Image operations
+        //-----------------------------------------------------------------------------------------
+
+        //! Combine an image on top of another one using a specific effect (Blend, SoftLight, 
+		//! Hardlight, Burn...). And optionally a mask.
+        IM_LAYER,
+
+        //! Apply a colour on top of an image using a specific effect (Blend, SoftLight, 
+		//! Hardlight, Burn...), optionally using a mask.
+        IM_LAYERCOLOUR,        
+
+        //! Convert between pixel formats
+        IM_PIXELFORMAT,
+
+        //! Generate mipmaps up to a provided level
+        IM_MIPMAP,
+
+        //! Resize the image to a constant size
+        IM_RESIZE,
+
+        //! Resize the image to the size of another image
+        IM_RESIZELIKE,
+
+        //! Resize the image by a relative factor
+        IM_RESIZEREL,
+
+        //! Create an empty image to hold a particular layout.
+        IM_BLANKLAYOUT,
+
+        //! Copy an image into a rect of another one.
+        IM_COMPOSE,
+
+        //! Interpolate between 2 images taken from a row of targets (2 consecutive targets).
+        IM_INTERPOLATE,
+
+        //! Change the saturation of the image.
+        IM_SATURATE,
+
+        //! Generate a one-channel image with the luminance of the source image.
+        IM_LUMINANCE,
+
+        //! Recombine the channels of several images into one.
+        IM_SWIZZLE,
+
+        //! Convert the source image colours using a "palette" image sampled with the source
+        //! grey-level.
+        IM_COLOURMAP,
+
+        //! Build a horizontal gradient image from two colours
+        IM_GRADIENT,
+
+        //! Generate a black and white image from an image and a threshold.
+        IM_BINARISE,
+
+        //! Generate a plain colour image
+        IM_PLAINCOLOUR,
+
+        //! Cut a rect from an image
+        IM_CROP,
+
+        //! Replace a subrect of an image with another one
+        IM_PATCH,
+
+        //! Render a mesh texture layout into a mask
+        IM_RASTERMESH,
+
+        //! Create an image displacement encoding the grow operation for a mask
+        IM_MAKEGROWMAP,
+
+        //! Apply an image displacement on another image.
+        IM_DISPLACE,
+
+        //! Repeately apply
+        IM_MULTILAYER,
+
+        //! Inverts the colors of an image
+        IM_INVERT,
+
+        //! Modifiy roughness channel of an image based on normal variance.
+        IM_NORMALCOMPOSITE,
+
+		//! Apply linear transform to Image content. Resulting samples outside the original image are tiled.
+		IM_TRANSFORM,
+
+        //-----------------------------------------------------------------------------------------
+        // Mesh operations
+        //-----------------------------------------------------------------------------------------
+
+        //! Apply a layout to a mesh texture coordinates channel
+        ME_APPLYLAYOUT,
+
+        //! Compare two meshes and extract a morph from the first to the second
+        //! The meshes must have the same topology, etc.
+        ME_DIFFERENCE,
+
+        //! Apply a one morphs on a base. 
+        ME_MORPH,
+
+        //! Merge a mesh to a mesh
+        ME_MERGE,
+
+        //! Interpolate between several meshes.
+        ME_INTERPOLATE,
+
+        //! Create a new mask mesh selecting all the faces of a source that are inside a given
+        //! clip mesh.
+        ME_MASKCLIPMESH,
+
+        /** Create a new mask mesh selecting the faces of a source that have UVs inside the region marked in an image mask. */
+		ME_MASKCLIPUVMASK,
+
+        //! Create a new mask mesh selecting all the faces of a source that match another mesh.
+        ME_MASKDIFF,
+
+        //! Remove all the geometry selected by a mask.
+        ME_REMOVEMASK,
+
+        //! Change the mesh format to match the format of another one.
+        ME_FORMAT,
+
+        //! Extract a fragment of a mesh containing a specific layout block.
+        ME_EXTRACTLAYOUTBLOCK,
+
+        //! Apply a transform in a 4x4 matrix to the geometry channels of the mesh
+        ME_TRANSFORM,
+
+        //! Clip the mesh with a plane and morph it when it is near until it becomes an ellipse on
+        //! the plane.
+        ME_CLIPMORPHPLANE,
+
+        //! Clip the mesh with another mesh.
+        ME_CLIPWITHMESH,
+
+        //! Replace the skeleton data from a mesh with another one.
+        ME_SETSKELETON,
+
+        //! Project a mesh using a projector and clipping the irrelevant faces
+        ME_PROJECT,
+
+        //! Deform a skinned mesh applying a skeletal pose
+        ME_APPLYPOSE,
+
+		//! Apply a geometry core operation to a mesh.
+		//! \TODO: Deprecated?
+		ME_GEOMETRYOPERATION,
+
+		//! Calculate the binding of a mesh on a shape
+		ME_BINDSHAPE,
+
+		//! Apply a shape on a (previously bound) mesh
+		ME_APPLYSHAPE,
+
+		//! Clip Deform using bind data.
+		ME_CLIPDEFORM,
+	
+        //! Mesh morph with Skeleton Reshape based on the morphed mesh.
+        ME_MORPHRESHAPE,
+
+		//! Optimize skinning before adding a mesh to the component
+		ME_OPTIMIZESKINNING,
+
+		//! Add a set of tags to a mesh
+		ME_ADDTAGS,
+
+        //-----------------------------------------------------------------------------------------
+        // Instance operations
+        //-----------------------------------------------------------------------------------------
+
+        //! Add a mesh to an instance
+        IN_ADDMESH,
+
+        //! Add an image to an instance
+        IN_ADDIMAGE,
+
+        //! Add a vector to an instance
+        IN_ADDVECTOR,
+
+        //! Add a scalar to an instance
+        IN_ADDSCALAR,
+
+        //! Add a string to an instance
+        IN_ADDSTRING,
+
+        //! Add a surface to an instance component
+        IN_ADDSURFACE,
+
+        //! Add a component to an instance LOD
+        IN_ADDCOMPONENT,
+
+        //! Add all LODs to an instance. This operation can only appear once in a model.
+        IN_ADDLOD,
+
+		//! Add extension data to an instance
+		IN_ADDEXTENSIONDATA,
+
+        //-----------------------------------------------------------------------------------------
+        // Layout operations
+        //-----------------------------------------------------------------------------------------
+
+        //! Pack all the layout blocks from the source in the grid without overlapping
+        LA_PACK,
+
+        //! Merge two layouts
+        LA_MERGE,
+
+        //! Remove all layout blocks not used by any vertex of the mesh.
+        //! This operation is for the new way of managing layout blocks.
+        LA_REMOVEBLOCKS,
+
+		//! Extract a layout from a mesh
+		LA_FROMMESH,
+
+        //-----------------------------------------------------------------------------------------
+        // Utility values
+        //-----------------------------------------------------------------------------------------
+
+        //!
+        COUNT
+
+    };
+
+	enum class EMeshBindShapeFlags : uint32
+	{
+		None				   = 0,
+		ReshapeSkeleton		   = 1 << 0,
+		EnableRigidParts       = 1 << 2,
+		ReshapePhysicsVolumes  = 1 << 4,
+		ReshapeVertices		   = 1 << 5,
+		ApplyLaplacian		   = 1 << 6,
+	};
+	ENUM_CLASS_FLAGS(EMeshBindShapeFlags);
+
+	enum class EMeshBindColorChannelUsage : uint8
+	{
+		None       = 0,
+		ClusterId  = 1,
+		MaskWeight = 2,
+	};
+
+	struct FMeshBindColorChannelUsages
+	{
+		EMeshBindColorChannelUsage R;
+		EMeshBindColorChannelUsage G;
+		EMeshBindColorChannelUsage B;
+		EMeshBindColorChannelUsage A;
+	};
+
+
+	static_assert(sizeof(FMeshBindColorChannelUsages) == sizeof(uint32));
+
+    //---------------------------------------------------------------------------------------------
+    //!
+    //---------------------------------------------------------------------------------------------
+    template<typename ADDRESS_TYPE>
+    struct t_OP
+    {
+        typedef ADDRESS_TYPE ADDRESS;
+
+        //-----------------------------------------------------------------------------------------
+        // Arguments for every operation type
+        //-----------------------------------------------------------------------------------------
+        struct BoolConstantArgs
+        {
+            bool value;
+        };
+
+        struct IntConstantArgs
+        {
+            int32 value;
+        };
+
+        struct ScalarConstantArgs
+        {
+            float value;
+        };
+
+        struct ColourConstantArgs
+        {
+            float value[4];
+        };
+
+        struct ResourceConstantArgs
+        {
+            ADDRESS value;
+        };
+
+        struct MeshConstantArgs
+        {
+            // Index of the mesh in the mesh constant array
+            ADDRESS value;
+
+            // If not negative, index of the skeleton to set to the mesh from the skeleton
+            // constant array.
+            int32 skeleton;
+
+            // If not negative, index of the physics body to set to the mesh from the physics body
+            // constant array.
+			int32 physicsBody;
+        };
+
+        struct ParameterArgs
+        {
+            ADDRESS variable;
+        };
+
+        struct ConditionalArgs
+        {
+            ADDRESS condition, yes, no;
+        };
+
+		struct ResourceReferenceArgs
+		{
+			FImageDesc ImageDesc;
+			int32 ID;
+			int8 ForceLoad;
+		};
+
+        //-------------------------------------------------------------------------------------
+        struct BoolLessArgs
+        {
+            ADDRESS a,b;
+        };
+
+        struct BoolEqualScalarConstArgs
+        {
+            ADDRESS value;
+            int16 constant;
+        };
+
+        struct BoolBinaryArgs
+        {
+            ADDRESS a,b;
+        };
+
+        struct BoolNotArgs
+        {
+            ADDRESS source;
+        };
+
+        //-------------------------------------------------------------------------------------
+        struct ScalarMultiplyAddArgs
+        {
+            ADDRESS factor0,factor1,add;
+        };
+
+        struct ScalarCurveArgs
+        {
+            ADDRESS time;   // Operation generating the time value to sample the curve
+            ADDRESS curve;  // Constant curve
+        };
+
+        //-------------------------------------------------------------------------------------
+        struct ColourSampleImageArgs
+        {
+            ADDRESS image;
+            ADDRESS x,y;
+            uint8 filter;
+        };
+
+        struct ColourSwizzleArgs
+        {
+            uint8 sourceChannels[ MUTABLE_OP_MAX_SWIZZLE_CHANNELS ];
+            ADDRESS sources[ MUTABLE_OP_MAX_SWIZZLE_CHANNELS ];
+        };
+
+		struct ColourMultiplyArgs
+        {
+            ADDRESS a,b;
+        };
+
+		struct ColourFromScalarsArgs
+		{
+			ADDRESS v[MUTABLE_OP_MAX_SWIZZLE_CHANNELS];
+		};
+
+        struct ArithmeticArgs
+		{
+			typedef enum
+			{
+				NONE,
+				ADD,
+				SUBTRACT,
+				MULTIPLY,
+				DIVIDE
+			} OPERATION;
+            uint8 operation;
+
+			ADDRESS a, b;
+		};
+
+
+        //-------------------------------------------------------------------------------------
+        struct ImageLayerArgs
+        {
+            ADDRESS base;
+            ADDRESS mask;
+            ADDRESS blended;
+			uint8 blendType;		// One of EBlendType
+
+			/** One of EBlendType. If this is different than NONE, it will be applied to the alpha with the channel BlendAlphaSourceChannel of the blended. */
+			uint8 blendTypeAlpha;
+
+			/** Channel to use from the source blended argument to apply blendTypeAlpha, if any. */
+			uint8 BlendAlphaSourceChannel;
+
+			uint8 flags;
+			typedef enum
+            {
+                F_NONE           = 0,
+                /** The mask is considered binary : 0 means 0% and any other value means 100% */
+                F_BINARY_MASK    = 1 << 0,
+				/** If the image has 4 channels, apply to the fourth channel as well. */
+				F_APPLY_TO_ALPHA = 1 << 1,
+				/** Use the alpha channel of the blended image as mask. Mask should be null.*/
+				F_USE_MASK_FROM_BLENDED = 1 << 2,
+				/** Use the alpha channel of the base image as its RGB.*/
+				F_BASE_RGB_FROM_ALPHA = 1 << 3,
+				/** Use the alpha channel of the blended image as its RGB.*/
+				F_BLENDED_RGB_FROM_ALPHA = 1 << 4,
+			} FLAGS;
+        };
+
+        struct ImageMultiLayerArgs
+        {
+            ADDRESS base;
+            ADDRESS mask;
+            ADDRESS blended;
+            ADDRESS rangeSize;
+            uint16 rangeId;
+			uint8 blendType;		// One of EBlendType
+			
+			/** One of EBlendType. If this is different than NONE, it will be applied to the alpha with the channel BlendAlphaSourceChannel of the blended. */
+			uint8 blendTypeAlpha;
+
+			/** Channel to use from the source color argument to apply blendTypeAlpha, if any. */
+			uint8 BlendAlphaSourceChannel;
+
+			uint8 bUseMaskFromBlended;	
+		};
+
+        struct ImageLayerColourArgs
+        {
+            ADDRESS base;
+            ADDRESS mask;
+            ADDRESS colour;
+			uint8 blendType;		// One of EBlendType
+
+			/** One of EBlendType. If this is different than NONE, it will be applied to the alpha but with the channel BlendAlphaSourceChannel of the color. */
+			uint8 blendTypeAlpha;	
+
+			/** Channel to use from the source color argument to apply blendTypeAlpha, if any. */
+			uint8 BlendAlphaSourceChannel;
+
+			/** Like in ImageLayerArgs. */
+			uint8 flags;
+		};
+
+        struct ImagePixelFormatArgs
+        {
+            ADDRESS source;
+			EImageFormat format;
+			EImageFormat formatIfAlpha;
+        };
+
+        struct ImageMipmapArgs
+        {
+            ADDRESS source;
+
+            //! Number of mipmaps to build. If zero, it means all.
+            uint8 levels;
+
+            //! Number of mipmaps that can be generated for a single layout block.
+            uint8 blockLevels;
+
+            //! This is true if this operation is supposed to build only the tail mipmaps.
+            //! It is used during the code optimisation phase, and to validate the code.
+            bool onlyTail;
+
+            //! Mipmap generation settings. 
+            float sharpenFactor;
+            EAddressMode addressMode;
+            EMipmapFilterType filterType;
+			bool ditherMipmapAlpha;
+        };
+
+        struct ImageResizeArgs
+        {
+            ADDRESS source;
+            uint16 size[2];
+
+            inline uint16 GetSize(int i) const { return size[i]; }
+        };
+
+        struct ImageResizeLikeArgs
+        {
+            //! Image that will be resized
+            ADDRESS source;
+
+            //! Image whose size will be used to resize the source.
+            ADDRESS sizeSource;
+        };
+
+        struct ImageResizeVarArgs
+        {
+            //! Image that will be resized
+            ADDRESS source;
+
+            //! Size expression.
+            ADDRESS size;
+        };
+
+        struct ImageResizeRelArgs
+        {
+            //! Image that will be resized
+            ADDRESS source;
+
+            //! Factor for each axis.
+            float factor[2];
+        };
+
+        struct ImageBlankLayoutArgs
+        {
+            ADDRESS layout;
+
+            //! Size of a layout block in pixels.
+            uint16 blockSize[2];
+			EImageFormat format;
+
+            //! If true, generate mipmaps
+            uint8 generateMipmaps;
+
+            //! Mipmaps to generate if mipmaps are to be generated. 0 means all.
+            uint8 mipmapCount;
+        };
+
+        struct ImageComposeArgs
+        {
+            ADDRESS layout, base, blockImage;
+            ADDRESS mask;
+            uint32 blockIndex;
+        };
+
+        struct ImageInterpolateArgs
+        {
+            ADDRESS factor;
+            ADDRESS targets[ MUTABLE_OP_MAX_INTERPOLATE_COUNT ];
+        };
+
+        struct ImageInterpolate3Args
+        {
+            ADDRESS factor1, factor2;
+            ADDRESS target0, target1, target2;
+        };
+
+        struct ImageSaturateArgs
+        {
+            //! Image to modify.
+            ADDRESS base;
+
+            //! Saturation factor: 0 desaturates, 1 leaves the same, >1 saturates
+            ADDRESS factor;
+
+        };
+
+        struct ImageLuminanceArgs
+        {
+            //! Image to modify.
+            ADDRESS base;
+        };
+
+        struct ImageSwizzleArgs
+        {
+			EImageFormat format;
+            uint8 sourceChannels[ MUTABLE_OP_MAX_SWIZZLE_CHANNELS ];
+            ADDRESS sources[ MUTABLE_OP_MAX_SWIZZLE_CHANNELS ];
+        };
+
+        struct ImageColourMapArgs
+        {
+            ADDRESS base;
+            ADDRESS mask;
+            ADDRESS map;
+        };
+
+        struct ImageGradientArgs
+        {
+            ADDRESS colour0;
+            ADDRESS colour1;
+            uint16 size[2];
+        };
+
+        struct ImageBinariseArgs
+        {
+            ADDRESS base;
+            ADDRESS threshold;
+        };
+
+        struct ImagePlainColourArgs
+        {
+            ADDRESS colour;
+			EImageFormat format;
+            uint16 size[2];
+
+			/** Number of mipmaps to generate. 0 means all the chain. */
+			uint8 LODs;
+        };
+
+        struct ImageGPUArgs
+        {
+            ADDRESS program;
+        };
+
+        struct ImageCropArgs
+        {
+            ADDRESS source;
+            uint16 minX, minY, sizeX, sizeY;
+        };
+
+        struct ImagePatchArgs
+        {
+            ADDRESS base;
+            ADDRESS patch;
+            uint16 minX, minY;
+        };
+
+        struct ImageRasterMeshArgs
+        {
+            ADDRESS mesh;
+
+			//! These are used in case of projected mesh raster.
+			ADDRESS image;
+			ADDRESS angleFadeProperties;
+
+			//! Mask selecting the pixels in the destination image that may receive projection.
+			ADDRESS mask;
+
+			//! A projector may be needed for some kind of per-pixel raster operations
+			//! like cylindrical projections.
+			ADDRESS projector;
+			
+			int32 blockId;
+			uint16 sizeX, sizeY;
+			uint16 SourceSizeX, SourceSizeY;
+			uint16 CropMinX, CropMinY;
+			uint16 UncroppedSizeX, UncroppedSizeY;
+			uint8 bIsRGBFadingEnabled : 1;
+			uint8 bIsAlphaFadingEnabled : 1;
+			
+			// Currently only 2 sampling methods are contemplated, but reserve 3 bits for future uses. 
+			uint8 SamplingMethod : 3;
+			// Currently only 2 min filter methods are contemplated, but reserve 3 bits for future uses. 
+			uint8 MinFilterMethod : 3;
+
+			uint8 LayoutIndex;
+        };
+
+        struct ImageMakeGrowMapArgs
+        {
+            ADDRESS mask;
+            int32 border;
+        };
+
+        struct ImageDisplaceArgs
+        {
+            ADDRESS source;
+            ADDRESS displacementMap;
+        };
+
+		struct ImageInvertArgs
+		{
+			ADDRESS base;
+		};
+
+        struct ImageNormalCompositeArgs
+        {
+            ADDRESS base;
+            ADDRESS normal;
+
+            float power;
+            ECompositeImageMode mode;
+        };
+
+		struct ImageTransformArgs
+		{
+			ADDRESS Base = 0;
+			ADDRESS OffsetX = 0;
+			ADDRESS OffsetY = 0;
+			ADDRESS ScaleX = 0;
+			ADDRESS ScaleY = 0;
+			ADDRESS Rotation = 0;
+
+			uint32 AddressMode      : 31;
+			uint32 bKeepAspectRatio : 1;
+
+            /** Size of the image to create. If 0, reuse size from base.*/
+            uint16 SizeX = 0;
+            uint16 SizeY = 0;
+
+			uint16 SourceSizeX = 0;
+			uint16 SourceSizeY = 0;
+		};
+
+        //-------------------------------------------------------------------------------------
+        struct MeshApplyLayoutArgs
+        {
+            ADDRESS mesh;
+            ADDRESS layout;
+            uint16 channel;
+        };
+
+        struct MeshMergeArgs
+        {
+            ADDRESS base;
+            ADDRESS added;
+
+            // If 0, it merges the surfaces, otherwise, add a new surface
+            // for the added mesh.
+            uint32 newSurfaceID;
+        };
+
+        struct MeshInterpolateArgs
+        {
+            ADDRESS factor;
+            ADDRESS base;
+            ADDRESS targets[ MUTABLE_OP_MAX_INTERPOLATE_COUNT-1 ];
+        };
+
+		struct MeshMaskClipMeshArgs
+		{
+			ADDRESS source;
+			ADDRESS clip;
+		};
+
+		struct MeshMaskClipUVMaskArgs
+		{
+			ADDRESS Source = 0;
+			ADDRESS Mask = 0;
+			uint8 LayoutIndex = 0;
+		};
+
+        struct MeshMaskDiffArgs
+        {
+            ADDRESS source;
+            ADDRESS fragment;
+        };
+
+        struct MeshFormatArgs
+        {
+            ADDRESS source;
+            ADDRESS format;
+
+            typedef enum
+            {
+                BT_VERTEX			= 1,
+                BT_INDEX			= 2,
+                BT_FACE				= 4,
+                //! This flag will not add blank channels for the channels in the format mesh but not
+                //! in the source mesh.
+                BT_IGNORE_MISSING	= 16,
+                //! This flag will force the reset of buffer indices to 0
+                BT_RESETBUFFERINDICES = 32
+            } BUFFER_TYPE;
+
+            //! Flag combination, selecting the buffers to reformat. The rest are left untouched.
+            uint8 buffers;
+
+        };
+
+		struct MeshTransformArgs
+		{
+			ADDRESS source;
+			ADDRESS matrix;
+		};
+
+		struct MeshClipMorphPlaneArgs
+		{
+			ADDRESS source;
+
+			ADDRESS morphShape;
+			ADDRESS vertexSelectionShapeOrBone;
+			
+			typedef enum 
+			{
+				VS_NONE,
+				VS_SHAPE,
+				VS_BONE_HIERARCHY
+			} VERTEX_SELECTION_TYPE;
+            uint8 vertexSelectionType;
+
+			float dist, factor, maxBoneRadius;
+		};
+
+        struct MeshClipWithMeshArgs
+        {
+            ADDRESS source;
+            ADDRESS clipMesh;
+        };
+
+        struct MeshSetSkeletonArgs
+        {
+            ADDRESS source;
+            ADDRESS skeleton;
+        };
+
+        struct MeshProjectArgs
+        {
+            ADDRESS mesh;
+            ADDRESS projector;
+        };
+
+        struct MeshApplyPoseArgs
+        {
+            ADDRESS base;
+            ADDRESS pose;
+        };
+
+		struct MeshGeometryOperationArgs
+		{
+			ADDRESS meshA;
+			ADDRESS meshB;
+			ADDRESS scalarA;
+			ADDRESS scalarB;
+		};
+
+		struct MeshBindShapeArgs
+		{
+			ADDRESS mesh;
+			ADDRESS shape;
+			uint32 flags;
+			uint32 bindingMethod;
+			uint32 ColorUsage;
+		};
+
+		struct MeshApplyShapeArgs
+		{
+			ADDRESS mesh;
+			ADDRESS shape;
+			uint32 flags;
+		};
+
+		struct MeshMorphReshapeArgs
+		{
+			ADDRESS Morph;
+			ADDRESS Reshape;
+		};
+
+
+		struct MeshClipDeformArgs
+		{
+			ADDRESS mesh;
+			ADDRESS clipShape;
+
+			float clipWeightThreshold = 0.9f;
+		};
+
+		struct MeshOptimizeSkinningArgs
+		{
+			ADDRESS source;
+		};
+
+
+        //-------------------------------------------------------------------------------------
+        struct InstanceAddArgs
+        {
+            ADDRESS instance;
+            ADDRESS value;
+            uint32 id;
+            uint32 ExternalId;
+
+			// Id used to identify shared surfaces between lods.
+			int32 SharedSurfaceId;
+
+            ADDRESS name;
+
+            // Index in the FProgram::m_parameterLists with the parameters that are relevant
+            // for whatever is added by this operation. This is used only for resources like
+            // images or meshes.
+            ADDRESS relevantParametersListIndex;
+        };
+
+        struct InstanceAddLODArgs
+        {
+            ADDRESS lod[ MUTABLE_OP_MAX_ADD_COUNT ];
+        };
+
+		struct InstanceAddExtensionDataArgs
+		{
+			// This is a reference to an op that produces the Instance that the ExtensionData will
+			// be added to.
+			ADDRESS Instance;
+			// An op that produces the ExtensionData to add to the Instance
+			ADDRESS ExtensionData;
+			// The name to associate with the ExtensionData
+			//
+			// This is an index into the string table
+			ADDRESS ExtensionDataName;
+		};
+
+        //-------------------------------------------------------------------------------------
+        struct LayoutPackArgs
+        {
+            ADDRESS Source;
+        };
+
+        struct LayoutMergeArgs
+        {
+            ADDRESS Base;
+            ADDRESS Added;
+        };
+
+		struct LayoutRemoveBlocksArgs
+		{
+			/** Layout to be processedand modified. */
+			ADDRESS Source;
+
+			/** Source layout to scan for active blocks. */
+			ADDRESS ReferenceLayout;
+		};
+
+		struct LayoutFromMeshArgs
+		{
+			/** Source mesh to retrieve the layout from. */
+			ADDRESS Mesh;
+			uint8 LayoutIndex;
+		};
+
+
+        //-------------------------------------------------------------------------------------
+        //-------------------------------------------------------------------------------------
+        //-------------------------------------------------------------------------------------
+        t_OP()
+        {
+            type = OP_TYPE::NONE;
+            unused = 0;
+			FMemory::Memzero( &args, sizeof(args) );
+        }
+
+        bool operator==( const t_OP<ADDRESS>& o ) const
+        {
+            if (o.type!=type) return false;
+            //if (o.flags!=flags) return false;
+            if (FMemory::Memcmp(&o.args, &args, sizeof(FArgs))) return false;
+            return true;
+        }
+
+        bool operator<( const t_OP<ADDRESS>& o ) const
+        {
+            if (o.type<type) return true;
+            if (o.type>type) return false;
+//            if (o.flags<flags) return true;
+//            if (o.flags>flags) return false;
+            if (FMemory::Memcmp(&o.args, &args, sizeof(FArgs))<0) return true;
+            return false;
+        }
+
+        OP_TYPE type;
+
+        typedef enum
+        {
+            //! In a particular state, this op is used only once. Knowing this help optimising the
+            //! build process since it allows reusing the resources in the operation result.
+            //F_SINGLE_USE = 1 << 0
+
+        } FLAGS;
+
+        //! Bitmask of FLAGS values
+        uint16 unused;
+
+        //! Constant arguments per operation type
+		//! Everything in this union needs to be converted to their own AST class eventually.
+        typedef union
+        {
+            IntConstantArgs IntConstant;
+            ScalarConstantArgs ScalarConstant;
+            ColourConstantArgs ColourConstant;
+            //ParameterArgs Parameter;
+
+            //-------------------------------------------------------------------------------------
+            BoolLessArgs BoolLess;
+            BoolEqualScalarConstArgs BoolEqualScalarConst;
+            BoolBinaryArgs BoolBinary;
+            BoolNotArgs BoolNot;
+
+            //-------------------------------------------------------------------------------------
+            ScalarMultiplyAddArgs ScalarMultiplyAdd;
+            ArithmeticArgs ScalarArithmetic;
+
+            //-------------------------------------------------------------------------------------
+            ColourSampleImageArgs ColourSampleImage;
+            ColourSwizzleArgs ColourSwizzle;
+            ColourMultiplyArgs ColourMultiply;
+            ColourFromScalarsArgs ColourFromScalars;
+            ArithmeticArgs ColourArithmetic;
+
+            //-------------------------------------------------------------------------------------
+            ImageResizeArgs ImageResize;
+            ImageResizeLikeArgs ImageResizeLike;
+            ImageResizeVarArgs ImageResizeVar;
+            ImageResizeRelArgs ImageResizeRel;
+            ImageBlankLayoutArgs ImageBlankLayout;
+            ImageInterpolateArgs ImageInterpolate;
+            ImageInterpolate3Args ImageInterpolate3;
+            ImageSaturateArgs ImageSaturate;
+            ImageLuminanceArgs ImageLuminance;
+            ImageColourMapArgs ImageColourMap;
+            ImageGradientArgs ImageGradient;
+            ImageBinariseArgs ImageBinarise;
+            ImagePlainColourArgs ImagePlainColour;
+            ImageGPUArgs ImageGPU;
+            ImageDisplaceArgs ImageDisplace;
+			ImageInvertArgs ImageInvert;
+
+            //-------------------------------------------------------------------------------------
+            MeshApplyLayoutArgs MeshApplyLayout;
+            MeshMergeArgs MeshMerge;
+            MeshInterpolateArgs MeshInterpolate;
+            MeshMaskDiffArgs MeshMaskDiff;
+			MeshClipMorphPlaneArgs MeshClipMorphPlane;
+            MeshClipWithMeshArgs MeshClipWithMesh;
+            MeshSetSkeletonArgs MeshSetSkeleton;
+			MeshProjectArgs MeshProject;
+
+        } FArgs;
+
+        FArgs args;
+    };
+
+    typedef t_OP<uint32> OP;
+
+    // OP is not serialisable anymore, we link code now.
+    MUTABLE_DEFINE_POD_SERIALISABLE(OP);
+    MUTABLE_DEFINE_POD_VECTOR_SERIALISABLE(OP);
+    
+    // Check that we didn't go out of control with the operation size
+    static_assert( sizeof(OP::FArgs)==28, "Argument union has an unexpected size." );
+    static_assert( sizeof(OP)==32, "Operation struct has an unexpected size." );
+
+
+	//! Types of data handled by the Mutable runtime.
+	typedef enum
+	{
+		DT_NONE,
+		DT_BOOL,
+		DT_INT,
+		DT_SCALAR,
+		DT_COLOUR,
+		DT_IMAGE,
+		DT_VOLUME_DEPRECATED,
+		DT_LAYOUT,
+		DT_MESH,
+		DT_INSTANCE,
+		DT_PROJECTOR,
+		DT_STRING,
+		DT_EXTENSION_DATA,
+
+		// Supporting data types : Never returned as an actual data type for any operation.
+		DT_MATRIX,
+		DT_SHAPE,
+		DT_CURVE,
+		DT_SKELETON,
+		DT_PHYSICS_ASSET,
+		
+		DT_COUNT
+	} DATATYPE;
+
+
+	// Generic data about a Mutable runtime operation.
+    struct OP_DESC
+    {
+		//! Type of data generated by the instruction
+		DATATYPE type;
+
+        //! True if the instruction is worth caching when generating models
+        bool cached;
+
+        //! For image instructions, for every image format, true if it is supported as the base
+        //! format of the operation.
+        //! TODO: Move to tools library?
+		bool supportedBasePixelFormats[uint8(EImageFormat::IF_COUNT)];
+    };
+
+	MUTABLERUNTIME_API extern const OP_DESC& GetOpDesc( OP_TYPE type );
+
+
+    //!
+    inline static DATATYPE GetOpDataType( OP_TYPE type )
+    {
+		return GetOpDesc(type).type;
+    }
+
+    //! Utility function to apply a function to all operation references to other operations.
+	MUTABLERUNTIME_API extern void ForEachReference( OP& op, const TFunctionRef<void(OP::ADDRESS*)> );
+
+    //! Utility function to apply a function to all operation references to other operations.
+	MUTABLERUNTIME_API extern void ForEachReference( const struct FProgram& program, OP::ADDRESS at, const TFunctionRef<void(OP::ADDRESS)> );
+
+	//!
+	MUTABLERUNTIME_API inline OP_TYPE GetSwitchForType( DATATYPE d )
+    {
+        switch (d)
+        {
+        case DT_INSTANCE: return OP_TYPE::IN_SWITCH;
+        case DT_MESH: return OP_TYPE::ME_SWITCH;
+		case DT_IMAGE: return OP_TYPE::IM_SWITCH;
+		case DT_LAYOUT: return OP_TYPE::LA_SWITCH;
+        case DT_COLOUR: return OP_TYPE::CO_SWITCH;
+        case DT_SCALAR: return OP_TYPE::SC_SWITCH;
+        case DT_INT: return OP_TYPE::NU_SWITCH;
+		case DT_EXTENSION_DATA: return OP_TYPE::ED_SWITCH;
+
+        default:
+			check(false);
+			break;
+        }
+        return OP_TYPE::NONE;
+    }
+}
+
diff --git a/Engine/Plugins/Experimental/Mutable/Source/MutableRuntime/Internal/MuR/ParametersPrivate.h b/Engine/Plugins/Experimental/Mutable/Source/MutableRuntime/Internal/MuR/ParametersPrivate.h
new file mode 100644
--- /dev/null
+++ b/Engine/Plugins/Experimental/Mutable/Source/MutableRuntime/Internal/MuR/ParametersPrivate.h
@@ -0,0 +1,610 @@
+// Copyright Epic Games, Inc. All Rights Reserved.
+
+#pragma once
+
+#include "Misc/TVariant.h"
+#include "Math/MathFwd.h"
+
+#include "MuR/Parameters.h"
+#include "MuR/SerialisationPrivate.h"
+#include "MuR/Operations.h"
+#include "MuR/Image.h"
+
+#include "MuR/Model.h"
+
+#include <string>
+
+namespace mu
+{
+
+namespace Private
+{
+
+    template<class T>
+    class TIndirectObject
+    {
+        TUniquePtr<T> StoragePtr;
+
+    public:
+        template<typename... TArgs>
+        TIndirectObject(TArgs&&... Args) : StoragePtr(MakeUnique<T>(Forward<TArgs>(Args)...)) 
+        {
+        }
+        
+        TIndirectObject(const TIndirectObject<T>& Other) : StoragePtr(MakeUnique<T>()) 
+        { 
+            Get() = Other; 
+        }
+
+		TIndirectObject(TIndirectObject<T>&& Other) : StoragePtr(MoveTemp(Other.StoragePtr)) 
+        {
+        }
+
+        TIndirectObject(const T& Object) : StoragePtr(MakeUnique<T>()) 
+        { 
+            Get() = Object; 
+        }
+        
+        TIndirectObject(T&& Object) : StoragePtr(MakeUnique<T>()) 
+        { 
+            Get() = MoveTemp(Object); 
+        }
+
+        TIndirectObject& operator=(TIndirectObject<T>&&) = default;
+
+        TIndirectObject& operator=(const TIndirectObject<T>& Other) 
+        { 
+            Get() = Other; 
+			return *this;
+        }
+
+        TIndirectObject& operator=(const T& Object) 
+        { 
+            Get() = Object; 
+            return *this;
+        }
+
+        TIndirectObject& operator=(const T&& Object) 
+        { 
+            Get() = MoveTemp(Object); 
+            return *this;
+        }
+
+        const T& Get() const 
+        { 
+            check(StoragePtr); 
+            return *StoragePtr; 
+        }
+
+        T& Get() 
+        { 
+            check(StoragePtr); 
+            return *StoragePtr; 
+        }
+
+        operator T&() 
+        { 
+            return Get(); 
+        }
+
+        operator const T&() const 
+        { 
+            return Get(); 
+        }
+
+        T* operator &() 
+        { 
+            return StoragePtr.Get(); 
+        }
+
+        const T* operator &() const 
+        { 
+            return StoragePtr.Get(); 
+        }
+
+        T& operator *() 
+        { 
+            return Get(); 
+        }
+
+        const T& operator *() const 
+        { 
+            return Get(); 
+        }
+
+        bool operator==(const TIndirectObject<T>& Other) const 
+        { 
+            return *StoragePtr == Other; 
+        }
+
+        bool operator==(const T& Object) const 
+        { 
+            return *StoragePtr == Object; 
+        }
+
+		//!
+		void Serialise(OutputArchive& Arch) const
+		{
+			Arch << Get();
+		}
+
+		//!
+		void Unserialise(InputArchive& Arch)
+		{
+			Arch >> Get();
+		} 
+    };
+
+}
+
+    MUTABLE_DEFINE_ENUM_SERIALISABLE(PARAMETER_TYPE)
+    MUTABLE_DEFINE_ENUM_SERIALISABLE(PROJECTOR_TYPE)
+
+
+    /** Description of a projector to project an image on a mesh. */
+    struct FProjector
+    {
+        PROJECTOR_TYPE type = PROJECTOR_TYPE::PLANAR;
+        FVector3f position = {0,0,0};
+		FVector3f direction = {0,0,0};
+		FVector3f up = {0,0,0};
+		FVector3f scale = {0,0,0};
+        float projectionAngle = 0.0f;
+
+        //!
+        inline void GetDirectionSideUp(FVector3f& OutDirection, FVector3f& OutSide, FVector3f& OutUp) const
+        {
+			OutDirection = direction;
+            OutUp = up;
+            OutSide = FVector3f::CrossProduct( up, direction );
+			OutSide.Normalize();
+        }
+
+
+        //!
+        void Serialise( OutputArchive& arch ) const
+        {
+            arch << type;
+            arch << position;
+            arch << direction;
+            arch << up;
+            arch << scale;
+            arch << projectionAngle;
+        }
+
+        //!
+        void Unserialise( InputArchive& arch )
+        {
+            arch >> type;
+            arch >> position;
+            arch >> direction;
+            arch >> up;
+            arch >> scale;
+            arch >> projectionAngle;
+        }
+
+        bool operator==( const FProjector& o ) const
+        {
+            return     type==o.type
+                    && position==o.position
+                    && up==o.up
+                    && direction==o.direction
+                    && scale==o.scale
+                    && projectionAngle==o.projectionAngle;
+        }
+
+    };
+
+	//---------------------------------------------------------------------------------------------
+	//! Information about a generic shape in space.
+	//---------------------------------------------------------------------------------------------
+	struct FShape
+	{
+		// Transform
+		vec3f position;
+		vec3f up;
+		vec3f side;
+		
+		vec3f size;
+
+		// 
+		enum class Type : uint8
+		{
+			None = 0,
+			Ellipse,
+			AABox
+		};
+        uint8 type = 0;
+		
+
+		//!
+		void Serialise(OutputArchive& arch) const
+		{
+			arch << position;
+			arch << up;
+			arch << side;
+			arch << size;
+			arch << type;
+		}
+
+		//!
+		void Unserialise(InputArchive& arch)
+		{
+			arch >> position;
+			arch >> up;
+			arch >> side;
+			arch >> size;
+			arch >> type;
+		}
+
+        bool operator==( const FShape& o ) const
+        {
+            return type==o.type
+                    && position==o.position
+                    && up==o.up
+                    && side==o.side
+                    && size==o.size;
+        }
+	};
+
+
+	//---------------------------------------------------------------------------------------------
+	//! Information about a curve interpolated from keyframes.
+	//---------------------------------------------------------------------------------------------
+	struct CurveKeyFrame
+	{
+		float time;
+		float value;
+		float in_tangent;
+		float in_tangent_weight;
+		float out_tangent;
+		float out_tangent_weight;
+
+		enum class InterpMode : uint8
+		{
+			Linear = 0,
+			Constant,
+			Cubic,
+			None,
+
+			//! Utility enumeration value, not really a Interp mode.
+			Count
+		};
+        uint8 interp_mode;
+
+
+		enum class TangentMode : uint8
+		{
+			Auto = 0,
+			Manual,
+			Broken,
+			None,
+
+			//! Utility enumeration value, not really a Tangent mode.
+			Count
+		};
+        uint8 tangent_mode;
+
+
+		enum class TangentWeightMode : uint8
+		{
+			Auto = 0,
+			Manual,
+			Broken,
+			None,
+
+			//! Utility enumeration value, not really a Tangent weight mode.
+			Count
+		};
+        uint8 tangent_weight_mode;
+
+        //!
+        bool operator==( const CurveKeyFrame& other ) const
+        {
+            return time==other.time &&
+                    value==other.value&&
+                    in_tangent==other.in_tangent&&
+                    in_tangent_weight==other.in_tangent_weight&&
+                    out_tangent==other.out_tangent&&
+                    out_tangent_weight==other.out_tangent_weight&&
+                    interp_mode==other.interp_mode&&
+                    tangent_mode==other.tangent_mode&&
+                    tangent_weight_mode==other.tangent_weight_mode;
+        }
+
+		//!
+		void Serialise(OutputArchive& arch) const
+		{
+			arch << time;
+			arch << value;
+			arch << in_tangent;
+			arch << in_tangent_weight;
+			arch << out_tangent;
+			arch << out_tangent_weight;
+
+			arch << interp_mode;
+			arch << tangent_mode;
+			arch << tangent_weight_mode;
+		}
+
+		//!
+		void Unserialise(InputArchive& arch)
+		{
+			arch >> time;
+			arch >> value;
+			arch >> in_tangent;
+			arch >> in_tangent_weight;
+			arch >> out_tangent;
+			arch >> out_tangent_weight;
+
+			arch >> interp_mode;
+			arch >> tangent_mode;
+			arch >> tangent_weight_mode;
+		}
+	};
+
+
+	struct Curve
+	{
+		TArray<CurveKeyFrame> keyFrames;
+		float defaultValue = 0.f;
+
+        //!
+        bool operator==( const Curve& other ) const
+        {
+            return keyFrames == other.keyFrames && defaultValue == other.defaultValue;
+        }
+
+		//!
+		void Serialise(OutputArchive& arch) const
+		{
+            arch << keyFrames;
+            arch << defaultValue;
+		}
+
+		//!
+		void Unserialise(InputArchive& arch)
+		{
+			arch >> keyFrames;
+            arch >> defaultValue;
+        }
+	};
+
+	using ParamBoolType = bool;
+	using ParamIntType = int32;
+	using ParamFloatType = float;
+	using ParamColorType = FVector3f;
+	using ParamProjectorType = Private::TIndirectObject<FProjector>;
+	using ParamImageType = FName;
+	using ParamStringType = Private::TIndirectObject<FString>;
+	
+	using PARAMETER_VALUE = TVariant<
+            ParamBoolType, ParamIntType, ParamFloatType, ParamColorType, ParamProjectorType, ParamImageType, ParamStringType>;
+
+    // static_assert to track PARAMETER_VALUE size changes. It is ok to change if needed.
+    static_assert(sizeof(PARAMETER_VALUE) == 8*3, "PARAMETER_VALUE size has changed.");
+
+	// TVariant currently does not support this operator. Once supported remove it.
+	inline bool operator==(const PARAMETER_VALUE& ValueA, const PARAMETER_VALUE& ValueB)
+	{
+		const int32 IndexA = ValueA.GetIndex();
+		const int32 IndexB = ValueB.GetIndex();
+
+		if (IndexA != IndexB)
+		{
+			return false;
+		}
+
+		return Visit([&ValueB](const auto& StoredValueA)
+		{
+			using Type = typename TDecay<decltype(StoredValueA)>::Type;
+			return StoredValueA == ValueB.Get<Type>();
+		}, ValueA);
+	}
+	
+
+    struct FParameterDesc
+    {
+        FString m_name;
+
+        //! Unique id (provided externally, so no actual guarantee that it is unique.)
+		FString m_uid;
+
+        PARAMETER_TYPE m_type = PARAMETER_TYPE::T_NONE;
+
+        PARAMETER_VALUE m_defaultValue;
+
+        //! Ranges, if the parameter is multi-dimensional. The indices refer to the Model's program
+        //! vector of range descriptors.
+		TArray<uint32> m_ranges;
+
+        //! Possible values of the parameter in case of being an integer, and its names
+        struct FIntValueDesc
+        {
+            int16 m_value;
+			FString m_name;
+
+            //!
+            bool operator==( const FIntValueDesc& other ) const
+            {
+                return m_value==other.m_value &&
+                        m_name==other.m_name;
+            }
+
+            //!
+            void Serialise( OutputArchive& arch ) const
+            {
+                arch << m_value;
+                arch << m_name;
+            }
+
+            //!
+            void Unserialise( InputArchive& arch )
+            {
+                arch >> m_value;
+                arch >> m_name;
+            }
+        };
+
+        //! For integer parameters, this contains the description of the possible values.
+        //! If empty, the integer may have any value.
+		TArray<FIntValueDesc> m_possibleValues;
+
+        //!
+        bool operator==( const FParameterDesc& other ) const
+        {
+            return m_name == other.m_name && m_uid == other.m_uid && m_type == other.m_type &&
+                   m_defaultValue == other.m_defaultValue &&
+                   m_ranges == other.m_ranges &&
+                   m_possibleValues == other.m_possibleValues;
+        }
+
+        //!
+        void Serialise( OutputArchive& arch ) const
+        {
+            const int32 ver = 8;
+            arch << ver;
+
+			arch << m_name;
+			arch << m_uid;
+            arch << m_type;
+            arch << m_defaultValue;
+            arch << m_ranges;
+            arch << m_possibleValues;
+        }
+
+        //!
+        void Unserialise( InputArchive& arch )
+        {
+            int32 ver;
+            arch >> ver;
+            check( ver==8 );
+
+			arch >> m_name;
+			arch >> m_uid;
+            arch >> m_type;
+            arch >> m_defaultValue;
+            arch >> m_ranges;
+			arch >> m_possibleValues;
+        }
+    };
+
+
+    struct FRangeDesc
+    {
+		FString m_name;
+		FString m_uid;
+
+		/** Parameter that controls the size of this range, if any. */
+		int32 m_dimensionParameter = -1;
+
+        //!
+        bool operator==( const FRangeDesc& other ) const
+        {
+            return m_name==other.m_name
+				&&
+				m_uid == other.m_uid 
+				&&
+				m_dimensionParameter == other.m_dimensionParameter;
+        }
+
+        //!
+        void Serialise( OutputArchive& arch ) const
+        {
+            const uint32 ver = 3;
+            arch << ver;
+
+            arch << m_name;
+            arch << m_uid;
+			arch << m_dimensionParameter;
+        }
+
+        //!
+        void Unserialise( InputArchive& arch )
+        {
+            uint32 ver;
+            arch >> ver;
+            check( ver == 3 );
+
+			arch >> m_name;
+			arch >> m_uid;
+			arch >> m_dimensionParameter;
+        }
+    };
+
+
+    class MUTABLERUNTIME_API RangeIndex::Private : public Base
+    {
+    public:
+
+        //! Run-time data
+        Ptr<const Parameters> m_pParameters;
+
+        //! Index of the parameter for which we are a range index
+        int m_parameter = -1;
+
+        //! Position in the several dimension of the range, as defined in m_pParameters
+		TArray<int> m_values;
+    };
+
+
+    class MUTABLERUNTIME_API Parameters::Private : public Base
+    {
+    public:
+
+        //! Warning: update Parameters::Clone method if this members change.
+
+        //! Run-time data
+        TSharedPtr<const Model> m_pModel;
+ 
+        //! Values for the parameters if they are not multidimensional.
+		TArray<PARAMETER_VALUE> m_values;
+
+
+        //! If the parameter is multidemensional, the values are stored here.
+        //! The key of the map is the vector of values stored in a RangeIndex
+		TArray< TMap< TArray<int32>, PARAMETER_VALUE > > m_multiValues;
+
+
+        //!
+        void Serialise( OutputArchive& arch ) const
+        {
+            const uint32 ver = 3;
+            arch << ver;
+
+            arch << m_values;
+            arch << m_multiValues;
+        }
+
+        //!
+        void Unserialise( InputArchive& arch )
+        {
+            uint32 ver;
+            arch >> ver;
+			check(ver == 3);
+       
+		    arch >> m_values;
+			arch >> m_multiValues;
+        }
+
+        //!
+        int32 Find( const FString& Name ) const;
+
+        //!
+        FProjector GetProjectorValue( int index, const Ptr<const RangeIndex>& pos ) const;
+
+		/** Return true if the parameter has any multi-dimensional values set. This is independent to if the model
+		* accepts multi-dimensional parameters for this particular parameter.
+		*/
+		inline bool HasMultipleValues(int32 ParamIndex) const
+		{
+			if (ParamIndex >= m_multiValues.Num())
+			{
+				return false;
+			}
+
+			return m_multiValues[ParamIndex].Num()>0;
+		}
+    };
+
+}
diff --git a/Engine/Plugins/Experimental/Mutable/Source/MutableRuntime/Internal/MuR/SerialisationPrivate.h b/Engine/Plugins/Experimental/Mutable/Source/MutableRuntime/Internal/MuR/SerialisationPrivate.h
new file mode 100644
--- /dev/null
+++ b/Engine/Plugins/Experimental/Mutable/Source/MutableRuntime/Internal/MuR/SerialisationPrivate.h
@@ -0,0 +1,573 @@
+// Copyright Epic Games, Inc. All Rights Reserved.
+
+#pragma once
+
+#include "Containers/StaticArray.h"
+#include "MuR/Serialisation.h"
+
+#include "Misc/TVariant.h"
+#include "Math/Color.h"
+#include "Math/Quat.h"
+#include "Math/Vector.h"
+#include "Math/IntVector.h"
+#include "Math/Vector4.h"
+#include "Math/Transform.h"
+#include "MuR/Ptr.h"
+#include "MuR/RefCounted.h"
+#include "MuR/MutableMemory.h"
+#include "MuR/MutableMath.h"
+#include "MuR/Types.h"
+
+#include <string>
+
+
+namespace mu
+{	
+
+    //!
+    class MUTABLERUNTIME_API InputMemoryStream : public InputStream
+    {
+    public:
+
+        //-----------------------------------------------------------------------------------------
+        // Life cycle
+        //-----------------------------------------------------------------------------------------
+
+        //! Create the stream using an external buffer.
+        //! The buffer will not be owned by this object, so it cannot be deallocated while this
+        //! objects is in use.
+        InputMemoryStream( const void* pBuffer, uint64 size );
+
+        ~InputMemoryStream();
+
+
+        //-----------------------------------------------------------------------------------------
+        // InputStream interface
+        //-----------------------------------------------------------------------------------------
+        void Read( void* pData, uint64 size ) override;
+
+
+    private:
+
+        class Private;
+        Private* m_pD;
+
+    };
+
+
+    //! This stream doesn-t store any data, it just counts de amonut of data serialised.
+    class MUTABLERUNTIME_API OutputSizeStream : public OutputStream
+    {
+    public:
+
+        //-----------------------------------------------------------------------------------------
+        // Life cycle
+        //-----------------------------------------------------------------------------------------
+
+        //!
+        OutputSizeStream();
+
+        //-----------------------------------------------------------------------------------------
+        // OutputStream interface
+        //-----------------------------------------------------------------------------------------
+        void Write( const void* pData, uint64 size ) override;
+
+        //-----------------------------------------------------------------------------------------
+        // Own interface
+        //-----------------------------------------------------------------------------------------
+
+        //! Get the amount of data serialised, in bytes.
+        uint64 GetBufferSize() const;
+
+    private:
+
+        uint64 m_writtenBytes;
+
+    };
+
+
+    //---------------------------------------------------------------------------------------------
+    class MUTABLERUNTIME_API InputMemoryStream::Private : public Base
+    {
+    public:
+
+        Private()
+        {
+            m_pBuffer = 0;
+            m_size = 0;
+            m_pos = 0;
+        }
+
+        const void* m_pBuffer;
+
+		uint64 m_size;
+
+		uint64 m_pos;
+    };
+
+
+    //---------------------------------------------------------------------------------------------
+    class MUTABLERUNTIME_API OutputMemoryStream::Private : public Base
+    {
+    public:
+        TArray<uint8,TSizedHeapAllocator<64>> m_buffer;
+    };
+
+
+	//---------------------------------------------------------------------------------------------
+	class MUTABLERUNTIME_API InputArchive::Private : public Base
+	{
+	public:
+
+		//! Not owned
+		InputStream* m_pStream;
+
+		//! Already read pointers
+		TArray< Ptr<RefCounted> > m_history;
+
+	};
+
+	//---------------------------------------------------------------------------------------------
+	class MUTABLERUNTIME_API OutputArchive::Private : public Base
+	{
+	public:
+
+		//! Not owned
+		OutputStream* m_pStream;
+
+		//! Already written pointers and their ids
+        TMap< const void*, int32 > m_history;
+
+	};
+
+
+	//---------------------------------------------------------------------------------------------
+#define MUTABLE_IMPLEMENT_POD_SERIALISABLE(T)							\
+		template<>														\
+		void DLLEXPORT operator<< <T>( OutputArchive& arch, const T& t )			\
+		{																\
+			arch.GetPrivate()->m_pStream->Write( &t, sizeof(T) );		\
+		}																\
+																		\
+		template<>														\
+		void DLLEXPORT operator>> <T>( InputArchive& arch, T& t )					\
+		{																\
+			arch.GetPrivate()->m_pStream->Read( &t, sizeof(T) );		\
+		}																\
+		
+
+#define MUTABLE_IMPLEMENT_POD_VECTOR_SERIALISABLE(T)									     \
+	template<typename Alloc>                                                             \
+	void operator <<(OutputArchive& arch, const TArray<T, Alloc>& v)					 \
+	{                                                                                    \
+		uint32 Num = uint32(v.Num());													 \
+		arch << Num;																	 \
+		if (Num)																		 \
+		{                                                                                \
+			arch.GetPrivate()->m_pStream->Write(&v[0], Num * sizeof(T));				 \
+		}                                                                                \
+	}                                                                                    \
+	                                                                                     \
+	template<typename Alloc>                                                             \
+	void operator >>(InputArchive& arch, TArray<T, Alloc>& v)                            \
+	{                                                                                    \
+		uint32 Num;																		 \
+		arch >> Num;																	 \
+		v.SetNum(Num);																	 \
+		if (Num)																		 \
+		{                                                                                \
+			arch.GetPrivate()->m_pStream->Read(&v[0], Num * sizeof(T));					 \
+		}                                                                                \
+	}                                                                                    \
+
+	//---------------------------------------------------------------------------------------------
+	// TVariant custom serialize. Based on the default serialization.
+	//---------------------------------------------------------------------------------------------
+	template <typename... Ts>
+	void operator<<(OutputArchive& Ar, const TVariant<Ts...>& Variant)
+	{
+		const uint8 Index = static_cast<uint8>(Variant.GetIndex());
+		Ar << Index;
+		
+		Visit([&Ar](auto& StoredValue)
+		{
+			Ar << StoredValue;
+		}, Variant);
+	}
+
+
+	template <typename T, typename VariantType>
+	struct TVariantLoadFromInputArchiveCaller
+	{
+		/** Default construct the type and load it from the FArchive */
+		static void Load(InputArchive& Ar, VariantType& OutVariant)
+		{
+			OutVariant.template Emplace<T>();
+			Ar >> OutVariant.template Get<T>();
+		}
+	};
+
+	
+	template <typename... Ts>
+	struct TVariantLoadFromInputArchiveLookup
+	{
+		using VariantType = TVariant<Ts...>;
+		static_assert((std::is_default_constructible<Ts>::value && ...), "Each type in TVariant template parameter pack must be default constructible in order to use FArchive serialization");
+
+		/** Load the type at the specified index from the FArchive and emplace it into the TVariant */
+		static void Load(SIZE_T TypeIndex, InputArchive& Ar, VariantType& OutVariant)
+		{
+			static constexpr void(*Loaders[])(InputArchive&, VariantType&) = { &TVariantLoadFromInputArchiveCaller<Ts, VariantType>::Load... };
+			check(TypeIndex < UE_ARRAY_COUNT(Loaders));
+			Loaders[TypeIndex](Ar, OutVariant);
+		}
+	};
+
+	
+	template <typename... Ts>
+	void operator>>(InputArchive& Ar, TVariant<Ts...>& Variant)
+	{
+		uint8 Index;
+		Ar >> Index;
+		check(Index < sizeof...(Ts));
+
+		TVariantLoadFromInputArchiveLookup<Ts...>::Load(static_cast<SIZE_T>(Index), Ar, Variant);
+	}
+	
+
+	//---------------------------------------------------------------------------------------------
+	// Bool size is not a standard
+	//---------------------------------------------------------------------------------------------
+	template<>														
+	inline void operator<< <bool>( OutputArchive& arch, const bool& t )
+	{
+        uint8 s = t ? 1 : 0;
+        arch.GetPrivate()->m_pStream->Write( &s, sizeof(uint8) );
+	}
+
+	template<>
+	inline void operator>> <bool>( InputArchive& arch, bool& t )
+	{
+        uint8 s;
+        arch.GetPrivate()->m_pStream->Read( &s, sizeof(uint8) );
+		t = s!=0;
+	}
+
+
+	//---------------------------------------------------------------------------------------------
+#define MUTABLE_IMPLEMENT_ENUM_SERIALISABLE(T)								\
+		template<>															\
+        void DLLEXPORT operator<< <T>( OutputArchive& arch, const T& t )	\
+		{																	\
+            auto v = (uint32)t;                                         	\
+            arch.GetPrivate()->m_pStream->Write( &v, sizeof(uint32) );  	\
+		}																	\
+																			\
+		template<>															\
+        void DLLEXPORT operator>> <T>( InputArchive& arch, T& t )   		\
+		{																	\
+            uint32 v;														\
+            arch.GetPrivate()->m_pStream->Read( &v, sizeof(uint32) );		\
+			t = (T)v;														\
+		}																	\
+
+    //---------------------------------------------------------------------------------------------
+    template< typename T0, typename T1 >
+    inline void operator<< ( OutputArchive& arch, const std::pair<T0,T1>& v )
+    {
+        arch << v.first;
+        arch << v.second;
+    }
+
+    template< typename T0, typename T1 >
+    inline void operator>> ( InputArchive& arch, std::pair<T0,T1>& v )
+    {
+        arch >> v.first;
+        arch >> v.second;
+    }
+
+	
+	//---------------------------------------------------------------------------------------------
+	template<typename T, uint32 Size> void operator<<(OutputArchive& arch, const TStaticArray<T, Size>& v)
+	{
+		for (uint32 i = 0; i < Size; ++i)
+		{
+			arch << v[i];
+		}
+	}
+
+	template<typename T, uint32 Size> void operator>>(InputArchive& arch, TStaticArray<T, Size>& v)
+	{
+		for (uint32 i = 0; i < Size; ++i)
+		{
+			arch >> v[i];
+		}
+	}
+	
+
+	//---------------------------------------------------------------------------------------------
+	template< typename K, typename T >
+	inline void operator<< (OutputArchive& Arch, const TMap<K, T>& V)
+	{
+		Arch << (uint32)V.Num();
+		for (const TPair<K, T>& Element : V)
+		{
+			Arch << Element.Key;
+			Arch << Element.Value;
+		}
+	}
+
+	template< typename K, typename T >
+	inline void operator>> (InputArchive& Arch, TMap<K, T>& V)
+	{
+		uint32 Num;
+		Arch >> Num;
+
+		for (uint32 Index = 0; Index < Num; ++Index)
+		{
+			K Key;
+			T Element;
+			Arch >> Key;
+			Arch >> Element;
+
+			V.Emplace(MoveTemp(Key), MoveTemp(Element));
+		}
+	}
+
+	MUTABLE_DEFINE_POD_SERIALISABLE(vec2f);
+	MUTABLE_DEFINE_POD_SERIALISABLE(vec3f);
+	MUTABLE_DEFINE_POD_SERIALISABLE(mat3f);
+	MUTABLE_DEFINE_POD_SERIALISABLE(mat4f);
+	MUTABLE_DEFINE_POD_SERIALISABLE(vec2<int>);
+
+	// Unreal POD Serializables
+	MUTABLE_DEFINE_POD_SERIALISABLE(FUintVector2);
+	MUTABLE_DEFINE_POD_SERIALISABLE(UE::Math::TIntVector2<uint16>);
+	MUTABLE_DEFINE_POD_SERIALISABLE(UE::Math::TIntVector2<int16>);
+	MUTABLE_DEFINE_POD_SERIALISABLE(FVector4f);
+	MUTABLE_DEFINE_POD_SERIALISABLE(UE::Math::TVector4<float>);
+
+	MUTABLE_IMPLEMENT_POD_VECTOR_SERIALISABLE(float);
+	MUTABLE_IMPLEMENT_POD_VECTOR_SERIALISABLE(double);
+	MUTABLE_IMPLEMENT_POD_VECTOR_SERIALISABLE(uint8);
+	MUTABLE_IMPLEMENT_POD_VECTOR_SERIALISABLE(uint16);
+	MUTABLE_IMPLEMENT_POD_VECTOR_SERIALISABLE(uint32);
+	MUTABLE_IMPLEMENT_POD_VECTOR_SERIALISABLE(uint64);
+	MUTABLE_IMPLEMENT_POD_VECTOR_SERIALISABLE(int8);
+	MUTABLE_IMPLEMENT_POD_VECTOR_SERIALISABLE(int16);
+	MUTABLE_IMPLEMENT_POD_VECTOR_SERIALISABLE(int32);
+	MUTABLE_IMPLEMENT_POD_VECTOR_SERIALISABLE(int64);
+
+	MUTABLE_IMPLEMENT_POD_VECTOR_SERIALISABLE(vec2f);
+	MUTABLE_IMPLEMENT_POD_VECTOR_SERIALISABLE(vec3f);
+	MUTABLE_IMPLEMENT_POD_VECTOR_SERIALISABLE(mat3f);
+	MUTABLE_IMPLEMENT_POD_VECTOR_SERIALISABLE(mat4f);
+	MUTABLE_IMPLEMENT_POD_VECTOR_SERIALISABLE(vec2<int>);
+	MUTABLE_IMPLEMENT_POD_VECTOR_SERIALISABLE(TCHAR);
+
+	MUTABLE_IMPLEMENT_POD_VECTOR_SERIALISABLE(FUintVector2);
+	MUTABLE_IMPLEMENT_POD_VECTOR_SERIALISABLE(UE::Math::TIntVector2<uint16>);
+	MUTABLE_IMPLEMENT_POD_VECTOR_SERIALISABLE(UE::Math::TIntVector2<int16>);
+	MUTABLE_IMPLEMENT_POD_VECTOR_SERIALISABLE(FVector4f);
+
+	//---------------------------------------------------------------------------------------------
+	//template<>
+	//inline void operator<< <std::string>( OutputArchive& arch, const std::string& v )
+	//{
+	//	arch << (uint32)v.size();
+	//	if ( v.size() )
+	//	{
+	//		arch.GetPrivate()->m_pStream->Write( &v[0], (unsigned)v.size()*sizeof(char) );
+	//	}
+	//}
+
+	template<>
+	inline void operator>> <std::string>( InputArchive& arch, std::string& v )
+	{
+        uint32 size;
+		arch >> size;
+		v.resize( size );
+		if (size)
+		{
+			arch.GetPrivate()->m_pStream->Read( &v[0], (unsigned)size*sizeof(char) );
+		}
+	}
+
+
+	//---------------------------------------------------------------------------------------------
+	template< typename T >
+	inline void operator<< ( OutputArchive& arch, const Ptr<const T>& p )
+	{
+		if (!p.get())
+		{
+            arch << (int32)-1;
+		}
+		else
+		{
+            int32* it = arch.GetPrivate()->m_history.Find(p.get());
+
+			if ( !it )
+			{
+                int32 id = arch.GetPrivate()->m_history.Num();
+                arch.GetPrivate()->m_history.Add( p.get(), id );
+				arch << id;
+				T::Serialise( p.get(), arch );
+			}
+			else
+			{
+				arch << *it;
+			}
+		}
+	}
+
+	template< typename T >
+	inline void operator>> ( InputArchive& arch, Ptr<const T>& p )
+	{
+        int32 id;
+		arch >> id;
+
+		if ( id == -1 )
+		{
+			p = 0;
+		}
+		else
+		{
+			if ( id < arch.GetPrivate()->m_history.Num() )
+			{
+				p = static_cast<T*>( arch.GetPrivate()->m_history[id].get() );
+
+				// If the pointer was 0 it means the position in history is used, but not set yet
+				// option 1: we have a smart pointer loop which is very bad.
+				// option 2: the resource in this Ptr is also pointed by a Proxy that has absorbed it
+				//			 and this reference should also be a proxy instead of a ptr.
+                check( p );
+			}
+			else
+			{
+                // Ids come in order, but they may have been absorbed outside in some serialisations
+                // like proxies.
+                //check( id == (int)arch.GetPrivate()->m_history.size() );
+                arch.GetPrivate()->m_history.SetNum( id+1 );
+
+				Ptr<T> t = T::StaticUnserialise( arch );
+				p = t;
+				arch.GetPrivate()->m_history[id] = t;
+			}
+		}
+	}
+
+
+	//---------------------------------------------------------------------------------------------
+	template< typename T >
+	inline void operator<< ( OutputArchive& arch, const Ptr<T>& p )
+	{
+		operator<<( arch, (const Ptr<const T>&) p );
+	}
+
+	template< typename T >
+	inline void operator>> ( InputArchive& arch, Ptr<T>& p )
+	{
+		operator>>( arch, (Ptr<const T>&) p );
+	}
+
+
+	//---------------------------------------------------------------------------------------------
+	//---------------------------------------------------------------------------------------------
+	//---------------------------------------------------------------------------------------------
+	template<typename T0, typename T1> inline void operator<<(OutputArchive& arch, const TPair<T0,T1>& v)
+	{
+		arch << v.Key;
+		arch << v.Value;
+	}
+
+	template<typename T0, typename T1> inline void operator>>(InputArchive& arch, TPair<T0, T1>& v)
+	{
+		arch >> v.Key;
+		arch >> v.Value;
+	}
+
+
+	//---------------------------------------------------------------------------------------------
+	// TODO: As POD?
+	template< typename T >
+	inline void operator<< (OutputArchive& arch, const UE::Math::TQuat<T>& v)
+	{		
+		arch << v.X;
+		arch << v.Y;
+		arch << v.Z;
+		arch << v.W;
+	}
+
+	template< typename T >
+	inline void operator>> (InputArchive& arch, UE::Math::TQuat<T>& v)
+	{
+		arch >> v.X;
+		arch >> v.Y;
+		arch >> v.Z;
+		arch >> v.W;
+	}
+
+
+	//---------------------------------------------------------------------------------------------
+	// TODO: As POD?
+	template< typename T >
+	inline void operator<< (OutputArchive& arch, const UE::Math::TVector<T>& v)
+	{
+		arch << v.X;
+		arch << v.Y;
+		arch << v.Z;
+	}
+
+	template< typename T >
+	inline void operator>> (InputArchive& arch, UE::Math::TVector<T>& v)
+	{
+		arch >> v.X;
+		arch >> v.Y;
+		arch >> v.Z;
+	}
+
+
+	//---------------------------------------------------------------------------------------------
+	template< typename T >
+	inline void operator<< (OutputArchive& arch, const UE::Math::TTransform<T>& v)
+	{
+		arch << v.GetRotation();
+		arch << v.GetTranslation();
+		arch << v.GetScale3D();
+	}
+
+	template< typename T >
+	inline void operator>> (InputArchive& arch, UE::Math::TTransform<T>& v)
+	{
+		UE::Math::TQuat<T> Rot;
+		UE::Math::TVector<T> Trans;
+		UE::Math::TVector<T> Scale;
+
+		arch >> Rot;
+		arch >> Trans;
+		arch >> Scale;
+
+		v.SetComponents(Rot, Trans, Scale);
+	}
+
+
+	//---------------------------------------------------------------------------------------------
+	template< typename T >
+	inline void operator<< (OutputArchive& arch, const UE::Math::TVector4<T>& v)
+	{
+		arch << v.X;
+		arch << v.Y;
+		arch << v.Z;
+		arch << v.W;
+	}
+
+	template< typename T >
+	inline void operator>> (InputArchive& arch, UE::Math::TVector4<T>& v)
+	{
+		arch >> v.X;
+		arch >> v.Y;
+		arch >> v.Z;
+		arch >> v.W;
+	}
+
+}
diff --git a/Engine/Plugins/Experimental/Mutable/Source/MutableRuntime/Internal/MuR/SystemPrivate.h b/Engine/Plugins/Experimental/Mutable/Source/MutableRuntime/Internal/MuR/SystemPrivate.h
new file mode 100644
--- /dev/null
+++ b/Engine/Plugins/Experimental/Mutable/Source/MutableRuntime/Internal/MuR/SystemPrivate.h
@@ -0,0 +1,1898 @@
+// Copyright Epic Games, Inc. All Rights Reserved.
+
+#pragma once
+
+#include "MuR/System.h"
+
+#include "MuR/Settings.h"
+#include "MuR/Operations.h"
+#include "MuR/Image.h"
+#include "MuR/MutableString.h"
+#include "MuR/ModelPrivate.h"
+#include "MuR/Mesh.h"
+#include "MuR/Instance.h"
+#include "MuR/ParametersPrivate.h"
+#include "MuR/MutableTrace.h"
+
+#if !(UE_BUILD_SHIPPING || UE_BUILD_TEST)
+#include "HAL/Thread.h"
+#include "HAL/PlatformTLS.h"
+#endif
+
+namespace mu::MemoryCounters
+{
+	struct FMemoryTrackerInternalMemoryCounterTag {};
+	using FMemoryTrackerInternalMemoryCounter = TMemoryCounter<FMemoryTrackerInternalMemoryCounterTag>;
+}
+
+namespace mu
+{
+	class ExtensionDataStreamer;
+
+	// Call the tick of the LLM system (we do this to simulate a frame since the LLM system is not entirelly designed to run over a program)
+	inline void UpdateLLMStats()
+	{
+		// This code will only be compiled (and ran) if the global definition to enable LLM tracking is set to 1 for the host program
+		// Ex : 			GlobalDefinitions.Add("LLM_ENABLED_IN_CONFIG=1");
+#if ENABLE_LOW_LEVEL_MEM_TRACKER && IS_PROGRAM
+		FLowLevelMemTracker& MemTracker = FLowLevelMemTracker::Get();
+		if (MemTracker.IsEnabled())
+		{
+			MemTracker.UpdateStatsPerFrame();
+		}
+#endif
+	}
+
+	constexpr uint64 AllParametersMask = TNumericLimits<uint64> ::Max();
+
+
+	/** ExecutinIndex stores the location inside all ranges of the execution of a specific
+	* operation. The first integer on each pair is the dimension/range index in the program
+	* array of ranges, and the second integer is the value inside that range.
+	* The vector order is undefined.
+	*/
+	class ExecutionIndex : public  TArray<TPair<int32, int32>>
+	{
+	public:
+		//! Set or add a value to the index
+		void SetFromModelRangeIndex(uint16 rangeIndex, int rangeValue)
+		{
+			auto Index = IndexOfByPredicate([=](const ElementType& v) { return v.Key >= rangeIndex; });
+			if (Index != INDEX_NONE && (*this)[Index].Key == rangeIndex)
+			{
+				// Update
+				(*this)[Index].Value = rangeValue;
+			}
+			else
+			{
+				// Add new
+				Push(ElementType(rangeIndex, rangeValue));
+			}
+		}
+
+		//! Get the value of the index from the range index in the model.
+		int GetFromModelRangeIndex(int modelRangeIndex) const
+		{
+			for (const auto& e : *this)
+			{
+				if (e.Key == modelRangeIndex)
+				{
+					return e.Value;
+				}
+			}
+			return 0;
+		}
+	};
+
+
+	/** This structure stores the data about an ongoing mutable operation that needs to be executed. */
+	struct FScheduledOp
+	{
+		inline FScheduledOp()
+		{
+			Stage = 0;
+			Type = EType::Full;
+		}
+
+		inline FScheduledOp(OP::ADDRESS InAt, const FScheduledOp& InOpTemplate, uint8 InStage = 0, uint32 InCustomState = 0)
+		{
+			check(InStage < 120);
+			At = InAt;
+			ExecutionOptions = InOpTemplate.ExecutionOptions;
+			ExecutionIndex = InOpTemplate.ExecutionIndex;
+			Stage = InStage;
+			CustomState = InCustomState;
+			Type = InOpTemplate.Type;
+		}
+
+		static inline FScheduledOp FromOpAndOptions(OP::ADDRESS InAt, const FScheduledOp& InOpTemplate, uint8 InExecutionOptions)
+		{
+			FScheduledOp r;
+			r.At = InAt;
+			r.ExecutionOptions = InExecutionOptions;
+			r.ExecutionIndex = InOpTemplate.ExecutionIndex;
+			r.Stage = 0;
+			r.CustomState = InOpTemplate.CustomState;
+			r.Type = InOpTemplate.Type;
+			return r;
+		}
+
+		//! Address of the operation
+		OP::ADDRESS At = 0;
+
+		//! Additional custom state data that the operation can store. This is usually used to pass information
+		//! between execution stages of an operation.
+		uint32 CustomState = 0;
+
+		//! Index of the operation execution: This is used for iteration of different ranges.
+		//! It is an index into the CodeRunner::GetMemory()::m_rangeIndex vector.
+		//! executionIndex 0 is always used for empty ExecutionIndex, which is the most common
+		//! one.
+		uint16 ExecutionIndex = 0;
+
+		//! Additional execution options. Set externally to this op, it usually alters the result.
+		//! For example, this is used to keep track of the mipmaps to skip in image operations.
+		uint8 ExecutionOptions = 0;
+
+		//! Internal stage of the operation.
+		//! Stage 0 is usually scheduling of children, and 1 is execution. Some instructions
+		//! may have more steges to schedule children that are optional for execution, etc.
+		uint8 Stage : 7;
+
+		//! Type of calculation we are requesting for this operation.
+		enum class EType : uint8
+		{
+			//! Execute the operation to calculate the full result
+			Full,
+
+			//! Execute the operation to obtain the descriptor of an image.
+			ImageDesc
+		};
+		EType Type : 1;
+	};
+
+	inline uint32 GetTypeHash(const FScheduledOp& Op)
+	{
+		return HashCombine(::GetTypeHash(Op.At), HashCombine(::GetTypeHash(Op.Stage), ::GetTypeHash(Op.ExecutionIndex)));
+	}
+
+
+	/** A cache address is the operation plus the context of execution(iteration indices, etc...). */
+	struct FCacheAddress
+	{
+		/** The meaning of all these fields is the same than the FScheduledOp struct. */
+		OP::ADDRESS At = 0;
+		uint16 ExecutionIndex = 0;
+		uint8 ExecutionOptions = 0;
+		FScheduledOp::EType Type = FScheduledOp::EType::Full;
+
+		FCacheAddress() {}
+
+		FCacheAddress(OP::ADDRESS InAt, uint16 InExecutionIndex, uint8 InExecutionOptions)
+		{
+			At = InAt;
+			ExecutionIndex = InExecutionIndex;
+			ExecutionOptions = InExecutionOptions;
+		}
+
+		FCacheAddress(OP::ADDRESS InAt, const FScheduledOp& Item)
+		{
+			At = InAt;
+			ExecutionIndex = Item.ExecutionIndex;
+			ExecutionOptions = Item.ExecutionOptions;
+			Type = Item.Type;
+		}
+
+		FCacheAddress(const FScheduledOp& Item)
+		{
+			At = Item.At;
+			ExecutionIndex = Item.ExecutionIndex;
+			ExecutionOptions = Item.ExecutionOptions;
+			Type = Item.Type;
+		}
+	};
+
+
+	inline uint32 GetTypeHash(const FCacheAddress& a)
+	{
+		return HashCombine(::GetTypeHash(a.At), a.ExecutionIndex);
+	}
+
+
+	inline uint32 GetTypeHash(const Ptr<const Resource>& V)
+	{
+		return ::GetTypeHash(V.get());
+	}
+
+
+	//! Container that stores data per executable code operation (indexed by address and execution
+	//! index).
+	template<class DATA>
+	class CodeContainer
+	{
+		using MemoryCounter = MemoryCounters::FMemoryTrackerInternalMemoryCounter;
+		using ArrayDataContainerType = TArray<DATA, FDefaultMemoryTrackingAllocator<MemoryCounter>>;
+		using MapDataContainerType = TMap<FCacheAddress, DATA, FDefaultMemoryTrackingSetAllocator<MemoryCounter>>;
+
+	public:
+		void resize(size_t s)
+		{
+			m_index0.SetNumZeroed(s);
+		}
+
+		uint32 size_code() const
+		{
+			return uint32(m_index0.Num());
+		}
+
+		void clear()
+		{
+			m_index0.Empty();
+			m_otherIndex.Empty();
+		}
+
+		inline void erase(const FCacheAddress& at)
+		{
+			if (at.ExecutionIndex == 0 && at.ExecutionOptions == 0)
+			{
+				m_index0[at.At] = nullptr;
+			}
+			else
+			{
+				m_otherIndex.Remove(at);
+			}
+		}
+
+		inline DATA get(const FCacheAddress& at) const
+		{
+			if (at.ExecutionIndex == 0 && at.ExecutionOptions == 0)
+			{
+				if (at.At < uint32(m_index0.Num()))
+				{
+					return m_index0[at.At];
+				}
+				else
+				{
+					return 0;
+				}
+			}
+			else
+			{
+				const DATA* it = m_otherIndex.Find(at);
+				if (it)
+				{
+					return *it;
+				}
+			}
+			return 0;
+		}
+
+		inline DATA* get_ptr(const FCacheAddress& at)
+		{
+			if (at.ExecutionIndex == 0 && at.ExecutionOptions == 0)
+			{
+				if (at.At < uint32(m_index0.Num()))
+				{
+					return &m_index0[at.At];
+				}
+				else
+				{
+					return nullptr;
+				}
+			}
+			else
+			{
+				DATA* it = m_otherIndex.Find(at);
+				if (it)
+				{
+					return it;
+				}
+			}
+			return nullptr;
+		}
+
+		inline const DATA* get_ptr(const FCacheAddress& at) const
+		{
+			if (at.ExecutionIndex == 0 && at.ExecutionOptions == 0)
+			{
+				if (at.At < uint32(m_index0.Num()))
+				{
+					return &m_index0[at.At];
+				}
+				else
+				{
+					return nullptr;
+				}
+			}
+			else
+			{
+				const DATA* it = m_otherIndex.Find(at);
+				if (it)
+				{
+					return it;
+				}
+			}
+			return nullptr;
+		}
+
+		inline DATA& operator[](const FCacheAddress& at)
+		{
+			if (at.ExecutionIndex == 0 && at.ExecutionOptions == 0)
+			{
+				return m_index0[at.At];
+			}
+			else
+			{
+				return m_otherIndex.FindOrAdd(at);
+			}
+		}
+
+		inline const DATA& operator[](const FCacheAddress& at) const
+		{
+			if (at.ExecutionIndex == 0 && at.ExecutionOptions == 0)
+			{
+				return m_index0[at.At];
+			}
+			else
+			{
+				return m_otherIndex[at];
+			}
+		}
+
+		struct iterator
+		{
+		private:
+			friend class CodeContainer<DATA>;
+
+			const CodeContainer<DATA>* container;
+			typename CodeContainer<DATA>::ArrayDataContainerType::TIterator it0;
+			typename CodeContainer<DATA>::MapDataContainerType::TIterator it1;
+
+			iterator(CodeContainer<DATA>* InContainer)
+				: container(InContainer)
+				, it0(InContainer->m_index0.CreateIterator())
+				, it1(InContainer->m_otherIndex.CreateIterator())
+			{
+			}
+
+		public:
+
+			inline void operator++(int)
+			{
+				if (it0)
+				{
+					++it0;
+				}
+				else
+				{
+					++it1;
+				}
+			}
+
+			iterator operator++()
+			{
+				if (it0)
+				{
+					++it0;
+				}
+				else
+				{
+					++it1;
+				}
+				return *this;
+			}
+
+			inline bool operator!=(const iterator& o) const
+			{
+				return it0 != o.it0 || it1 != o.it1;
+			}
+
+			inline DATA& operator*()
+			{
+				if (it0)
+				{
+					return *it0;
+				}
+				else
+				{
+					return it1->Value;
+				}
+			}
+
+			inline FCacheAddress get_address() const
+			{
+				if (it0)
+				{
+					return { uint32_t(it0.GetIndex()), 0, 0};
+				}
+				else
+				{
+					return it1->Key;
+				}
+			}
+
+			inline bool IsValid() const
+			{
+				return bool(it0) || bool(it1);
+			}
+		};
+
+		inline iterator begin()
+		{
+			iterator it(this);
+			return it;
+		}
+
+		inline int32 GetAllocatedSize() const
+		{
+			return m_index0.GetAllocatedSize()
+				+ m_otherIndex.GetAllocatedSize();
+		}
+	private:
+		// For index 0
+		ArrayDataContainerType m_index0;
+
+		// For index>0
+		MapDataContainerType m_otherIndex;
+	};
+
+
+	/** Interface for storage of data while Mutable code is being executed. */
+	class FProgramCache
+	{
+	public:
+
+		using AllocType = FDefaultMemoryTrackingAllocator<MemoryCounters::FMemoryTrackerInternalMemoryCounter>;
+
+		template<class Type, class Alloc = AllocType >
+		using TMemoryTrackedArray = TArray<Type, Alloc>;
+		
+		TMemoryTrackedArray<ExecutionIndex, TInlineAllocator<4, AllocType>> m_usedRangeIndices;
+
+		/** Runtime data for each program op. */
+		struct FOpExecutionData
+		{
+			uint16 OpHitCount;
+
+			/** Enabled if the op descriptor has been calculated. */
+			uint8 IsDescCacheValid : 1;
+			uint8 IsValueValid : 1;
+			uint8 IsCacheLocked : 1;
+
+			/** The operation DATA_TYPE. */
+			uint8 DataType;
+
+			/** The position in the type-specific array where the result data is stored. 
+			* 0 means index not valid.
+			* For small types like bool, int and float, the actual value is the index, instead of having their own array.
+			*/
+			union
+			{
+				int32 DataTypeIndex;
+				float ScalarResult;
+			};
+		};
+
+		/** Cached resources while the program is executing. 
+		* first value of the pair:
+		* 0 : value not valid (not set).
+		* 1 : valid, not worth freeing for memory
+		* 2 : valid, worth freeing
+		*/
+		CodeContainer<FOpExecutionData> OpExecutionData;
+
+		/** */
+
+		TMemoryTrackedArray<FVector4f> ColorResults;
+		TMemoryTrackedArray<Ptr<const Image>> ImageResults;
+		TMemoryTrackedArray<Ptr<const Layout>> LayoutResults;
+		TMemoryTrackedArray<Ptr<const Mesh>> MeshResults;
+		TMemoryTrackedArray<Ptr<const Instance>> InstanceResults;
+		TMemoryTrackedArray<FProjector> ProjectorResults;
+		TMemoryTrackedArray<Ptr<const String>> StringResults;
+		TMemoryTrackedArray<Ptr<const ExtensionData>> ExtensionDataResults;
+
+		/** */
+		inline const ExecutionIndex& GetRangeIndex(uint32_t i)
+		{
+			// Make sure we have the default element.
+			if (m_usedRangeIndices.IsEmpty())
+			{
+				m_usedRangeIndices.Push(ExecutionIndex());
+			}
+
+			check(i < uint32_t(m_usedRangeIndices.Num()));
+			return m_usedRangeIndices[i];
+		}
+
+		//!
+		inline uint32_t GetRangeIndexIndex(const ExecutionIndex& rangeIndex)
+		{
+			if (rangeIndex.IsEmpty())
+			{
+				return 0;
+			}
+
+			// Make sure we have the default element.
+			if (m_usedRangeIndices.IsEmpty())
+			{
+				m_usedRangeIndices.Push(ExecutionIndex());
+			}
+
+			// Look for or add the new element.
+			auto ElemIndex = m_usedRangeIndices.Find(rangeIndex);
+			if (ElemIndex != INDEX_NONE)
+			{
+				return ElemIndex;
+			}
+
+			m_usedRangeIndices.Push(rangeIndex);
+			return uint32_t(m_usedRangeIndices.Num()) - 1;
+		}
+
+		void Init(uint32 Size)
+		{
+			// This clear would prevent live update cache reusal
+			//OpExecutionData.clear();
+
+			OpExecutionData.resize(Size);
+
+			if (ColorResults.IsEmpty())
+			{
+				// Insert dafault/null values
+				ColorResults.Add(FVector4f());
+				ImageResults.Add(nullptr);
+				LayoutResults.Add(nullptr);
+				MeshResults.Add(nullptr);
+				InstanceResults.Add(nullptr);
+				ProjectorResults.Add(FProjector());
+				StringResults.Add(nullptr);
+				ExtensionDataResults.Add(nullptr);
+			}
+		}
+
+		void SetUnused(FOpExecutionData& Data)
+		{
+			// Only clear datatypes that have results that use relevant amounts of memory			
+			uint32 DataTypeIndex = Data.DataTypeIndex;
+			Data.IsValueValid = false;
+
+			if (DataTypeIndex)
+			{
+				check(Data.DataType < DATATYPE::DT_COUNT);
+				switch ((DATATYPE)Data.DataType)
+				{
+				case DATATYPE::DT_IMAGE:					
+					ImageResults[DataTypeIndex] = nullptr;
+					break;
+
+				case DATATYPE::DT_MESH:
+					MeshResults[DataTypeIndex] = nullptr;
+					break;
+
+				case DATATYPE::DT_INSTANCE:
+					InstanceResults[DataTypeIndex] = nullptr;
+					break;
+
+				case DATATYPE::DT_EXTENSION_DATA:
+					ExtensionDataResults[DataTypeIndex] = nullptr;
+					break;
+
+				default:
+					break;
+				}
+			}
+		}
+
+
+		bool IsValid(FCacheAddress at) const
+		{
+			if (at.At == 0 || at.At>=OpExecutionData.size_code() )
+			{
+				return false;
+			}
+
+			const FOpExecutionData* Data = OpExecutionData.get_ptr(at);
+			if (!Data)
+			{
+				return false;
+			}
+
+			// Is it a desc data query?
+			if (at.Type == FScheduledOp::EType::ImageDesc)
+			{
+				return Data->IsDescCacheValid;
+			}
+
+			// It's a full data query.
+			return Data->IsValueValid != 0;
+		}
+
+
+		/** */
+		void CheckHitCountsCleared()
+		{
+#if !(UE_BUILD_SHIPPING || UE_BUILD_TEST)
+			//MUTABLE_CPUPROFILER_SCOPE(CheckHitCountsCleared);
+
+			//int32 IncorrectCount = 0;
+			//CodeContainer<int>::iterator it = m_opHitCount.begin();
+			//for (; it.IsValid(); ++it)
+			//{
+			//	int32 Count = *it;
+			//	if (Count>0 && Count < UE_MUTABLE_CACHE_COUNT_LIMIT)
+			//	{
+			//		// We don't manage the hitcounts of small types that don't use much memory
+			//		if (m_resources[it.get_address()].Key==2)
+			//		{
+			//			// Op hitcount should have reached 0, otherwise, it means we requested an operation but never read 
+			//			// the result.
+			//			++IncorrectCount;
+			//		}
+			//	}
+			//}
+
+			//if (IncorrectCount > 0)
+			//{
+			//	UE_LOG(LogMutableCore, Log, TEXT("The op-hit-count didn't hit 0 for %5d operations. This may mean that too much memory is cached."), IncorrectCount);
+			//	//check(false);
+			//}
+#endif
+		}
+
+
+		void Clear()
+		{
+			MUTABLE_CPUPROFILER_SCOPE(ProgramCacheClear);
+
+			uint32 CodeSize = OpExecutionData.size_code();
+			OpExecutionData.clear();
+			OpExecutionData.resize(CodeSize);
+		}
+
+
+		void ClearDescCache()
+		{
+			MUTABLE_CPUPROFILER_SCOPE(ProgramDescCacheClear);
+
+			CodeContainer<FProgramCache::FOpExecutionData>::iterator it = OpExecutionData.begin();
+			for (; it.IsValid(); ++it)
+			{
+				(*it).IsDescCacheValid = false;
+			}
+		}
+
+
+		bool GetBool(FCacheAddress at)
+		{
+			if (!at.At) return false;
+			const FOpExecutionData* Data = OpExecutionData.get_ptr(at);
+			if (!Data) return false;
+
+			check(Data->DataType==DATATYPE::DT_BOOL);
+			return Data->DataTypeIndex!=0;
+		}
+
+		float GetScalar(FCacheAddress at)
+		{
+			if (!at.At) return 0.0f;
+			const FOpExecutionData* Data = OpExecutionData.get_ptr(at);
+			if (!Data) return 0.0f;
+
+			check(Data->DataType == DATATYPE::DT_SCALAR);
+			return Data->ScalarResult;
+		}
+
+		int32 GetInt(FCacheAddress at)
+		{
+			if (!at.At) return 0;
+			const FOpExecutionData* Data = OpExecutionData.get_ptr(at);
+			if (!Data) return 0;
+
+			check(Data->DataType == DATATYPE::DT_INT);
+			return Data->DataTypeIndex;
+		}
+
+		FVector4f GetColour(FCacheAddress at)
+		{
+			if (!at.At) return FVector4f();
+			const FOpExecutionData* Data = OpExecutionData.get_ptr(at);
+			if (!Data) return FVector4f();
+
+			check(Data->DataType == DATATYPE::DT_COLOUR);
+			return ColorResults[Data->DataTypeIndex];
+		}
+
+		FProjector GetProjector(FCacheAddress at)
+		{
+			if (!at.At) return FProjector();
+			const FOpExecutionData* Data = OpExecutionData.get_ptr(at);
+			if (!Data) return FProjector();
+
+			check(Data->DataType == DATATYPE::DT_PROJECTOR);
+			return ProjectorResults[Data->DataTypeIndex];
+		}
+
+		Ptr<const Instance> GetInstance(FCacheAddress at)
+		{
+			if (!at.At) return nullptr;
+			FOpExecutionData* Data = OpExecutionData.get_ptr(at);
+			if (!Data) return nullptr;
+
+			check(Data->DataType == DATATYPE::DT_INSTANCE);
+			Ptr<const Instance> Result = InstanceResults[Data->DataTypeIndex];
+
+			// We need to decrease the hit-count even if the result is null.
+			check(Data->OpHitCount > 0);
+
+			--Data->OpHitCount;
+			if (Data->OpHitCount == 0 && !Data->IsCacheLocked)
+			{
+				SetUnused(*Data);
+			}
+
+			return Result;
+		}
+
+
+		Ptr<const Image> GetImage(FCacheAddress at, bool& bIsLastReference)
+		{
+			bIsLastReference = false;
+
+			if (!at.At) return nullptr;
+			if (at.At >= OpExecutionData.size_code()) return nullptr;
+			FOpExecutionData* Data = OpExecutionData.get_ptr(at);
+			if (!Data) return nullptr;
+
+			check(Data->DataType == DATATYPE::DT_IMAGE);
+			Ptr<const Image> Result = ImageResults[Data->DataTypeIndex];
+
+			// We need to decrease the hit-count even if the result is null.
+			check(Data->OpHitCount > 0);
+
+			--Data->OpHitCount;
+			if (Data->OpHitCount == 0 && !Data->IsCacheLocked)
+			{
+				SetUnused(*Data);
+				bIsLastReference = true;
+			}
+
+			return Result;
+		}
+
+
+		Ptr<const Mesh> GetMesh(FCacheAddress at, bool& bIsLastReference)
+		{
+			bIsLastReference = false;
+
+			if (!at.At) return nullptr;
+			if (at.At >= OpExecutionData.size_code()) return nullptr;
+			FOpExecutionData* Data = OpExecutionData.get_ptr(at);
+			if (!Data) return nullptr;
+
+			check(Data->DataType == DATATYPE::DT_MESH);
+			Ptr<const Mesh> Result = MeshResults[Data->DataTypeIndex];
+
+			// We need to decrease the hit-count even if the result is null.
+			check(Data->OpHitCount > 0);
+
+			--Data->OpHitCount;
+			if (Data->OpHitCount == 0 && !Data->IsCacheLocked)
+			{
+				SetUnused(*Data);
+				bIsLastReference = true;
+			}
+
+			return Result;
+		}
+
+
+		Ptr<const Layout> GetLayout(FCacheAddress at)
+		{
+			if (!at.At) return nullptr;
+			FOpExecutionData* Data = OpExecutionData.get_ptr(at);
+			if (!Data) return nullptr;
+
+			check(Data->DataType == DATATYPE::DT_LAYOUT);
+			return LayoutResults[Data->DataTypeIndex];
+		}
+
+
+		Ptr<const String> GetString(FCacheAddress at)
+		{
+			if (!at.At) return nullptr;
+			FOpExecutionData* Data = OpExecutionData.get_ptr(at);
+			if (!Data) return nullptr;
+
+			check(Data->DataType == DATATYPE::DT_STRING);
+			return StringResults[Data->DataTypeIndex];
+		}
+
+
+		Ptr<const ExtensionData> GetExtensionData(FCacheAddress at)
+		{
+			if (!at.At) return nullptr;
+			FOpExecutionData* Data = OpExecutionData.get_ptr(at);
+			if (!Data) return nullptr;
+
+			check(Data->DataType == DATATYPE::DT_EXTENSION_DATA);
+			return ExtensionDataResults[Data->DataTypeIndex];
+		}
+
+		void SetValidDesc(FCacheAddress at)
+		{
+			check(at.Type == FScheduledOp::EType::ImageDesc);
+			check(at.At < OpExecutionData.size_code());
+			FOpExecutionData* Data = OpExecutionData.get_ptr(at);
+			Data->IsDescCacheValid = true;
+		}
+
+		void SetBool(FCacheAddress at, bool v)
+		{
+			check(at.At < OpExecutionData.size_code());
+			FOpExecutionData* Data = OpExecutionData.get_ptr(at);
+			check(Data->DataType == DATATYPE::DT_BOOL || Data->DataType == DATATYPE::DT_NONE);
+			Data->DataType = DATATYPE::DT_BOOL;
+			Data->DataTypeIndex = v;
+			Data->IsValueValid = true;
+		}
+
+		void SetInt(FCacheAddress at, int32 v)
+		{
+			check(at.At < OpExecutionData.size_code());
+			FOpExecutionData* Data = OpExecutionData.get_ptr(at);
+			check(Data->DataType == DATATYPE::DT_INT || Data->DataType == DATATYPE::DT_NONE);
+			Data->DataType = DATATYPE::DT_INT;
+			Data->DataTypeIndex = v;
+			Data->IsValueValid = true;
+		}
+
+		void SetScalar(FCacheAddress at, float v)
+		{
+			check(at.At < OpExecutionData.size_code());
+			FOpExecutionData* Data = OpExecutionData.get_ptr(at);
+			check(Data->DataType == DATATYPE::DT_SCALAR || Data->DataType == DATATYPE::DT_NONE);
+			Data->DataType = DATATYPE::DT_SCALAR;
+			Data->ScalarResult = v;
+			Data->IsValueValid = true;
+		}
+
+		void SetColour(FCacheAddress at, const FVector4f& v)
+		{
+			check(at.At < OpExecutionData.size_code());
+			FOpExecutionData* Data = OpExecutionData.get_ptr(at);
+			check(Data->DataType == DATATYPE::DT_COLOUR || Data->DataType == DATATYPE::DT_NONE);
+			Data->DataType = DATATYPE::DT_COLOUR;
+			Data->IsValueValid = true;
+
+			if (!Data->DataTypeIndex)
+			{
+				Data->DataTypeIndex = ColorResults.Num();
+				ColorResults.Add(v);
+			}
+			else
+			{
+				ColorResults[Data->DataTypeIndex] = v;
+			}
+			check(Data->DataTypeIndex != 0);
+		}
+
+		void SetProjector(FCacheAddress at, const FProjector& v)
+		{
+			check(at.At < OpExecutionData.size_code());
+			FOpExecutionData* Data = OpExecutionData.get_ptr(at);
+			check(Data->DataType == DATATYPE::DT_PROJECTOR || Data->DataType == DATATYPE::DT_NONE);
+			Data->DataType = DATATYPE::DT_PROJECTOR;
+			Data->IsValueValid = true;
+
+			if (!Data->DataTypeIndex)
+			{
+				Data->DataTypeIndex = ProjectorResults.Num();
+				ProjectorResults.Add(v);
+			}
+			else
+			{
+				ProjectorResults[Data->DataTypeIndex] = v;
+			}
+			check(Data->DataTypeIndex != 0);
+		}
+
+		void SetInstance(FCacheAddress at, Ptr<const Instance> v)
+		{
+			check(at.At < OpExecutionData.size_code());
+			FOpExecutionData* Data = OpExecutionData.get_ptr(at);
+			check(Data->DataType == DATATYPE::DT_INSTANCE || Data->DataType == DATATYPE::DT_NONE);
+			Data->DataType = DATATYPE::DT_INSTANCE;
+			Data->IsValueValid = true;
+
+			if (!Data->DataTypeIndex)
+			{
+				Data->DataTypeIndex = InstanceResults.Num();
+				InstanceResults.Add(v);
+			}
+			else
+			{
+				InstanceResults[Data->DataTypeIndex] = v;
+			}
+			check(Data->DataTypeIndex != 0);
+		}
+
+		void SetExtensionData(FCacheAddress at, Ptr<const ExtensionData> v)
+		{
+			check(at.At < OpExecutionData.size_code());
+			FOpExecutionData* Data = OpExecutionData.get_ptr(at);
+			check(Data->DataType == DATATYPE::DT_EXTENSION_DATA || Data->DataType == DATATYPE::DT_NONE);
+			Data->DataType = DATATYPE::DT_EXTENSION_DATA;
+			Data->IsValueValid = true;
+
+			if (!Data->DataTypeIndex)
+			{
+				Data->DataTypeIndex = ExtensionDataResults.Num();
+				ExtensionDataResults.Add(v);
+			}
+			else
+			{
+				ExtensionDataResults[Data->DataTypeIndex] = v;
+			}
+			check(Data->DataTypeIndex != 0);
+		}
+
+		void SetImage(FCacheAddress at, Ptr<const Image> v)
+		{
+			check(at.At < OpExecutionData.size_code());
+			FOpExecutionData* Data = OpExecutionData.get_ptr(at);
+			check(Data->DataType == DATATYPE::DT_IMAGE || Data->DataType == DATATYPE::DT_NONE);
+			Data->DataType = DATATYPE::DT_IMAGE;
+			Data->IsValueValid = true;
+
+			if (!Data->DataTypeIndex)
+			{
+				Data->DataTypeIndex = ImageResults.Num();
+				ImageResults.Add(v);
+			}
+			else
+			{
+				ImageResults[Data->DataTypeIndex] = v;
+			}
+			check(Data->DataTypeIndex != 0);
+
+			mu::UpdateLLMStats();
+		}
+
+		void SetMesh(FCacheAddress at, Ptr<const Mesh> v)
+		{
+			check(at.At < OpExecutionData.size_code());
+			FOpExecutionData* Data = OpExecutionData.get_ptr(at);
+			check(Data->DataType == DATATYPE::DT_MESH || Data->DataType == DATATYPE::DT_NONE);
+			Data->DataType = DATATYPE::DT_MESH;
+			Data->IsValueValid = true;
+
+			if (!Data->DataTypeIndex)
+			{
+				Data->DataTypeIndex = MeshResults.Num();
+				MeshResults.Add(v);
+			}
+			else
+			{
+				MeshResults[Data->DataTypeIndex] = v;
+			}
+			check(Data->DataTypeIndex != 0);
+
+			mu::UpdateLLMStats();
+		}
+
+		void SetLayout(FCacheAddress at, Ptr<const Layout> v)
+		{
+			check(at.At < OpExecutionData.size_code());
+			FOpExecutionData* Data = OpExecutionData.get_ptr(at);
+			check(Data->DataType == DATATYPE::DT_LAYOUT || Data->DataType == DATATYPE::DT_NONE);
+			Data->DataType = DATATYPE::DT_LAYOUT;
+			Data->IsValueValid = true;
+
+			if (!Data->DataTypeIndex)
+			{
+				Data->DataTypeIndex = LayoutResults.Num();
+				LayoutResults.Add(v);
+			}
+			else
+			{
+				LayoutResults[Data->DataTypeIndex] = v;
+			}
+			check(Data->DataTypeIndex != 0);
+
+			mu::UpdateLLMStats();
+		}
+
+		void SetString(FCacheAddress at, Ptr<const String> v)
+		{
+			check(at.At < OpExecutionData.size_code());
+			FOpExecutionData* Data = OpExecutionData.get_ptr(at);
+			check(Data->DataType == DATATYPE::DT_STRING || Data->DataType == DATATYPE::DT_NONE);
+			Data->DataType = DATATYPE::DT_STRING;
+			Data->IsValueValid = true;
+
+			if (!Data->DataTypeIndex)
+			{
+				Data->DataTypeIndex = StringResults.Num();
+				StringResults.Add(v);
+			}
+			else
+			{
+				StringResults[Data->DataTypeIndex] = v;
+			}
+			check(Data->DataTypeIndex!=0);
+		}
+
+
+		inline void IncreaseHitCount(FCacheAddress at)
+		{
+			// Don't count hits for instruction 0, which is always null. It is usually already
+			// check that At is not 0, and then it is not requested, generating a stray non-zero count
+			// at its position.
+			if (at.At)
+			{
+				check(at.At < OpExecutionData.size_code());
+				FOpExecutionData& Data = OpExecutionData[at];
+				Data.OpHitCount++;
+			}
+		}
+
+		inline void SetForceCached(OP::ADDRESS at)
+		{
+			// \TODO: It only locks at,0,0
+			if (at)
+			{
+				check(at < OpExecutionData.size_code());
+				FOpExecutionData& Data = OpExecutionData[FCacheAddress(at, 0, 0)];
+				Data.IsCacheLocked = true;
+			}
+		}
+	};
+
+
+	inline bool operator==(const FCacheAddress& a, const FCacheAddress& b)
+	{
+		return a.At == b.At
+			&&
+			a.ExecutionIndex == b.ExecutionIndex
+			&&
+			a.ExecutionOptions == b.ExecutionOptions
+			&&
+			a.Type == b.Type;
+	}
+
+	inline bool operator<(const FCacheAddress& a, const FCacheAddress& b)
+	{
+		if (a.At < b.At) return true;
+		if (a.At > b.At) return false;
+		if (a.ExecutionIndex < b.ExecutionIndex) return true;
+		if (a.ExecutionIndex > b.ExecutionIndex) return false;
+		if (a.ExecutionOptions < b.ExecutionOptions) return true;
+		if (a.ExecutionOptions > b.ExecutionOptions) return false;
+		return a.Type < b.Type;
+	}
+
+	/** Data for an instance that is currently being processed in the mutable system. This means it is
+	* between a BeginUpdate and EndUpdate, or during an "atomic" operation (like generate a single resource).
+	*/
+	struct FLiveInstance
+	{
+		Instance::ID InstanceID;
+		int32 State = 0;
+		Ptr<const Instance> Instance;
+		TSharedPtr<const Model> Model;
+
+		Ptr<Parameters> OldParameters;
+
+		/** Mask of the parameters that have changed since the last update.
+		* Every bit represents a state parameter.
+		*/
+		uint64 UpdatedParameters = 0;
+
+		/** Cached data for the generation of this instance. */
+		TSharedPtr<FProgramCache> Cache;
+
+		~FLiveInstance()
+		{
+			// Manually done to trace mem deallocations
+			MUTABLE_CPUPROFILER_SCOPE(LiveInstanceDestructor);
+			Cache = nullptr;
+			OldParameters = nullptr;
+			Instance = nullptr;
+			Model = nullptr;
+		}
+	};
+
+
+    /** Struct to manage all the memory allocated for resources used during mutable operation. */
+    struct FWorkingMemoryManager
+    {
+		using MemoryCounter = MemoryCounters::FMemoryTrackerInternalMemoryCounter;
+
+		template<class Type>
+		using TMemoryTrackedArray = TArray<Type, FDefaultMemoryTrackingAllocator<MemoryCounter>>;
+
+		template<class KeyType, class ValueType>
+		using TMemoryTrackedMap = TMap<KeyType, ValueType, FDefaultMemoryTrackingSetAllocator<MemoryCounter>>;
+
+		/** Cached traking for streamed model data for one model. */
+        struct FModelCacheEntry
+        {
+            /** Model who's data is being tracked. */
+            TWeakPtr<const Model> Model;
+
+            /** For each model rom, the last time its streamed data was used. */
+            TMemoryTrackedArray<TPair<uint64, uint64>> RomWeights;
+
+			/** Count of pending operations for every rom index. */
+			TMemoryTrackedArray<uint16> PendingOpsPerRom;
+        };
+
+		//! Management of generated resources
+		//! @{
+
+		//! This is used to uniquely identify a generated resource like meshes or images.
+		struct FGeneratedResourceData
+		{
+			/** Model for this resource. */
+			TWeakPtr<const Model> Model;
+
+			//! The id assigned to the generated resource.
+			FResourceID Id;
+
+			//! The last request operation for this resource
+			uint32 LastRequestId;
+
+			//! An opaque blob with the values of the relevant parameters
+			TMemoryTrackedArray<uint8> ParameterValuesBlob;
+		};
+
+		//! The last id generated for a resource
+		uint32 LastResourceKeyId = 0;
+
+		//! The last id generated for a resource request. This is used to check the
+		//! relevancy of the resources when flushing the cache
+		uint32 LastResourceResquestId = 0;
+
+		//! Cached ids for returned assets
+		//! This is non-persistent runtime data
+		TMemoryTrackedArray<FGeneratedResourceData> GeneratedResources;
+
+		/** */
+		FResourceID GetResourceKey(const TSharedPtr<const Model>&, const Parameters*, uint32 ParamListIndex, OP::ADDRESS RootAt);
+
+		//! @}
+
+		/** Maximum working memory that mutable should be using. */
+		int64 BudgetBytes = 0;
+
+		/** Maximum excess memory reached suring the current operation. */
+		int64 BudgetExcessBytes = 0;
+
+		/** Maximum number of resource keys that will be stored for resource reusal. */
+		int32 MaxGeneratedResourceCacheSize = 1024;
+
+
+		/** This value is used to track the order of loading of roms. */
+        uint64 RomTick = 0;
+
+		/** Control info for the per-model cache of streamed data. */
+        TMemoryTrackedArray<FModelCacheEntry> CachePerModel;
+
+		/** Data for each mutable instance that is being updated. */
+		TMemoryTrackedArray<FLiveInstance> LiveInstances;
+
+		/** Temporary reference to the memory of the current instance being updated. Only valid during a mutable "atomic" operation, like a BeginUpdate or a GetImage. */
+		TSharedPtr<FProgramCache> CurrentInstanceCache;
+
+		/** Resources that have been used in the past, but haven't been deallocated because they still fitted the memory budget and they could be reused. */
+		TMemoryTrackedArray<Ptr<Image>> PooledImages;
+
+		/** List of intermediate resources that are not soterd anywhere yet. They are still locally referenced by code. */
+		TMemoryTrackedArray<Ptr<const Image>> TempImages;
+		TMemoryTrackedArray<Ptr<const Mesh>> TempMeshes;
+
+		/** List of resources that are currently in any cache position, and the number of positions they are in. */
+		TMemoryTrackedMap<Ptr<const Resource>, int32> CacheResources;
+
+		/** Given a mutable model, find or create its rom cache. */
+		FModelCacheEntry* FindModelCache(const Model*);
+		FModelCacheEntry& FindOrAddModelCache(const TSharedPtr<const Model>&);
+
+        /** Make sure the working memory is below the internal budget, even counting with the passed additional memory. 
+		* An optional function can be passed to "block" the unload of certain roms of data.
+		* Return true if it succeeded, false otherwise.
+		*/
+		bool EnsureBudgetBelow(uint64 AdditionalMemory);
+		
+		/** Return true if the memory budget is 90% full. */
+		bool IsMemoryBudgetFull() const;
+
+		/** Calculate the current usage of memory as used to calculate the budget. */
+		int64 GetCurrentMemoryBytes() const;
+
+        /** Register that a specific rom has been requested and update the heuristics to keep it in memory. */
+        void MarkRomUsed( int32 RomIndex, const TSharedPtr<const Model>& );
+
+		/** */
+		UE_NODISCARD Ptr<Image> CreateImage(uint32 SizeX, uint32 SizeY, uint32 Lods, EImageFormat Format, EInitializationType Init)
+		{
+			CheckRunnerThread();
+
+			uint32 DataSize = Image::CalculateDataSize(SizeX, SizeY, Lods, Format);
+
+			// Look for an unused image in the pool that can be reused
+			int32 PooledImageCount = PooledImages.Num();
+			for (int Index = 0; DataSize>0 && Index<PooledImageCount; ++Index)
+			{
+				Ptr<Image>& Candidate = PooledImages[Index];
+				if (Candidate->GetFormat() == Format
+					&& Candidate->GetSizeX() == SizeX
+					&& Candidate->GetSizeY() == SizeY
+					&& Candidate->GetLODCount() == Lods
+					)
+				{
+					Ptr<Image> Result = Candidate;
+					PooledImages.RemoveAtSwap(Index);
+					
+					if (Init == EInitializationType::Black)
+					{
+						Result->InitToBlack();
+					}
+					else
+					{
+						Result->m_flags = 0;
+						Result->RelevancyMinY = 0;
+						Result->RelevancyMaxY = 0;
+					}
+					return Result;
+				}
+			}
+
+			// Make room in the budget
+			EnsureBudgetBelow(DataSize);
+
+			// Create it
+			Ptr<Image> Result = new Image(SizeX, SizeY, Lods, Format, Init);
+
+			TempImages.Add(Result);
+			return Result;
+		}
+
+		/** Ref will be nulled and relesed in any case. */
+		UE_NODISCARD Ptr<Image> CloneOrTakeOver(Ptr<const Image>& Resource)
+		{
+			CheckRunnerThread();
+
+			TempImages.RemoveSingle(Resource);
+			
+			check(!TempImages.Contains(Resource));
+			check(!PooledImages.Contains(Resource));
+
+			Ptr<Image> Result;
+			if (!Resource->IsUnique())
+			{
+				// TODO: try to grab from the pool
+
+				uint32 DataSize = Resource->GetDataSize();
+				EnsureBudgetBelow(DataSize);
+
+				Result = Resource->Clone();
+				Release(Resource);
+			}
+			else
+			{
+				Result = const_cast<Image*>(Resource.get());
+				Resource = nullptr;
+			}
+
+			return Result;
+		}
+
+		/** */
+		void Release(Ptr<const Image>& Resource)
+		{
+			CheckRunnerThread();
+
+			if (!Resource)
+			{
+				return;
+			}
+
+			const int32 ResourceDataSize = Resource->GetDataSize();
+			TempImages.RemoveSingle(Resource);
+
+			check(!TempImages.Contains(Resource));
+			check(!PooledImages.Contains(Resource));
+
+			if (IsBudgetTemp(Resource))
+			{
+				// Check if we are exceeding the budget
+				bool bInBudget = EnsureBudgetBelow(ResourceDataSize);
+				if (bInBudget)
+				{
+					PooledImages.Add(const_cast<Image*>(Resource.get()));
+				}
+			}
+			else
+			{
+				// Check if we are exceeding the budget
+				EnsureBudgetBelow(0);
+			}
+
+			Resource = nullptr;
+		}
+
+		/** */
+		void Release(Ptr<Image>& Resource)
+		{
+			CheckRunnerThread();
+
+			if (!Resource)
+			{
+				return;
+			}
+
+			const int32 ResourceDataSize = Resource->GetDataSize();
+			TempImages.RemoveSingle(Resource);
+
+			check(!TempImages.Contains(Resource));
+			check(!PooledImages.Contains(Resource));
+
+			if (IsBudgetTemp(Resource))
+			{
+				// Check if we are exceeding the budget
+				bool bInBudget = EnsureBudgetBelow(ResourceDataSize);
+				if (bInBudget)
+				{
+					PooledImages.Add(Resource.get());
+				}
+			}
+			else
+			{
+				// Check if we are exceeding the budget
+				EnsureBudgetBelow(0);
+			}
+
+			Resource = nullptr;
+		}
+
+		UE_NODISCARD Ptr<Mesh> CreateMesh(int32 BudgetReserveSize)
+		{
+			CheckRunnerThread();
+
+			EnsureBudgetBelow(BudgetReserveSize);
+
+			Ptr<Mesh> Result = new Mesh();
+
+			TempMeshes.Add(Result);
+
+			return Result;
+		}
+
+		UE_NODISCARD Ptr<Mesh> CloneOrTakeOver(Ptr<const Mesh>& Resource)
+		{
+			CheckRunnerThread();
+
+			const int32 ResourceDataSize = Resource->GetDataSize();
+			TempMeshes.RemoveSingle(Resource);
+
+			Ptr<Mesh> Result;
+			if (!Resource->IsUnique())
+			{
+				Result = CreateMesh(ResourceDataSize); 
+				Result->CopyFrom(*Resource);
+				Release(Resource);
+			}
+			else
+			{
+				Result = Ptr<Mesh>(const_cast<Mesh*>(Resource.get()));
+				Resource = nullptr;
+			}
+
+			return Result;
+		}
+
+		void Release(Ptr<const Mesh>& Resource)
+		{
+			CheckRunnerThread();
+
+			if (!Resource)
+			{
+				return;
+			}
+
+			TempMeshes.RemoveSingle(Resource);
+			check(!TempMeshes.Contains(Resource));
+
+			EnsureBudgetBelow(0);
+
+			Resource = nullptr;
+		}
+
+		void Release(Ptr<Mesh>& Resource)
+		{
+			Ptr<const Mesh> ConstPtr = Resource;
+
+			Release(ConstPtr);
+
+			Resource = nullptr;
+		}
+
+		/** */
+		UE_NODISCARD Ptr<const Mesh> LoadMesh(const FCacheAddress& From, bool bTakeOwnership = false)
+		{
+			bool bIsLastReference = false;
+			Ptr<const Mesh> Result = CurrentInstanceCache->GetMesh(From, bIsLastReference);
+			if (!Result)
+			{
+				return nullptr;
+			}
+
+			// If we retrieved the last reference to this resource in "From" cache position (it could still be in other cache positions as well)
+			if (bIsLastReference)
+			{
+				int32* CountPtr = CacheResources.Find(Result);
+				check(CountPtr);
+				*CountPtr = (*CountPtr) - 1;
+				if (!*CountPtr)
+				{
+					CacheResources.FindAndRemoveChecked(Result);
+				}
+			}
+
+			if (!bTakeOwnership && Result->IsUnique())
+			{
+				TempMeshes.Add(Result);
+			}
+
+			return Result;
+		}
+
+		/** */
+		UE_NODISCARD Ptr<const Image> LoadImage(const FCacheAddress& From, bool bTakeOwnership = false)
+		{
+			bool bIsLastReference = false;
+			Ptr<const Image> Result = CurrentInstanceCache->GetImage(From, bIsLastReference);
+			if (!Result)
+			{
+				return nullptr;
+			}
+
+			// If we retrieved the last reference to this resource in "From" cache position (it could still be in other cache positions as well)
+			if (bIsLastReference)
+			{
+				int32* CountPtr = CacheResources.Find(Result);
+				check(CountPtr);
+				*CountPtr = (*CountPtr) - 1;
+				if (!*CountPtr)
+				{
+					CacheResources.FindAndRemoveChecked(Result);
+				}
+			}
+
+
+			if (!bTakeOwnership && Result->IsUnique())
+			{
+				TempImages.Add(Result);
+			}
+
+			return Result;
+		}
+
+		/** */
+		void StoreImage(const FCacheAddress& To, Ptr<const Image> Resource)
+		{
+			if (Resource)
+			{
+				int32 ResourceDataSize = Resource->GetDataSize();
+				TempImages.RemoveSingle(Resource);
+
+				check(!TempImages.Contains(Resource));
+
+				int32& Count = CacheResources.FindOrAdd(Resource, 0);
+				++Count;
+			}
+
+			CurrentInstanceCache->SetImage(To, Resource);
+		}
+
+		void StoreMesh(const FCacheAddress& To, Ptr<const Mesh> Resource)
+		{
+			if (Resource)
+			{
+				TempMeshes.RemoveSingle(Resource);
+
+				int32& Count = CacheResources.FindOrAdd(Resource, 0);
+				++Count;
+			}
+
+			CurrentInstanceCache->SetMesh(To, Resource);
+		}
+
+		/** Return true if the resource is not in any cache (0,1,rom). */
+		bool IsBudgetTemp(const Ptr<const Resource>& Resource)
+		{
+			if (!Resource)
+			{
+				return false;
+			}
+
+			bool bIsTemp = Resource->IsUnique();
+			return bIsTemp;
+		}
+
+		int32 GetPooledBytes() const
+		{
+			int32 Result = 0;
+			for (const Ptr<Image>& Value : PooledImages)
+			{
+				Result += Value->GetDataSize();
+			}
+
+			return Result;
+		}
+
+		int32 GetTempBytes() const
+		{
+			int32 Result = 0;
+			for (const Ptr<const Image>& Value : TempImages)
+			{
+				Result += Value->GetDataSize();
+			}
+
+			for (const Ptr<const Mesh>& Value : TempMeshes)
+			{
+				Result += Value->GetDataSize();
+			}
+
+			return Result;
+		}
+
+
+		int32 GetRomBytes() const
+		{
+			int32 Result = 0;
+
+			TArray<const Model*> Models;
+			for (const FLiveInstance& Instance : LiveInstances)
+			{
+				Models.AddUnique(Instance.Model.Get());
+			}
+
+			// Data stored per-model, but related to instance construction
+			for (const Model* Model : Models)
+			{
+				// Count streamable and currently-loaded resources
+				const FProgram& Program = Model->GetPrivate()->m_program;
+				for (const TPair<int32, Ptr<const Image>>& Rom : Program.m_constantImageLODs)
+				{
+					if (Rom.Value && Rom.Key >= 0)
+					{
+						Result += Rom.Value->GetDataSize();
+					}
+				}
+				for (const TPair<int32, Ptr<const Mesh>>& Rom : Program.m_constantMeshes)
+				{
+					if (Rom.Value && Rom.Key >= 0)
+					{
+						Result += Rom.Value->GetDataSize();
+					}
+				}
+			}
+
+			return Result;
+		}
+
+
+		int32 GetTrackedCacheBytes() const
+		{
+			int32 Result = 0;
+			for (const TTuple<Ptr<const Resource>,int32>& It : CacheResources)
+			{
+				Result += It.Key->GetDataSize();
+			}
+
+			return Result;
+		}
+
+		/** Calculate the amount of bytes in data cached in the level 0 and 1 cache in all live instances. */
+		int32 GetCacheBytes() const
+		{
+			int32 Result = 0;
+			TSet<const Resource*> Cache0Unique;
+			TSet<const Resource*> Cache1Unique;
+
+			for (const FLiveInstance& Instance : LiveInstances)
+			{
+				CodeContainer<FProgramCache::FOpExecutionData>::iterator it = Instance.Cache->OpExecutionData.begin();
+				for (; it.IsValid(); ++it)
+				{
+					if (!(*it).DataTypeIndex)
+					{
+						continue;
+					}
+
+					TSet<const Resource*>* TargetSet = &Cache0Unique;
+					if ((*it).IsCacheLocked)
+					{
+						TargetSet = &Cache1Unique;
+					}
+
+					const Resource* Value = nullptr;
+					switch ((*it).DataType)
+					{
+					case DATATYPE::DT_IMAGE:
+						Value = Instance.Cache->ImageResults[(*it).DataTypeIndex].get();
+						if (Value)
+						{
+							TargetSet->Add(Value);
+						}
+						break;
+
+					case DATATYPE::DT_MESH:
+						Value = Instance.Cache->MeshResults[(*it).DataTypeIndex].get();
+						if (Value)
+						{
+							TargetSet->Add(Value);
+						}
+						break;
+
+					default:
+						break;
+					}
+				}
+			}
+
+			// Merge
+			Cache0Unique.Append(Cache1Unique);
+
+			// Count
+			for (const Resource* Value : Cache0Unique)
+			{
+				Result += Value->GetDataSize();
+			}
+
+			return Result;
+		}
+
+		/** Remove all intermediate data (big and small) from the memory except for the one that has been explicitely
+		* marked as state cache.
+		*/
+		void ClearCacheLayer0()
+		{
+			check(CurrentInstanceCache);
+
+			MUTABLE_CPUPROFILER_SCOPE(ClearLayer0);
+
+			CodeContainer<FProgramCache::FOpExecutionData>::iterator it = CurrentInstanceCache->OpExecutionData.begin();
+			for (; it.IsValid(); ++it)
+			{
+				FProgramCache::FOpExecutionData& Data = *it;
+				if (!Data.DataTypeIndex
+					||
+					Data.IsCacheLocked)
+				{
+					continue;
+				}
+
+				const Resource* Value = nullptr;
+
+				switch (Data.DataType)
+				{
+				case DATATYPE::DT_IMAGE:
+					Value = CurrentInstanceCache->ImageResults[Data.DataTypeIndex].get();
+					if (Value)
+					{
+						CacheResources.Remove(Value);
+						CurrentInstanceCache->ImageResults[Data.DataTypeIndex] = nullptr;
+					}
+					break;
+
+				case DATATYPE::DT_MESH:
+					Value = CurrentInstanceCache->MeshResults[Data.DataTypeIndex].get();
+					if (Value)
+					{
+						CacheResources.Remove(Value);
+						CurrentInstanceCache->MeshResults[Data.DataTypeIndex] = nullptr;
+					}
+					break;
+
+				case DATATYPE::DT_LAYOUT:
+					Value = CurrentInstanceCache->LayoutResults[Data.DataTypeIndex].get();
+					CurrentInstanceCache->LayoutResults[Data.DataTypeIndex] = nullptr;
+					break;
+
+				case DATATYPE::DT_INSTANCE:
+					Value = CurrentInstanceCache->InstanceResults[Data.DataTypeIndex].get();
+					CurrentInstanceCache->InstanceResults[Data.DataTypeIndex] = nullptr;
+					break;
+
+				default:
+					break;
+				}
+
+				Data.OpHitCount = 0;
+				Data.IsValueValid = false;
+			}
+		}
+
+
+		/** Remove all intermediate data (big and small) from the memory including the one that has been explicitely
+		* marked as state cache.
+		*/
+		void ClearCacheLayer1()
+		{
+			MUTABLE_CPUPROFILER_SCOPE(ClearLayer1);
+
+			CodeContainer<FProgramCache::FOpExecutionData>::iterator it = CurrentInstanceCache->OpExecutionData.begin();
+			for (; it.IsValid(); ++it)
+			{
+				FProgramCache::FOpExecutionData& Data = *it;
+				if (!Data.DataTypeIndex)
+				{
+					continue;
+				}
+
+				const Resource* Value = nullptr;
+
+				switch (Data.DataType)
+				{
+				case DATATYPE::DT_IMAGE:
+					Value = CurrentInstanceCache->ImageResults[Data.DataTypeIndex].get();
+					CacheResources.Remove(Value);
+					CurrentInstanceCache->ImageResults[Data.DataTypeIndex] = nullptr;
+					break;
+
+				case DATATYPE::DT_MESH:
+					Value = CurrentInstanceCache->MeshResults[Data.DataTypeIndex].get();
+					CacheResources.Remove(Value);
+					CurrentInstanceCache->MeshResults[Data.DataTypeIndex] = nullptr;
+					break;
+
+				case DATATYPE::DT_LAYOUT:
+					Value = CurrentInstanceCache->LayoutResults[Data.DataTypeIndex].get();
+					CurrentInstanceCache->LayoutResults[Data.DataTypeIndex] = nullptr;
+					break;
+
+				case DATATYPE::DT_INSTANCE:
+					Value = CurrentInstanceCache->InstanceResults[Data.DataTypeIndex].get();
+					CurrentInstanceCache->InstanceResults[Data.DataTypeIndex] = nullptr;
+					break;
+
+				default:
+					break;
+				}
+
+				Data.OpHitCount = 0;
+				Data.IsValueValid = false;
+			}
+		}
+
+
+		void LogWorkingMemory(const class CodeRunner* CurrentRunner) const;
+
+
+#if !(UE_BUILD_SHIPPING || UE_BUILD_TEST)
+		/** Temp variable with the ID of the thread running code, for debugging. */
+		uint32 DebugRunnerThreadID = FThread::InvalidThreadId;
+#endif
+
+		/** This is a development-only check to make sure calls to resource management happen in the correct thread. */
+		inline void BeginRunnerThread()
+		{
+#if !(UE_BUILD_SHIPPING || UE_BUILD_TEST)
+			// If this check fails it means have not set up correctly all the paths to debug threading for resource management.
+			check(DebugRunnerThreadID == FThread::InvalidThreadId);
+
+			DebugRunnerThreadID = FPlatformTLS::GetCurrentThreadId();
+#endif
+		}
+
+
+		/** This is a development-only check to make sure calls to resource management happen in the correct thread. */
+		inline void CheckRunnerThread()
+		{
+#if !(UE_BUILD_SHIPPING || UE_BUILD_TEST)
+			// If this check fails it means have not set up correctly all the paths to debug threading for resource management.
+			check(DebugRunnerThreadID != FThread::InvalidThreadId);
+
+			// If this check fails it means we are doing resource management from a thread that is not the CodeRunner::RunCode
+			// thread, and this is not allowed.
+			check(DebugRunnerThreadID== FPlatformTLS::GetCurrentThreadId());
+#endif
+		}
+
+
+		/** This is a development-only check to make sure calls to resource management happen in the correct thread. */
+		inline void EndRunnerThread()
+		{
+			CurrentInstanceCache->CheckHitCountsCleared();
+
+			// If this check fails it means some operation is not correctly handling resource management and didn't release
+			// a resource it created.
+			// This should be reported and reviewed, but it is not fatal. Some unnecessary memory may be used temporarily.
+			ensure(TempImages.Num() == 0);
+			ensure(TempMeshes.Num() == 0);
+
+#if !(UE_BUILD_SHIPPING || UE_BUILD_TEST)
+			// If this check fails it means have not set up correctly all the paths to debug threading for resource management.
+			check(DebugRunnerThreadID != FThread::InvalidThreadId);
+
+			DebugRunnerThreadID = FThread::InvalidThreadId;
+#endif
+		}
+
+	};
+
+
+	/** */
+    class System::Private : public Base
+    {
+    public:
+
+        Private( Ptr<Settings>, const TSharedPtr<ExtensionDataStreamer>& );
+        virtual ~Private();
+
+        //-----------------------------------------------------------------------------------------
+        //! Own interface
+        //-----------------------------------------------------------------------------------------
+
+		/** This method can be used to internally prepare for code execution. */
+		MUTABLERUNTIME_API void BeginBuild(const TSharedPtr<const Model>&);
+		MUTABLERUNTIME_API void EndBuild();
+
+		MUTABLERUNTIME_API bool BuildBool(const TSharedPtr<const Model>&, const Parameters*, OP::ADDRESS) ;
+		MUTABLERUNTIME_API int32 BuildInt(const TSharedPtr<const Model>&, const Parameters*, OP::ADDRESS) ;
+		MUTABLERUNTIME_API float BuildScalar(const TSharedPtr<const Model>&, const Parameters*, OP::ADDRESS) ;
+		MUTABLERUNTIME_API FVector4f BuildColour(const TSharedPtr<const Model>&, const Parameters*, OP::ADDRESS) ;
+		MUTABLERUNTIME_API Ptr<const String> BuildString(const TSharedPtr<const Model>&, const Parameters*, OP::ADDRESS) ;
+		MUTABLERUNTIME_API Ptr<const Image> BuildImage(const TSharedPtr<const Model>&, const Parameters*, OP::ADDRESS, int32 MipsToSkip, int32 LOD) ;
+		MUTABLERUNTIME_API Ptr<const Mesh> BuildMesh(const TSharedPtr<const Model>&, const Parameters*, OP::ADDRESS) ;
+		MUTABLERUNTIME_API Ptr<const Layout> BuildLayout(const TSharedPtr<const Model>&, const Parameters*, OP::ADDRESS) ;
+    	MUTABLERUNTIME_API FProjector BuildProjector(const TSharedPtr<const Model>&, const Parameters*, OP::ADDRESS) ;
+
+		ExtensionDataStreamer* GetExtensionDataStreamer() const { return ExtensionDataStreamer.Get(); }
+
+        //!
+        Ptr<const Settings> Settings;
+
+        //! Data streaming interface, if any.
+		TSharedPtr<ModelReader> StreamInterface;
+
+		TSharedPtr<ImageParameterGenerator> ImageParameterGenerator;
+
+		/** If this is set, it will be tried first instead of the internal formatting function. */
+		FImageOperator::FImagePixelFormatFunc ImagePixelFormatOverride;
+
+		/** */
+		FWorkingMemoryManager WorkingMemoryManager;
+
+		/** Counter used to generate unique IDs for every new instance created in the system. */
+        Instance::ID LastInstanceID = 0;
+
+		/** The pointer returned by this function is only valid for the duration of the current mutable operation. */
+		inline FLiveInstance* FindLiveInstance(Instance::ID id);
+
+        //!
+        bool CheckUpdatedParameters( const FLiveInstance*, const Ptr<const Parameters>&, uint64& OutUpdatedParameters );
+
+
+		void RunCode(const TSharedPtr<const Model>&, const Parameters*, OP::ADDRESS at, uint32 LODs = System::AllLODs, uint8 executionOptions = 0, int32 LOD = 0);
+
+		//!
+		void PrepareCache(const Model*, int32 State);
+
+	private:
+
+		/** Owned by this system. */
+		TSharedPtr<ExtensionDataStreamer> ExtensionDataStreamer = nullptr;
+
+		/** This flag is turned on when a streaming error or similar happens. Results are not usable.
+		* This should only happen in-editor.
+		*/
+		bool bUnrecoverableError = false;
+
+    };
+
+}
diff --git a/Engine/Plugins/Experimental/Mutable/Source/MutableRuntime/Private/MuR/CodeRunnerTasked.cpp b/Engine/Plugins/Experimental/Mutable/Source/MutableRuntime/Private/MuR/CodeRunnerTasked.cpp
--- a/Engine/Plugins/Experimental/Mutable/Source/MutableRuntime/Private/MuR/CodeRunnerTasked.cpp
+++ b/Engine/Plugins/Experimental/Mutable/Source/MutableRuntime/Private/MuR/CodeRunnerTasked.cpp
@@ -34,6 +34,7 @@
 #include "MuR/SerialisationPrivate.h"
 #include "MuR/Settings.h"
 #include "MuR/SystemPrivate.h"
+#include "MuR/MutableRuntimeModule.h"
 #include "Stats/Stats2.h"
 #include "Templates/RefCounting.h"
 #include "Templates/SharedPointer.h"
diff --git a/Engine/Plugins/Experimental/Mutable/Source/MutableRuntime/Private/MuR/ImagePrivate.h b/Engine/Plugins/Experimental/Mutable/Source/MutableRuntime/Private/MuR/ImagePrivate.h
--- a/Engine/Plugins/Experimental/Mutable/Source/MutableRuntime/Private/MuR/ImagePrivate.h
+++ b/Engine/Plugins/Experimental/Mutable/Source/MutableRuntime/Private/MuR/ImagePrivate.h
@@ -12,54 +12,10 @@
 
 namespace mu
 {
-	MUTABLE_DEFINE_ENUM_SERIALISABLE( EBlendType );
 	MUTABLE_DEFINE_ENUM_SERIALISABLE( EMipmapFilterType );
-	MUTABLE_DEFINE_ENUM_SERIALISABLE( ECompositeImageMode );
-	MUTABLE_DEFINE_ENUM_SERIALISABLE( ESamplingMethod );
-	MUTABLE_DEFINE_ENUM_SERIALISABLE( EMinFilterMethod );
 	MUTABLE_DEFINE_ENUM_SERIALISABLE( EImageFormat );
 	MUTABLE_DEFINE_ENUM_SERIALISABLE( EAddressMode );
 
-	//---------------------------------------------------------------------------------------------
-	//!
-	//---------------------------------------------------------------------------------------------
-	MUTABLERUNTIME_API inline EImageFormat GetUncompressedFormat(EImageFormat Format)
-	{
-		check(Format < EImageFormat::IF_COUNT);
-
-		EImageFormat Result = Format;
-
-		switch (Result)
-		{
-		case EImageFormat::IF_L_UBIT_RLE: Result = EImageFormat::IF_L_UBYTE; break;
-		case EImageFormat::IF_L_UBYTE_RLE: Result = EImageFormat::IF_L_UBYTE; break;
-		case EImageFormat::IF_RGB_UBYTE_RLE: Result = EImageFormat::IF_RGB_UBYTE; break;
-        case EImageFormat::IF_RGBA_UBYTE_RLE: Result = EImageFormat::IF_RGBA_UBYTE; break;
-        case EImageFormat::IF_BC1: Result = EImageFormat::IF_RGBA_UBYTE; break;
-        case EImageFormat::IF_BC2: Result = EImageFormat::IF_RGBA_UBYTE; break;
-        case EImageFormat::IF_BC3: Result = EImageFormat::IF_RGBA_UBYTE; break;
-        case EImageFormat::IF_BC4: Result = EImageFormat::IF_L_UBYTE; break;
-        case EImageFormat::IF_BC5: Result = EImageFormat::IF_RGB_UBYTE; break;
-		case EImageFormat::IF_ASTC_4x4_RGB_LDR: Result = EImageFormat::IF_RGB_UBYTE; break;
-		case EImageFormat::IF_ASTC_4x4_RGBA_LDR: Result = EImageFormat::IF_RGBA_UBYTE; break;
-		case EImageFormat::IF_ASTC_4x4_RG_LDR: Result = EImageFormat::IF_RGB_UBYTE; break;
-		case EImageFormat::IF_ASTC_6x6_RGB_LDR: Result = EImageFormat::IF_RGB_UBYTE; break;
-		case EImageFormat::IF_ASTC_6x6_RGBA_LDR: Result = EImageFormat::IF_RGBA_UBYTE; break;
-		case EImageFormat::IF_ASTC_6x6_RG_LDR: Result = EImageFormat::IF_RGB_UBYTE; break;
-		case EImageFormat::IF_ASTC_8x8_RGB_LDR: Result = EImageFormat::IF_RGB_UBYTE; break;
-		case EImageFormat::IF_ASTC_8x8_RGBA_LDR: Result = EImageFormat::IF_RGBA_UBYTE; break;
-		case EImageFormat::IF_ASTC_8x8_RG_LDR: Result = EImageFormat::IF_RGB_UBYTE; break;
-		case EImageFormat::IF_ASTC_10x10_RGB_LDR: Result = EImageFormat::IF_RGB_UBYTE; break;
-		case EImageFormat::IF_ASTC_10x10_RGBA_LDR: Result = EImageFormat::IF_RGBA_UBYTE; break;
-		case EImageFormat::IF_ASTC_10x10_RG_LDR: Result = EImageFormat::IF_RGB_UBYTE; break;
-		case EImageFormat::IF_ASTC_12x12_RGB_LDR: Result = EImageFormat::IF_RGB_UBYTE; break;
-		case EImageFormat::IF_ASTC_12x12_RGBA_LDR: Result = EImageFormat::IF_RGBA_UBYTE; break;
-		case EImageFormat::IF_ASTC_12x12_RG_LDR: Result = EImageFormat::IF_RGB_UBYTE; break;
-		default: break;
-		}
-
-		return Result;
-	}
 
     //---------------------------------------------------------------------------------------------
     //!
@@ -425,25 +381,5 @@ namespace mu
 		return DecompressionFuncRefType(Private::DecompressionFuncNotFoundFunc);
 	}
 
-
-    //---------------------------------------------------------------------------------------------
-    //! Use with care.
-    //---------------------------------------------------------------------------------------------
-    template<class T>
-    Ptr<T> CloneOrTakeOver( const T* source )
-    {
-        Ptr<T> result;
-        if (source->IsUnique())
-        {
-            result = const_cast<T*>(source);
-        }
-        else
-        {
-            result = source->Clone();
-        }
-
-        return result;
-    }
-
 }
 
diff --git a/Engine/Plugins/Experimental/Mutable/Source/MutableRuntime/Private/MuR/MemoryTrackingUtils.cpp b/Engine/Plugins/Experimental/Mutable/Source/MutableRuntime/Private/MuR/MemoryTrackingUtils.cpp
--- a/Engine/Plugins/Experimental/Mutable/Source/MutableRuntime/Private/MuR/MemoryTrackingUtils.cpp
+++ b/Engine/Plugins/Experimental/Mutable/Source/MutableRuntime/Private/MuR/MemoryTrackingUtils.cpp
@@ -4,6 +4,7 @@
 
 
 #include "MuR/Platform.h"
+#include "MuR/MutableRuntimeModule.h"
 #include "Logging/LogMacros.h"
 #include "HAL/LowLevelMemTracker.h"
 #include "HAL/IConsoleManager.h"
diff --git a/Engine/Plugins/Experimental/Mutable/Source/MutableRuntime/Private/MuR/ModelPrivate.h b/Engine/Plugins/Experimental/Mutable/Source/MutableRuntime/Private/MuR/ModelPrivate.h
deleted file mode 100644
diff --git a/Engine/Plugins/Experimental/Mutable/Source/MutableRuntime/Private/MuR/OpMeshFormat.h b/Engine/Plugins/Experimental/Mutable/Source/MutableRuntime/Private/MuR/OpMeshFormat.h
deleted file mode 100644
diff --git a/Engine/Plugins/Experimental/Mutable/Source/MutableRuntime/Private/MuR/OpMeshSmoothing.h b/Engine/Plugins/Experimental/Mutable/Source/MutableRuntime/Private/MuR/OpMeshSmoothing.h
--- a/Engine/Plugins/Experimental/Mutable/Source/MutableRuntime/Private/MuR/OpMeshSmoothing.h
+++ b/Engine/Plugins/Experimental/Mutable/Source/MutableRuntime/Private/MuR/OpMeshSmoothing.h
@@ -15,6 +15,7 @@
 #include "MuR/Mesh.h"
 #include "MuR/MeshPrivate.h"
 #include "MuR/MutableTrace.h"
+#include "MuR/MutableRuntimeModule.h"
 
 
 namespace mu
diff --git a/Engine/Plugins/Experimental/Mutable/Source/MutableRuntime/Private/MuR/Operations.h b/Engine/Plugins/Experimental/Mutable/Source/MutableRuntime/Private/MuR/Operations.h
deleted file mode 100644
diff --git a/Engine/Plugins/Experimental/Mutable/Source/MutableRuntime/Private/MuR/ParametersPrivate.h b/Engine/Plugins/Experimental/Mutable/Source/MutableRuntime/Private/MuR/ParametersPrivate.h
deleted file mode 100644
diff --git a/Engine/Plugins/Experimental/Mutable/Source/MutableRuntime/Private/MuR/Platform.h b/Engine/Plugins/Experimental/Mutable/Source/MutableRuntime/Private/MuR/Platform.h
--- a/Engine/Plugins/Experimental/Mutable/Source/MutableRuntime/Private/MuR/Platform.h
+++ b/Engine/Plugins/Experimental/Mutable/Source/MutableRuntime/Private/MuR/Platform.h
@@ -5,11 +5,9 @@
 // This file contains some hacks to solve differences between platforms
 
 #include "HAL/PlatformTime.h"
-#include "Logging/LogMacros.h"
 #include "MuR/Types.h"
 #include "HAL/UnrealMemory.h"
 
-MUTABLERUNTIME_API DECLARE_LOG_CATEGORY_EXTERN(LogMutableCore, Log, All);
 
 
 //! Unify debug defines
diff --git a/Engine/Plugins/Experimental/Mutable/Source/MutableRuntime/Private/MuR/SerialisationPrivate.h b/Engine/Plugins/Experimental/Mutable/Source/MutableRuntime/Private/MuR/SerialisationPrivate.h
deleted file mode 100644
diff --git a/Engine/Plugins/Experimental/Mutable/Source/MutableRuntime/Private/MuR/System.cpp b/Engine/Plugins/Experimental/Mutable/Source/MutableRuntime/Private/MuR/System.cpp
--- a/Engine/Plugins/Experimental/Mutable/Source/MutableRuntime/Private/MuR/System.cpp
+++ b/Engine/Plugins/Experimental/Mutable/Source/MutableRuntime/Private/MuR/System.cpp
@@ -29,6 +29,7 @@
 #include "MuR/Platform.h"
 #include "MuR/Serialisation.h"
 #include "MuR/SystemPrivate.h"
+#include "MuR/MutableRuntimeModule.h"
 #include "Templates/SharedPointer.h"
 #include "Templates/Tuple.h"
 #include "Trace/Detail/Channel.h"
diff --git a/Engine/Plugins/Experimental/Mutable/Source/MutableRuntime/Private/MuR/SystemPrivate.h b/Engine/Plugins/Experimental/Mutable/Source/MutableRuntime/Private/MuR/SystemPrivate.h
deleted file mode 100644
diff --git a/Engine/Plugins/Experimental/Mutable/Source/MutableRuntime/Public/MuR/Image.h b/Engine/Plugins/Experimental/Mutable/Source/MutableRuntime/Public/MuR/Image.h
--- a/Engine/Plugins/Experimental/Mutable/Source/MutableRuntime/Public/MuR/Image.h
+++ b/Engine/Plugins/Experimental/Mutable/Source/MutableRuntime/Public/MuR/Image.h
@@ -24,13 +24,17 @@ namespace mu::MemoryCounters
 
 namespace mu
 {
-
 	// Forward declarations
 	class Image;
 
 	typedef Ptr<Image> ImagePtr;
 	typedef Ptr<const Image> ImagePtrConst;
 
+	MUTABLE_DEFINE_ENUM_SERIALISABLE( EBlendType );
+	MUTABLE_DEFINE_ENUM_SERIALISABLE( ECompositeImageMode );
+	MUTABLE_DEFINE_ENUM_SERIALISABLE( ESamplingMethod );
+	MUTABLE_DEFINE_ENUM_SERIALISABLE( EMinFilterMethod );
+
     //! \brief 2D image resource with mipmaps.
 	//! \ingroup runtime
     class MUTABLERUNTIME_API Image : public Resource
@@ -410,6 +414,62 @@ namespace mu
 	/** */
 	MUTABLERUNTIME_API void ImageSwizzle(Image* Result, const Ptr<const Image> Sources[], const uint8 Channels[]);
 
+	MUTABLERUNTIME_API inline EImageFormat GetUncompressedFormat(EImageFormat Format)
+	{
+		check(Format < EImageFormat::IF_COUNT);
+
+		EImageFormat Result = Format;
+
+		switch (Result)
+		{
+		case EImageFormat::IF_L_UBIT_RLE: Result = EImageFormat::IF_L_UBYTE; break;
+		case EImageFormat::IF_L_UBYTE_RLE: Result = EImageFormat::IF_L_UBYTE; break;
+		case EImageFormat::IF_RGB_UBYTE_RLE: Result = EImageFormat::IF_RGB_UBYTE; break;
+        case EImageFormat::IF_RGBA_UBYTE_RLE: Result = EImageFormat::IF_RGBA_UBYTE; break;
+        case EImageFormat::IF_BC1: Result = EImageFormat::IF_RGBA_UBYTE; break;
+        case EImageFormat::IF_BC2: Result = EImageFormat::IF_RGBA_UBYTE; break;
+        case EImageFormat::IF_BC3: Result = EImageFormat::IF_RGBA_UBYTE; break;
+        case EImageFormat::IF_BC4: Result = EImageFormat::IF_L_UBYTE; break;
+        case EImageFormat::IF_BC5: Result = EImageFormat::IF_RGB_UBYTE; break;
+		case EImageFormat::IF_ASTC_4x4_RGB_LDR: Result = EImageFormat::IF_RGB_UBYTE; break;
+		case EImageFormat::IF_ASTC_4x4_RGBA_LDR: Result = EImageFormat::IF_RGBA_UBYTE; break;
+		case EImageFormat::IF_ASTC_4x4_RG_LDR: Result = EImageFormat::IF_RGB_UBYTE; break;
+		case EImageFormat::IF_ASTC_6x6_RGB_LDR: Result = EImageFormat::IF_RGB_UBYTE; break;
+		case EImageFormat::IF_ASTC_6x6_RGBA_LDR: Result = EImageFormat::IF_RGBA_UBYTE; break;
+		case EImageFormat::IF_ASTC_6x6_RG_LDR: Result = EImageFormat::IF_RGB_UBYTE; break;
+		case EImageFormat::IF_ASTC_8x8_RGB_LDR: Result = EImageFormat::IF_RGB_UBYTE; break;
+		case EImageFormat::IF_ASTC_8x8_RGBA_LDR: Result = EImageFormat::IF_RGBA_UBYTE; break;
+		case EImageFormat::IF_ASTC_8x8_RG_LDR: Result = EImageFormat::IF_RGB_UBYTE; break;
+		case EImageFormat::IF_ASTC_10x10_RGB_LDR: Result = EImageFormat::IF_RGB_UBYTE; break;
+		case EImageFormat::IF_ASTC_10x10_RGBA_LDR: Result = EImageFormat::IF_RGBA_UBYTE; break;
+		case EImageFormat::IF_ASTC_10x10_RG_LDR: Result = EImageFormat::IF_RGB_UBYTE; break;
+		case EImageFormat::IF_ASTC_12x12_RGB_LDR: Result = EImageFormat::IF_RGB_UBYTE; break;
+		case EImageFormat::IF_ASTC_12x12_RGBA_LDR: Result = EImageFormat::IF_RGBA_UBYTE; break;
+		case EImageFormat::IF_ASTC_12x12_RG_LDR: Result = EImageFormat::IF_RGB_UBYTE; break;
+		default: break;
+		}
+
+		return Result;
+	}
+
+	//---------------------------------------------------------------------------------------------
+	//! Use with care.
+	//---------------------------------------------------------------------------------------------
+	template<class T>
+	Ptr<T> CloneOrTakeOver( const T* source )
+	{
+		Ptr<T> result;
+		if (source->IsUnique())
+		{
+			result = const_cast<T*>(source);
+		}
+		else
+		{
+			result = source->Clone();
+		}
 
+		return result;
+	}
+	
 }
 
diff --git a/Engine/Plugins/Experimental/Mutable/Source/MutableRuntime/Public/MuR/MutableRuntimeModule.h b/Engine/Plugins/Experimental/Mutable/Source/MutableRuntime/Public/MuR/MutableRuntimeModule.h
--- a/Engine/Plugins/Experimental/Mutable/Source/MutableRuntime/Public/MuR/MutableRuntimeModule.h
+++ b/Engine/Plugins/Experimental/Mutable/Source/MutableRuntime/Public/MuR/MutableRuntimeModule.h
@@ -3,6 +3,7 @@
 #pragma once
 
 #include "Modules/ModuleInterface.h"
+#include "Logging/LogMacros.h"
 
 
 // This will use unreal's malloc instead of system malloc for mutable memory operations
@@ -16,6 +17,9 @@
 #define USE_STAT_MUTABLE_MEMORY					0
 
 
+MUTABLERUNTIME_API DECLARE_LOG_CATEGORY_EXTERN(LogMutableCore, Log, All);
+
+
 class MUTABLERUNTIME_API FMutableRuntimeModule : public IModuleInterface
 {
 public:
diff --git a/Engine/Plugins/Experimental/Mutable/Source/MutableTools/Private/MuT/AST.h b/Engine/Plugins/Experimental/Mutable/Source/MutableTools/Private/MuT/AST.h
--- a/Engine/Plugins/Experimental/Mutable/Source/MutableTools/Private/MuT/AST.h
+++ b/Engine/Plugins/Experimental/Mutable/Source/MutableTools/Private/MuT/AST.h
@@ -30,6 +30,8 @@
 #include <utility>
 #include <set>
 
+#include "MuT/Compiler.h"
+
 namespace std
 {
 
@@ -178,8 +180,6 @@ namespace mu
     typedef TSet<Ptr<ASTOp>> ASTOpSet;
 
 
-	typedef TFunction<Ptr<Image>(int32)> FReferencedResourceFunc;
-
     //! Detailed optimization flags
     struct FModelOptimizationOptions
     {
diff --git a/Engine/Plugins/Experimental/Mutable/Source/MutableTools/Private/MuT/CodeOptimiser.cpp b/Engine/Plugins/Experimental/Mutable/Source/MutableTools/Private/MuT/CodeOptimiser.cpp
--- a/Engine/Plugins/Experimental/Mutable/Source/MutableTools/Private/MuT/CodeOptimiser.cpp
+++ b/Engine/Plugins/Experimental/Mutable/Source/MutableTools/Private/MuT/CodeOptimiser.cpp
@@ -23,6 +23,8 @@
 
 #include <unordered_set>
 
+#include "MuR/MutableRuntimeModule.h"
+
 namespace mu
 {
 
diff --git a/Engine/Plugins/Experimental/Mutable/Source/MutableTools/Private/MuT/CodeOptimiser_Parameters.cpp b/Engine/Plugins/Experimental/Mutable/Source/MutableTools/Private/MuT/CodeOptimiser_Parameters.cpp
--- a/Engine/Plugins/Experimental/Mutable/Source/MutableTools/Private/MuT/CodeOptimiser_Parameters.cpp
+++ b/Engine/Plugins/Experimental/Mutable/Source/MutableTools/Private/MuT/CodeOptimiser_Parameters.cpp
@@ -36,6 +36,8 @@
 #include <memory>
 #include <utility>
 
+#include "MuR/MutableRuntimeModule.h"
+
 
 namespace mu
 {
diff --git a/Engine/Plugins/Experimental/Mutable/Source/MutableTools/Private/MuT/Compiler.cpp b/Engine/Plugins/Experimental/Mutable/Source/MutableTools/Private/MuT/Compiler.cpp
--- a/Engine/Plugins/Experimental/Mutable/Source/MutableTools/Private/MuT/Compiler.cpp
+++ b/Engine/Plugins/Experimental/Mutable/Source/MutableTools/Private/MuT/Compiler.cpp
@@ -20,6 +20,7 @@
 #include "MuR/ParametersPrivate.h"
 #include "MuR/Platform.h"
 #include "MuR/Serialisation.h"
+#include "MuR/MutableRuntimeModule.h"
 #include "MuR/System.h"
 #include "MuT/AST.h"
 #include "MuT/ASTOpParameter.h"
diff --git a/Engine/Plugins/Experimental/Mutable/Source/MutableTools/Private/MuT/ErrorLog.cpp b/Engine/Plugins/Experimental/Mutable/Source/MutableTools/Private/MuT/ErrorLog.cpp
--- a/Engine/Plugins/Experimental/Mutable/Source/MutableTools/Private/MuT/ErrorLog.cpp
+++ b/Engine/Plugins/Experimental/Mutable/Source/MutableTools/Private/MuT/ErrorLog.cpp
@@ -20,6 +20,7 @@
 #include "MuR/Platform.h"
 #include "MuR/Skeleton.h"
 #include "MuR/System.h"
+#include "MuR/MutableRuntimeModule.h"
 #include "MuT/ErrorLogPrivate.h"
 #include "Templates/Tuple.h"
 #include "Trace/Detail/Channel.h"
diff --git a/Engine/Plugins/Experimental/Mutable/Source/MutableTools/Private/MuT/NodeImageLayerColourPrivate.h b/Engine/Plugins/Experimental/Mutable/Source/MutableTools/Private/MuT/NodeImageLayerColourPrivate.h
--- a/Engine/Plugins/Experimental/Mutable/Source/MutableTools/Private/MuT/NodeImageLayerColourPrivate.h
+++ b/Engine/Plugins/Experimental/Mutable/Source/MutableTools/Private/MuT/NodeImageLayerColourPrivate.h
@@ -6,6 +6,7 @@
 #include "MuT/NodeImageLayerPrivate.h"
 #include "MuT/NodeImageLayerColour.h"
 #include "MuT/NodeColour.h"
+#include "MuR/Image.h"
 
 
 namespace mu
diff --git a/Engine/Plugins/Experimental/Mutable/Source/MutableTools/Private/MuT/NodeImageNormalComposite.cpp b/Engine/Plugins/Experimental/Mutable/Source/MutableTools/Private/MuT/NodeImageNormalComposite.cpp
--- a/Engine/Plugins/Experimental/Mutable/Source/MutableTools/Private/MuT/NodeImageNormalComposite.cpp
+++ b/Engine/Plugins/Experimental/Mutable/Source/MutableTools/Private/MuT/NodeImageNormalComposite.cpp
@@ -6,6 +6,7 @@
 #include "Misc/AssertionMacros.h"
 #include "MuT/NodeImageNormalCompositePrivate.h"
 #include "MuT/NodePrivate.h"
+#include "MuR/Image.h"
 
 
 namespace mu
diff --git a/Engine/Plugins/Experimental/Mutable/Source/MutableTools/Private/MuT/NodeImageProject.cpp b/Engine/Plugins/Experimental/Mutable/Source/MutableTools/Private/MuT/NodeImageProject.cpp
--- a/Engine/Plugins/Experimental/Mutable/Source/MutableTools/Private/MuT/NodeImageProject.cpp
+++ b/Engine/Plugins/Experimental/Mutable/Source/MutableTools/Private/MuT/NodeImageProject.cpp
@@ -10,7 +10,7 @@
 #include "MuT/NodePrivate.h"
 #include "MuT/NodeProjector.h"
 #include "MuT/NodeScalar.h"
-
+#include "MuR/Image.h"
 
 
 namespace mu
diff --git a/Engine/Plugins/Experimental/Mutable/Source/MutableTools/Private/MuT/Platform.h b/Engine/Plugins/Experimental/Mutable/Source/MutableTools/Private/MuT/Platform.h
--- a/Engine/Plugins/Experimental/Mutable/Source/MutableTools/Private/MuT/Platform.h
+++ b/Engine/Plugins/Experimental/Mutable/Source/MutableTools/Private/MuT/Platform.h
@@ -2,8 +2,6 @@
 
 #pragma once
 
-#include "MuR/Platform.h"
-
 #include <vector>
 #include <map>
 #include <set>
diff --git a/Engine/Plugins/Experimental/Mutable/Source/MutableTools/Public/MuT/Compiler.h b/Engine/Plugins/Experimental/Mutable/Source/MutableTools/Public/MuT/Compiler.h
--- a/Engine/Plugins/Experimental/Mutable/Source/MutableTools/Public/MuT/Compiler.h
+++ b/Engine/Plugins/Experimental/Mutable/Source/MutableTools/Public/MuT/Compiler.h
@@ -2,7 +2,6 @@
 
 #pragma once
 
-#include "MuT/AST.h"
 #include "MuR/Ptr.h"
 #include "MuR/RefCounted.h"
 #include "MuR/Image.h"
@@ -31,6 +30,8 @@ namespace mu
     class Node;
     class NodeTransformedObject;
 
+	typedef TFunction<Ptr<Image>(int32)> FReferencedResourceFunc;
+
     //! \brief Options used to compile the models with a compiler.
     class MUTABLETOOLS_API CompilerOptions : public RefCounted
     {
diff --git a/Engine/Plugins/Experimental/Mutable/Source/MutableTools/Public/MuT/TypeInfo.h b/Engine/Plugins/Experimental/Mutable/Source/MutableTools/Public/MuT/TypeInfo.h
--- a/Engine/Plugins/Experimental/Mutable/Source/MutableTools/Public/MuT/TypeInfo.h
+++ b/Engine/Plugins/Experimental/Mutable/Source/MutableTools/Public/MuT/TypeInfo.h
@@ -2,11 +2,11 @@
 
 #pragma once
 
-#include "../Private/MuR/ParametersPrivate.h"
 #include "HAL/PlatformMath.h"
 #include "MuR/Image.h"
 #include "MuR/MeshBufferSet.h"
 #include "MuR/Parameters.h"
+#include "MuR/ParametersPrivate.h"
 
 // This header file provides additional information about data types defined in the runtime, that
 // can be useful only for the tools. For instance, it provides strings for some enumeration values.
