[Mutable] Remove old code and CVars used to roll back fixes.

- Remove CVarApplyFixDoComponentsNeedUpdate, CVarApplyFixPrepareSkeletons and CVarRevertCacheTextureParameters.

#rnx

diff --git a/Engine/Plugins/Experimental/Mutable/Source/CustomizableObject/Private/MuCO/CustomizableObjectInstance.cpp b/Engine/Plugins/Experimental/Mutable/Source/CustomizableObject/Private/MuCO/CustomizableObjectInstance.cpp
--- a/Engine/Plugins/Experimental/Mutable/Source/CustomizableObject/Private/MuCO/CustomizableObjectInstance.cpp
+++ b/Engine/Plugins/Experimental/Mutable/Source/CustomizableObject/Private/MuCO/CustomizableObjectInstance.cpp
@@ -1614,13 +1614,6 @@ void SetMeshUVChannelDensity(FMeshUVChannelInfo& UVChannelInfo, float Density =
 	}
 }
 
-// TODO PRP: MTBL-1653 Remove CVar
-static bool bApplyFixDoComponentsNeedUpdate = true;
-FAutoConsoleVariableRef CVarApplyFixDoComponentsNeedUpdate(
-	TEXT("Mutable.ApplyFixDoComponentsNeedUpdate"),
-	bApplyFixDoComponentsNeedUpdate,
-	TEXT("If true, the new version of DoComponentsNeedUpdate will be used. This fix tries to detect cases where the output mesh would be invalid and cancel the update."));
-
 
 bool UCustomizableInstancePrivateData::DoComponentsNeedUpdate(UCustomizableObjectInstance* Public, const TSharedRef<FUpdateContextPrivate>& OperationData, TArray<bool>& OutComponentNeedsUpdate, bool& bHasInvalidMesh)
 {
@@ -1630,190 +1623,95 @@ bool UCustomizableInstancePrivateData::DoComponentsNeedUpdate(UCustomizableObjec
 	const int32 NumLODs = OperationData->InstanceUpdateData.LODs.Num();
 	const int32 NumComponents = CustomizableObject->GetComponentCount();
 
-	if (bApplyFixDoComponentsNeedUpdate)
-	{
-		TArray<bool> ComponentWithMesh;
-		ComponentWithMesh.Init(false, NumComponents);
+	TArray<bool> ComponentWithMesh;
+	ComponentWithMesh.Init(false, NumComponents);
+
+	TArray<mu::FResourceID> MeshIDs;
+	MeshIDs.Init(MAX_uint64, NumComponents * MAX_MESH_LOD_COUNT);
 
-		TArray<mu::FResourceID> MeshIDs;
-		MeshIDs.Init(MAX_uint64, NumComponents * MAX_MESH_LOD_COUNT);
+	// Gather the Mesh Ids of all components, and validate the integrity of the meshes to generate. 
+	for (int32 LODIndex = OperationData->CurrentMinLOD; LODIndex <= OperationData->CurrentMaxLOD && LODIndex < NumLODs; ++LODIndex)
+	{
+		const FInstanceUpdateData::FLOD& LOD = OperationData->InstanceUpdateData.LODs[LODIndex];
 
-		// Gather the Mesh Ids of all components, and validate the integrity of the meshes to generate. 
-		for (int32 LODIndex = OperationData->CurrentMinLOD; LODIndex <= OperationData->CurrentMaxLOD && LODIndex < NumLODs; ++LODIndex)
+		for (int32 ComponentIndex = 0; ComponentIndex < LOD.ComponentCount; ++ComponentIndex)
 		{
-			const FInstanceUpdateData::FLOD& LOD = OperationData->InstanceUpdateData.LODs[LODIndex];
+			const FInstanceUpdateData::FComponent& Component = OperationData->InstanceUpdateData.Components[LOD.FirstComponent + ComponentIndex];
 
-			for (int32 ComponentIndex = 0; ComponentIndex < LOD.ComponentCount; ++ComponentIndex)
+			if (Component.Id >= NumComponents)
 			{
-				const FInstanceUpdateData::FComponent& Component = OperationData->InstanceUpdateData.Components[LOD.FirstComponent + ComponentIndex];
-
-				if (Component.Id >= NumComponents)
-				{
-					ensureMsgf(false, TEXT("Mutable: Failed to generate SkeletalMesh. Invalid ComponentIndex. Index [%d]. Valid range [0 .. %d)."), Component.Id, NumComponents);
-					bHasInvalidMesh = true;
-					return false;
-				}
-
-				if (!Component.bGenerated || !Component.Mesh || Component.SurfaceCount == 0) // else
-				{
-					continue;
-				}
-
-				// Unreal does not support empty sections.
-				if (Component.Mesh->GetVertexCount() == 0) // else
-				{
-					UE_LOG(LogMutable, Error, TEXT("Failed to generate SkeletalMesh for CO Instance [%s]. CO [%s] has invalid geometry for LOD [%d] Component [%d]."),
-						*Public->GetName(), *CustomizableObject->GetName(),
-						OperationData->CurrentMaxLOD, ComponentIndex);
-					bHasInvalidMesh = true;
-					continue;
-				}
-
-				ComponentWithMesh[Component.Id] = true;
-				MeshIDs[Component.Id * MAX_MESH_LOD_COUNT + LODIndex] = Component.MeshID;
+				ensureMsgf(false, TEXT("Mutable: Failed to generate SkeletalMesh. Invalid ComponentIndex. Index [%d]. Valid range [0 .. %d)."), Component.Id, NumComponents);
+				bHasInvalidMesh = true;
+				return false;
 			}
-		}
-
-		// Find which components need an update
-		OutComponentNeedsUpdate.AddDefaulted(NumComponents);
 
-		for (int32 ComponentIndex = 0; ComponentIndex < NumComponents; ++ComponentIndex)
-		{
-			if (OperationData->bUseMeshCache)
+			if (!Component.bGenerated || !Component.Mesh || Component.SurfaceCount == 0) // else
 			{
-				if (CustomizableObject->GetPrivate()->MeshCache.Get(OperationData->MeshDescriptors[ComponentIndex]))
-				{
-					OutComponentNeedsUpdate[ComponentIndex] = true;
-					ComponentWithMesh[ComponentIndex] = true;
-					continue;					
-				}
+				continue;
 			}
-			
-			// Components with mesh must have valid geometry at CurrentMaxLOD
-			if (ComponentWithMesh[ComponentIndex] && MeshIDs[ComponentIndex * MAX_MESH_LOD_COUNT + OperationData->CurrentMaxLOD] == MAX_uint64)
+
+			// Unreal does not support empty sections.
+			if (Component.Mesh->GetVertexCount() == 0) // else
 			{
-				UE_LOG(LogMutable, Error, TEXT("Failed to generate SkeletalMesh for CO Instance [%s]. CO [%s] is missing geometry for LOD [%d] Component [%d]."),
+				UE_LOG(LogMutable, Error, TEXT("Failed to generate SkeletalMesh for CO Instance [%s]. CO [%s] has invalid geometry for LOD [%d] Component [%d]."),
 					*Public->GetName(), *CustomizableObject->GetName(),
 					OperationData->CurrentMaxLOD, ComponentIndex);
 				bHasInvalidMesh = true;
 				continue;
 			}
 
-			// Update the component if there is a mesh and it shouldn't, or the other way around.
-			const bool bHasSkeletalMesh = Public->SkeletalMeshes.IsValidIndex(ComponentIndex) && Public->SkeletalMeshes[ComponentIndex];
-			OutComponentNeedsUpdate[ComponentIndex] = (ComponentWithMesh[ComponentIndex] != bHasSkeletalMesh);
-
-			const FCustomizableInstanceComponentData* ComponentData = GetComponentData(ComponentIndex);
-			if (!ComponentData) // Could be nullptr if the component has not been generated.
-			{
-				continue;
-			}
-
-			// Update if MeshIDs are different
-			const int32 ComponentOffset = ComponentIndex * MAX_MESH_LOD_COUNT;
-			for (int32 MeshIndex = 0; !OutComponentNeedsUpdate[ComponentIndex] && MeshIndex < MAX_MESH_LOD_COUNT; ++MeshIndex)
-			{
-				OutComponentNeedsUpdate[ComponentIndex] = MeshIDs[ComponentOffset + MeshIndex] != ComponentData->LastMeshIdPerLOD[MeshIndex];
-			}
+			ComponentWithMesh[Component.Id] = true;
+			MeshIDs[Component.Id * MAX_MESH_LOD_COUNT + LODIndex] = Component.MeshID;
 		}
-
-		// Mark as invalid if all components have empty meshes.
-		bHasInvalidMesh = bHasInvalidMesh || ComponentWithMesh.Find(true) == INDEX_NONE;
-
-		return !bHasInvalidMesh && OutComponentNeedsUpdate.Find(true) != INDEX_NONE;
 	}
-	else // TODO PRP: MTBL-1653 Remove old version
-	{
-		OutComponentNeedsUpdate.AddDefaulted(NumComponents);
 
-		// Return true if at least one mesh needs to be updated
-		bool bUpdateMeshes = false;
+	// Find which components need an update
+	OutComponentNeedsUpdate.AddDefaulted(NumComponents);
 
-		for (int32 ComponentIndex = 0; ComponentIndex < NumComponents; ++ComponentIndex)
+	for (int32 ComponentIndex = 0; ComponentIndex < NumComponents; ++ComponentIndex)
+	{
+		if (OperationData->bUseMeshCache)
 		{
-			USkeletalMesh* SkeletalMesh = Public->SkeletalMeshes.IsValidIndex(ComponentIndex) ? Public->SkeletalMeshes[ComponentIndex] : nullptr;
-			FCustomizableInstanceComponentData* ComponentData = GetComponentData(ComponentIndex);
-
-			// Check if the mesh is/isn't generated and it should/shouldn't
-			bool bMeshNeedsUpdate = (SkeletalMesh && !ComponentData) || (!SkeletalMesh && ComponentData);
-
-			if (SkeletalMesh)
-			{
-				bMeshNeedsUpdate = bMeshNeedsUpdate || (SkeletalMesh->GetLODNum() - 1) != OperationData->CurrentMaxLOD;
-			}
-
-			if (ComponentData)
+			if (CustomizableObject->GetPrivate()->MeshCache.Get(OperationData->MeshDescriptors[ComponentIndex]))
 			{
-				for (int32 LODIndex = 0; !bMeshNeedsUpdate && LODIndex < NumLODs; ++LODIndex)
-				{
-					// Check if an LOD is generated and it shouldn't
-					const bool LODCanBeGenerated = LODIndex >= OperationData->CurrentMinLOD && LODIndex <= OperationData->CurrentMaxLOD;
-					bMeshNeedsUpdate = !LODCanBeGenerated && ComponentData->LastMeshIdPerLOD[LODIndex] != MAX_uint64;
-				}
+				OutComponentNeedsUpdate[ComponentIndex] = true;
+				ComponentWithMesh[ComponentIndex] = true;
+				continue;
 			}
-
-			OutComponentNeedsUpdate[ComponentIndex] = bMeshNeedsUpdate;
-
-			bUpdateMeshes = bUpdateMeshes || bMeshNeedsUpdate;
 		}
 
-		bool bFoundEmptyMesh = false;
-		bool bFoundEmptyLOD = false;
-
-		bool bFoundNonEmptyMesh = false;
-
-		for (int32 LODIndex = OperationData->CurrentMinLOD; LODIndex <= OperationData->CurrentMaxLOD && LODIndex < NumLODs; ++LODIndex)
+		// Components with mesh must have valid geometry at CurrentMaxLOD
+		if (ComponentWithMesh[ComponentIndex] && MeshIDs[ComponentIndex * MAX_MESH_LOD_COUNT + OperationData->CurrentMaxLOD] == MAX_uint64)
 		{
-
-			const FInstanceUpdateData::FLOD& LOD = OperationData->InstanceUpdateData.LODs[LODIndex];
-
-			for (int32 ComponentIndex = 0; ComponentIndex < LOD.ComponentCount; ++ComponentIndex)
-			{
-				const FInstanceUpdateData::FComponent& Component = OperationData->InstanceUpdateData.Components[LOD.FirstComponent + ComponentIndex];
-
-				const mu::FResourceID LastMeshID = GetLastMeshId(Component.Id, LODIndex);
-
-				if (Component.bGenerated)
-				{
-					if (Component.Mesh)
-					{
-						OutComponentNeedsUpdate[Component.Id] = OutComponentNeedsUpdate[Component.Id] || LastMeshID != Component.MeshID;
-
-						// Don't build a degenerated mesh if something went wrong
-						if (Component.Mesh->GetVertexCount() > 0)
-						{
-							bFoundNonEmptyMesh = true;
-						}
-						else if (LOD.ComponentCount == 1)
-						{
-							bFoundEmptyMesh = true;
-						}
-					}
-					else
-					{
-						OutComponentNeedsUpdate[Component.Id] = true;
-					}
-				}
-				else
-				{
-					OutComponentNeedsUpdate[Component.Id] = OutComponentNeedsUpdate[Component.Id] || LastMeshID != MAX_uint64;
-				}
-
-				bUpdateMeshes = bUpdateMeshes || OutComponentNeedsUpdate[Component.Id];
-			}
+			UE_LOG(LogMutable, Error, TEXT("Failed to generate SkeletalMesh for CO Instance [%s]. CO [%s] is missing geometry for LOD [%d] Component [%d]."),
+				*Public->GetName(), *CustomizableObject->GetName(),
+				OperationData->CurrentMaxLOD, ComponentIndex);
+			bHasInvalidMesh = true;
+			continue;
 		}
 
+		// Update the component if there is a mesh and it shouldn't, or the other way around.
+		const bool bHasSkeletalMesh = Public->SkeletalMeshes.IsValidIndex(ComponentIndex) && Public->SkeletalMeshes[ComponentIndex];
+		OutComponentNeedsUpdate[ComponentIndex] = (ComponentWithMesh[ComponentIndex] != bHasSkeletalMesh);
 
-		if (!bFoundNonEmptyMesh)
+		const FCustomizableInstanceComponentData* ComponentData = GetComponentData(ComponentIndex);
+		if (!ComponentData) // Could be nullptr if the component has not been generated.
 		{
-			bFoundEmptyLOD = true;
-			bUpdateMeshes = true;
-			UE_LOG(LogMutable, Warning, TEXT("Building instance: An LOD has no mesh geometry. This cannot be handled by Unreal."));
+			continue;
 		}
 
-		bHasInvalidMesh = bFoundEmptyLOD || bFoundEmptyMesh;
-		return bUpdateMeshes;
+		// Update if MeshIDs are different
+		const int32 ComponentOffset = ComponentIndex * MAX_MESH_LOD_COUNT;
+		for (int32 MeshIndex = 0; !OutComponentNeedsUpdate[ComponentIndex] && MeshIndex < MAX_MESH_LOD_COUNT; ++MeshIndex)
+		{
+			OutComponentNeedsUpdate[ComponentIndex] = MeshIDs[ComponentOffset + MeshIndex] != ComponentData->LastMeshIdPerLOD[MeshIndex];
+		}
 	}
+
+	// Mark as invalid if all components have empty meshes.
+	bHasInvalidMesh = bHasInvalidMesh || ComponentWithMesh.Find(true) == INDEX_NONE;
+
+	return !bHasInvalidMesh && OutComponentNeedsUpdate.Find(true) != INDEX_NONE;
 }
 
 
diff --git a/Engine/Plugins/Experimental/Mutable/Source/CustomizableObject/Private/MuCO/CustomizableObjectSystem.cpp b/Engine/Plugins/Experimental/Mutable/Source/CustomizableObject/Private/MuCO/CustomizableObjectSystem.cpp
--- a/Engine/Plugins/Experimental/Mutable/Source/CustomizableObject/Private/MuCO/CustomizableObjectSystem.cpp
+++ b/Engine/Plugins/Experimental/Mutable/Source/CustomizableObject/Private/MuCO/CustomizableObjectSystem.cpp
@@ -740,13 +740,6 @@ FAutoConsoleVariableRef CVarMaxTextureSizeToGenerate(
 	TEXT("Max texture size on Mutable textures. Mip 0 will be the first mip with max size equal or less than MaxTextureSizeToGenerate."
 		"If a texture doesn't have small enough mips, mip 0 will be the last mip available."));
 
-static bool bApplyFixPrepareSkeletons = true;
-
-static FAutoConsoleVariableRef CVarApplyFixPrepareSkeletons(
-	TEXT("mutable.ApplyFixPrepareSkeletons"), bApplyFixPrepareSkeletons,
-	TEXT("If true, Fix missing SkeletonsData when FirstLODToGenerate is greater than 0. If false, There may be a crash when generating meshes on platform that skip LODs."),
-	ECVF_Default);
-
 static FAutoConsoleVariable CVarDescriptorDebugPrint(
 	TEXT("mutable.DescriptorDebugPrint"),
 	false,
@@ -2050,135 +2043,53 @@ namespace impl
 	// This runs in a worker thread
 	void Subtask_Mutable_PrepareSkeletonData(const TSharedRef<FUpdateContextPrivate>& OperationData)
 	{
-		MUTABLE_CPUPROFILER_SCOPE(Subtask_Mutable_PrepareSkeletonData)
+		MUTABLE_CPUPROFILER_SCOPE(Subtask_Mutable_PrepareSkeletonData);
 
-		if (bApplyFixPrepareSkeletons)
+		for (FInstanceUpdateData::FComponent& Component : OperationData->InstanceUpdateData.Components)
 		{
-			for (FInstanceUpdateData::FComponent& Component : OperationData->InstanceUpdateData.Components)
+			if (!OperationData->InstanceUpdateData.Skeletons.IsValidIndex(Component.Id))
 			{
-				if (!OperationData->InstanceUpdateData.Skeletons.IsValidIndex(Component.Id))
-				{
-					OperationData->InstanceUpdateData.Skeletons.SetNum(Component.Id + 1);
-				}
-
-				FInstanceUpdateData::FSkeletonData& SkeletonData = OperationData->InstanceUpdateData.Skeletons[Component.Id];
-				SkeletonData.ComponentIndex = Component.Id;
-
-				mu::MeshPtrConst Mesh = Component.Mesh;
-				if (!Mesh)
-				{
-					continue;
-				}
-
-				// Add SkeletonIds 
-				const int32 SkeletonIDsCount = Mesh->GetSkeletonIDsCount();
-				for (int32 SkeletonIndex = 0; SkeletonIndex < SkeletonIDsCount; ++SkeletonIndex)
-				{
-					SkeletonData.SkeletonIds.AddUnique(Mesh->GetSkeletonID(SkeletonIndex));
-				}
-
-				// Append BoneMap to the array of BoneMaps
-				const TArray<uint16>& BoneMap = Mesh->GetBoneMap();
-				Component.FirstBoneMap = OperationData->InstanceUpdateData.BoneMaps.Num();
-				Component.BoneMapCount = BoneMap.Num();
-				OperationData->InstanceUpdateData.BoneMaps.Append(BoneMap);
-
-				// Add active bone indices and poses
-				const int32 MaxBoneIndex = Mesh->GetBonePoseCount();
-				Component.ActiveBones.Reserve(MaxBoneIndex);
-				for (int32 BonePoseIndex = 0; BonePoseIndex < MaxBoneIndex; ++BonePoseIndex)
-				{
-					const uint16 BoneId = Mesh->GetBonePoseBoneId(BonePoseIndex);
-
-					Component.ActiveBones.Add(BoneId);
-
-					if (SkeletonData.BoneIds.Find(BoneId) == INDEX_NONE)
-					{
-						SkeletonData.BoneIds.Add(BoneId);
-
-						FTransform3f Transform;
-						Mesh->GetBoneTransform(BonePoseIndex, Transform);
-						SkeletonData.BoneMatricesWithScale.Emplace(Transform.Inverse().ToMatrixWithScale());
-					}
-				}
+				OperationData->InstanceUpdateData.Skeletons.SetNum(Component.Id + 1);
 			}
-		
-			return;
-		}
-
-
-		const int32 LODCount = OperationData->InstanceUpdateData.LODs.Num();
-		const FInstanceUpdateData::FLOD& MinLOD = OperationData->InstanceUpdateData.LODs[OperationData->CurrentMinLOD];
-		const int32 ComponentCount = MinLOD.ComponentCount;
 
-		// Add SkeletonData for each component
-		OperationData->InstanceUpdateData.Skeletons.AddDefaulted(ComponentCount);
+			FInstanceUpdateData::FSkeletonData& SkeletonData = OperationData->InstanceUpdateData.Skeletons[Component.Id];
+			SkeletonData.ComponentIndex = Component.Id;
 
-		for (int32 ComponentIndex = 0; ComponentIndex < ComponentCount; ++ComponentIndex)
-		{
-			FInstanceUpdateData::FComponent& MinLODComponent = OperationData->InstanceUpdateData.Components[MinLOD.FirstComponent + ComponentIndex];
-
-			// Set the ComponentIndex
-			OperationData->InstanceUpdateData.Skeletons[ComponentIndex].ComponentIndex = MinLODComponent.Id;
-
-			// Fill the data used to generate the RefSkeletalMesh
-			TArray<uint16>& SkeletonIds = OperationData->InstanceUpdateData.Skeletons[ComponentIndex].SkeletonIds;
-			TArray<uint16>& BoneIds = OperationData->InstanceUpdateData.Skeletons[ComponentIndex].BoneIds;
-			TArray<FMatrix44f>& BoneMatricesWithScale = OperationData->InstanceUpdateData.Skeletons[ComponentIndex].BoneMatricesWithScale;
-
-			// Use first valid LOD bone count as a potential total number of bones, used for pre-allocating data arrays
-			if (MinLODComponent.Mesh && MinLODComponent.Mesh->GetSkeleton())
+			mu::MeshPtrConst Mesh = Component.Mesh;
+			if (!Mesh)
 			{
-				const int32 TotalPossibleBones = MinLODComponent.Mesh->GetSkeleton()->GetBoneCount();
-
-				// Out Data
-				BoneIds.Reserve(TotalPossibleBones);
-				BoneMatricesWithScale.Reserve(TotalPossibleBones);
+				continue;
 			}
 
-			for (int32 LODIndex = OperationData->CurrentMinLOD; LODIndex <= OperationData->CurrentMaxLOD && LODIndex < LODCount; ++LODIndex)
+			// Add SkeletonIds 
+			const int32 SkeletonIDsCount = Mesh->GetSkeletonIDsCount();
+			for (int32 SkeletonIndex = 0; SkeletonIndex < SkeletonIDsCount; ++SkeletonIndex)
 			{
-				MUTABLE_CPUPROFILER_SCOPE(PrepareSkeletonData_LODs);
-
-				const FInstanceUpdateData::FLOD& CurrentLOD = OperationData->InstanceUpdateData.LODs[LODIndex];
-				FInstanceUpdateData::FComponent& CurrentLODComponent = OperationData->InstanceUpdateData.Components[CurrentLOD.FirstComponent + ComponentIndex];
-				mu::MeshPtrConst Mesh = CurrentLODComponent.Mesh;
+				SkeletonData.SkeletonIds.AddUnique(Mesh->GetSkeletonID(SkeletonIndex));
+			}
 
-				if (!Mesh)
-				{
-					continue;
-				}
+			// Append BoneMap to the array of BoneMaps
+			const TArray<uint16>& BoneMap = Mesh->GetBoneMap();
+			Component.FirstBoneMap = OperationData->InstanceUpdateData.BoneMaps.Num();
+			Component.BoneMapCount = BoneMap.Num();
+			OperationData->InstanceUpdateData.BoneMaps.Append(BoneMap);
 
-				// Add SkeletonIds 
-				const int32 SkeletonIDsCount = Mesh->GetSkeletonIDsCount();
-				for (int32 SkeletonIndex = 0; SkeletonIndex < SkeletonIDsCount; ++SkeletonIndex)
-				{
-					SkeletonIds.AddUnique(Mesh->GetSkeletonID(SkeletonIndex));
-				}
+			// Add active bone indices and poses
+			const int32 MaxBoneIndex = Mesh->GetBonePoseCount();
+			Component.ActiveBones.Reserve(MaxBoneIndex);
+			for (int32 BonePoseIndex = 0; BonePoseIndex < MaxBoneIndex; ++BonePoseIndex)
+			{
+				const uint16 BoneId = Mesh->GetBonePoseBoneId(BonePoseIndex);
 
-				// Append BoneMap to the array of BoneMaps
-				const TArray<uint16>& BoneMap = Mesh->GetBoneMap();
-				CurrentLODComponent.FirstBoneMap = OperationData->InstanceUpdateData.BoneMaps.Num();
-				CurrentLODComponent.BoneMapCount = BoneMap.Num();
-				OperationData->InstanceUpdateData.BoneMaps.Append(BoneMap);
+				Component.ActiveBones.Add(BoneId);
 
-				// Add active bone indices and poses
-				const int32 MaxBoneIndex = Mesh->GetBonePoseCount();
-				CurrentLODComponent.ActiveBones.Reserve(MaxBoneIndex);
-				for (int32 BonePoseIndex = 0; BonePoseIndex < MaxBoneIndex; ++BonePoseIndex)
+				if (SkeletonData.BoneIds.Find(BoneId) == INDEX_NONE)
 				{
-					const uint16 BoneId = Mesh->GetBonePoseBoneId(BonePoseIndex);
+					SkeletonData.BoneIds.Add(BoneId);
 
-					CurrentLODComponent.ActiveBones.Add(BoneId);
-
-					if(BoneIds.Find(BoneId) == INDEX_NONE)
-					{
-						BoneIds.Add(BoneId);
-
-						FTransform3f Transform;
-						Mesh->GetBoneTransform(BonePoseIndex, Transform);
-						BoneMatricesWithScale.Emplace(Transform.Inverse().ToMatrixWithScale());
-					}
+					FTransform3f Transform;
+					Mesh->GetBoneTransform(BonePoseIndex, Transform);
+					SkeletonData.BoneMatricesWithScale.Emplace(Transform.Inverse().ToMatrixWithScale());
 				}
 			}
 		}
@@ -3355,19 +3266,9 @@ void UCustomizableObjectSystem::UnregisterImageProvider(UCustomizableSystemImage
 	GetPrivateChecked()->GetImageProviderChecked()->ImageProviders.Remove(Provider);
 }
 
-static bool bRevertCacheTextureParameters = false;
-static FAutoConsoleVariableRef CVarRevertCacheTextureParameters(
-	TEXT("mutable.RevertCacheTextureParameters"), bRevertCacheTextureParameters,
-	TEXT("If true, FUpdateContextPrivate will not cache/uncache texture parameters. If false, FUpdateContextPrivate will add an additional reference to the TextureParameters being used in the update."));
-
 
 void CacheTexturesParameters(const TArray<FName>& TextureParameters)
 {
-	if (bRevertCacheTextureParameters)
-	{
-		return;
-	}
-
 	if (!TextureParameters.IsEmpty() && UCustomizableObjectSystem::IsCreated())
 	{
 		FUnrealMutableImageProvider* ImageProvider = UCustomizableObjectSystem::GetInstance()->GetPrivateChecked()->GetImageProviderChecked();
@@ -3383,11 +3284,6 @@ void CacheTexturesParameters(const TArray<FName>& TextureParameters)
 
 void UnCacheTexturesParameters(const TArray<FName>& TextureParameters)
 {
-	if (bRevertCacheTextureParameters)
-	{
-		return;
-	}
-
 	if (!TextureParameters.IsEmpty() && UCustomizableObjectSystem::IsCreated())
 	{
 		FUnrealMutableImageProvider* ImageProvider = UCustomizableObjectSystem::GetInstance()->GetPrivateChecked()->GetImageProviderChecked();
