[Mutable] Remove COEditor Projector temporal fix.

#rnx

diff --git a/Engine/Plugins/Experimental/Mutable/Source/CustomizableObjectEditor/Private/MuCOE/CustomizableObjectEditorViewportClient.cpp b/Engine/Plugins/Experimental/Mutable/Source/CustomizableObjectEditor/Private/MuCOE/CustomizableObjectEditorViewportClient.cpp
--- a/Engine/Plugins/Experimental/Mutable/Source/CustomizableObjectEditor/Private/MuCOE/CustomizableObjectEditorViewportClient.cpp
+++ b/Engine/Plugins/Experimental/Mutable/Source/CustomizableObjectEditor/Private/MuCOE/CustomizableObjectEditorViewportClient.cpp
@@ -2605,7 +2605,7 @@ void FCustomizableObjectEditorViewportClient::SetWidgetType(EWidgetType Type)
 {
 	WidgetType = Type;
 	
-	//SetWidgetMode(UE::Widget::WM_Translate); // TODO GMTFuture Uncomment once UE-191354 fixed
+	SetWidgetMode(UE::Widget::WM_Translate);
 	Widget->SetDefaultVisibility(Type != EWidgetType::Hidden);
 }	
 
