[Mutable] Fix NodeMeshClipMorph incorrect offset backwards compatible case.

#rnx

diff --git a/Engine/Plugins/Experimental/Mutable/Source/CustomizableObjectEditor/Classes/MuCOE/Nodes/CustomizableObjectNodeMeshClipMorph.h b/Engine/Plugins/Experimental/Mutable/Source/CustomizableObjectEditor/Classes/MuCOE/Nodes/CustomizableObjectNodeMeshClipMorph.h
--- a/Engine/Plugins/Experimental/Mutable/Source/CustomizableObjectEditor/Classes/MuCOE/Nodes/CustomizableObjectNodeMeshClipMorph.h
+++ b/Engine/Plugins/Experimental/Mutable/Source/CustomizableObjectEditor/Classes/MuCOE/Nodes/CustomizableObjectNodeMeshClipMorph.h
@@ -63,7 +63,12 @@ public:
 
 	UPROPERTY()
 	FVector Normal;
-	
+
+private:
+	UPROPERTY()
+	bool bOldOffset_DEPRECATED;
+
+public:
 	UPROPERTY(EditAnywhere, Category = MeshClipParameters, meta = (DisplayName = "Max Effect Radius", ToolTip = "The maximum distance from the origin of the widget where vertices will be affected. If negative, there will be no limit."))
 	float MaxEffectRadius;
 
diff --git a/Engine/Plugins/Experimental/Mutable/Source/CustomizableObjectEditor/Private/MuCOE/Nodes/CustomizableObjectNodeMeshClipMorph.cpp b/Engine/Plugins/Experimental/Mutable/Source/CustomizableObjectEditor/Private/MuCOE/Nodes/CustomizableObjectNodeMeshClipMorph.cpp
--- a/Engine/Plugins/Experimental/Mutable/Source/CustomizableObjectEditor/Private/MuCOE/Nodes/CustomizableObjectNodeMeshClipMorph.cpp
+++ b/Engine/Plugins/Experimental/Mutable/Source/CustomizableObjectEditor/Private/MuCOE/Nodes/CustomizableObjectNodeMeshClipMorph.cpp
@@ -115,8 +115,9 @@ void UCustomizableObjectNodeMeshClipMorph::Serialize(FArchive& Ar)
 
 	Ar.UsingCustomVersion(FCustomizableObjectCustomVersion::GUID);
 
-	if (Ar.CustomVer(FCustomizableObjectCustomVersion::GUID) < FCustomizableObjectCustomVersion::PostLoadToCustomVersion
-		&& bLocalStartOffset)
+	if (Ar.CustomVer(FCustomizableObjectCustomVersion::GUID) < FCustomizableObjectCustomVersion::PostLoadToCustomVersion &&
+		bOldOffset_DEPRECATED &&
+		bLocalStartOffset)
 	{
 		// Previous Offset
 		FVector Tangent, Binormal;
