[Mutable] Fix ModelDiskStreamer data race

All access to the internal members now happen in the MutableThread. CVar toggable.

#rnx

diff --git a/Engine/Plugins/Experimental/Mutable/Source/CustomizableObject/Classes/MuCO/CustomizableObjectSystem.h b/Engine/Plugins/Experimental/Mutable/Source/CustomizableObject/Classes/MuCO/CustomizableObjectSystem.h
--- a/Engine/Plugins/Experimental/Mutable/Source/CustomizableObject/Classes/MuCO/CustomizableObjectSystem.h
+++ b/Engine/Plugins/Experimental/Mutable/Source/CustomizableObject/Classes/MuCO/CustomizableObjectSystem.h
@@ -47,6 +47,9 @@ extern TAutoConsoleVariable<bool> CVarPreserveUserLODsOnFirstGeneration;
 
 extern TAutoConsoleVariable<bool> CVarEnableMeshCache;
 
+extern TAutoConsoleVariable<bool> CVarRollbackFixModelDiskStreamerDataRace;
+
+
 #if WITH_EDITOR
 
 // Struct used to keep a copy of the EditorSettings needed to compile Customizable Objects.
diff --git a/Engine/Plugins/Experimental/Mutable/Source/CustomizableObject/Private/MuCO/CustomizableObjectSystem.cpp b/Engine/Plugins/Experimental/Mutable/Source/CustomizableObject/Private/MuCO/CustomizableObjectSystem.cpp
--- a/Engine/Plugins/Experimental/Mutable/Source/CustomizableObject/Private/MuCO/CustomizableObjectSystem.cpp
+++ b/Engine/Plugins/Experimental/Mutable/Source/CustomizableObject/Private/MuCO/CustomizableObjectSystem.cpp
@@ -70,14 +70,12 @@ static TAutoConsoleVariable<int32> CVarWorkingMemory(
 	TEXT("Limit the amount of memory (in KB) to use as working memory when building characters. More memory reduces the object construction time. 0 means no restriction. Defaults: Desktop = 50,000 KB, Others = 10,000 KB"),
 	ECVF_Scalability);
 
-
 TAutoConsoleVariable<bool> CVarClearWorkingMemoryOnUpdateEnd(
 	TEXT("mutable.ClearWorkingMemoryOnUpdateEnd"),
 	false,
 	TEXT("Clear the working memory and cache after every Mutable operation."),
 	ECVF_Scalability);
 
-
 TAutoConsoleVariable<bool> CVarReuseImagesBetweenInstances(
 	TEXT("mutable.ReuseImagesBetweenInstances"),
 	true,
@@ -96,13 +94,17 @@ TAutoConsoleVariable<bool> CVarPreserveUserLODsOnFirstGeneration(
 	TEXT("If false, force disable UCustomizableObject::bPreserveUserLODsOnFirstGeneration."),
 	ECVF_Scalability);
 
-
 TAutoConsoleVariable<bool> CVarEnableMeshCache(
 	TEXT("mutable.EnableMeshCache"),
 	true,
 	TEXT("Enables or disables the reuse of meshes."),
 	ECVF_Scalability);
 
+TAutoConsoleVariable<bool> CVarRollbackFixModelDiskStreamerDataRace(
+	TEXT("mutable.Rollback.FixModelDiskStreamerDataRace"),
+	true,
+	TEXT("When true, use the new code path that fixes a data race in UnrealMutableModelDiskStreamer."));
+
 
 int32 FCustomizableObjectSystemPrivate::SkeletalMeshMinLodQualityLevel = -1;
 
@@ -539,14 +541,26 @@ void UCustomizableObjectSystem::BeginDestroy()
 		Private->PendingTasks.Empty();
 
 		// Complete pending taskgraph tasks
+		if (CVarRollbackFixModelDiskStreamerDataRace.GetValueOnAnyThread())
+		{
+			Private->MutableTaskGraph.AllowLaunchingMutableTaskLowPriority(false, false);
+			check(Private->Streamer);
+			Private->MutableTaskGraph.AddMutableThreadTask(TEXT("EndStream"), [Streamer = Private->Streamer]()
+			{
+				Streamer->EndStreaming();
+			});
+		}
 		Private->MutableTaskGraph.WaitForMutableTasks();
 
 		// Clear the ongoing operation
 		Private->CurrentMutableOperation = nullptr;
 
-		// Deallocate streaming
-		check(Private->Streamer != nullptr);
-		Private->Streamer->EndStreaming();
+		if (!CVarRollbackFixModelDiskStreamerDataRace.GetValueOnAnyThread())
+		{
+			// Deallocate streaming
+			check(Private->Streamer != nullptr);
+			Private->Streamer->EndStreaming();
+		}
 
 		Private->CurrentInstanceBeingUpdated = nullptr;
 
@@ -1591,6 +1605,16 @@ namespace impl
 		UCustomizableObjectSystem* System = UCustomizableObjectSystem::GetInstanceChecked(); // Save since UCustomizableObjectSystem::BeginDestroy always waits for all tasks to finish
 		const FCustomizableObjectSystemPrivate* SystemPrivate = System->GetPrivateChecked();
 		
+		Operation->UpdateStartBytes = mu::FGlobalMemoryCounter::GetCounter();
+		mu::FGlobalMemoryCounter::Zero();
+
+		if (CVarRollbackFixModelDiskStreamerDataRace.GetValueOnAnyThread())
+		{
+			// Prepare streaming for the current customizable object
+			check(SystemPrivate->Streamer != nullptr);
+			SystemPrivate->Streamer->PrepareStreamingForObject(Operation->Instance->GetCustomizableObject());			
+		}
+
 		const mu::Ptr<mu::System> MutableSystem = SystemPrivate->MutableSystem;
 
 		const TSharedPtr<mu::Model> Model = Operation->Instance->GetCustomizableObject()->GetPrivate()->GetModel();
@@ -2102,9 +2126,6 @@ namespace impl
 		MUTABLE_CPUPROFILER_SCOPE(Task_Mutable_Update_GetMesh)
 		FMutableScopeTimer Timer(OperationData->TaskGetMeshTime);
 		
-		OperationData->UpdateStartBytes = mu::FGlobalMemoryCounter::GetCounter();
-		mu::FGlobalMemoryCounter::Zero();
-
 #if WITH_EDITOR
 		const uint32 StartCycles = FPlatformTime::Cycles();
 #endif
@@ -2725,9 +2746,12 @@ namespace impl
 
 		SystemPrivateData->CurrentInstanceBeingUpdated = CandidateInstance;
 
-		// Prepare streaming for the current customizable object
-		check(SystemPrivateData->Streamer != nullptr);
-		SystemPrivateData->Streamer->PrepareStreamingForObject(CustomizableObject);
+		if (!CVarRollbackFixModelDiskStreamerDataRace.GetValueOnAnyThread())
+		{
+			// Prepare streaming for the current customizable object
+			check(SystemPrivateData->Streamer != nullptr);
+			SystemPrivateData->Streamer->PrepareStreamingForObject(CustomizableObject);
+		}
 
 		check(SystemPrivateData->ExtensionDataStreamer != nullptr);
 		SystemPrivateData->ExtensionDataStreamer->SetActiveObject(CustomizableObject);
diff --git a/Engine/Plugins/Experimental/Mutable/Source/CustomizableObject/Private/MuCO/UnrealMutableModelDiskStreamer.cpp b/Engine/Plugins/Experimental/Mutable/Source/CustomizableObject/Private/MuCO/UnrealMutableModelDiskStreamer.cpp
--- a/Engine/Plugins/Experimental/Mutable/Source/CustomizableObject/Private/MuCO/UnrealMutableModelDiskStreamer.cpp
+++ b/Engine/Plugins/Experimental/Mutable/Source/CustomizableObject/Private/MuCO/UnrealMutableModelDiskStreamer.cpp
@@ -49,15 +49,15 @@ void UnrealMutableInputStream::Read(void* pData, uint64 size)
 //-------------------------------------------------------------------------------------------------
 FUnrealMutableModelBulkReader::~FUnrealMutableModelBulkReader()
 {
-	EndStreaming();
+	if (!CVarRollbackFixModelDiskStreamerDataRace.GetValueOnAnyThread())
+	{
+		EndStreaming();
+	}
 }
 
 
 bool FUnrealMutableModelBulkReader::PrepareStreamingForObject(UCustomizableObject* CustomizableObject)
 {
-	// This happens in the game thread
-	check(IsInGameThread());
-
 	if (!CustomizableObject)
 	{
 		check(false);
