[Mutable] Converted CustomizableObjectPrivateData to serializable UObject.

- Renamed to CustomizableObjectPrivate.
- Renamed PrivateData varaible to Private.

#rnx

diff --git a/Engine/Plugins/Experimental/Mutable/Source/CustomizableObject/Classes/MuCO/CustomizableObject.h b/Engine/Plugins/Experimental/Mutable/Source/CustomizableObject/Classes/MuCO/CustomizableObject.h
--- a/Engine/Plugins/Experimental/Mutable/Source/CustomizableObject/Classes/MuCO/CustomizableObject.h
+++ b/Engine/Plugins/Experimental/Mutable/Source/CustomizableObject/Classes/MuCO/CustomizableObject.h
@@ -108,7 +108,7 @@ enum class ECustomizableObjectRelevancy : uint8
 
 
 // This is used to hide Mutable SDK members in the public headers.
-class FCustomizableObjectPrivateData;
+class UCustomizableObjectPrivate;
 
 USTRUCT()
 struct FProfileParameterDat
@@ -1059,7 +1059,7 @@ private:
 UCLASS( BlueprintType, config=Engine )
 class CUSTOMIZABLEOBJECT_API UCustomizableObject : public UObject
 {
-	friend FCustomizableObjectPrivateData;
+	friend UCustomizableObjectPrivate;
 	
 public:
 	GENERATED_BODY()
@@ -1445,8 +1445,7 @@ public:
 	// Begin UObject interface.
 	void PreSave(FObjectPreSaveContext ObjectSaveContext) override;
 	void PostRename(UObject* OldOuter, const FName OldName) override;
-	void PostDuplicate(EDuplicateMode::Type DuplicateMode) override;
-
+	
 	void BeginCacheForCookedPlatformData(const ITargetPlatform* TargetPlatform) override;
 	bool IsCachedCookedPlatformDataLoaded(const ITargetPlatform* TargetPlatform) override;
 	// End UObject interface.
@@ -1610,7 +1609,7 @@ public:
 	// Return a pointer to the BulkData subobject, only valid in packaged builds
 	const UCustomizableObjectBulk* GetStreamableBulkData() const { return BulkData; }
 
-	FCustomizableObjectPrivateData* GetPrivate() const;
+	UCustomizableObjectPrivate* GetPrivate() const;
 
 	/** Check if the CustomizableObject asset has been compiled. This will always be true in a packaged game, but it could be false in the editor. */
 	UFUNCTION(BlueprintCallable, Category = CustomizableObject)
@@ -1670,8 +1669,9 @@ private:
 	UPROPERTY()
 	TMap<TObjectPtr<const UObject>, FGuid> ParticipatingObjects;
 #endif
-	
-	TSharedPtr<FCustomizableObjectPrivateData> PrivateData;
+
+	UPROPERTY()
+	UCustomizableObjectPrivate* Private;
 };
 
 #if UE_ENABLE_INCLUDE_ORDER_DEPRECATED_IN_5_2
diff --git a/Engine/Plugins/Experimental/Mutable/Source/CustomizableObject/Internal/MuCO/CustomizableObjectPrivate.h b/Engine/Plugins/Experimental/Mutable/Source/CustomizableObject/Internal/MuCO/CustomizableObjectPrivate.h
--- a/Engine/Plugins/Experimental/Mutable/Source/CustomizableObject/Internal/MuCO/CustomizableObjectPrivate.h
+++ b/Engine/Plugins/Experimental/Mutable/Source/CustomizableObject/Internal/MuCO/CustomizableObjectPrivate.h
@@ -10,6 +10,8 @@
 #include "Misc/Guid.h"
 #endif
 
+#include "CustomizableObjectPrivate.generated.h"
+
 namespace mu { class Model; }
 class UCustomizableObject;
 class USkeletalMesh;
@@ -27,8 +29,11 @@ private:
 };
 
 
-class FCustomizableObjectPrivateData
+UCLASS()
+class UCustomizableObjectPrivate : public UObject
 {
+	GENERATED_BODY()
+	
 	TSharedPtr<mu::Model, ESPMode::ThreadSafe> MutableModel;
 
 public:
diff --git a/Engine/Plugins/Experimental/Mutable/Source/CustomizableObject/Private/MuCO/CustomizableObject.cpp b/Engine/Plugins/Experimental/Mutable/Source/CustomizableObject/Private/MuCO/CustomizableObject.cpp
--- a/Engine/Plugins/Experimental/Mutable/Source/CustomizableObject/Private/MuCO/CustomizableObject.cpp
+++ b/Engine/Plugins/Experimental/Mutable/Source/CustomizableObject/Private/MuCO/CustomizableObject.cpp
@@ -50,7 +50,7 @@ static bool bUsesOnCookStart = false;
 UCustomizableObject::UCustomizableObject()
 	: UObject()
 {
-	PrivateData = MakeShared<FCustomizableObjectPrivateData>();
+	Private = CreateDefaultSubobject<UCustomizableObjectPrivate>(FName("Private"));
 
 #if WITH_EDITORONLY_DATA
 	const FString CVarName = TEXT("r.SkeletalMesh.MinLodQualityLevel");
@@ -218,9 +218,9 @@ void UCustomizableObject::PostLoad()
 
 bool UCustomizableObject::IsLocked() const
 {
-	if (PrivateData.IsValid())
+	if (Private)
 	{
-		return PrivateData->bLocked;
+		return Private->bLocked;
 	}
 
 	return false;
@@ -276,18 +276,6 @@ void UCustomizableObject::PostRename(UObject * OldOuter, const FName OldName)
 }
 
 
-void UCustomizableObject::PostDuplicate(EDuplicateMode::Type DuplicateMode)
-{
-	Super::PostDuplicate(DuplicateMode);
-
-	if (DuplicateMode == EDuplicateMode::Normal)
-	{
-		// Create a new Private Data or it will use the same as the original Customizable Object
-		PrivateData = TSharedPtr<FCustomizableObjectPrivateData>(new FCustomizableObjectPrivateData());
-	}
-}
-
-
 void UCustomizableObject::BeginCacheForCookedPlatformData(const ITargetPlatform* TargetPlatform)
 {
 	bool bIsRelevantForThisTarget =
@@ -297,7 +285,7 @@ void UCustomizableObject::BeginCacheForCookedPlatformData(const ITargetPlatform*
 	
 	if (TargetPlatform && bIsRelevantForThisTarget)
 	{
-		if (PrivateData->CachedPlatformNames.Find(TargetPlatform->PlatformName()) == INDEX_NONE)
+		if (Private->CachedPlatformNames.Find(TargetPlatform->PlatformName()) == INDEX_NONE)
 		{
 			if (!bUsesOnCookStart)
 			{
@@ -309,23 +297,23 @@ void UCustomizableObject::BeginCacheForCookedPlatformData(const ITargetPlatform*
 				// Load from Disk
 				LoadCompiledDataFromDisk(false, TargetPlatform);
 
-				if (!PrivateData->bModelCompiledForCook)
+				if (!Private->bModelCompiledForCook)
 				{
 					LoadReferencedObjects();
-					PrivateData->bModelCompiledForCook = true;
+					Private->bModelCompiledForCook = true;
 				}
 			}
 
-			PrivateData->CachedPlatformNames.Add(TargetPlatform->PlatformName());
+			Private->CachedPlatformNames.Add(TargetPlatform->PlatformName());
 		}
 	}
 	else
 	{
 		ClearCompiledData();
-		PrivateData->SetModel(nullptr, FGuid()); // Discard compilation
+		Private->SetModel(nullptr, FGuid()); // Discard compilation
 		if (TargetPlatform)
 		{
-			PrivateData->CachedPlatformNames.Add(TargetPlatform->PlatformName());
+			Private->CachedPlatformNames.Add(TargetPlatform->PlatformName());
 		}
 	}
 }
@@ -333,7 +321,7 @@ void UCustomizableObject::BeginCacheForCookedPlatformData(const ITargetPlatform*
 
 bool UCustomizableObject::IsCachedCookedPlatformDataLoaded( const ITargetPlatform* TargetPlatform ) 
 { 
-	return PrivateData->CachedPlatformNames.Find(TargetPlatform->PlatformName()) != INDEX_NONE;
+	return Private->CachedPlatformNames.Find(TargetPlatform->PlatformName()) != INDEX_NONE;
 }
 
 // TODO COOK: Remove Hack to add new references to the package
@@ -392,7 +380,7 @@ void UCustomizableObject::ClearCompiledData()
 
 void UCustomizableObject::UpdateCompiledDataFromModel()
 {
-	TSharedPtr<mu::Model, ESPMode::ThreadSafe> Model = PrivateData->GetModel();
+	TSharedPtr<mu::Model, ESPMode::ThreadSafe> Model = Private->GetModel();
 
 	// Generate a map that using the resource id tells the offset and size of the resource inside the bulk data
 	if (Model)
@@ -611,7 +599,7 @@ void UCustomizableObject::SaveCompiledData(FArchive& MemoryWriter, bool bIsCooki
 
 void UCustomizableObject::LoadCompiledData(FArchive& MemoryReader, const ITargetPlatform* InTargetPlatform, bool bIsCooking)
 {
-	PrivateData->SetModel(nullptr, FGuid());
+	Private->SetModel(nullptr, FGuid());
 	ClearCompiledData();
 
 	MutableCompiledDataStreamHeader Header;
@@ -726,7 +714,7 @@ void UCustomizableObject::LoadCompiledData(FArchive& MemoryReader, const ITarget
 			mu::InputArchive arch(&stream);
 			TSharedPtr<mu::Model, ESPMode::ThreadSafe> Model = mu::Model::StaticUnserialise( arch );
 
-			PrivateData->SetModel(Model, Identifier);
+			Private->SetModel(Model, Identifier);
 		}
 	}
 
@@ -855,7 +843,7 @@ void UCustomizableObject::CompileForTargetPlatform(const ITargetPlatform* Target
 		(TargetPlatform && Relevancy == ECustomizableObjectRelevancy::ClientOnly && !TargetPlatform->IsServerOnly());
 
 	// Discard any older compilation
-	PrivateData->SetModel(nullptr, FGuid());
+	Private->SetModel(nullptr, FGuid());
 
 	if (bIsRootObject && bIsRelevantForThisTarget)
 	{
@@ -917,7 +905,7 @@ bool UCustomizableObject::ConditionalAutoCompile()
 	}
 
 	// Discard any older compilation
-	PrivateData->SetModel(nullptr, FGuid());
+	Private->SetModel(nullptr, FGuid());
 
 	// Sync/Async compilation
 	if (System->IsAutoCompilationSync())
@@ -1014,10 +1002,10 @@ void UCustomizableObject::SaveEmbeddedData(FArchive& Ar)
 {
 	UE_LOG(LogMutable, Verbose, TEXT("Saving embedded data for Customizable Object [%s] now at position %d."), *GetName(), int(Ar.Tell()));
 
-	int32 InternalVersion = PrivateData->GetModel() ? CurrentSupportedVersion : -1;
+	int32 InternalVersion = Private->GetModel() ? CurrentSupportedVersion : -1;
 	Ar << InternalVersion;
 
-	if (PrivateData->GetModel())
+	if (Private->GetModel())
 	{
 		// Serialize morph data
 		{
@@ -1033,11 +1021,11 @@ void UCustomizableObject::SaveEmbeddedData(FArchive& Ar)
 
 		// Serialise the entire model, but unload the streamable data first.
 		{
-			PrivateData->GetModel()->UnloadExternalData();
+			Private->GetModel()->UnloadExternalData();
 
 			UnrealMutableOutputStream stream(Ar);
 			mu::OutputArchive arch(&stream);
-			mu::Model::Serialise(PrivateData->GetModel().Get(), arch);
+			mu::Model::Serialise(Private->GetModel().Get(), arch);
 		}
 
 		UE_LOG(LogMutable, Verbose, TEXT("Saved embedded data for Customizable Object [%s] now at position %d."), *GetName(), int(Ar.Tell()));
@@ -1075,7 +1063,7 @@ void UCustomizableObject::LoadEmbeddedData(FArchive& Ar)
 		UnrealMutableInputStream stream(Ar);
 		mu::InputArchive arch(&stream);
 		TSharedPtr<mu::Model, ESPMode::ThreadSafe> Model = mu::Model::StaticUnserialise( arch );
-		PrivateData->SetModel( Model, FGuid());
+		Private->SetModel( Model, FGuid());
 
 		// Create parameter properties
 		UpdateParameterPropertiesFromModel();
@@ -1083,16 +1071,16 @@ void UCustomizableObject::LoadEmbeddedData(FArchive& Ar)
 }
 
 
-FCustomizableObjectPrivateData* UCustomizableObject::GetPrivate() const
+UCustomizableObjectPrivate* UCustomizableObject::GetPrivate() const
 {
-	check(PrivateData);
-	return PrivateData.Get();
+	check(Private);
+	return Private;
 }
 
 
 bool UCustomizableObject::IsCompiled() const
 {
-	bool IsCompiled = PrivateData->GetModel() != nullptr;
+	bool IsCompiled = Private->GetModel() != nullptr;
 
 	return IsCompiled;
 }
@@ -1212,9 +1200,9 @@ void UCustomizableObject::UnCacheInvalidSkeletons()
 int32 UCustomizableObject::FindState( const FString& Name ) const
 {
 	int32 Result = -1;
-	if (PrivateData->GetModel())
+	if (Private->GetModel())
 	{
-		Result = PrivateData->GetModel()->FindState(Name);
+		Result = Private->GetModel()->FindState(Name);
 	}
 
 	return Result;
@@ -1225,9 +1213,9 @@ int32 UCustomizableObject::GetStateCount() const
 {
 	int32 Result = 0;
 
-	if (PrivateData->GetModel())
+	if (Private->GetModel())
 	{
-		Result = PrivateData->GetModel()->GetStateCount();
+		Result = Private->GetModel()->GetStateCount();
 	}
 
 	return Result;
@@ -1238,9 +1226,9 @@ FString UCustomizableObject::GetStateName(int32 StateIndex) const
 {
 	FString Result;
 
-	if (PrivateData->GetModel())
+	if (Private->GetModel())
 	{
-		Result = PrivateData->GetModel()->GetStateName(StateIndex);
+		Result = Private->GetModel()->GetStateName(StateIndex);
 	}
 
 	return Result;
@@ -1251,9 +1239,9 @@ int32 UCustomizableObject::GetStateParameterCount( int32 StateIndex ) const
 {
 	int32 Result = 0;
 
-	if (PrivateData->GetModel())
+	if (Private->GetModel())
 	{
-		Result = PrivateData->GetModel()->GetStateParameterCount(StateIndex);
+		Result = Private->GetModel()->GetStateParameterCount(StateIndex);
 	}
 
 	return Result;
@@ -1263,9 +1251,9 @@ int32 UCustomizableObject::GetStateParameterIndex(int32 StateIndex, int32 Parame
 {
 	int32 Result = 0;
 
-	if (PrivateData->GetModel())
+	if (Private->GetModel())
 	{
-		Result = PrivateData->GetModel()->GetStateParameterIndex(StateIndex, ParameterIndex);
+		Result = Private->GetModel()->GetStateParameterIndex(StateIndex, ParameterIndex);
 	}
 
 	return Result;
@@ -1319,13 +1307,13 @@ UCustomizableObjectInstance* UCustomizableObject::CreateInstance()
 
 TSharedPtr<mu::Model, ESPMode::ThreadSafe> UCustomizableObject::GetModel() const
 {
-	return PrivateData->GetModel();
+	return Private->GetModel();
 }
 
 #if WITH_EDITOR
 void UCustomizableObject::SetModel(TSharedPtr<mu::Model, ESPMode::ThreadSafe> Model)
 {
-	PrivateData->SetModel(Model, GenerateIdentifier(*this));
+	Private->SetModel(Model, GenerateIdentifier(*this));
 	
 	UpdateCompiledDataFromModel();
 }
@@ -1414,9 +1402,9 @@ const FString & UCustomizableObject::GetParameterName(int32 ParamIndex) const
 
 void UCustomizableObject::UpdateParameterPropertiesFromModel()
 {
-	if (PrivateData->GetModel())
+	if (Private->GetModel())
 	{
-		mu::ParametersPtr MutableParameters = mu::Model::NewParameters(PrivateData->GetModel());
+		mu::ParametersPtr MutableParameters = mu::Model::NewParameters(Private->GetModel());
 		int paramCount = MutableParameters->GetCount();
 
 		ParameterProperties.Reset(paramCount);
@@ -1793,9 +1781,9 @@ bool UCustomizableObject::IsParameterMultidimensional(const FString& InParameter
 bool UCustomizableObject::IsParameterMultidimensional(const int32& InParamIndex) const
 {
 	check(InParamIndex != INDEX_NONE);
-	if (PrivateData->GetModel())
+	if (Private->GetModel())
 	{
-		return PrivateData->GetModel()->IsParameterMultidimensional(InParamIndex);
+		return Private->GetModel()->IsParameterMultidimensional(InParamIndex);
 	}
 
 	return false;
@@ -1942,7 +1930,7 @@ void FMeshCache::Add(const TArray<mu::FResourceID>& Key, USkeletalMesh* Value)
 }
 
 
-void FCustomizableObjectPrivateData::SetModel(const TSharedPtr<mu::Model, ESPMode::ThreadSafe>& Model, const FGuid Id)
+void UCustomizableObjectPrivate::SetModel(const TSharedPtr<mu::Model, ESPMode::ThreadSafe>& Model, const FGuid Id)
 {
 #if WITH_EDITOR
 	if (MutableModel)
@@ -1959,20 +1947,20 @@ void FCustomizableObjectPrivateData::SetModel(const TSharedPtr<mu::Model, ESPMod
 }
 
 
-const TSharedPtr<mu::Model, ESPMode::ThreadSafe>& FCustomizableObjectPrivateData::GetModel()
+const TSharedPtr<mu::Model, ESPMode::ThreadSafe>& UCustomizableObjectPrivate::GetModel()
 {
 	return MutableModel;
 }
 
 
-TSharedPtr<const mu::Model, ESPMode::ThreadSafe> FCustomizableObjectPrivateData::GetModel() const
+TSharedPtr<const mu::Model, ESPMode::ThreadSafe> UCustomizableObjectPrivate::GetModel() const
 {
 	return MutableModel;
 }
 
 
 #if WITH_EDITORONLY_DATA
-TMap<TObjectPtr<const UObject>, FGuid>& FCustomizableObjectPrivateData::GetParticipatingObjects(UCustomizableObject& Public)
+TMap<TObjectPtr<const UObject>, FGuid>& UCustomizableObjectPrivate::GetParticipatingObjects(UCustomizableObject& Public)
 {
 	return Public.ParticipatingObjects;
 }
diff --git a/Engine/Plugins/Experimental/Mutable/Source/CustomizableObject/Private/MuCO/CustomizableObjectSystem.cpp b/Engine/Plugins/Experimental/Mutable/Source/CustomizableObject/Private/MuCO/CustomizableObjectSystem.cpp
--- a/Engine/Plugins/Experimental/Mutable/Source/CustomizableObject/Private/MuCO/CustomizableObjectSystem.cpp
+++ b/Engine/Plugins/Experimental/Mutable/Source/CustomizableObject/Private/MuCO/CustomizableObjectSystem.cpp
@@ -2580,7 +2580,7 @@ namespace impl
 		FCustomizableObjectSystemPrivate* SystemPrivate = System->GetPrivateChecked();
 
 		const UCustomizableObject* CustomizableObject = Operation->Instance->GetCustomizableObject();
-		FCustomizableObjectPrivateData* CustomizableObjectPrivate = CustomizableObject->GetPrivate();
+		UCustomizableObjectPrivate* CustomizableObjectPrivate = CustomizableObject->GetPrivate();
 		
 		for (const TArray<mu::FResourceID>& MeshId : Operation->MeshDescriptors)
 		{
