[Mutable] Remove COE private dependencies 1/3.

- Moved required files to the Internal directory.
- Removed unused includes.

#rnx

diff --git a/Engine/Plugins/Experimental/Mutable/Source/CustomizableObjectEditor/Private/MuCOE/SMutableCodeViewer.h b/Engine/Plugins/Experimental/Mutable/Source/CustomizableObjectEditor/Private/MuCOE/SMutableCodeViewer.h
--- a/Engine/Plugins/Experimental/Mutable/Source/CustomizableObjectEditor/Private/MuCOE/SMutableCodeViewer.h
+++ b/Engine/Plugins/Experimental/Mutable/Source/CustomizableObjectEditor/Private/MuCOE/SMutableCodeViewer.h
@@ -3,7 +3,7 @@
 #pragma once
 
 #include "MuR/ModelPrivate.h"
-#include "MuT/ErrorLogPrivate.h"
+#include "MuT/ErrorLog.h"
 #include "MuT/TypeInfo.h"
 #include "UObject/GCObject.h"
 #include "Widgets/SCompoundWidget.h"
diff --git a/Engine/Plugins/Experimental/Mutable/Source/CustomizableObjectEditor/Private/MuCOE/SMutableGraphViewer.cpp b/Engine/Plugins/Experimental/Mutable/Source/CustomizableObjectEditor/Private/MuCOE/SMutableGraphViewer.cpp
--- a/Engine/Plugins/Experimental/Mutable/Source/CustomizableObjectEditor/Private/MuCOE/SMutableGraphViewer.cpp
+++ b/Engine/Plugins/Experimental/Mutable/Source/CustomizableObjectEditor/Private/MuCOE/SMutableGraphViewer.cpp
@@ -17,6 +17,11 @@
 #include "Widgets/Input/SNumericDropDown.h"
 #include "Widgets/Views/STreeView.h"
 #include "ScopedTransaction.h"
+#include "MuT/NodeObjectGroup.h"
+#include "MuT/NodeObjectNew.h"
+#include "MuT/NodeSurfaceEdit.h"
+#include "MuT/NodeSurfaceSwitch.h"
+#include "MuT/NodeSurfaceVariation.h"
 #include "MuT/Streams.h"
 
 // This is necessary because of problems with rtti information in other platforms. In any case, this part of the debugger is only useful in the standard editor.
diff --git a/Engine/Plugins/Experimental/Mutable/Source/MutableTools/Private/MuT/Compiler.cpp b/Engine/Plugins/Experimental/Mutable/Source/MutableTools/Private/MuT/Compiler.cpp
--- a/Engine/Plugins/Experimental/Mutable/Source/MutableTools/Private/MuT/Compiler.cpp
+++ b/Engine/Plugins/Experimental/Mutable/Source/MutableTools/Private/MuT/Compiler.cpp
@@ -217,11 +217,11 @@ namespace mu
             genErrorLog = gen.m_pErrorLog;
         }
 
-        vector<Ptr<ASTOp>> roots;
-        for( const STATE_COMPILATION_DATA& s: states)
-        {
-            roots.push_back(s.root);
-        }
+        //vector<Ptr<ASTOp>> roots;
+        //for( const STATE_COMPILATION_DATA& s: states)
+        //{
+        //    roots.push_back(s.root);
+        //}
 
         // Slow AST code verification for debugging.
         //ASTOp::FullAssert(roots);
diff --git a/Engine/Plugins/Experimental/Mutable/Source/MutableTools/Private/MuT/ErrorLogPrivate.h b/Engine/Plugins/Experimental/Mutable/Source/MutableTools/Private/MuT/ErrorLogPrivate.h
--- a/Engine/Plugins/Experimental/Mutable/Source/MutableTools/Private/MuT/ErrorLogPrivate.h
+++ b/Engine/Plugins/Experimental/Mutable/Source/MutableTools/Private/MuT/ErrorLogPrivate.h
@@ -12,8 +12,6 @@
 
 namespace mu
 {
-	MUTABLETOOLS_API extern const TCHAR* s_opNames[(int)OP_TYPE::COUNT];
-
 	class ErrorLog::Private : public Base
 	{
 	public:
diff --git a/Engine/Plugins/Experimental/Mutable/Source/MutableTools/Public/MuT/ErrorLog.h b/Engine/Plugins/Experimental/Mutable/Source/MutableTools/Public/MuT/ErrorLog.h
--- a/Engine/Plugins/Experimental/Mutable/Source/MutableTools/Public/MuT/ErrorLog.h
+++ b/Engine/Plugins/Experimental/Mutable/Source/MutableTools/Public/MuT/ErrorLog.h
@@ -100,5 +100,6 @@ namespace mu
 		Private* m_pD;
 
 	};
-
+	
+	MUTABLETOOLS_API extern const TCHAR* s_opNames[];
 }
