[Mutable] Allow ConvertResources to generate empty meshes on all components without marking it as an error.

#rnx

diff --git a/Engine/Plugins/Experimental/Mutable/Source/CustomizableObject/Private/MuCO/CustomizableObjectInstance.cpp b/Engine/Plugins/Experimental/Mutable/Source/CustomizableObject/Private/MuCO/CustomizableObjectInstance.cpp
--- a/Engine/Plugins/Experimental/Mutable/Source/CustomizableObject/Private/MuCO/CustomizableObjectInstance.cpp
+++ b/Engine/Plugins/Experimental/Mutable/Source/CustomizableObject/Private/MuCO/CustomizableObjectInstance.cpp
@@ -1708,9 +1708,6 @@ bool UCustomizableInstancePrivateData::DoComponentsNeedUpdate(UCustomizableObjec
 		}
 	}
 
-	// Mark as invalid if all components have empty meshes.
-	bHasInvalidMesh = bHasInvalidMesh || ComponentWithMesh.Find(true) == INDEX_NONE;
-
 	return !bHasInvalidMesh && OutComponentNeedsUpdate.Find(true) != INDEX_NONE;
 }
 
@@ -1748,12 +1745,6 @@ bool UCustomizableInstancePrivateData::UpdateSkeletalMesh_PostBeginUpdate0(UCust
 
 	TextureReuseCache.Empty(); // Sections may have changed, so invalidate the texture reuse cache because it's indexed by section
 
-	// Reset last mesh IDs.
-	for (FCustomizableInstanceComponentData& ComponentData : ComponentsData)
-	{
-		ComponentData.LastMeshIdPerLOD.Init(MAX_uint64, MAX_MESH_LOD_COUNT);
-	}
-
 	TArray<TObjectPtr<USkeletalMesh>> OldSkeletalMeshes = Public->SkeletalMeshes;
 
 	bool bSuccess = true;
@@ -1795,6 +1786,9 @@ bool UCustomizableInstancePrivateData::UpdateSkeletalMesh_PostBeginUpdate0(UCust
 			}
 		}
 
+		// Reset last mesh IDs.
+		ComponentsData[Component.Id].LastMeshIdPerLOD.Init(MAX_uint64, MAX_MESH_LOD_COUNT);
+
 		if (!Component.bGenerated || !Component.Mesh || Component.SurfaceCount == 0)
 		{
 			continue;
