[Mutable] Fix child objects ignored during compilation.

#rnx

diff --git a/Engine/Plugins/Experimental/Mutable/Source/CustomizableObjectEditor/Private/MuCOE/CustomizableObjectCompiler.cpp b/Engine/Plugins/Experimental/Mutable/Source/CustomizableObjectEditor/Private/MuCOE/CustomizableObjectCompiler.cpp
--- a/Engine/Plugins/Experimental/Mutable/Source/CustomizableObjectEditor/Private/MuCOE/CustomizableObjectCompiler.cpp
+++ b/Engine/Plugins/Experimental/Mutable/Source/CustomizableObjectEditor/Private/MuCOE/CustomizableObjectCompiler.cpp
@@ -325,10 +325,10 @@ void FCustomizableObjectCompiler::UpdateArrayGCProtect()
 }
 
 
-void FCustomizableObjectCompiler::ProcessChildObjectsRecursively(UCustomizableObject* Object, FMutableGraphGenerationContext &GenerationContext)
+void FCustomizableObjectCompiler::ProcessChildObjectsRecursively(UCustomizableObject* ParentObject, FMutableGraphGenerationContext& GenerationContext)
 {
 	TArray<FName> ArrayReferenceNames;
-	AddCachedReferencers(*Object->GetOuter()->GetPathName(), ArrayReferenceNames);
+	AddCachedReferencers(*ParentObject->GetOuter()->GetPathName(), ArrayReferenceNames);
 	UpdateArrayGCProtect();
 
 	bool bMultipleBaseObjectsFound = false;
@@ -339,56 +339,51 @@ void FCustomizableObjectCompiler::ProcessChildObjectsRecursively(UCustomizableOb
 		{
 			continue;
 		}
-		
-		ArrayAlreadyProcessedChild.Add(ReferenceName);
 
 		const FAssetData* AssetData = GetCachedAssetData(ReferenceName.ToString());
-		if (!AssetData) // Elements in ArrayAssetData are already of static class UCustomizableObject
+
+		UCustomizableObject* ChildObject = AssetData ? Cast<UCustomizableObject>(AssetData->GetAsset()) : nullptr;
+		if (!ChildObject || ChildObject->HasAnyFlags(RF_Transient))
 		{
+			ArrayAlreadyProcessedChild.Add(ReferenceName);
 			continue;
 		}
-		
-		UCustomizableObject* ChildObject = Cast<UCustomizableObject>(AssetData->GetAsset());
-		if (!ChildObject)
+
+		UCustomizableObjectNodeObject* Root = GetRootNode(ChildObject, bMultipleBaseObjectsFound);
+		if (Root->ParentObject != ParentObject)
 		{
 			continue;
 		}
-		
-		if (ChildObject != Object && !ChildObject->HasAnyFlags(RF_Transient))
-		{
-			UCustomizableObjectNodeObject* ChildRoot = GetRootNode(ChildObject, bMultipleBaseObjectsFound);
 
-			if (ChildRoot && !bMultipleBaseObjectsFound)
-			{
-				if (ChildRoot->ParentObject == Object)
-				{
-					if (const FGroupNodeIdsTempData* GroupGuid = GenerationContext.DuplicatedGroupNodeIds.FindPair(Object, FGroupNodeIdsTempData(ChildRoot->ParentObjectGroupId)))
-					{
-						ChildRoot->ParentObjectGroupId = GroupGuid->NewGroupNodeId;
-					}
+		ArrayAlreadyProcessedChild.Add(ReferenceName);
 
-					GenerationContext.GroupIdToExternalNodeMap.Add(ChildRoot->ParentObjectGroupId, ChildRoot);
-					GenerationContext.CustomizableObjectGuidsInCompilation.Add(ChildObject->GetVersionId());
-				}
+		if (!bMultipleBaseObjectsFound)
+		{
+			if (const FGroupNodeIdsTempData* GroupGuid = GenerationContext.DuplicatedGroupNodeIds.FindPair(ParentObject, FGroupNodeIdsTempData(Root->ParentObjectGroupId)))
+			{
+				Root->ParentObjectGroupId = GroupGuid->NewGroupNodeId;
 			}
-		}
 
-		TArray<UCustomizableObjectNodeObjectGroup*> GroupNodes;
-		ChildObject->Source->GetNodesOfClass<UCustomizableObjectNodeObjectGroup>(GroupNodes);
+			GenerationContext.GroupIdToExternalNodeMap.Add(Root->ParentObjectGroupId, Root);
+			GenerationContext.CustomizableObjectGuidsInCompilation.Add(ChildObject->GetVersionId());
 
-		if (GroupNodes.Num() > 0) // Only grafs with group nodes should have child grafs
-		{
-			for (int32 i = 0; i < GroupNodes.Num(); ++i)
+			TArray<UCustomizableObjectNodeObjectGroup*> GroupNodes;
+			ChildObject->Source->GetNodesOfClass<UCustomizableObjectNodeObjectGroup>(GroupNodes);
+
+			if (GroupNodes.Num() > 0) // Only grafs with group nodes should have child grafs
 			{
-				const FGuid NodeId = GenerationContext.GetNodeIdUnique(GroupNodes[i]);
-				if (NodeId != GroupNodes[i]->NodeGuid)
+				for (int32 i = 0; i < GroupNodes.Num(); ++i)
 				{
-					GenerationContext.DuplicatedGroupNodeIds.Add(ChildObject, FGroupNodeIdsTempData(GroupNodes[i]->NodeGuid, NodeId));
-					GroupNodes[i]->NodeGuid = NodeId;
+					const FGuid NodeId = GenerationContext.GetNodeIdUnique(GroupNodes[i]);
+					if (NodeId != GroupNodes[i]->NodeGuid)
+					{
+						GenerationContext.DuplicatedGroupNodeIds.Add(ChildObject, FGroupNodeIdsTempData(GroupNodes[i]->NodeGuid, NodeId));
+						GroupNodes[i]->NodeGuid = NodeId;
+					}
 				}
-			}
 
-			ProcessChildObjectsRecursively(ChildObject, GenerationContext);
+				ProcessChildObjectsRecursively(ChildObject, GenerationContext);
+			}
 		}
 	}
 }
