[Mutable] Remove unused private dependency.

#rnx

diff --git a/Engine/Plugins/Experimental/Mutable/Source/CustomizableObject/CustomizableObject.Build.cs b/Engine/Plugins/Experimental/Mutable/Source/CustomizableObject/CustomizableObject.Build.cs
--- a/Engine/Plugins/Experimental/Mutable/Source/CustomizableObject/CustomizableObject.Build.cs
+++ b/Engine/Plugins/Experimental/Mutable/Source/CustomizableObject/CustomizableObject.Build.cs
@@ -41,11 +41,6 @@ public class CustomizableObject : ModuleRules
 			"StructUtils"
 		});
 
-        PrivateIncludePathModuleNames.AddRange(
-        new string[] {
-                "TargetPlatform"
-        });
-
 		if (TargetRules.bBuildEditor == true)
         {
             PublicDependencyModuleNames.Add("UnrealEd");	// @todo api: Only public because of WITH_EDITOR
@@ -53,6 +48,5 @@ public class CustomizableObject : ModuleRules
             PublicDependencyModuleNames.Add("EditorStyle");
             PublicDependencyModuleNames.Add("MessageLog");
         }
-
     }
 }
