[Mutable] Remove the remaining OnCookStart code and other outdated code.
Remove unused compiler options bSaveCookedDataToDisk and bCheckChildrenGuids. CheckChildrenGuids was deprecated in favor of ParticipatingObjects.

#rnx

diff --git a/Engine/Plugins/Experimental/Mutable/Source/CustomizableObject/Classes/MuCO/CustomizableObject.h b/Engine/Plugins/Experimental/Mutable/Source/CustomizableObject/Classes/MuCO/CustomizableObject.h
--- a/Engine/Plugins/Experimental/Mutable/Source/CustomizableObject/Classes/MuCO/CustomizableObject.h
+++ b/Engine/Plugins/Experimental/Mutable/Source/CustomizableObject/Classes/MuCO/CustomizableObject.h
@@ -192,15 +192,9 @@ struct FCompilationOptions
 	// Compiling for cook
 	bool bIsCooking = false;
 
-	// Save compilation data for cook on disk 
-	bool bSaveCookedDataToDisk = false;
-
 	// This can be set for additional settings
 	const class ITargetPlatform* TargetPlatform = nullptr;
 
-	// Used to force a check to compare the guids stored in the base object compared to the children to see if any of them has changes and thus mark the object as modified.
-	bool bCheckChildrenGuids = false;
-
 	// Used to prevent the current model from being stored in the editor's streamed data and cache when the compilation finishes
 	bool bDontUpdateStreamedDataAndCache = false;
 
@@ -1226,10 +1220,6 @@ public:
 	UPROPERTY()
 	FGuid VersionId;
 
-	// Set of all the guids of all the CustomizableObjects in the compilation
-	UPROPERTY()
-	TSet<FGuid> CustomizableObjectGuidsInCompilation;
-
 	/* Map to identify what CustomizableObject owns a parameter. Used to display a tooltip when hovering a parameter
 	   in the Prev. instance panel */
 	UPROPERTY(Transient)
@@ -1409,6 +1399,12 @@ private:
 	// Warning: If while merging code both versions have changed, take the highest+1.
 	static const int32 CurrentSupportedVersion = 414;
 
+	// Compile the object for a specific platform - Compile for Cook Customizable Object
+	void CompileForTargetPlatform(const ITargetPlatform* TargetPlatform);
+
+	// Unless we are packaging there is no need for keeping all the data generated during compilation, this information is stored in the derived data.
+	void ClearCompiledData();
+
 public:
 
 	void AddUncompiledCOWarning(const FString& AdditionalLoggingInfo);
@@ -1429,7 +1425,7 @@ public:
 	int32 GetCurrentSupportedVersion() const { return CurrentSupportedVersion; };
 
 	// Compose folder name where the data is stored
-	FString GetCompiledDataFolderPath(bool bIsEditorData) const;
+	FString GetCompiledDataFolderPath() const;
 
 	// Compose file name 
 	FString GetCompiledDataFileName(bool bIsModel, const ITargetPlatform* InTargetPlatform = nullptr, bool bIsDiskStreamer = false);
@@ -1452,15 +1448,6 @@ public:
 	bool IsCachedCookedPlatformDataLoaded(const ITargetPlatform* TargetPlatform) override;
 	// End UObject interface.
 
-	// To avoid missing resources in packages we load all the resources. Loading the references added in the BeginCacheForCookedPlatformData will help us ensure they are included in the final package
-	void LoadReferencedObjects();
-	
-	// Compile the object for a specific platform - Compile for Cook Customizable Object
-	void CompileForTargetPlatform(const ITargetPlatform* TargetPlatform, bool bIsOnCookStart = false);
-
-	// Unless we are packaging there is no need for keeping all the data generated during compilation, this information is stored in the derived data.
-	void ClearCompiledData();
-
 	// Rebuild HashToStreamableBlocks and ParameterProperties from the current compiled model.
 	void UpdateCompiledDataFromModel();
 	
@@ -1468,8 +1455,8 @@ public:
 	void SaveCompiledData(FArchive& Ar, bool bSkipEditorOnlyData = false);
 	void LoadCompiledData(FArchive& Ar, const ITargetPlatform* InTargetPlatform, bool bSkipEditorOnlyData = false);
 
-	/** Load compiled data from disk, this is used to load Editor Compilations and Cook Compilations (when using OnCookStart) */
-	void LoadCompiledDataFromDisk(bool bIsEditorData = true, const ITargetPlatform* InTargetPlatform = nullptr);
+	/** Load compiled data for the running platform from disk, this is used to load Editor Compilations. */
+	void LoadCompiledDataFromDisk();
 
 	/** Cache platform data for cook */
 	void CachePlatformData(const ITargetPlatform* InTargetPlatform, const TArray64<uint8>& InObjectBytes, const TArray64<uint8>& InBulkBytes);
diff --git a/Engine/Plugins/Experimental/Mutable/Source/CustomizableObject/Private/MuCO/CustomizableObject.cpp b/Engine/Plugins/Experimental/Mutable/Source/CustomizableObject/Private/MuCO/CustomizableObject.cpp
--- a/Engine/Plugins/Experimental/Mutable/Source/CustomizableObject/Private/MuCO/CustomizableObject.cpp
+++ b/Engine/Plugins/Experimental/Mutable/Source/CustomizableObject/Private/MuCO/CustomizableObject.cpp
@@ -43,9 +43,6 @@
 
 DEFINE_LOG_CATEGORY(LogMutable);
 
-#if WITH_EDITOR
-static bool bUsesOnCookStart = false;
-#endif
 //-------------------------------------------------------------------------------------------------
 
 UCustomizableObject::UCustomizableObject()
@@ -258,9 +255,7 @@ void UCustomizableObject::Serialize(FArchive& Ar_Asset)
 		
 		if (Ar_Asset.IsLoading())
 		{
-			ITargetPlatformManagerModule& TargetPlatformManager = GetTargetPlatformManagerRef();
-			const ITargetPlatform* RunningPlatform = TargetPlatformManager.GetRunningTargetPlatform();
-			LoadCompiledDataFromDisk(true, RunningPlatform);
+			LoadCompiledDataFromDisk();
 		}
 	}
 #else
@@ -286,74 +281,23 @@ void UCustomizableObject::PostRename(UObject * OldOuter, const FName OldName)
 
 void UCustomizableObject::BeginCacheForCookedPlatformData(const ITargetPlatform* TargetPlatform)
 {
-	bool bIsRelevantForThisTarget =
-		(Relevancy == ECustomizableObjectRelevancy::All)
-		||
-		(TargetPlatform && Relevancy == ECustomizableObjectRelevancy::ClientOnly && !TargetPlatform->IsServerOnly());
-	
-	if (TargetPlatform && bIsRelevantForThisTarget)
+	if (!TargetPlatform)
 	{
-		if (Private->CachedPlatformNames.Find(TargetPlatform->PlatformName()) == INDEX_NONE)
-		{
-			if (!bUsesOnCookStart)
-			{
-				// Compile and save in the CachedPlatformsData map
-				CompileForTargetPlatform(TargetPlatform);
-			}
-			else
-			{
-				// Load from Disk
-				LoadCompiledDataFromDisk(false, TargetPlatform);
-
-				if (!Private->bModelCompiledForCook)
-				{
-					LoadReferencedObjects();
-					Private->bModelCompiledForCook = true;
-				}
-			}
-
-			Private->CachedPlatformNames.Add(TargetPlatform->PlatformName());
-		}
+		return;
 	}
-	else
+
+	if (Private->CachedPlatformNames.Find(TargetPlatform->PlatformName()) == INDEX_NONE)
 	{
-		ClearCompiledData();
-		Private->SetModel(nullptr, FGuid()); // Discard compilation
-		if (TargetPlatform)
-		{
-			Private->CachedPlatformNames.Add(TargetPlatform->PlatformName());
-		}
+		// Compile and save in the CachedPlatformsData map
+		CompileForTargetPlatform(TargetPlatform);
+		Private->CachedPlatformNames.Add(TargetPlatform->PlatformName());
 	}
 }
 
 
 bool UCustomizableObject::IsCachedCookedPlatformDataLoaded( const ITargetPlatform* TargetPlatform ) 
-{ 
-	return Private->CachedPlatformNames.Find(TargetPlatform->PlatformName()) != INDEX_NONE;
-}
-
-// TODO COOK: Remove Hack to add new references to the package
-void UCustomizableObject::LoadReferencedObjects()
 {
-	for (TSoftObjectPtr<UMaterialInterface> r : ReferencedMaterials)
-	{
-		r.LoadSynchronous();
-	}
-
-	for(TPair<FString, TSoftObjectPtr<UPhysicsAsset>>& p : PhysicsAssetsMap)
-	{
-		p.Value.LoadSynchronous();
-	}
-	
-	for (TPair<FString, FParameterUIData>& i : ParameterUIDataMap)
-	{
-		i.Value.LoadResources();
-	}
-
-	for (TPair<FString, FParameterUIData>& s : StateUIDataMap)
-	{
-		s.Value.LoadResources();
-	}
+	return !TargetPlatform || Private->CachedPlatformNames.Find(TargetPlatform->PlatformName()) != INDEX_NONE;
 }
 
 
@@ -729,25 +673,18 @@ void UCustomizableObject::LoadCompiledData(FArchive& MemoryReader, const ITarget
 	UpdateParameterPropertiesFromModel();
 }
 
-void UCustomizableObject::LoadCompiledDataFromDisk(bool bIsEditorData, const ITargetPlatform* InTargetPlatform)
+void UCustomizableObject::LoadCompiledDataFromDisk()
 {
-	FString PlatformName = InTargetPlatform ? InTargetPlatform->PlatformName() : FPlatformProperties::PlatformName();
-
-	if (!bIsEditorData) // Loading cooked data
-	{
-		// If we don't use OnCookStart there's nothing to be loaded. Same case for child objects.
-		if (!bUsesOnCookStart || bIsChildObject) return;
-
-		// Platform Data already cached
-		if (CachedPlatformsData.Contains(PlatformName)) return;
-	}
+	ITargetPlatformManagerModule& TargetPlatformManager = GetTargetPlatformManagerRef();
+	const ITargetPlatform* RunningPlatform = TargetPlatformManager.GetRunningTargetPlatform();
+	check(RunningPlatform);
 
 	// Compose Folder Name
-	const FString FolderPath = GetCompiledDataFolderPath(bIsEditorData);
+	const FString FolderPath = GetCompiledDataFolderPath();
 
 	// Compose File Names
-	FString ModelFileName = FolderPath + GetCompiledDataFileName(true, InTargetPlatform);
-	FString StreamableFileName = FolderPath + GetCompiledDataFileName(false, InTargetPlatform);
+	const FString ModelFileName = FolderPath + GetCompiledDataFileName(true, RunningPlatform);
+	const FString StreamableFileName = FolderPath + GetCompiledDataFileName(false, RunningPlatform);
 
 	IFileManager& FileManager = IFileManager::Get();
 	if (FileManager.FileExists(*ModelFileName) && FileManager.FileExists(*StreamableFileName))
@@ -780,7 +717,7 @@ void UCustomizableObject::LoadCompiledDataFromDisk(bool bIsEditorData, const ITa
 			&&
 			CompiledDataHeader.VersionId == StreamableDataHeader.VersionId)
 		{
-			if ( bIsEditorData && (IsRunningGame() || CompiledDataHeader.VersionId == VersionId) )
+			if (IsRunningGame() || CompiledDataHeader.VersionId == VersionId)
 			{ 
 				int64 CompiledDataSize = CompiledDataFileHandle->Size() - HeaderSize;
 				TArray64<uint8> CompiledDataBytes;
@@ -790,24 +727,7 @@ void UCustomizableObject::LoadCompiledDataFromDisk(bool bIsEditorData, const ITa
 				CompiledDataFileHandle->Read(CompiledDataBytes.GetData(), CompiledDataSize);
 
 				FMemoryReaderView MemoryReader(CompiledDataBytes);
-				LoadCompiledData(MemoryReader, InTargetPlatform);
-			}
-			else if (!bIsEditorData)// Caching Cooked Data
-			{
-				FMutableCachedPlatformData& CachedData = CachedPlatformsData.Add(PlatformName);
-				int64 CompiledDataSize = CompiledDataFileHandle->Size() - HeaderSize;
-				int64 StreamableDataSize = StreamableDataFileHandle->Size() - HeaderSize;
-
-				CachedData.ModelData.SetNumUninitialized(CompiledDataSize);
-				CachedData.StreamableData.SetNumUninitialized(StreamableDataSize);
-
-				// Change the current read position to exclude the header
-				CompiledDataFileHandle->Seek(HeaderSize);
-				StreamableDataFileHandle->Seek(HeaderSize);
-
-				// Read Data
-				CompiledDataFileHandle->Read(CachedData.ModelData.GetData(), CompiledDataSize);
-				StreamableDataFileHandle->Read(CachedData.StreamableData.GetData(), StreamableDataSize);
+				LoadCompiledData(MemoryReader, RunningPlatform);
 			}
 		}
 
@@ -836,10 +756,12 @@ void UCustomizableObject::CachePlatformData(const ITargetPlatform* InTargetPlatf
 }
 
 
-void UCustomizableObject::CompileForTargetPlatform(const ITargetPlatform* TargetPlatform, bool bIsOnCookStart)
+void UCustomizableObject::CompileForTargetPlatform(const ITargetPlatform* TargetPlatform)
 {
-	// TEMP COOK: Remove together with OnCookStart
-	bUsesOnCookStart = bIsOnCookStart;
+	if (!TargetPlatform)
+	{
+		return;
+	}
 
 	FCustomizableObjectCompilerBase* Compiler = UCustomizableObjectSystem::GetInstance()->GetNewCompiler();
 
@@ -848,7 +770,7 @@ void UCustomizableObject::CompileForTargetPlatform(const ITargetPlatform* Target
 	bool bIsRelevantForThisTarget =
 		(Relevancy == ECustomizableObjectRelevancy::All)
 		||
-		(TargetPlatform && Relevancy == ECustomizableObjectRelevancy::ClientOnly && !TargetPlatform->IsServerOnly());
+		(Relevancy == ECustomizableObjectRelevancy::ClientOnly && !TargetPlatform->IsServerOnly());
 
 	// Discard any older compilation
 	Private->SetModel(nullptr, FGuid());
@@ -859,7 +781,6 @@ void UCustomizableObject::CompileForTargetPlatform(const ITargetPlatform* Target
 		Options.OptimizationLevel = 2;	// max optimization when packaging.
 		Options.TextureCompression = ECustomizableObjectTextureCompression::HighQuality;
 		Options.bIsCooking = true;
-		Options.bSaveCookedDataToDisk = bUsesOnCookStart;
 		Options.TargetPlatform = TargetPlatform;
 		Options.CustomizableObjectNumBoneInfluences = ICustomizableObjectModule::Get().GetNumBoneInfluences();
 
@@ -969,10 +890,9 @@ FReply UCustomizableObject::AddNewParameterProfile(FString Name, UCustomizableOb
 }
 
 
-FString UCustomizableObject::GetCompiledDataFolderPath(bool bIsEditorData) const
+FString UCustomizableObject::GetCompiledDataFolderPath() const
 {	
-	const FString FolderName = bIsEditorData ? TEXT("MutableStreamedDataEditor/") : TEXT("MutableStreamedData/");
-	return FPaths::ConvertRelativePathToFull(FPaths::ProjectSavedDir() + FolderName);
+	return FPaths::ConvertRelativePathToFull(FPaths::ProjectSavedDir() + TEXT("MutableStreamedDataEditor/"));
 }
 
 
diff --git a/Engine/Plugins/Experimental/Mutable/Source/CustomizableObject/Private/MuCO/UnrealMutableModelDiskStreamer.cpp b/Engine/Plugins/Experimental/Mutable/Source/CustomizableObject/Private/MuCO/UnrealMutableModelDiskStreamer.cpp
--- a/Engine/Plugins/Experimental/Mutable/Source/CustomizableObject/Private/MuCO/UnrealMutableModelDiskStreamer.cpp
+++ b/Engine/Plugins/Experimental/Mutable/Source/CustomizableObject/Private/MuCO/UnrealMutableModelDiskStreamer.cpp
@@ -96,7 +96,7 @@ bool FUnrealMutableModelBulkReader::PrepareStreamingForObject(UCustomizableObjec
 		NewData.Model = TWeakPtr<const mu::Model>(CustomizableObject->GetModel());
 
 #if WITH_EDITOR
-		FString FolderPath = CustomizableObject->GetCompiledDataFolderPath(true);
+		FString FolderPath = CustomizableObject->GetCompiledDataFolderPath();
 		FString FullFileName = FolderPath + CustomizableObject->GetCompiledDataFileName(false, nullptr, true);
 
 		const TSharedPtr<IAsyncReadFileHandle> ReadFileHandle = MakeShareable(FPlatformFileManager::Get().GetPlatformFile().OpenAsyncRead(*FullFileName));
diff --git a/Engine/Plugins/Experimental/Mutable/Source/CustomizableObjectEditor/Private/MuCOE/CustomizableObjectCompileRunnable.cpp b/Engine/Plugins/Experimental/Mutable/Source/CustomizableObjectEditor/Private/MuCOE/CustomizableObjectCompileRunnable.cpp
--- a/Engine/Plugins/Experimental/Mutable/Source/CustomizableObjectEditor/Private/MuCOE/CustomizableObjectCompileRunnable.cpp
+++ b/Engine/Plugins/Experimental/Mutable/Source/CustomizableObjectEditor/Private/MuCOE/CustomizableObjectCompileRunnable.cpp
@@ -230,10 +230,10 @@ FCustomizableObjectSaveDDRunnable::FCustomizableObjectSaveDDRunnable(UCustomizab
 	CustomizableObjectHeader.InternalVersion = CustomizableObject->GetCurrentSupportedVersion();
 	CustomizableObjectHeader.VersionId = Options.bIsCooking? FGuid::NewGuid() : CustomizableObject->GetVersionId();
 
-	if (!Options.bIsCooking || Options.bSaveCookedDataToDisk)
+	if (!Options.bIsCooking)
 	{
 		// We will be saving all compilation data in two separate files, write CO Data
-		FolderPath = CustomizableObject->GetCompiledDataFolderPath(!InOptions.bIsCooking);
+		FolderPath = CustomizableObject->GetCompiledDataFolderPath();
 		CompileDataFullFileName = FolderPath + CustomizableObject->GetCompiledDataFileName(true, InOptions.TargetPlatform);
 		StreamableDataFullFileName = FolderPath + CustomizableObject->GetCompiledDataFileName(false, InOptions.TargetPlatform);
 
@@ -248,7 +248,7 @@ uint32 FCustomizableObjectSaveDDRunnable::Run()
 {
 	bool bModelSerialized = Model.Get() != nullptr;
 
-	if (Options.bIsCooking && !Options.bSaveCookedDataToDisk)
+	if (Options.bIsCooking)
 	{
 		// Serialize mu::Model and streamable resources 
 		FMemoryWriter64 ModelMemoryWriter(Bytes, false, true);
diff --git a/Engine/Plugins/Experimental/Mutable/Source/CustomizableObjectEditor/Private/MuCOE/CustomizableObjectCompiler.cpp b/Engine/Plugins/Experimental/Mutable/Source/CustomizableObjectEditor/Private/MuCOE/CustomizableObjectCompiler.cpp
--- a/Engine/Plugins/Experimental/Mutable/Source/CustomizableObjectEditor/Private/MuCOE/CustomizableObjectCompiler.cpp
+++ b/Engine/Plugins/Experimental/Mutable/Source/CustomizableObjectEditor/Private/MuCOE/CustomizableObjectCompiler.cpp
@@ -365,7 +365,6 @@ void FCustomizableObjectCompiler::ProcessChildObjectsRecursively(UCustomizableOb
 			}
 
 			GenerationContext.GroupIdToExternalNodeMap.Add(Root->ParentObjectGroupId, Root);
-			GenerationContext.CustomizableObjectGuidsInCompilation.Add(ChildObject->GetVersionId());
 
 			TArray<UCustomizableObjectNodeObjectGroup*> GroupNodes;
 			ChildObject->Source->GetNodesOfClass<UCustomizableObjectNodeObjectGroup>(GroupNodes);
@@ -953,25 +952,6 @@ void FCustomizableObjectCompiler::CompileInternal(UCustomizableObject* Object, c
 	}
 	else
 	{
-		if (Options.bCheckChildrenGuids)
-		{
-			// Check if all children are exactly the same as in the last compilation, if not, mark the object as modified and change VersionId to force a cache recompilation in other pcs
-			if ((GenerationContext.CustomizableObjectGuidsInCompilation.Num() != Object->CustomizableObjectGuidsInCompilation.Num()) ||
-				(GenerationContext.CustomizableObjectGuidsInCompilation.Intersect(Object->CustomizableObjectGuidsInCompilation).Num() != GenerationContext.CustomizableObjectGuidsInCompilation.Num()))
-			{
-				if (!ParamNamesToSelectedOptions.Num()) // Don't marked the object as modified because of a partial compilation
-				{
-					Object->CustomizableObjectGuidsInCompilation = GenerationContext.CustomizableObjectGuidsInCompilation;
-
-					if (!Options.bIsCooking)
-					{
-						Object->UpdateVersionId();
-						Object->MarkPackageDirty();
-					}
-				}
-			}
-		}
-
 		// Morph target generated data does not need extra processing so move semantics can be used
 		// to avoid a possibly expensive copy.
 		Object->ContributingMorphTargetsInfo = MoveTemp(GenerationContext.ContributingMorphTargetsInfo);
diff --git a/Engine/Plugins/Experimental/Mutable/Source/CustomizableObjectEditor/Private/MuCOE/CustomizableObjectEditor.cpp b/Engine/Plugins/Experimental/Mutable/Source/CustomizableObjectEditor/Private/MuCOE/CustomizableObjectEditor.cpp
--- a/Engine/Plugins/Experimental/Mutable/Source/CustomizableObjectEditor/Private/MuCOE/CustomizableObjectEditor.cpp
+++ b/Engine/Plugins/Experimental/Mutable/Source/CustomizableObjectEditor/Private/MuCOE/CustomizableObjectEditor.cpp
@@ -748,10 +748,7 @@ void FCustomizableObjectEditor::CompileObjectUserPressedButton()
 	}
 
 	Compiler.ClearAllCompileOnlySelectedOption();
-
-	CustomizableObject->CompileOptions.bCheckChildrenGuids = true;
 	CompileObject();
-	CustomizableObject->CompileOptions.bCheckChildrenGuids = false;
 }
 
 
diff --git a/Engine/Plugins/Experimental/Mutable/Source/CustomizableObjectEditor/Private/MuCOE/CustomizableObjectEditorModule.cpp b/Engine/Plugins/Experimental/Mutable/Source/CustomizableObjectEditor/Private/MuCOE/CustomizableObjectEditorModule.cpp
--- a/Engine/Plugins/Experimental/Mutable/Source/CustomizableObjectEditor/Private/MuCOE/CustomizableObjectEditorModule.cpp
+++ b/Engine/Plugins/Experimental/Mutable/Source/CustomizableObjectEditor/Private/MuCOE/CustomizableObjectEditorModule.cpp
@@ -79,12 +79,6 @@ const FName CustomizableObjectDebuggerAppIdentifier = FName(TEXT("CustomizableOb
 
 #define LOCTEXT_NAMESPACE "MutableSettings"
 
-static TAutoConsoleVariable<bool> CVarMutableOnCookStartEnabled(
-	TEXT("b.OnCookStartEnabled"),
-	true,
-	TEXT("If enabled, Customizable Objects will be compiled before the actual cook starts. Compiled data will be stored on de DDC and cached during BeginCache.\n"),
-	ECVF_Scalability);
-
 /**
  * StaticMesh editor module
  */
diff --git a/Engine/Plugins/Experimental/Mutable/Source/CustomizableObjectEditor/Private/MuCOE/GenerateMutableSource/GenerateMutableSource.h b/Engine/Plugins/Experimental/Mutable/Source/CustomizableObjectEditor/Private/MuCOE/GenerateMutableSource/GenerateMutableSource.h
--- a/Engine/Plugins/Experimental/Mutable/Source/CustomizableObjectEditor/Private/MuCOE/GenerateMutableSource/GenerateMutableSource.h
+++ b/Engine/Plugins/Experimental/Mutable/Source/CustomizableObjectEditor/Private/MuCOE/GenerateMutableSource/GenerateMutableSource.h
@@ -719,9 +719,6 @@ struct FMutableGraphGenerationContext
 	TMap<FGraphCycleKey, const UCustomizableObject*> VisitedPins;
 	const UCustomizableObject* CustomizableObjectWithCycle = nullptr;
 
-	// Set of all the guids of all the child CustomizableObjects in the compilation
-	TSet<FGuid> CustomizableObjectGuidsInCompilation;
-
 	/** Stores the physics assets gathered from the SkeletalMesh nodes during compilation, to be used in mesh generation in-game */
 	TMap<FString, TSoftObjectPtr<UPhysicsAsset>> PhysicsAssetMap;
 
diff --git a/Engine/Plugins/Experimental/Mutable/Source/CustomizableObjectPopulation/Private/MuCOP/CustomizableObjectPopulation.cpp b/Engine/Plugins/Experimental/Mutable/Source/CustomizableObjectPopulation/Private/MuCOP/CustomizableObjectPopulation.cpp
--- a/Engine/Plugins/Experimental/Mutable/Source/CustomizableObjectPopulation/Private/MuCOP/CustomizableObjectPopulation.cpp
+++ b/Engine/Plugins/Experimental/Mutable/Source/CustomizableObjectPopulation/Private/MuCOP/CustomizableObjectPopulation.cpp
@@ -91,7 +91,7 @@ void UCustomizableObjectPopulation::BeginCacheForCookedPlatformData(const ITarge
 	{
 		for (int32 i = 0; i < ClassWeights.Num(); ++i)
 		{
-			ClassWeights[i].Class->CustomizableObject->LoadCompiledDataFromDisk(false, TargetPlatform);
+			check(ClassWeights[i].Class->CustomizableObject->IsCompiled());
 		}
 
 		CompilePopulation(Generator);
