[Mutable] Removed Unreal Insights Mutable cateogry.

This category was introduced to be able to enable Mutable trace in isolation. Since this can not be done (CPU has to always be enabled) it no longer makes sence to maintain this category.

#rnx

diff --git a/Engine/Plugins/Experimental/Mutable/Source/MutableRuntime/Private/MuR/MutableTrace.cpp b/Engine/Plugins/Experimental/Mutable/Source/MutableRuntime/Private/MuR/MutableTrace.cpp
deleted file mode 100644
diff --git a/Engine/Plugins/Experimental/Mutable/Source/MutableRuntime/Public/MuR/MutableTrace.h b/Engine/Plugins/Experimental/Mutable/Source/MutableRuntime/Public/MuR/MutableTrace.h
--- a/Engine/Plugins/Experimental/Mutable/Source/MutableRuntime/Public/MuR/MutableTrace.h
+++ b/Engine/Plugins/Experimental/Mutable/Source/MutableRuntime/Public/MuR/MutableTrace.h
@@ -5,13 +5,10 @@
 #include "Trace/Trace.h"
 #include "HAL/PlatformTime.h"
 
-namespace UE::Trace { class FChannel; }
-
-UE_TRACE_CHANNEL_EXTERN(MutableChannel, MUTABLERUNTIME_API)
 
 /** Custom Mutable profiler scope. */
 #define MUTABLE_CPUPROFILER_SCOPE(Name) \
-	TRACE_CPUPROFILER_EVENT_SCOPE_ON_CHANNEL(Mutable_##Name, MutableChannel)
+	TRACE_CPUPROFILER_EVENT_SCOPE(Mutable_##Name)
 
 
 /** Simple class that saves the time the scope is alive in the given location. */
