[Mutable] Fix incorrect UpdateResult.

#rnx

diff --git a/Engine/Plugins/Experimental/Mutable/Source/CustomizableObject/Private/MuCO/CustomizableObjectSystem.cpp b/Engine/Plugins/Experimental/Mutable/Source/CustomizableObject/Private/MuCO/CustomizableObjectSystem.cpp
--- a/Engine/Plugins/Experimental/Mutable/Source/CustomizableObject/Private/MuCO/CustomizableObjectSystem.cpp
+++ b/Engine/Plugins/Experimental/Mutable/Source/CustomizableObject/Private/MuCO/CustomizableObjectSystem.cpp
@@ -279,7 +279,7 @@ void FMutablePendingInstanceWork::AddDiscard(const FMutablePendingInstanceDiscar
 {
 	if (const FMutablePendingInstanceUpdate* ExistingUpdate = PendingInstanceUpdates.Find(TaskToEnqueue.CustomizableObjectInstance.Get()))
 	{
-		ExistingUpdate->Context->UpdateResult = EUpdateResult::ErrorReplaced;
+		ExistingUpdate->Context->UpdateResult = EUpdateResult::ErrorDiscarded;
 		FinishUpdateGlobal(ExistingUpdate->Context);
 		RemoveUpdate(ExistingUpdate->Context->Instance);
 	}
