[Mutable] Show CO in "Failed to remap old pins" warning.

#rnx

diff --git a/Engine/Plugins/Experimental/Mutable/Source/CustomizableObjectEditor/Private/MuCOE/Nodes/CustomizableObjectNode.cpp b/Engine/Plugins/Experimental/Mutable/Source/CustomizableObjectEditor/Private/MuCOE/Nodes/CustomizableObjectNode.cpp
--- a/Engine/Plugins/Experimental/Mutable/Source/CustomizableObjectEditor/Private/MuCOE/Nodes/CustomizableObjectNode.cpp
+++ b/Engine/Plugins/Experimental/Mutable/Source/CustomizableObjectEditor/Private/MuCOE/Nodes/CustomizableObjectNode.cpp
@@ -286,6 +286,7 @@ void UCustomizableObjectNode::ReconstructNode(UCustomizableObjectNodeRemapPins*
 		if (bOrphanedPin)
 		{
 			FCustomizableObjectEditorLogger::CreateLog(LOCTEXT("OrphanPinsWarningReconstruct", "Failed to remap old pins"))
+			.BaseObject()
 			.Severity(EMessageSeverity::Warning)
 			.Context(*this)
 			.Log();
