[Mutable] Fix LodBenchmarkUtil crash when shutting down the editor.

#rnx

diff --git a/Engine/Plugins/Experimental/Mutable/Source/CustomizableObject/Private/MuCO/CustomizableObjectSystem.cpp b/Engine/Plugins/Experimental/Mutable/Source/CustomizableObject/Private/MuCO/CustomizableObjectSystem.cpp
--- a/Engine/Plugins/Experimental/Mutable/Source/CustomizableObject/Private/MuCO/CustomizableObjectSystem.cpp
+++ b/Engine/Plugins/Experimental/Mutable/Source/CustomizableObject/Private/MuCO/CustomizableObjectSystem.cpp
@@ -850,6 +850,11 @@ void FinishUpdateGlobal(const TSharedRef<FUpdateContextPrivate>& Context)
 		FFunctionGraphTask::CreateAndDispatchWhenReady( // Calling Benchmark in a task so we make sure we exited all scopes.
 		[Context]()
 		{
+			if (!UCustomizableObjectSystem::IsCreated()) // We are shutting down
+			{
+				return;	
+			}
+			
 			UCustomizableObjectSystem* System = UCustomizableObjectSystem::GetInstance();
 			if (!System)
 			{
