[Mutable] Better random parameters.

#rnx

diff --git a/Engine/Plugins/Experimental/Mutable/Source/CustomizableObject/Private/MuCO/CustomizableObjectInstance.cpp b/Engine/Plugins/Experimental/Mutable/Source/CustomizableObject/Private/MuCO/CustomizableObjectInstance.cpp
--- a/Engine/Plugins/Experimental/Mutable/Source/CustomizableObject/Private/MuCO/CustomizableObjectInstance.cpp
+++ b/Engine/Plugins/Experimental/Mutable/Source/CustomizableObject/Private/MuCO/CustomizableObjectInstance.cpp
@@ -2432,7 +2432,7 @@ int32 UCustomizableObjectInstance::FindProjectorParameterNameIndex(const FString
 
 void UCustomizableObjectInstance::SetRandomValues()
 {
-	Descriptor.SetRandomValues(FMath::Rand());
+	Descriptor.SetRandomValues();
 }
 
 void UCustomizableObjectInstance::SetRandomValuesFromStream(const FRandomStream& Stream)
diff --git a/Engine/Plugins/Experimental/Mutable/Source/CustomizableObject/Private/MuCO/CustomizableObjectInstanceDescriptor.cpp b/Engine/Plugins/Experimental/Mutable/Source/CustomizableObject/Private/MuCO/CustomizableObjectInstanceDescriptor.cpp
--- a/Engine/Plugins/Experimental/Mutable/Source/CustomizableObject/Private/MuCO/CustomizableObjectInstanceDescriptor.cpp
+++ b/Engine/Plugins/Experimental/Mutable/Source/CustomizableObject/Private/MuCO/CustomizableObjectInstanceDescriptor.cpp
@@ -2,6 +2,8 @@
 
 #include "MuCO/CustomizableObjectInstanceDescriptor.h"
 
+#include "CustomizableObjectSystemPrivate.h"
+#include "UnrealMutableImageProvider.h"
 #include "MuCO/CustomizableObject.h"
 #include "MuCO/CustomizableObjectPrivate.h"
 #include "MuCO/DefaultImageProvider.h"
@@ -2314,40 +2316,90 @@ void FCustomizableObjectInstanceDescriptor::SetCurrentState(const FString& State
 }
 
 
-void FCustomizableObjectInstanceDescriptor::SetRandomValues(const int32& InRandomizationSeed)
+void FCustomizableObjectInstanceDescriptor::SetRandomValues()
+{
+	SetRandomValues(FMath::SRand() * TNumericLimits<int32>::Max());
+}
+
+
+void FCustomizableObjectInstanceDescriptor::SetRandomValues(const int32 InRandomizationSeed)
 {
 	check(CustomizableObject);
 	RETURN_ON_UNCOMPILED_CO(CustomizableObject, TEXT("Error: Cannot set random values"))
 
 	// Set seed for deterministic behaviour
-	FMath::SRandInit(InRandomizationSeed);
-	
-	for (int32 i = 0; i < FloatParameters.Num(); ++i)
+	const FRandomStream Random(InRandomizationSeed);
+		
+	for (FCustomizableObjectFloatParameterValue& FloatParameter : FloatParameters)
 	{
-		FloatParameters[i].ParameterValue = FMath::SRand();
+		FloatParameter.ParameterValue = Random.GetFraction();
+
+		for (float& RangeValue : FloatParameter.ParameterRangeValues)
+		{
+			RangeValue = Random.GetFraction();
+		}
 	}
 
-	for (int32 i = 0; i < BoolParameters.Num(); ++i)
+	for (FCustomizableObjectBoolParameterValue& BoolParameter : BoolParameters)
 	{
-		BoolParameters[i].ParameterValue = (int32)FMath::SRand() % 2 == 0;
+		BoolParameter.ParameterValue = static_cast<bool>(Random.RandRange(0, 1));
 	}
-
-	for (int32 i = 0; i < IntParameters.Num(); ++i)
+	
+	for (FCustomizableObjectIntParameterValue& IntParameter : IntParameters)
 	{
-		const int32 ParameterIndexInCO = CustomizableObject->FindParameter(IntParameters[i].ParameterName);
+		const int32 ParamIndex = CustomizableObject->FindParameter(IntParameter.ParameterName);
+		const int32 NumValues = CustomizableObject->GetIntParameterNumOptions(ParamIndex);
 
-		// TODO: Randomize multidimensional parameters
-		if (ParameterIndexInCO >= 0 && !CustomizableObject->IsParameterMultidimensional(ParameterIndexInCO))
+		if (NumValues)
 		{
-			const int32 NumValues = CustomizableObject->GetIntParameterNumOptions(ParameterIndexInCO);
-			if (NumValues > 0)
+			IntParameter.ParameterValueName = CustomizableObject->GetIntParameterAvailableOption(ParamIndex, NumValues * Random.GetFraction());
+
+			for (FString& RangeValue : IntParameter.ParameterRangeValueNames)
 			{
-				const int32 Index = (int32)FMath::SRand() % NumValues;
-				FString Option = CustomizableObject->GetIntParameterAvailableOption(ParameterIndexInCO, Index);
-				SetIntParameterSelectedOption(i, Option);
+				RangeValue = CustomizableObject->GetIntParameterAvailableOption(ParamIndex, NumValues * Random.GetFraction());
 			}
+		}		
+	}
+
+	for (FCustomizableObjectVectorParameterValue& VectorParameter : VectorParameters)
+	{
+		VectorParameter.ParameterValue.R = Random.GetFraction();
+		VectorParameter.ParameterValue.G = Random.GetFraction();
+		VectorParameter.ParameterValue.B = Random.GetFraction();
+		VectorParameter.ParameterValue.A = Random.GetFraction();
+	}
+
+	const FCustomizableObjectSystemPrivate* SystemPrivate = UCustomizableObjectSystem::GetInstance()->GetPrivate();
+
+	TArray<FName> PossibleValues;
+
+	// Get all possible values
+	TArray<FCustomizableObjectExternalTexture> ProviderValues;
+	for (const TWeakObjectPtr<UCustomizableSystemImageProvider>& Provider : SystemPrivate->ImageProvider->ImageProviders)
+	{
+		ProviderValues.Reset();
+		Provider->GetTextureParameterValues(ProviderValues);
+
+		for (const FCustomizableObjectExternalTexture& ProviderValue : ProviderValues)
+		{
+			PossibleValues.Add(ProviderValue.Value);
 		}
 	}
+
+	if (const int32 NumPossibleValues = PossibleValues.Num())
+	{
+		for (FCustomizableObjectTextureParameterValue& TextureParameter : TextureParameters)
+		{
+			TextureParameter.ParameterValue = PossibleValues[NumPossibleValues * Random.GetFraction()];
+		
+			for (FName& RangeValue : TextureParameter.ParameterRangeValues)
+			{
+				RangeValue = PossibleValues[NumPossibleValues * Random.GetFraction()];
+			}				
+		}		
+	}
+	
+	// Currently we are not randomizing the projectors since we do not know the valid range of values.
 }
 
 
diff --git a/Engine/Plugins/Experimental/Mutable/Source/CustomizableObject/Public/MuCO/CustomizableObjectInstanceDescriptor.h b/Engine/Plugins/Experimental/Mutable/Source/CustomizableObject/Public/MuCO/CustomizableObjectInstanceDescriptor.h
--- a/Engine/Plugins/Experimental/Mutable/Source/CustomizableObject/Public/MuCO/CustomizableObjectInstanceDescriptor.h
+++ b/Engine/Plugins/Experimental/Mutable/Source/CustomizableObject/Public/MuCO/CustomizableObjectInstanceDescriptor.h
@@ -292,8 +292,10 @@ struct CUSTOMIZABLEOBJECT_API FCustomizableObjectInstanceDescriptor
 	void SetCurrentState(const FString& StateName);
 
 	// ------------------------------------------------------------
+
+	void SetRandomValues();
 	
-	void SetRandomValues(const int32& InRandomizationSeed);
+	void SetRandomValues(int32 InRandomizationSeed);
 
 	void SetRandomValuesFromStream(const FRandomStream& Stream);
 
