#rnx
#tests PIE

diff --git a/Engine/Plugins/Experimental/Mutable/Source/CustomizableObject/Private/MuCO/CustomizableObjectInstance.cpp b/Engine/Plugins/Experimental/Mutable/Source/CustomizableObject/Private/MuCO/CustomizableObjectInstance.cpp
--- a/Engine/Plugins/Experimental/Mutable/Source/CustomizableObject/Private/MuCO/CustomizableObjectInstance.cpp
+++ b/Engine/Plugins/Experimental/Mutable/Source/CustomizableObject/Private/MuCO/CustomizableObjectInstance.cpp
@@ -58,68 +58,87 @@ struct FMutableMaterialPlaceholder
 	struct FMutableMaterialPlaceHolderParam
 	{
 		FName ParamName;
+		EPlaceHolderParamType Type;
 		int32 LayerIndex; // Set to -1 for non-multilayer params
-		FLinearColor Vector;
 		float Scalar;
+		FLinearColor Vector;
 		FGeneratedTexture Texture;
-		EPlaceHolderParamType Type;
 
 		FMutableMaterialPlaceHolderParam(const FName& InParamName, const int32 InLayerIndex, const FLinearColor& InVector)
-			: ParamName(InParamName), LayerIndex(InLayerIndex), Vector(InVector), Type(EPlaceHolderParamType::Vector) {}
+			: ParamName(InParamName), Type(EPlaceHolderParamType::Vector), LayerIndex(InLayerIndex), Vector(InVector) {}
 
 		FMutableMaterialPlaceHolderParam(const FName& InParamName, const int32 InLayerIndex, const float InScalar)
-			: ParamName(InParamName), LayerIndex(InLayerIndex), Scalar(InScalar), Type(EPlaceHolderParamType::Scalar) {}
+			: ParamName(InParamName), Type(EPlaceHolderParamType::Scalar), LayerIndex(InLayerIndex), Scalar(InScalar) {}
 
 		FMutableMaterialPlaceHolderParam(const FName& InParamName, const int32 InLayerIndex, const FGeneratedTexture& InTexture)
-			: ParamName(InParamName), LayerIndex(InLayerIndex), Texture(InTexture), Type(EPlaceHolderParamType::Texture) {}
+			: ParamName(InParamName), Type(EPlaceHolderParamType::Texture), LayerIndex(InLayerIndex), Texture(InTexture) {}
 
 		bool operator<(const FMutableMaterialPlaceHolderParam& Other) const
 		{
 			return Type < Other.Type || ParamName.CompareIndexes(Other.ParamName);
 		}
+
+		bool operator==(const FMutableMaterialPlaceHolderParam& Other) const = default;
 	};
 
 	uint32 ParentMaterialID = 0;
-	TArray<FMutableMaterialPlaceHolderParam> Params;
 	int32 MatIndex = -1;
 
+private:
+	mutable TArray<FMutableMaterialPlaceHolderParam> Params;
+
+public:
 	void AddParam(const FMutableMaterialPlaceHolderParam& NewParam) { Params.Add(NewParam); }
+	
+	const TArray<FMutableMaterialPlaceHolderParam>& GetParams() const { return Params; }
 
-	// Return a hash of the material and its parameters
-	uint32 GetHash()
-	{
-		uint32 Hash = ParentMaterialID;
+	bool operator==(const FMutableMaterialPlaceholder& Other) const;
 
-		// Sort parameters before building the hash.
-		Params.Sort();
+	friend uint32 GetTypeHash(const FMutableMaterialPlaceholder& PlaceHolder);
+};
 
-		for (const FMutableMaterialPlaceHolderParam& Param : Params)
-		{
-			uint32 ParamHash = GetTypeHash(Param.ParamName);
-			ParamHash = HashCombineFast(ParamHash, (uint32)Param.LayerIndex);
-			ParamHash = HashCombineFast(ParamHash, (uint32)Param.Type);
 
-			switch (Param.Type)
-			{
-			case EPlaceHolderParamType::Vector:
-				ParamHash = HashCombineFast(ParamHash, GetTypeHash(Param.Vector));
-				break;
+bool FMutableMaterialPlaceholder::operator==(const FMutableMaterialPlaceholder& Other) const
+{
+	return ParentMaterialID == Other.ParentMaterialID &&
+		   Params == Other.Params;
+}
 
-			case EPlaceHolderParamType::Scalar:
-				ParamHash = HashCombineFast(ParamHash, GetTypeHash(Param.Scalar));
-				break;
 
-			case EPlaceHolderParamType::Texture:
-				ParamHash = HashCombineFast(ParamHash, Param.Texture.Texture->GetUniqueID());
-				break;
-			}
+// Return a hash of the material and its parameters
+uint32 GetTypeHash(const FMutableMaterialPlaceholder& PlaceHolder)
+{
+	uint32 Hash = GetTypeHash(PlaceHolder.ParentMaterialID);
+
+	// Sort parameters before building the hash.
+	PlaceHolder.Params.Sort();
 
-			Hash = HashCombineFast(Hash, ParamHash);
+	for (const FMutableMaterialPlaceholder::FMutableMaterialPlaceHolderParam& Param : PlaceHolder.Params)
+	{
+		uint32 ParamHash = GetTypeHash(Param.ParamName);
+		ParamHash = HashCombineFast(ParamHash, (uint32)Param.LayerIndex);
+		ParamHash = HashCombineFast(ParamHash, (uint32)Param.Type);
+
+		switch (Param.Type)
+		{
+		case FMutableMaterialPlaceholder::EPlaceHolderParamType::Vector:
+			ParamHash = HashCombineFast(ParamHash, GetTypeHash(Param.Vector));
+			break;
+
+		case FMutableMaterialPlaceholder::EPlaceHolderParamType::Scalar:
+			ParamHash = HashCombineFast(ParamHash, GetTypeHash(Param.Scalar));
+			break;
+
+		case FMutableMaterialPlaceholder::EPlaceHolderParamType::Texture:
+			ParamHash = HashCombineFast(ParamHash, Param.Texture.Texture->GetUniqueID());
+			break;
 		}
 
-		return Hash;
+		Hash = HashCombineFast(Hash, ParamHash);
 	}
-};
+
+	return Hash;
+}
 
 
 UTexture2D* UCustomizableInstancePrivateData::CreateTexture()
@@ -5356,7 +5375,7 @@ void UCustomizableInstancePrivateData::BuildMaterials(const TSharedRef<FUpdateCo
 		TArray<FSkeletalMaterial> Materials;
 
 		// Maps serializations of FMutableMaterialPlaceholder to Created Dynamic Material instances, used to reuse materials across LODs
-		TMap<uint32, TSharedPtr<FMutableMaterialPlaceholder>> ReuseMaterialCache;
+		TSet<FMutableMaterialPlaceholder> ReuseMaterialCache;
 
 		// SurfaceId per MaterialSlotIndex
 		TArray<int32> SurfaceIdToMaterialIndex;
@@ -5417,8 +5436,7 @@ void UCustomizableInstancePrivateData::BuildMaterials(const TSharedRef<FUpdateCo
 				const int32 LODMaterialIndex = SkeletalMesh->GetLODInfoArray()[LODIndex].LODMaterialMap.Add(MaterialSlotIndex);
 				SkeletalMesh->GetResourceForRendering()->LODRenderData[LODIndex].RenderSections[SurfaceIndex].MaterialIndex = LODMaterialIndex;
 
-				TSharedPtr<FMutableMaterialPlaceholder> MutableMaterialPlaceholderPtr(new FMutableMaterialPlaceholder);
-				FMutableMaterialPlaceholder& MutableMaterialPlaceholder = *MutableMaterialPlaceholderPtr;
+				FMutableMaterialPlaceholder MutableMaterialPlaceholder;
 				MutableMaterialPlaceholder.ParentMaterialID = MaterialTemplate->GetUniqueID();
 				MutableMaterialPlaceholder.MatIndex = MaterialSlotIndex;
 
@@ -5770,16 +5788,15 @@ void UCustomizableInstancePrivateData::BuildMaterials(const TSharedRef<FUpdateCo
 				// Find or create the material for this slot
 				UMaterialInterface* MaterialInterface = MaterialSlot.MaterialInterface;
 
-				const uint32 MaterialParameterHash = MutableMaterialPlaceholder.GetHash();
-				if (TSharedPtr<FMutableMaterialPlaceholder>* FoundMaterialPlaceholder = ReuseMaterialCache.Find(MaterialParameterHash))
+				if (FMutableMaterialPlaceholder* FoundMaterialPlaceholder = ReuseMaterialCache.Find(MutableMaterialPlaceholder))
 				{
-					MaterialInterface = Materials[(*FoundMaterialPlaceholder)->MatIndex].MaterialInterface;
+					MaterialInterface = Materials[FoundMaterialPlaceholder->MatIndex].MaterialInterface;
 				}
 				else // Material not cached, create a new one
 				{
 					MUTABLE_CPUPROFILER_SCOPE(BuildMaterials_CreateMaterial);
 
-					ReuseMaterialCache.Add(MaterialParameterHash, MutableMaterialPlaceholderPtr);
+					ReuseMaterialCache.Add(MutableMaterialPlaceholder);
 
 					FGeneratedMaterial& Material = GeneratedMaterials.AddDefaulted_GetRef();
 					Material.SurfaceId = Surface.SurfaceId;
@@ -5795,7 +5812,7 @@ void UCustomizableInstancePrivateData::BuildMaterials(const TSharedRef<FUpdateCo
 						Material.MaterialInterface = OldMaterial.MaterialInterface;
 					}
 					
-					if (!MaterialInstance && MutableMaterialPlaceholder.Params.Num() != 0)
+					if (!MaterialInstance && MutableMaterialPlaceholder.GetParams().Num() != 0)
 					{
 						MaterialInstance = UMaterialInstanceDynamic::Create(MaterialTemplate, GetTransientPackage());
 						Material.MaterialInterface = MaterialInstance;
@@ -5803,7 +5820,7 @@ void UCustomizableInstancePrivateData::BuildMaterials(const TSharedRef<FUpdateCo
 
 					if (MaterialInstance)
 					{
-						for (const FMutableMaterialPlaceholder::FMutableMaterialPlaceHolderParam& Param : MutableMaterialPlaceholder.Params)
+						for (const FMutableMaterialPlaceholder::FMutableMaterialPlaceHolderParam& Param : MutableMaterialPlaceholder.GetParams())
 						{
 							switch (Param.Type)
 							{
diff --git a/Engine/Plugins/Experimental/Mutable/Source/CustomizableObject/Private/MuCO/CustomizableObjectSystemPrivate.h b/Engine/Plugins/Experimental/Mutable/Source/CustomizableObject/Private/MuCO/CustomizableObjectSystemPrivate.h
--- a/Engine/Plugins/Experimental/Mutable/Source/CustomizableObject/Private/MuCO/CustomizableObjectSystemPrivate.h
+++ b/Engine/Plugins/Experimental/Mutable/Source/CustomizableObject/Private/MuCO/CustomizableObjectSystemPrivate.h
@@ -202,6 +202,8 @@ struct FGeneratedTexture
 
 	UPROPERTY(Category = CustomizableObjectInstance, VisibleAnywhere)
 	TObjectPtr<UTexture> Texture = nullptr;
+
+	bool operator==(const FGeneratedTexture& Other) const = default;
 };
 
 
