[Mutable] Fix mip genertion for block compressed images with less than 4 channels.

- Currently the case can never happen because the scratch memory forces a decompression to RGBA, but fix it just so it does not break if at some point other formats are used.

- Fix incorrect use of the out of bound mask.

#rnx

diff --git a/Engine/Plugins/Experimental/Mutable/Source/MutableRuntime/Private/MuR/OpImageMipmap.cpp b/Engine/Plugins/Experimental/Mutable/Source/MutableRuntime/Private/MuR/OpImageMipmap.cpp
--- a/Engine/Plugins/Experimental/Mutable/Source/MutableRuntime/Private/MuR/OpImageMipmap.cpp
+++ b/Engine/Plugins/Experimental/Mutable/Source/MutableRuntime/Private/MuR/OpImageMipmap.cpp
@@ -365,8 +365,8 @@ namespace OpImageMipmap_Detail
 								uint64 Row1Bits;
 								FMemory::Memcpy(&Row1Bits, SrcRow1, sizeof(uint64));
 								
-								const uint64 OutOfBoundsMask = Row0Offset.X + 1 >= BatchDecSizeInPixels.X ? 0 : ~0;
-
+								const bool bOutOfBounds = Row0Offset.X + 1 >= BatchDecSizeInPixels.X;
+								
 								constexpr uint64 ShiftMask = 0xFEFEFEFEFEFEFEFE;
 								
 								const uint64 XorRow0Row1Bits = Row0Bits ^ Row1Bits;
@@ -374,7 +374,7 @@ namespace OpImageMipmap_Detail
 								
 								// Average of 2 unsigned integers without overflow extended to work on multiple bytes.
 								const uint64 AvgLowBits = (Row0Bits & Row1Bits) + ((XorRow0Row1Bits & ShiftMask) >> 1) + ErrorCorrection;
-								const uint64 AvgHighBits = OutOfBoundsMask ? AvgLowBits : (AvgLowBits >> 32);
+								const uint64 AvgHighBits = bOutOfBounds ? AvgLowBits : (AvgLowBits >> NumChannels*8);
 								const uint32 Result = (AvgLowBits & AvgHighBits) + (((AvgLowBits ^ AvgHighBits) & ShiftMask) >> 1);
 								
 								FMemory::Memcpy(DestPixel, &Result, NumChannels);
