[mutable] Fix code optimization bug with incorrect fading angles when a swizzle of channels is combined with a raster operation.
#rnx

diff --git a/Engine/Plugins/Experimental/Mutable/Source/MutableTools/Private/MuT/ASTOpImageSwizzle.cpp b/Engine/Plugins/Experimental/Mutable/Source/MutableTools/Private/MuT/ASTOpImageSwizzle.cpp
--- a/Engine/Plugins/Experimental/Mutable/Source/MutableTools/Private/MuT/ASTOpImageSwizzle.cpp
+++ b/Engine/Plugins/Experimental/Mutable/Source/MutableTools/Private/MuT/ASTOpImageSwizzle.cpp
@@ -378,6 +378,20 @@ namespace mu
 				Ptr<ASTOpImageSwizzle> NewSwizzle = mu::Clone<ASTOpImageSwizzle>(this);
 				ReplaceAllSources(NewSwizzle, NewRaster->image.child());
 				NewRaster->image = NewSwizzle;
+
+				// If we are swapping rgb and alphas, we need to correct some flags
+				{
+					// We should only find these two cases
+					if (SourceChannels[0] == 3 || SourceChannels[1] == 3 || SourceChannels[2] == 3)
+					{
+						NewRaster->bIsRGBFadingEnabled = NewRaster->bIsAlphaFadingEnabled;
+					}
+					else if (Sources[3] && SourceChannels[3] < 3 )
+					{
+						NewRaster->bIsAlphaFadingEnabled = NewRaster->bIsRGBFadingEnabled;
+					}
+				}
+
 				at = NewRaster;
 				break;
 			}
@@ -736,6 +750,19 @@ namespace mu
 
 					NewRaster->image = NewSwizzle;
 
+					// If we are swapping rgb and alphas, we need to correct some flags
+					{
+						// We should only find these two cases
+						if (SourceChannels[0] == 3 || SourceChannels[1] == 3 || SourceChannels[2] == 3)
+						{
+							NewRaster->bIsRGBFadingEnabled = NewRaster->bIsAlphaFadingEnabled;
+						}
+						else if (Sources[3] && SourceChannels[3] < 3)
+						{
+							NewRaster->bIsAlphaFadingEnabled = NewRaster->bIsRGBFadingEnabled;
+						}
+					}
+
 					at = NewRaster;
 				}
 
