CustomizableObject: Fix indeterminism in saving CustomizableObjectInstances, by making sure the CustomizableObject has compiled for the targetplatform before we save the instance that might need to update from it.
#rnx

diff --git a/Engine/Plugins/Experimental/Mutable/Source/CustomizableObject/Classes/MuCO/CustomizableObject.h b/Engine/Plugins/Experimental/Mutable/Source/CustomizableObject/Classes/MuCO/CustomizableObject.h
--- a/Engine/Plugins/Experimental/Mutable/Source/CustomizableObject/Classes/MuCO/CustomizableObject.h
+++ b/Engine/Plugins/Experimental/Mutable/Source/CustomizableObject/Classes/MuCO/CustomizableObject.h
@@ -1461,6 +1461,12 @@ public:
 	/** Cache platform data for cook */
 	void CachePlatformData(const ITargetPlatform* InTargetPlatform, const TArray64<uint8>& InObjectBytes, const TArray64<uint8>& InBulkBytes);
 
+	/**
+	 * Loads data previously compiled in BeginCacheForCookedPlatformData onto the UProperties in *this,
+	 * in preparation for saving the cooked package for *this or for a CustomizableObjectInstance using *this.
+	 * Returns whether the data was successfully loaded.
+	 */
+	bool TryLoadCompiledCookDataForPlatform(const ITargetPlatform* TargetPlatform);
 #if WITH_EDITORONLY_DATA
 
 	/** Map of PlatformName to CachedPlatformData. Only valid while cooking. */
@@ -1471,6 +1477,8 @@ public:
 	void SaveEmbeddedData(FArchive& Ar);
 
 #endif
+	/** Compute bIsChildObject if currently possible to do so. Return whether it was computed. */
+	bool TryUpdateIsChildObject();
 
 	// Data that may be stored in the asset itself, only in packaged builds.
 	void LoadEmbeddedData(FArchive& Ar);
diff --git a/Engine/Plugins/Experimental/Mutable/Source/CustomizableObject/Classes/MuCO/CustomizableObjectInstance.h b/Engine/Plugins/Experimental/Mutable/Source/CustomizableObject/Classes/MuCO/CustomizableObjectInstance.h
--- a/Engine/Plugins/Experimental/Mutable/Source/CustomizableObject/Classes/MuCO/CustomizableObjectInstance.h
+++ b/Engine/Plugins/Experimental/Mutable/Source/CustomizableObject/Classes/MuCO/CustomizableObjectInstance.h
@@ -154,6 +154,9 @@ public:
 #if WITH_EDITOR
 	virtual void PostEditChangeProperty(FPropertyChangedEvent& PropertyChangedEvent) override;
 	virtual bool CanEditChange( const FProperty* InProperty ) const override;
+	virtual void BeginCacheForCookedPlatformData(const ITargetPlatform* TargetPlatform) override;
+	virtual bool IsCachedCookedPlatformDataLoaded(const ITargetPlatform* TargetPlatform) override;
+	virtual void PreSave(FObjectPreSaveContext ObjectSaveContext) override;
 #endif //WITH_EDITOR
 
 	virtual void Serialize(FArchive& Ar) override;
diff --git a/Engine/Plugins/Experimental/Mutable/Source/CustomizableObject/Private/MuCO/CustomizableObject.cpp b/Engine/Plugins/Experimental/Mutable/Source/CustomizableObject/Private/MuCO/CustomizableObject.cpp
--- a/Engine/Plugins/Experimental/Mutable/Source/CustomizableObject/Private/MuCO/CustomizableObject.cpp
+++ b/Engine/Plugins/Experimental/Mutable/Source/CustomizableObject/Private/MuCO/CustomizableObject.cpp
@@ -100,12 +100,8 @@ void UCustomizableObject::PreSave(FObjectPreSaveContext ObjectSaveContext)
 	Super::PreSave(ObjectSaveContext);
 
 	// Update the derived child object flag
-	FCustomizableObjectCompilerBase* Compiler = UCustomizableObjectSystem::GetInstance()->GetNewCompiler();
-	if (Compiler)
+	if (TryUpdateIsChildObject())
 	{
-		bIsChildObject = !Compiler->IsRootObject(this);
-		delete Compiler;
-
 		if (bIsChildObject)
 		{
 			GetPackage()->SetPackageFlags(PKG_EditorOnly);
@@ -116,16 +112,13 @@ void UCustomizableObject::PreSave(FObjectPreSaveContext ObjectSaveContext)
 		}
 	}
 
-#if WITH_EDITORONLY_DATA
+#if WITH_EDITOR
 	if (ObjectSaveContext.IsCooking() && !bIsChildObject)
 	{
 		const ITargetPlatform* TargetPlatform = ObjectSaveContext.GetTargetPlatform();
-		if (const FMutableCachedPlatformData* PlatformData = CachedPlatformsData.Find(TargetPlatform->PlatformName()))
-		{
 		// Load cached data before saving
-			FMemoryReaderView MemoryReader(PlatformData->ModelData);
-			LoadCompiledData(MemoryReader, TargetPlatform, true);
-
+		if (TryLoadCompiledCookDataForPlatform(TargetPlatform))
+		{
 			// Create an export object to manage the streamable data
 			BulkData = NewObject<UCustomizableObjectBulk>(this);
 			BulkData->Mark(OBJECTMARK_TagExp);
@@ -145,6 +138,36 @@ void UCustomizableObject::PreSave(FObjectPreSaveContext ObjectSaveContext)
 #endif
 }
 
+bool UCustomizableObject::TryUpdateIsChildObject()
+{
+	FCustomizableObjectCompilerBase* Compiler = UCustomizableObjectSystem::GetInstance()->GetNewCompiler();
+	if (Compiler)
+	{
+		bIsChildObject = !Compiler->IsRootObject(this);
+		delete Compiler;
+		return true;
+	}
+	else
+	{
+		return false;
+	}
+}
+
+#if WITH_EDITOR
+bool UCustomizableObject::TryLoadCompiledCookDataForPlatform(const ITargetPlatform* TargetPlatform)
+{
+	const FMutableCachedPlatformData* PlatformData = CachedPlatformsData.Find(TargetPlatform->PlatformName());
+	if (!PlatformData)
+	{
+		return false;
+	}
+
+	FMemoryReaderView MemoryReader(PlatformData->ModelData);
+	LoadCompiledData(MemoryReader, TargetPlatform, true);
+	return true;
+}
+#endif
+
 #endif // End WITH_EDITOR
 
 
diff --git a/Engine/Plugins/Experimental/Mutable/Source/CustomizableObject/Private/MuCO/CustomizableObjectInstance.cpp b/Engine/Plugins/Experimental/Mutable/Source/CustomizableObject/Private/MuCO/CustomizableObjectInstance.cpp
--- a/Engine/Plugins/Experimental/Mutable/Source/CustomizableObject/Private/MuCO/CustomizableObjectInstance.cpp
+++ b/Engine/Plugins/Experimental/Mutable/Source/CustomizableObject/Private/MuCO/CustomizableObjectInstance.cpp
@@ -38,6 +38,7 @@
 #include "RenderingThread.h"
 #include "SkeletalMergingLibrary.h"
 #include "UnrealMutableImageProvider.h"
+#include "UObject/ObjectSaveContext.h"
 #include "UObject/UObjectIterator.h"
 #include "Widgets/Notifications/SNotificationList.h"
 #include "PhysicsEngine/PhysicsConstraintTemplate.h"
@@ -558,10 +559,63 @@ void UCustomizableObjectInstance::PostLoad()
 	BindPostCompileDelegate(GetCustomizableObject());
 #endif
 
-	Descriptor.ReloadParameters();
+	// Skip the cost of ReloadParameters in the cook commandlet; it will be reloaded during PreSave. For cooked runtime
+	// and editor UI, reload on load because it will not otherwise reload unless the CustomizableObject recompiles.
+	if (!IsRunningCookCommandlet())
+	{
+		Descriptor.ReloadParameters();
+	}
 	PrivateData->InitCustomizableObjectData(GetCustomizableObject());
 }
 
+#if WITH_EDITOR
+void UCustomizableObjectInstance::BeginCacheForCookedPlatformData(const ITargetPlatform* TargetPlatform)
+{
+	UCustomizableObject* CustomizableObject = GetCustomizableObject();
+	if (!CustomizableObject)
+	{
+		return;
+	}
+	CustomizableObject->BeginCacheForCookedPlatformData(TargetPlatform);
+}
+
+bool UCustomizableObjectInstance::IsCachedCookedPlatformDataLoaded(const ITargetPlatform* TargetPlatform)
+{
+	UCustomizableObject* CustomizableObject = GetCustomizableObject();
+	if (!CustomizableObject)
+	{
+		return true;
+	}
+
+	return CustomizableObject->IsCachedCookedPlatformDataLoaded(TargetPlatform);
+}
+
+void UCustomizableObjectInstance::PreSave(FObjectPreSaveContext ObjectSaveContext)
+{
+	Super::PreSave(ObjectSaveContext);
+
+	const ITargetPlatform* TargetPlatform = ObjectSaveContext.GetTargetPlatform();
+	if (!TargetPlatform)
+	{
+		return;
+	}
+
+	UCustomizableObject* CustomizableObject = GetCustomizableObject();
+	if (!CustomizableObject)
+	{
+		return;
+	}
+	if (CustomizableObject->TryUpdateIsChildObject() && CustomizableObject->bIsChildObject)
+	{
+		UE_LOG(LogMutable, Error,
+			TEXT("CO Instance [%s] has an invalid dependency on CO that is a child object: [%s]. The instance will be unusable at runtime and may crash."),
+			*GetPathName(), *CustomizableObject->GetPathName());
+	}
+
+	CustomizableObject->TryLoadCompiledCookDataForPlatform(TargetPlatform);
+	Descriptor.ReloadParameters();
+}
+#endif
 
 FString UCustomizableObjectInstance::GetDesc()
 {
