[Mutable] Fix ReroutPinNode pins being deprecated.

#rnx

diff --git a/Engine/Plugins/Experimental/Mutable/Source/CustomizableObjectEditor/Classes/MuCOE/Nodes/CustomizableObjectNodeReroute.h b/Engine/Plugins/Experimental/Mutable/Source/CustomizableObjectEditor/Classes/MuCOE/Nodes/CustomizableObjectNodeReroute.h
--- a/Engine/Plugins/Experimental/Mutable/Source/CustomizableObjectEditor/Classes/MuCOE/Nodes/CustomizableObjectNodeReroute.h
+++ b/Engine/Plugins/Experimental/Mutable/Source/CustomizableObjectEditor/Classes/MuCOE/Nodes/CustomizableObjectNodeReroute.h
@@ -24,6 +24,7 @@ public:
 	virtual TSharedPtr<SGraphNode> CreateVisualWidget() override;
 	void NotifyPinConnectionListChanged(UEdGraphPin* Pin);
 	virtual void PinConnectionListChanged(UEdGraphPin* Pin) override;
+	virtual UCustomizableObjectNodeRemapPins* CreateRemapPinsDefault() const override;
 
 	// UCustomizableObjectNode interface
 	virtual void AllocateDefaultPins(UCustomizableObjectNodeRemapPins* RemapPins) override;
diff --git a/Engine/Plugins/Experimental/Mutable/Source/CustomizableObjectEditor/Private/MuCOE/Nodes/CustomizableObjectNodeReroute.cpp b/Engine/Plugins/Experimental/Mutable/Source/CustomizableObjectEditor/Private/MuCOE/Nodes/CustomizableObjectNodeReroute.cpp
--- a/Engine/Plugins/Experimental/Mutable/Source/CustomizableObjectEditor/Private/MuCOE/Nodes/CustomizableObjectNodeReroute.cpp
+++ b/Engine/Plugins/Experimental/Mutable/Source/CustomizableObjectEditor/Private/MuCOE/Nodes/CustomizableObjectNodeReroute.cpp
@@ -4,6 +4,7 @@
 
 #include "SGraphNodeKnot.h"
 #include "MuCOE/EdGraphSchema_CustomizableObject.h"
+#include "MuCOE/RemapPins/CustomizableObjectNodeRemapPinsByPosition.h"
 
 #define LOCTEXT_NAMESPACE "CustomizableObjectEditor"
 
@@ -100,6 +101,12 @@ void UCustomizableObjectNodeReroute::PinConnectionListChanged(UEdGraphPin* Pin)
 }
 
 
+UCustomizableObjectNodeRemapPins* UCustomizableObjectNodeReroute::CreateRemapPinsDefault() const
+{
+	return NewObject<UCustomizableObjectNodeRemapPinsByPosition>();
+}
+
+
 UEdGraphPin* UCustomizableObjectNodeReroute::GetInputPin() const
 {
 	return Pins[0];
