[mutable] Added timers and logs to report compilation and avg/combined instance update time.

diff --git a/Engine/Plugins/Experimental/Mutable/Source/MutableValidation/Private/MuV/CustomizableObjectValidationCommandlet.cpp b/Engine/Plugins/Experimental/Mutable/Source/MutableValidation/Private/MuV/CustomizableObjectValidationCommandlet.cpp
--- a/Engine/Plugins/Experimental/Mutable/Source/MutableValidation/Private/MuV/CustomizableObjectValidationCommandlet.cpp
+++ b/Engine/Plugins/Experimental/Mutable/Source/MutableValidation/Private/MuV/CustomizableObjectValidationCommandlet.cpp
@@ -29,6 +29,8 @@ int32 UCustomizableObjectValidationCommandlet::Main(const FString& Params)
 		UE_LOG(LogMutable,Display,TEXT("Instance generation count not specified. Using default value : %u"),InstancesToGenerate);
 	}
 	
+	// TODO: Detect target compilation platform based on argument value (string)
+	
 	// Load the resource
 	UObject* FoundObject = FSoftObjectPath(CustomizableObjectAssetPath).TryLoad();
 	if (!FoundObject)
@@ -59,8 +61,8 @@ int32 UCustomizableObjectValidationCommandlet::Main(const FString& Params)
     	
     	// Request a compiler to be able to locate the root and to compile it
     	const TUniquePtr<FCustomizableObjectCompilerBase> Compiler =
-    		TUniquePtr<FCustomizableObjectCompilerBase>(UCustomizableObjectSystem::GetInstance()->GetNewCompiler());
-
+    		TUniquePtr<FCustomizableObjectCompilerBase>(UCustomizableObjectSystem::GetInstanceChecked()->GetNewCompiler());
+		
 		// Override some configurations that may have been changed by the user
 		FCompilationOptions CompilationOptions = ToTestCustomizableObject->CompileOptions;
 		CompilationOptions.bSilentCompilation = false;
@@ -75,7 +77,14 @@ int32 UCustomizableObjectValidationCommandlet::Main(const FString& Params)
 
     	// Compile the CO with the provided compilation options
     	// Run Sync compilation -> Warning : Potentially long operation -------------
-    	Compiler->Compile(*ToTestCustomizableObject, CompilationOptions, false);
+		const double CompilationStartSeconds = FPlatformTime::Seconds();
+		{
+			Compiler->Compile(*ToTestCustomizableObject, CompilationOptions, false);
+		}
+		const double CompilationEndSeconds = FPlatformTime::Seconds() - CompilationStartSeconds;
+		UE_LOG(LogMutable, Log, TEXT("(double) model_compile_time_ms : %f "), CompilationEndSeconds * 1000);
+		UE_LOG(LogMutable, Display, TEXT("The compilation of the %s model took %f seconds."), *ToTestCustomizableObject->GetName(), CompilationEndSeconds);
+
     	// --------------------------------------------------------------------------
 		
     	// Get the compilation result
@@ -92,7 +101,10 @@ int32 UCustomizableObjectValidationCommandlet::Main(const FString& Params)
 		UE_LOG(LogMutable,Display,TEXT("Customizable Object was compiled succesfully."));
 		
 		// GHet the total size of the streaming data of the model ---------------------------------------------- //
-		if (const TSharedPtr<const mu::Model> MutableModel = ToTestCustomizableObject->GetModel())
+		const TSharedPtr<const mu::Model> MutableModel = ToTestCustomizableObject->GetModel();
+		check (MutableModel);
+
+		// Roms ---------------------- //
 		{
 			const int32 RomCount =  MutableModel->GetRomCount();
 			int64 TotalRomSizeBytes = 0;
@@ -124,7 +136,7 @@ int32 UCustomizableObjectValidationCommandlet::Main(const FString& Params)
 				}
 				else
 				{
-					UE_LOG(LogMutable,Error,TEXT("Failed to generate COI for Customizable Object with name : %s ."),*ToTestCustomizableObject->GetName());
+					UE_LOG(LogMutable,Error,TEXT("Failed to generate COI for the %s CO."),*ToTestCustomizableObject->GetName());
 					bWasInstancesCreationSuccessful = false;
 				}
 			}
@@ -132,9 +144,9 @@ int32 UCustomizableObjectValidationCommandlet::Main(const FString& Params)
 		// ---------------------------------------------------------------------------------------------------------- //
 		
 		// Update the instances generated --------------------------------------------------------------------------- //
+		UE_LOG(LogMutable,Display,TEXT("Updating generated instances..."));
+		const double InstanceUpdateStartSeconds = FPlatformTime::Seconds();
 		{
-			UE_LOG(LogMutable,Display,TEXT("Updating generated instances..."));
-            	
             // Now update the instances one by one
             while (!InstancesToProcess.IsEmpty() || InstanceBeingUpdated)
             {
@@ -169,11 +181,21 @@ int32 UCustomizableObjectValidationCommandlet::Main(const FString& Params)
             	}
             }
 		}	
+		
+		// Notify and log time required by the instances to get updated
+		const double CombinedInstanceUpdateSeconds = FPlatformTime::Seconds() - InstanceUpdateStartSeconds;
+		UE_LOG(LogMutable, Log,TEXT("(double) combined_update_time_ms : %f "), CombinedInstanceUpdateSeconds * 1000);
+
+		check(InstancesToGenerate > 0);
+		const double AverageInstanceUpdateSeconds = CombinedInstanceUpdateSeconds / InstancesToGenerate;
+		UE_LOG(LogMutable, Log,TEXT("(double) avg_update_time_ms : %f "), AverageInstanceUpdateSeconds * 1000);
+
+		UE_LOG(LogMutable,Display,TEXT("Generation of Customizable object instances took %f seconds (%f seconds avg)."), CombinedInstanceUpdateSeconds, AverageInstanceUpdateSeconds);
 		// ---------------------------------------------------------------------------------------------------------- //
 
 		// Compute instance update result
 		const bool bInstancesTestedSuccessfully = !bInstanceFailedUpdate && bWasInstancesCreationSuccessful;
-        if (bInstancesTestedSuccessfully)
+		if (bInstancesTestedSuccessfully)
         {
         	UE_LOG(LogMutable,Display,TEXT("Generation of Customizable object instances was succesfull."));
         }
@@ -196,18 +218,21 @@ int32 UCustomizableObjectValidationCommandlet::Main(const FString& Params)
 
 void UCustomizableObjectValidationCommandlet::OnInstanceUpdate(const FUpdateContext& Result)
 {
+	const FString InstanceName = InstanceBeingUpdated->GetName();
 	const EUpdateResult InstanceUpdateResult = Result.UpdateResult;
 	if (InstanceUpdateResult == EUpdateResult::Success)
 	{
-		UE_LOG(LogMutable,Display,TEXT("Instance %s finished update succesfully."),*InstanceBeingUpdated->GetName());
+		UE_LOG(LogMutable,Display,TEXT("Instance %s finished update succesfully."),*InstanceName);
 	}
 	else
 	{
 		const FString OutputStatus = UEnum::GetValueAsString(Result.UpdateResult);
-		UE_LOG(LogMutable,Error,TEXT("Instance %s finished update with anomalous state : %s."),*InstanceBeingUpdated->GetName(),*OutputStatus);
+		UE_LOG(LogMutable,Error,TEXT("Instance %s finished update with anomalous state : %s."),*InstanceName, *OutputStatus);
 		bInstanceFailedUpdate = true;
 	}
 	
+	// TODO: Wait for the mips before starting the update of another instance. 
+	
 	InstanceUpdateDelegate.Unbind();
 	InstanceBeingUpdated = nullptr;
 }
