CustomizableObject Cook-Indeterminism: Fix the indeterminism caused by compile-created objects changing their names when LoadCompiledData is called multiple times per cook.
#rnx

diff --git a/Engine/Plugins/Experimental/Mutable/Source/CustomizableObject/Private/MuCO/CustomizableObject.cpp b/Engine/Plugins/Experimental/Mutable/Source/CustomizableObject/Private/MuCO/CustomizableObject.cpp
--- a/Engine/Plugins/Experimental/Mutable/Source/CustomizableObject/Private/MuCO/CustomizableObject.cpp
+++ b/Engine/Plugins/Experimental/Mutable/Source/CustomizableObject/Private/MuCO/CustomizableObject.cpp
@@ -43,6 +43,29 @@
 
 DEFINE_LOG_CATEGORY(LogMutable);
 
+#if WITH_EDITORONLY_DATA
+
+namespace UE::Mutable::Private
+{
+
+template <typename T>
+T* MoveOldObjectAndCreateNew(UClass* Class, UObject* InOuter)
+{
+	FName ObjectFName = Class->GetFName();
+	FString ObjectNameStr = ObjectFName.ToString();
+	UObject* Existing = FindObject<UAssetUserData>(InOuter, *ObjectNameStr);
+	if (Existing)
+	{
+		// Move the old object out of the way
+		Existing->Rename(nullptr /* Rename will pick a free name*/, GetTransientPackage(), REN_DontCreateRedirectors);
+	}
+	return NewObject<T>(InOuter, Class, *ObjectNameStr);
+}
+
+}
+
+#endif
+
 //-------------------------------------------------------------------------------------------------
 
 UCustomizableObject::UCustomizableObject()
@@ -120,7 +143,8 @@ void UCustomizableObject::PreSave(FObjectPreSaveContext ObjectSaveContext)
 		if (TryLoadCompiledCookDataForPlatform(TargetPlatform))
 		{
 			// Create an export object to manage the streamable data
-			BulkData = NewObject<UCustomizableObjectBulk>(this);
+			BulkData = UE::Mutable::Private::MoveOldObjectAndCreateNew<UCustomizableObjectBulk>(
+				UCustomizableObjectBulk::StaticClass(), this);
 			BulkData->Mark(OBJECTMARK_TagExp);
 
 			// Split streamable data into smaller chunks and fix up the CO HashToStreamableBlock's FileIndex and Offset
@@ -2081,7 +2105,7 @@ void FMutableRefAssetUserData::InitResources(UCustomizableObject* InOuter)
 	UClass* AssetUserDataClass = FindObject<UClass>(nullptr, *ClassPath);
 	if (AssetUserDataClass)
 	{
-		AssetUserData = NewObject<UAssetUserData>(InOuter, AssetUserDataClass);
+		AssetUserData = UE::Mutable::Private::MoveOldObjectAndCreateNew<UAssetUserData>(AssetUserDataClass, InOuter);
 		if (AssetUserData)
 		{
 			FMemoryReaderView MemoryReader(Bytes);
