[mutable] Redesign code generator to not rely on the visitor pattern and reducing RTTI usage.
Change the handling of image sizes at code generation to simplify the process.
Fix some code optimizations.
Cleanup some old code.
#rnx

diff --git a/Engine/Plugins/Experimental/Mutable/Source/CustomizableObject/Private/MuCO/CustomizableObjectInstance.cpp b/Engine/Plugins/Experimental/Mutable/Source/CustomizableObject/Private/MuCO/CustomizableObjectInstance.cpp
--- a/Engine/Plugins/Experimental/Mutable/Source/CustomizableObject/Private/MuCO/CustomizableObjectInstance.cpp
+++ b/Engine/Plugins/Experimental/Mutable/Source/CustomizableObject/Private/MuCO/CustomizableObjectInstance.cpp
@@ -4217,7 +4217,7 @@ void UCustomizableInstancePrivateData::BuildOrCopyClothingData(const TSharedRef<
 
 			const TSet<int32>& SrcSelfCollisionVertexSet = SrcLodData.PhysicalMeshData.SelfCollisionVertexSet;
 			TSet<int32>& NewSelfCollisionVertexSet = NewLodData.PhysicalMeshData.SelfCollisionVertexSet;
-			TrimAndRemapVertexSet(NewSelfCollisionVertexSet, SrcSelfCollisionVertexSet);;
+			TrimAndRemapVertexSet(NewSelfCollisionVertexSet, SrcSelfCollisionVertexSet);
 						
 			{
 				MUTABLE_CPUPROFILER_SCOPE(BuildClothTetherData)
diff --git a/Engine/Plugins/Experimental/Mutable/Source/CustomizableObject/Private/MuCO/CustomizableObjectMipDataProvider.cpp b/Engine/Plugins/Experimental/Mutable/Source/CustomizableObject/Private/MuCO/CustomizableObjectMipDataProvider.cpp
--- a/Engine/Plugins/Experimental/Mutable/Source/CustomizableObject/Private/MuCO/CustomizableObjectMipDataProvider.cpp
+++ b/Engine/Plugins/Experimental/Mutable/Source/CustomizableObject/Private/MuCO/CustomizableObjectMipDataProvider.cpp
@@ -427,8 +427,16 @@ bool FMutableTextureMipDataProvider::PollMips(const FTextureUpdateSyncOptions& S
 				int32 MipDataSize = Image->GetLODDataSize(MipIndex);
 
 				// Check Mip DataSize for consistency, but skip if 0 because it's optional and might be zero in cooked mips
-				check(Level.DataSize == 0 || MipDataSize == Level.DataSize);
-				FMemory::Memcpy(Dest, Image->GetMipData(MipIndex), MipDataSize);
+				bool bCorrectDataSize = (Level.DataSize == 0 || MipDataSize == Level.DataSize);
+				if (bCorrectDataSize)
+				{
+					FMemory::Memcpy(Dest, Image->GetMipData(MipIndex), MipDataSize);
+				}
+				else
+				{
+					UE_LOG(LogMutable, Warning, TEXT("Mip data has incorrect size."));
+					FMemory::Memzero(Dest, Level.DataSize);
+				}
 			}
 			else
 			{
diff --git a/Engine/Plugins/Experimental/Mutable/Source/CustomizableObject/Private/MuCO/EditorImageProvider.cpp b/Engine/Plugins/Experimental/Mutable/Source/CustomizableObject/Private/MuCO/EditorImageProvider.cpp
--- a/Engine/Plugins/Experimental/Mutable/Source/CustomizableObject/Private/MuCO/EditorImageProvider.cpp
+++ b/Engine/Plugins/Experimental/Mutable/Source/CustomizableObject/Private/MuCO/EditorImageProvider.cpp
@@ -17,6 +17,6 @@ UCustomizableSystemImageProvider::ValueType UEditorImageProvider::HasTexturePara
 
 UTexture2D* UEditorImageProvider::GetTextureParameterValue(const FName& ID)
 {
-	return Cast<UTexture2D>(FSoftObjectPath(ID.ToString()).TryLoad());;
+	return Cast<UTexture2D>(FSoftObjectPath(ID.ToString()).TryLoad());
 }
 
diff --git a/Engine/Plugins/Experimental/Mutable/Source/CustomizableObjectEditor/Classes/MuCOE/Nodes/CustomizableObjectNodeMesh.h b/Engine/Plugins/Experimental/Mutable/Source/CustomizableObjectEditor/Classes/MuCOE/Nodes/CustomizableObjectNodeMesh.h
--- a/Engine/Plugins/Experimental/Mutable/Source/CustomizableObjectEditor/Classes/MuCOE/Nodes/CustomizableObjectNodeMesh.h
+++ b/Engine/Plugins/Experimental/Mutable/Source/CustomizableObjectEditor/Classes/MuCOE/Nodes/CustomizableObjectNodeMesh.h
@@ -27,7 +27,7 @@ public:
 	virtual void GetUVChannelForPin(const UEdGraphPin* Pin, TArray<FVector2f>& OutSegments, int32 UVIndex) const PURE_VIRTUAL(UCustomizableObjectNodeMesh::GetUVChannelForPin, );
 
 	/** Returns the Unreal mesh (e.g., USkeletalMesh, UStaticMesh...). */
-	virtual UObject* GetMesh() const PURE_VIRTUAL(UCustomizableObjectNodeMesh::GetMesh, return {};);;
+	virtual UObject* GetMesh() const PURE_VIRTUAL(UCustomizableObjectNodeMesh::GetMesh, return {};);
 
 	/** Returns the output Mesh pin associated to the given LODIndex and SectionIndex. Override. */
 	virtual UEdGraphPin* GetMeshPin(int32 LOD, int32 SectionIndex) const PURE_VIRTUAL(UCustomizableObjectNodeMesh::GetMeshPin, return {};);
diff --git a/Engine/Plugins/Experimental/Mutable/Source/CustomizableObjectEditor/Private/MuCOE/Nodes/CustomizableObjectNodeObjectGroup.cpp b/Engine/Plugins/Experimental/Mutable/Source/CustomizableObjectEditor/Private/MuCOE/Nodes/CustomizableObjectNodeObjectGroup.cpp
--- a/Engine/Plugins/Experimental/Mutable/Source/CustomizableObjectEditor/Private/MuCOE/Nodes/CustomizableObjectNodeObjectGroup.cpp
+++ b/Engine/Plugins/Experimental/Mutable/Source/CustomizableObjectEditor/Private/MuCOE/Nodes/CustomizableObjectNodeObjectGroup.cpp
@@ -88,7 +88,7 @@ FText UCustomizableObjectNodeObjectGroup::GetNodeTitle(ENodeTitleType::Type Titl
 	if (TitleType == ENodeTitleType::ListView ||
 		GroupName.IsEmpty())
 	{
-		return LOCTEXT("Object_Group", "Object Group");;
+		return LOCTEXT("Object_Group", "Object Group");
 	}
 	else
 	{
diff --git a/Engine/Plugins/Experimental/Mutable/Source/CustomizableObjectEditor/Private/MuCOE/SMutableCodeViewer.cpp b/Engine/Plugins/Experimental/Mutable/Source/CustomizableObjectEditor/Private/MuCOE/SMutableCodeViewer.cpp
--- a/Engine/Plugins/Experimental/Mutable/Source/CustomizableObjectEditor/Private/MuCOE/SMutableCodeViewer.cpp
+++ b/Engine/Plugins/Experimental/Mutable/Source/CustomizableObjectEditor/Private/MuCOE/SMutableCodeViewer.cpp
@@ -1228,7 +1228,7 @@ void SMutableCodeViewer::CacheOperationTypesPresentOnModel()
 			const mu::OP_TYPE OperationType = Program.GetOpType(ProgramAddressesIndex);
 			
 			// Increase the counter for this operation type
-			const uint16 TypeAsInteger = StaticCast<uint16>(OperationType);;
+			const uint16 TypeAsInteger = StaticCast<uint16>(OperationType);
 			ModelOperationTypes[TypeAsInteger].Value++;
 		}
 	}
@@ -2973,4 +2973,216 @@ FReply SMutableCodeViewer::OnDrop(const FGeometry& MyGeometry, const FDragDropEv
 }
 
 
+FMutableCodeTreeElement::FMutableCodeTreeElement(int32 InIndexOnTree, const int32& InMutableStateIndex, const TSharedPtr<mu::Model, ESPMode::ThreadSafe>& InModel, mu::OP::ADDRESS InOperation, const FString& InCaption, const FSlateColor InLabelColor, const TSharedPtr<FMutableCodeTreeElement>* InDuplicatedOf)
+{
+	MutableModel = InModel;
+	MutableOperation = InOperation;
+	Caption = InCaption;
+	LabelColor = InLabelColor;
+	IndexOnTree = InIndexOnTree;
+	if (InDuplicatedOf)
+	{
+		DuplicatedOf = *InDuplicatedOf;
+	}
+
+	// Generate the label to be used to display this operation in the operation tree
+	GenerateLabelText();
+
+	// Process the data that can be extracted from the current state
+	SetElementCurrentState(InMutableStateIndex);
+}
+
+
+void FMutableCodeTreeElement::SetElementCurrentState(const int32& InStateIndex)
+{
+	// Skip operation if state is the same
+	if (InStateIndex == CurrentMutableStateIndex)
+	{
+		return;
+	}
+
+	// Check for an out of bounds value
+	check(MutableModel);
+	mu::FProgram& MutableProgram = MutableModel->GetPrivate()->m_program;
+	check(InStateIndex >= 0 && InStateIndex < MutableProgram.m_states.Num());
+
+	CurrentMutableStateIndex = InStateIndex;
+	const mu::FProgram::FState& CurrentState = MutableProgram.m_states[CurrentMutableStateIndex];
+
+	// Check if it is a dynamic resource
+	for (auto& DynamicResource : CurrentState.m_dynamicResources)
+	{
+		// If the operation gets located then mark it as dynamic resource
+		if (DynamicResource.Key == MutableOperation)
+		{
+			bIsDynamicResource = true;
+			break;
+		}
+	}
+	// Early exit: A dynamic resource can not be at the same time a state constant
+	if (bIsDynamicResource)
+	{
+		return;
+	}
+
+	// Check if it is a state constant
+	bIsStateConstant = CurrentState.m_updateCache.Contains(MutableOperation);
+}
+
+
+void FMutableCodeTreeElement::GenerateLabelText()
+{
+	const mu::FProgram& Program = MutableModel->GetPrivate()->m_program;
+	const mu::OP_TYPE Type = Program.GetOpType(MutableOperation);
+	FString OpName = mu::s_opNames[static_cast<int32>(Type)];
+	OpName.TrimEndInline();
+
+	// See if the operation type accepts additional information in the label
+	switch (Type)
+	{
+	case mu::OP_TYPE::BO_PARAMETER:
+	case mu::OP_TYPE::NU_PARAMETER:
+	case mu::OP_TYPE::SC_PARAMETER:
+	case mu::OP_TYPE::CO_PARAMETER:
+	case mu::OP_TYPE::PR_PARAMETER:
+	case mu::OP_TYPE::IM_PARAMETER:
+	case mu::OP_TYPE::ST_PARAMETER:
+	{
+		mu::OP::ParameterArgs Args = Program.GetOpArgs<mu::OP::ParameterArgs>(MutableOperation);
+		OpName += TEXT(" ");
+		OpName += Program.m_parameters[int32(Args.variable)].m_name;
+		break;
+	}
+
+	case mu::OP_TYPE::IM_SWIZZLE:
+	{
+		mu::OP::ImageSwizzleArgs Args = Program.GetOpArgs<mu::OP::ImageSwizzleArgs>(MutableOperation);
+		OpName += TEXT(" ");
+		OpName += StringCast<TCHAR>(mu::TypeInfo::s_imageFormatName[int32(Args.format)]).Get();
+		break;
+	}
+
+	case mu::OP_TYPE::IM_PIXELFORMAT:
+	{
+		mu::OP::ImagePixelFormatArgs Args = Program.GetOpArgs<mu::OP::ImagePixelFormatArgs>(MutableOperation);
+		OpName += TEXT(" ");
+		OpName += StringCast<TCHAR>(mu::TypeInfo::s_imageFormatName[int32(Args.format)]).Get();
+		OpName += TEXT(" or ");
+		OpName += StringCast<TCHAR>(mu::TypeInfo::s_imageFormatName[int32(Args.formatIfAlpha)]).Get();
+		break;
+	}
+
+	case mu::OP_TYPE::IM_MIPMAP:
+	{
+		mu::OP::ImageMipmapArgs Args = Program.GetOpArgs<mu::OP::ImageMipmapArgs>(MutableOperation);
+		OpName += FString::Printf(TEXT(" levels: %d-%d tail: %d"), Args.levels, Args.blockLevels, int32(Args.onlyTail));
+		break;
+	}
+
+	case mu::OP_TYPE::IM_RESIZE:
+	{
+		mu::OP::ImageResizeArgs Args = Program.GetOpArgs<mu::OP::ImageResizeArgs>(MutableOperation);
+		OpName += FString::Printf(TEXT(" %d x %d"), int32(Args.size[0]), int32(Args.size[1]));
+		break;
+	}
+
+	case mu::OP_TYPE::IM_RESIZEREL:
+	{
+		mu::OP::ImageResizeRelArgs Args = Program.GetOpArgs<mu::OP::ImageResizeRelArgs>(MutableOperation);
+		OpName += FString::Printf(TEXT(" %.3f x %.3f"), Args.factor[0], Args.factor[1]);
+		break;
+	}
+
+	case mu::OP_TYPE::IM_MULTILAYER:
+	{
+		mu::OP::ImageMultiLayerArgs Args = Program.GetOpArgs<mu::OP::ImageMultiLayerArgs>(MutableOperation);
+		OpName += TEXT(" rgb: ");
+		OpName += mu::TypeInfo::s_blendModeName[int32(Args.blendType)];
+		OpName += TEXT(", a: ");
+		OpName += mu::TypeInfo::s_blendModeName[int32(Args.blendTypeAlpha)];
+		OpName += FString::Printf(TEXT(" a from %d "), Args.BlendAlphaSourceChannel);
+		OpName += FString::Printf(TEXT(" range-id: %d"), Args.rangeId);
+		OpName += FString::Printf(TEXT(" mask-from-alpha: %d"), int32(Args.bUseMaskFromBlended));
+		break;
+	}
+
+	case mu::OP_TYPE::IM_LAYER:
+	{
+		mu::OP::ImageLayerArgs Args = Program.GetOpArgs<mu::OP::ImageLayerArgs>(MutableOperation);
+		OpName += TEXT(" rgb: ");
+		OpName += mu::TypeInfo::s_blendModeName[int32(Args.blendType)];
+		OpName += TEXT(", a: ");
+		OpName += mu::TypeInfo::s_blendModeName[int32(Args.blendTypeAlpha)];
+		OpName += FString::Printf(TEXT(" a from %d "), Args.BlendAlphaSourceChannel);
+		OpName += FString::Printf(TEXT(" flags %d"), Args.flags);
+		break;
+	}
+
+	case mu::OP_TYPE::IM_LAYERCOLOUR:
+	{
+		mu::OP::ImageLayerColourArgs Args = Program.GetOpArgs<mu::OP::ImageLayerColourArgs>(MutableOperation);
+		OpName += TEXT(" rgb: ");
+		OpName += mu::TypeInfo::s_blendModeName[int32(Args.blendType)];
+		OpName += TEXT(" a: ");
+		OpName += mu::TypeInfo::s_blendModeName[int32(Args.blendTypeAlpha)];
+		OpName += TEXT(" a from ");
+		OpName += FString::Printf(TEXT(" a from %d "), Args.BlendAlphaSourceChannel);
+		OpName += FString::Printf(TEXT(" flags %d"), Args.flags);
+		break;
+	}
+
+	case mu::OP_TYPE::IM_PLAINCOLOUR:
+	{
+		mu::OP::ImagePlainColourArgs Args = Program.GetOpArgs<mu::OP::ImagePlainColourArgs>(MutableOperation);
+		OpName += TEXT(" format: ");
+		OpName += StringCast<TCHAR>(mu::TypeInfo::s_imageFormatName[int32(Args.format)]).Get();
+		OpName += FString::Printf(TEXT(" size %d x %d"), Args.size[0], Args.size[1]);
+		OpName += FString::Printf(TEXT(" mips %d"), Args.LODs);
+		break;
+	}
+
+	case mu::OP_TYPE::IN_ADDIMAGE:
+	{
+		mu::OP::InstanceAddArgs Args = Program.GetOpArgs<mu::OP::InstanceAddArgs>(MutableOperation);
+		if (Program.m_constantStrings.IsValidIndex(Args.name))
+		{
+			OpName += TEXT(" name: ");
+			OpName += Program.m_constantStrings[Args.name];
+		}
+		break;
+	}
+
+	default:
+		break;
+	}
+
+	// Prepare the text shown on the UI side of the operation tree
+	MainLabel = FString::Printf(TEXT("%s %d : %s"), *Caption, int32(MutableOperation), *OpName);
+
+	// DEBUG : 
+	// FString IndexOnTree = FString::FromInt(IndexOnTree);
+	// IndexOnTree.Append(TEXT("- "));
+	// MainLabel.InsertAt(0,IndexOnTree);
+
+	// DEBUG : 
+	// FString RowStateIndex = FString::FromInt(GetStateIndex());
+	// RowStateIndex.Append(TEXT(" st "));
+	// MainLabel.InsertAt(0,RowStateIndex);
+
+	if (DuplicatedOf)
+	{
+		MainLabel.Append(TEXT(" (duplicated)"));
+	}
+}
+
+
+int32 FMutableCodeTreeElement::GetStateIndex() const
+{
+	return CurrentMutableStateIndex;
+}
+
+
+#undef LOCTEXT_NAMESPACE
+
+
 #undef LOCTEXT_NAMESPACE
diff --git a/Engine/Plugins/Experimental/Mutable/Source/CustomizableObjectEditor/Private/MuCOE/SMutableCodeViewer.h b/Engine/Plugins/Experimental/Mutable/Source/CustomizableObjectEditor/Private/MuCOE/SMutableCodeViewer.h
--- a/Engine/Plugins/Experimental/Mutable/Source/CustomizableObjectEditor/Private/MuCOE/SMutableCodeViewer.h
+++ b/Engine/Plugins/Experimental/Mutable/Source/CustomizableObjectEditor/Private/MuCOE/SMutableCodeViewer.h
@@ -846,210 +846,13 @@ public:
 class FMutableCodeTreeElement : public TSharedFromThis<FMutableCodeTreeElement>
 {
 public:
-	FMutableCodeTreeElement(int32 InIndexOnTree  , const int32& InMutableStateIndex ,const TSharedPtr<mu::Model, ESPMode::ThreadSafe>& InModel, mu::OP::ADDRESS InOperation, const FString& InCaption,const FSlateColor InLabelColor, const TSharedPtr<FMutableCodeTreeElement>* InDuplicatedOf = nullptr)
-	{
-		MutableModel = InModel;
-		MutableOperation = InOperation;
-		Caption = InCaption;
-		LabelColor = InLabelColor;
-		IndexOnTree = InIndexOnTree;
-		if (InDuplicatedOf)
-		{
-			DuplicatedOf = *InDuplicatedOf;
-		}
-
-		// Generate the label to be used to display this operation in the operation tree
-		GenerateLabelText();
-		
-		// Process the data that can be extracted from the current state
-		SetElementCurrentState(InMutableStateIndex);
-	}
+	FMutableCodeTreeElement(int32 InIndexOnTree, const int32& InMutableStateIndex, const TSharedPtr<mu::Model, ESPMode::ThreadSafe>& InModel, mu::OP::ADDRESS InOperation, const FString& InCaption, const FSlateColor InLabelColor, const TSharedPtr<FMutableCodeTreeElement>* InDuplicatedOf = nullptr);
 	
-	void SetElementCurrentState(const int32& InStateIndex)
-	{
-		// Skip operation if state is the same
-		if (InStateIndex == CurrentMutableStateIndex)
-		{
-			return;
-		}
-
-		// Check for an out of bounds value
-		check(MutableModel)
-		mu::FProgram& MutableProgram = MutableModel->GetPrivate()->m_program;
-		check (InStateIndex >= 0 && InStateIndex < MutableProgram.m_states.Num());
-		
-		CurrentMutableStateIndex = InStateIndex;
-		const mu::FProgram::FState& CurrentState = MutableProgram.m_states[CurrentMutableStateIndex];
-		
-		// Check if it is a dynamic resource
-		for (auto& DynamicResource : CurrentState.m_dynamicResources)
-		{
-			// If the operation gets located then mark it as dynamic resource
-			if (DynamicResource.Key == MutableOperation)
-			{
-				bIsDynamicResource = true;
-				break;
-			}
-		}
-		// Early exit: A dynamic resource can not be at the same time a state constant
-		if (bIsDynamicResource)
-		{
-			return;
-		}
-		
-		// Check if it is a state constant
-		bIsStateConstant = CurrentState.m_updateCache.Contains(MutableOperation);
-	}
-
-	void GenerateLabelText()
-	{
-		const mu::FProgram& Program = MutableModel->GetPrivate()->m_program;
-		const mu::OP_TYPE Type = Program.GetOpType(MutableOperation);
-		FString OpName = mu::s_opNames[static_cast<int32>(Type)];
-		OpName.TrimEndInline();
-
-		// See if the operation type accepts additional information in the label
-		switch ( Type )
-		{
-		case mu::OP_TYPE::BO_PARAMETER:
-		case mu::OP_TYPE::NU_PARAMETER:
-		case mu::OP_TYPE::SC_PARAMETER:
-		case mu::OP_TYPE::CO_PARAMETER:
-		case mu::OP_TYPE::PR_PARAMETER:
-		case mu::OP_TYPE::IM_PARAMETER:
-		case mu::OP_TYPE::ST_PARAMETER:
-		{
-			mu::OP::ParameterArgs Args = Program.GetOpArgs<mu::OP::ParameterArgs>(MutableOperation);
-			OpName += TEXT(" ");
-			OpName += Program.m_parameters[int32(Args.variable)].m_name;
-			break;
-		}
-
-		case mu::OP_TYPE::IM_SWIZZLE:
-		{
-			mu::OP::ImageSwizzleArgs Args = Program.GetOpArgs<mu::OP::ImageSwizzleArgs>(MutableOperation);
-			OpName += TEXT(" ");
-			OpName += StringCast<TCHAR>(mu::TypeInfo::s_imageFormatName[int32(Args.format)]).Get();
-			break;
-		}
-
-		case mu::OP_TYPE::IM_PIXELFORMAT:
-		{
-			mu::OP::ImagePixelFormatArgs Args = Program.GetOpArgs<mu::OP::ImagePixelFormatArgs>(MutableOperation);
-			OpName += TEXT(" ");
-			OpName += StringCast<TCHAR>(mu::TypeInfo::s_imageFormatName[int32(Args.format)]).Get();
-			OpName += TEXT(" or ");
-			OpName += StringCast<TCHAR>(mu::TypeInfo::s_imageFormatName[int32(Args.formatIfAlpha)]).Get();
-			break;
-		}
-
-		case mu::OP_TYPE::IM_MIPMAP:
-		{
-			mu::OP::ImageMipmapArgs Args = Program.GetOpArgs<mu::OP::ImageMipmapArgs>(MutableOperation);
-			OpName += FString::Printf(TEXT(" levels: %d-%d tail: %d"), Args.levels, Args.blockLevels, int32(Args.onlyTail));
-			break;
-		}
-
-		case mu::OP_TYPE::IM_RESIZE:
-		{
-			mu::OP::ImageResizeArgs Args = Program.GetOpArgs<mu::OP::ImageResizeArgs>(MutableOperation);
-			OpName += FString::Printf(TEXT(" %d x %d"), int32(Args.size[0]), int32(Args.size[1]));
-			break;
-		}
-
-		case mu::OP_TYPE::IM_RESIZEREL:
-		{
-			mu::OP::ImageResizeRelArgs Args = Program.GetOpArgs<mu::OP::ImageResizeRelArgs>(MutableOperation);
-			OpName += FString::Printf(TEXT(" %.3f x %.3f"), Args.factor[0], Args.factor[1]);
-			break;
-		}
-
-		case mu::OP_TYPE::IM_MULTILAYER:
-		{
-			mu::OP::ImageMultiLayerArgs Args = Program.GetOpArgs<mu::OP::ImageMultiLayerArgs>(MutableOperation);
-			OpName += TEXT(" rgb: ");
-			OpName += mu::TypeInfo::s_blendModeName[int32(Args.blendType)];
-			OpName += TEXT(", a: ");
-			OpName += mu::TypeInfo::s_blendModeName[int32(Args.blendTypeAlpha)];
-			OpName += FString::Printf(TEXT(" a from %d "), Args.BlendAlphaSourceChannel);
-			OpName += FString::Printf(TEXT(" range-id: %d"), Args.rangeId);
-			OpName += FString::Printf(TEXT(" mask-from-alpha: %d"), int32(Args.bUseMaskFromBlended));
-			break;
-		}
-
-		case mu::OP_TYPE::IM_LAYER:
-		{
-			mu::OP::ImageLayerArgs Args = Program.GetOpArgs<mu::OP::ImageLayerArgs>(MutableOperation);
-			OpName += TEXT(" rgb: ");
-			OpName += mu::TypeInfo::s_blendModeName[int32(Args.blendType)];
-			OpName += TEXT(", a: ");
-			OpName += mu::TypeInfo::s_blendModeName[int32(Args.blendTypeAlpha)];
-			OpName += FString::Printf(TEXT(" a from %d "), Args.BlendAlphaSourceChannel);
-			OpName += FString::Printf(TEXT(" flags %d"),Args.flags);
-			break;
-		}
-
-		case mu::OP_TYPE::IM_LAYERCOLOUR:
-		{
-			mu::OP::ImageLayerColourArgs Args = Program.GetOpArgs<mu::OP::ImageLayerColourArgs>(MutableOperation);
-			OpName += TEXT(" rgb: ");
-			OpName += mu::TypeInfo::s_blendModeName[int32(Args.blendType)];
-			OpName += TEXT(" a: ");
-			OpName += mu::TypeInfo::s_blendModeName[int32(Args.blendTypeAlpha)];
-			OpName += TEXT(" a from ");
-			OpName += FString::Printf(TEXT(" a from %d "), Args.BlendAlphaSourceChannel);
-			OpName += FString::Printf(TEXT(" flags %d"), Args.flags);
-			break;
-		}
-
-		case mu::OP_TYPE::IM_PLAINCOLOUR:
-		{
-			mu::OP::ImagePlainColourArgs Args = Program.GetOpArgs<mu::OP::ImagePlainColourArgs>(MutableOperation);
-			OpName += TEXT(" format: ");
-			OpName += StringCast<TCHAR>(mu::TypeInfo::s_imageFormatName[int32(Args.format)]).Get();
-			OpName += FString::Printf(TEXT(" size %d x %d"), Args.size[0], Args.size[1]);
-			OpName += FString::Printf(TEXT(" mips %d"), Args.LODs);
-			break;
-		}
-
-		case mu::OP_TYPE::IN_ADDIMAGE:
-		{
-			mu::OP::InstanceAddArgs Args = Program.GetOpArgs<mu::OP::InstanceAddArgs>(MutableOperation);
-			if (Program.m_constantStrings.IsValidIndex(Args.name))
-			{
-				OpName += TEXT(" name: ");
-				OpName += Program.m_constantStrings[Args.name];
-			}
- 			break;
-		}
-
-		default:
-			break;
-		}
+	void SetElementCurrentState(const int32& InStateIndex);
 
-		// Prepare the text shown on the UI side of the operation tree
-		MainLabel = FString::Printf(TEXT("%s %d : %s"), *Caption, int32(MutableOperation), *OpName);
-
-		// DEBUG : 
-		// FString IndexOnTree = FString::FromInt(IndexOnTree);
-		// IndexOnTree.Append(TEXT("- "));
-		// MainLabel.InsertAt(0,IndexOnTree);
-
-		// DEBUG : 
-		// FString RowStateIndex = FString::FromInt(GetStateIndex());
-		// RowStateIndex.Append(TEXT(" st "));
-		// MainLabel.InsertAt(0,RowStateIndex);
-		
-		if (DuplicatedOf)
-		{
-			MainLabel.Append(TEXT(" (duplicated)"));
-		}
-	}
+	void GenerateLabelText();
 	
-	int32 GetStateIndex() const
-	{
-		return CurrentMutableStateIndex;
-	}
+	int32 GetStateIndex() const;
 	
 public:
 
diff --git a/Engine/Plugins/Experimental/Mutable/Source/CustomizableObjectEditor/Private/MuCOE/SMutableObjectViewer.cpp b/Engine/Plugins/Experimental/Mutable/Source/CustomizableObjectEditor/Private/MuCOE/SMutableObjectViewer.cpp
--- a/Engine/Plugins/Experimental/Mutable/Source/CustomizableObjectEditor/Private/MuCOE/SMutableObjectViewer.cpp
+++ b/Engine/Plugins/Experimental/Mutable/Source/CustomizableObjectEditor/Private/MuCOE/SMutableObjectViewer.cpp
@@ -140,7 +140,7 @@ void SMutableObjectViewer::Construct(const FArguments& InArgs, UCustomizableObje
 
 	// Initialize the debugger compile options
 	CompileOptions.TextureCompression = ECustomizableObjectTextureCompression::Fast;
-	CompileOptions.OptimizationLevel = 3;
+	CompileOptions.OptimizationLevel = 2;
 	{
 		ITargetPlatformManagerModule* TPM = GetTargetPlatformManager();
 		const TArray<ITargetPlatform*>& Platforms = TPM->GetActiveTargetPlatforms();
diff --git a/Engine/Plugins/Experimental/Mutable/Source/CustomizableObjectPopulationEditor/Private/MuCOPE/AssetDefinition_CustomizableObjectPopulation.h b/Engine/Plugins/Experimental/Mutable/Source/CustomizableObjectPopulationEditor/Private/MuCOPE/AssetDefinition_CustomizableObjectPopulation.h
--- a/Engine/Plugins/Experimental/Mutable/Source/CustomizableObjectPopulationEditor/Private/MuCOPE/AssetDefinition_CustomizableObjectPopulation.h
+++ b/Engine/Plugins/Experimental/Mutable/Source/CustomizableObjectPopulationEditor/Private/MuCOPE/AssetDefinition_CustomizableObjectPopulation.h
@@ -24,5 +24,5 @@ public:
 	
 	virtual TConstArrayView<FAssetCategoryPath> GetAssetCategories() const override;
 	
-	virtual EAssetCommandResult OpenAssets(const FAssetOpenArgs& OpenArgs) const override;;
+	virtual EAssetCommandResult OpenAssets(const FAssetOpenArgs& OpenArgs) const override;
 };
diff --git a/Engine/Plugins/Experimental/Mutable/Source/MutableRuntime/Internal/MuR/ModelPrivate.h b/Engine/Plugins/Experimental/Mutable/Source/MutableRuntime/Internal/MuR/ModelPrivate.h
--- a/Engine/Plugins/Experimental/Mutable/Source/MutableRuntime/Internal/MuR/ModelPrivate.h
+++ b/Engine/Plugins/Experimental/Mutable/Source/MutableRuntime/Internal/MuR/ModelPrivate.h
@@ -138,8 +138,7 @@ namespace mu
 
             //! List of root instructions for the dynamic resources that depend on the runtime
             //! parameters of this state, with a mask of relevant runtime parameters.
-            //! The mask has a bit on for every runtime parameter in the m_runtimeParameters
-            //! vector.
+            //! The mask has a bit on for every runtime parameter in the m_runtimeParameters array.
 			//! The uint64 is linked to MUTABLE_MAX_RUNTIME_PARAMETERS_PER_STATE
 			TArray< TPair<OP::ADDRESS,uint64> > m_dynamicResources;
 
@@ -669,7 +668,7 @@ namespace mu
 
 
     //!
-    class MUTABLERUNTIME_API Model::Private : public Base
+    class MUTABLERUNTIME_API Model::Private
     {
     public:
 
@@ -699,7 +698,7 @@ namespace mu
 
 
     //!
-    class MUTABLERUNTIME_API ModelParametersGenerator::Private : public Base
+    class MUTABLERUNTIME_API ModelParametersGenerator::Private
     {
     public:
 
diff --git a/Engine/Plugins/Experimental/Mutable/Source/MutableRuntime/Internal/MuR/ParametersPrivate.h b/Engine/Plugins/Experimental/Mutable/Source/MutableRuntime/Internal/MuR/ParametersPrivate.h
--- a/Engine/Plugins/Experimental/Mutable/Source/MutableRuntime/Internal/MuR/ParametersPrivate.h
+++ b/Engine/Plugins/Experimental/Mutable/Source/MutableRuntime/Internal/MuR/ParametersPrivate.h
@@ -533,7 +533,7 @@ namespace Private
     };
 
 
-    class MUTABLERUNTIME_API RangeIndex::Private : public Base
+    class MUTABLERUNTIME_API RangeIndex::Private
     {
     public:
 
@@ -548,7 +548,7 @@ namespace Private
     };
 
 
-    class MUTABLERUNTIME_API Parameters::Private : public Base
+    class MUTABLERUNTIME_API Parameters::Private
     {
     public:
 
diff --git a/Engine/Plugins/Experimental/Mutable/Source/MutableRuntime/Internal/MuR/SerialisationPrivate.h b/Engine/Plugins/Experimental/Mutable/Source/MutableRuntime/Internal/MuR/SerialisationPrivate.h
--- a/Engine/Plugins/Experimental/Mutable/Source/MutableRuntime/Internal/MuR/SerialisationPrivate.h
+++ b/Engine/Plugins/Experimental/Mutable/Source/MutableRuntime/Internal/MuR/SerialisationPrivate.h
@@ -87,7 +87,7 @@ namespace mu
 
 
     //---------------------------------------------------------------------------------------------
-    class MUTABLERUNTIME_API InputMemoryStream::Private : public Base
+    class MUTABLERUNTIME_API InputMemoryStream::Private
     {
     public:
 
@@ -107,7 +107,7 @@ namespace mu
 
 
     //---------------------------------------------------------------------------------------------
-    class MUTABLERUNTIME_API OutputMemoryStream::Private : public Base
+    class MUTABLERUNTIME_API OutputMemoryStream::Private
     {
     public:
         TArray<uint8,TSizedHeapAllocator<64>> m_buffer;
@@ -115,7 +115,7 @@ namespace mu
 
 
 	//---------------------------------------------------------------------------------------------
-	class MUTABLERUNTIME_API InputArchive::Private : public Base
+	class MUTABLERUNTIME_API InputArchive::Private
 	{
 	public:
 
@@ -128,7 +128,7 @@ namespace mu
 	};
 
 	//---------------------------------------------------------------------------------------------
-	class MUTABLERUNTIME_API OutputArchive::Private : public Base
+	class MUTABLERUNTIME_API OutputArchive::Private
 	{
 	public:
 
diff --git a/Engine/Plugins/Experimental/Mutable/Source/MutableRuntime/Internal/MuR/SystemPrivate.h b/Engine/Plugins/Experimental/Mutable/Source/MutableRuntime/Internal/MuR/SystemPrivate.h
--- a/Engine/Plugins/Experimental/Mutable/Source/MutableRuntime/Internal/MuR/SystemPrivate.h
+++ b/Engine/Plugins/Experimental/Mutable/Source/MutableRuntime/Internal/MuR/SystemPrivate.h
@@ -1827,7 +1827,7 @@ namespace mu
 
 
 	/** */
-    class System::Private : public Base
+    class System::Private
     {
     public:
 
diff --git a/Engine/Plugins/Experimental/Mutable/Source/MutableRuntime/Private/MuR/CodeRunner.cpp b/Engine/Plugins/Experimental/Mutable/Source/MutableRuntime/Private/MuR/CodeRunner.cpp
--- a/Engine/Plugins/Experimental/Mutable/Source/MutableRuntime/Private/MuR/CodeRunner.cpp
+++ b/Engine/Plugins/Experimental/Mutable/Source/MutableRuntime/Private/MuR/CodeRunner.cpp
@@ -64,7 +64,6 @@
 #include "Templates/Tuple.h"
 #include "Trace/Detail/Channel.h"
 
-
 namespace
 {
 
@@ -122,7 +121,7 @@ namespace mu
 		, m_pParams(InParams)
 		, m_lodMask(InLodMask)
 	{
-		FProgram& program = m_pModel->GetPrivate()->m_program;
+		const FProgram& program = m_pModel->GetPrivate()->m_program;
 		ScheduledStagePerOp.resize(program.m_opAddress.Num());
 
 		// We will read this in the end, so make sure we keep it.
@@ -219,8 +218,10 @@ namespace mu
     {
         MUTABLE_CPUPROFILER_SCOPE(RunCode_Conditional);
 
-		OP_TYPE type = pModel->GetPrivate()->m_program.GetOpType(item.At);
-		OP::ConditionalArgs args = pModel->GetPrivate()->m_program.GetOpArgs<OP::ConditionalArgs>(item.At);
+		const FProgram& Program = m_pModel->GetPrivate()->m_program;
+
+		OP_TYPE type = Program.GetOpType(item.At);
+		OP::ConditionalArgs args = Program.GetOpArgs<OP::ConditionalArgs>(item.At);
 
         // Conditionals have the following execution stages:
         // 0: we need to run the condition
@@ -289,9 +290,11 @@ namespace mu
 	//---------------------------------------------------------------------------------------------
 	void CodeRunner::RunCode_Switch(const FScheduledOp& item, const Model* pModel )
 	{
-		OP_TYPE type = pModel->GetPrivate()->m_program.GetOpType(item.At);
+		const FProgram& Program = m_pModel->GetPrivate()->m_program;
+
+		OP_TYPE type = Program.GetOpType(item.At);
 
-		const uint8* data = pModel->GetPrivate()->m_program.GetOpArgsPointer(item.At);
+		const uint8* data = Program.GetOpArgsPointer(item.At);
 
 		OP::ADDRESS VarAddress;
 		FMemory::Memcpy(&VarAddress, data, sizeof(OP::ADDRESS));
@@ -405,13 +408,15 @@ namespace mu
     {
         MUTABLE_CPUPROFILER_SCOPE(RunCode_Instance);
 
-		OP_TYPE type = pModel->GetPrivate()->m_program.GetOpType(item.At);
+		const FProgram& Program = m_pModel->GetPrivate()->m_program;
+
+		OP_TYPE type = Program.GetOpType(item.At);
         switch (type)
         {
 
         case OP_TYPE::IN_ADDVECTOR:
         {
-			OP::InstanceAddArgs args = pModel->GetPrivate()->m_program.GetOpArgs<OP::InstanceAddArgs>(item.At);
+			OP::InstanceAddArgs args = Program.GetOpArgs<OP::InstanceAddArgs>(item.At);
 
             switch (item.Stage)
             {
@@ -439,8 +444,8 @@ namespace mu
 					FVector4f value = LoadColor( FCacheAddress(args.value,item) );
 
                     OP::ADDRESS nameAd = args.name;
-                    check(  nameAd < (uint32)pModel->GetPrivate()->m_program.m_constantStrings.Num() );
-                    const FString& Name = pModel->GetPrivate()->m_program.m_constantStrings[ nameAd ];
+                    check(  nameAd < (uint32)Program.m_constantStrings.Num() );
+                    const FString& Name = Program.m_constantStrings[ nameAd ];
 
                     pResult->GetPrivate()->AddVector( 0, 0, 0, value, FName(Name) );
                 }
@@ -457,7 +462,7 @@ namespace mu
 
         case OP_TYPE::IN_ADDSCALAR:
         {
-			OP::InstanceAddArgs args = pModel->GetPrivate()->m_program.GetOpArgs<OP::InstanceAddArgs>(item.At);
+			OP::InstanceAddArgs args = Program.GetOpArgs<OP::InstanceAddArgs>(item.At);
             switch (item.Stage)
             {
             case 0:
@@ -484,8 +489,8 @@ namespace mu
                     float value = LoadScalar( FCacheAddress(args.value,item) );
 
                     OP::ADDRESS nameAd = args.name;
-                    check(  nameAd < (uint32)pModel->GetPrivate()->m_program.m_constantStrings.Num() );
-                    const FString& Name = pModel->GetPrivate()->m_program.m_constantStrings[ nameAd ];
+                    check(  nameAd < (uint32)Program.m_constantStrings.Num() );
+                    const FString& Name = Program.m_constantStrings[ nameAd ];
 
                     pResult->GetPrivate()->AddScalar( 0, 0, 0, value, FName(Name));
                 }
@@ -502,7 +507,7 @@ namespace mu
 
         case OP_TYPE::IN_ADDSTRING:
         {
-			OP::InstanceAddArgs args = pModel->GetPrivate()->m_program.GetOpArgs<OP::InstanceAddArgs>( item.At );
+			OP::InstanceAddArgs args = Program.GetOpArgs<OP::InstanceAddArgs>( item.At );
             switch ( item.Stage )
             {
             case 0:
@@ -530,8 +535,8 @@ namespace mu
                         LoadString( FCacheAddress( args.value, item ) );
 
                     OP::ADDRESS nameAd = args.name;
-                    check( nameAd < (uint32)pModel->GetPrivate()->m_program.m_constantStrings.Num() );
-                    const FString& Name = pModel->GetPrivate()->m_program.m_constantStrings[nameAd];
+                    check( nameAd < (uint32)Program.m_constantStrings.Num() );
+                    const FString& Name = Program.m_constantStrings[nameAd];
 
                     pResult->GetPrivate()->AddString( 0, 0, 0, value->GetValue(), FName(Name) );
                 }
@@ -548,7 +553,7 @@ namespace mu
 
         case OP_TYPE::IN_ADDCOMPONENT:
         {
-			OP::InstanceAddArgs args = pModel->GetPrivate()->m_program.GetOpArgs<OP::InstanceAddArgs>(item.At);
+			OP::InstanceAddArgs args = Program.GetOpArgs<OP::InstanceAddArgs>(item.At);
             switch (item.Stage)
             {
             case 0:
@@ -589,8 +594,8 @@ namespace mu
                     	
                         // Name
                         OP::ADDRESS nameAd = args.name;
-                        check( nameAd < (uint32)pModel->GetPrivate()->m_program.m_constantStrings.Num() );
-                        const FString& Name = pModel->GetPrivate()->m_program.m_constantStrings[ nameAd ];
+                        check( nameAd < (uint32)Program.m_constantStrings.Num() );
+                        const FString& Name = Program.m_constantStrings[ nameAd ];
                         pResult->GetPrivate()->SetComponentName( 0, cindex, FName(Name) );
                     }
                 }
@@ -607,7 +612,7 @@ namespace mu
 
         case OP_TYPE::IN_ADDSURFACE:
         {
-			OP::InstanceAddArgs args = pModel->GetPrivate()->m_program.GetOpArgs<OP::InstanceAddArgs>(item.At);
+			OP::InstanceAddArgs args = Program.GetOpArgs<OP::InstanceAddArgs>(item.At);
             switch (item.Stage)
             {
             case 0:
@@ -656,8 +661,8 @@ namespace mu
 
                     // Name
                     OP::ADDRESS nameAd = args.name;
-                    check( nameAd < (uint32)pModel->GetPrivate()->m_program.m_constantStrings.Num() );
-                    const FString& Name = pModel->GetPrivate()->m_program.m_constantStrings[ nameAd ];
+                    check( nameAd < (uint32)Program.m_constantStrings.Num() );
+                    const FString& Name = Program.m_constantStrings[ nameAd ];
                     pResult->GetPrivate()->SetSurfaceName( 0, 0, sindex, FName(Name) );
 
                     // IDs
@@ -678,7 +683,7 @@ namespace mu
 
         case OP_TYPE::IN_ADDLOD:
         {
-			OP::InstanceAddLODArgs args = pModel->GetPrivate()->m_program.GetOpArgs<OP::InstanceAddLODArgs>(item.At);
+			OP::InstanceAddLODArgs args = Program.GetOpArgs<OP::InstanceAddLODArgs>(item.At);
             switch (item.Stage)
             {
             case 0:
@@ -746,7 +751,7 @@ namespace mu
 
 		case OP_TYPE::IN_ADDEXTENSIONDATA:
 		{
-			OP::InstanceAddExtensionDataArgs Args = pModel->GetPrivate()->m_program.GetOpArgs<OP::InstanceAddExtensionDataArgs>(item.At);
+			OP::InstanceAddExtensionDataArgs Args = Program.GetOpArgs<OP::InstanceAddExtensionDataArgs>(item.At);
 			switch (item.Stage)
 			{
 				case 0:
@@ -774,8 +779,8 @@ namespace mu
 					if (ExtensionDataPtrConst ExtensionData = LoadExtensionData(FCacheAddress(Args.ExtensionData, item)))
 					{
 						const OP::ADDRESS NameAddress = Args.ExtensionDataName;
-						check(NameAddress < (uint32)pModel->GetPrivate()->m_program.m_constantStrings.Num());
-						const FString& NameString = pModel->GetPrivate()->m_program.m_constantStrings[NameAddress];
+						check(NameAddress < (uint32)Program.m_constantStrings.Num());
+						const FString& NameString = Program.m_constantStrings[NameAddress];
 
 						Result->GetPrivate()->AddExtensionData(ExtensionData, FName(NameString) );
 					}
@@ -807,13 +812,15 @@ namespace mu
 			return;
 		}
 
-		OP_TYPE type = InModel->GetPrivate()->m_program.GetOpType(item.At);
+		const FProgram& Program = m_pModel->GetPrivate()->m_program;
+
+		OP_TYPE type = Program.GetOpType(item.At);
         switch (type)
         {
 
         case OP_TYPE::IN_ADDMESH:
         {
-			OP::InstanceAddArgs args = InModel->GetPrivate()->m_program.GetOpArgs<OP::InstanceAddArgs>(item.At);
+			OP::InstanceAddArgs args = Program.GetOpArgs<OP::InstanceAddArgs>(item.At);
             switch (item.Stage)
             {
             case 0:
@@ -841,8 +848,8 @@ namespace mu
                 {
 					FResourceID MeshId = m_pSystem->WorkingMemoryManager.GetResourceKey(InModel,InParams,args.relevantParametersListIndex, args.value);
 					OP::ADDRESS NameAd = args.name;
-					check(NameAd < (uint32)InModel->GetPrivate()->m_program.m_constantStrings.Num());
-					const FString& Name = InModel->GetPrivate()->m_program.m_constantStrings[NameAd];
+					check(NameAd < (uint32)Program.m_constantStrings.Num());
+					const FString& Name = Program.m_constantStrings[NameAd];
 					pResult->GetPrivate()->AddMesh(0, 0, MeshId, FName(Name));
                 }
                 StoreInstance( item, pResult );
@@ -857,7 +864,7 @@ namespace mu
 
         case OP_TYPE::IN_ADDIMAGE:
         {
-			OP::InstanceAddArgs args = InModel->GetPrivate()->m_program.GetOpArgs<OP::InstanceAddArgs>(item.At);
+			OP::InstanceAddArgs args = Program.GetOpArgs<OP::InstanceAddArgs>(item.At);
             switch (item.Stage)
             {
             case 0:
@@ -882,8 +889,8 @@ namespace mu
                 {
 					FResourceID ImageId = m_pSystem->WorkingMemoryManager.GetResourceKey(InModel, InParams, args.relevantParametersListIndex, args.value);
 					OP::ADDRESS NameAd = args.name;
-					check(NameAd < (uint32)InModel->GetPrivate()->m_program.m_constantStrings.Num());
-					const FString& Name = InModel->GetPrivate()->m_program.m_constantStrings[NameAd];
+					check(NameAd < (uint32)Program.m_constantStrings.Num());
+					const FString& Name = Program.m_constantStrings[NameAd];
 					pResult->GetPrivate()->AddImage(0, 0, 0, ImageId, FName(Name) );
                 }
                 StoreInstance( item, pResult );
@@ -908,23 +915,23 @@ namespace mu
     {
 		MUTABLE_CPUPROFILER_SCOPE(RunCode_Constant);
 
-		OP_TYPE type = pModel->GetPrivate()->m_program.GetOpType(item.At);
+		const FProgram& Program = m_pModel->GetPrivate()->m_program;
+
+		OP_TYPE type = Program.GetOpType(item.At);
         switch (type)
         {
 
         case OP_TYPE::ME_CONSTANT:
         {
-			OP::MeshConstantArgs args = pModel->GetPrivate()->m_program.GetOpArgs<OP::MeshConstantArgs>(item.At);
+			OP::MeshConstantArgs args = Program.GetOpArgs<OP::MeshConstantArgs>(item.At);
 
             OP::ADDRESS cat = args.value;
 
-            FProgram& program = pModel->GetPrivate()->m_program;
-
             // Assume the ROM has been loaded previously
-            check(program.m_constantMeshes[cat].Value)
+            check(Program.m_constantMeshes[cat].Value)
 
             Ptr<const Mesh> SourceConst;
-            program.GetConstant(cat, SourceConst);
+			Program.GetConstant(cat, SourceConst);
 
 			check(SourceConst);
 			Ptr<Mesh> Source = CreateMesh(SourceConst->GetDataSize());
@@ -933,15 +940,15 @@ namespace mu
             // Set the separate skeleton if necessary
             if (args.skeleton >= 0)
             {
-                check(program.m_constantSkeletons.Num() > size_t(args.skeleton));
-                Ptr<const Skeleton> pSkeleton = program.m_constantSkeletons[args.skeleton];
+                check(Program.m_constantSkeletons.Num() > size_t(args.skeleton));
+                Ptr<const Skeleton> pSkeleton = Program.m_constantSkeletons[args.skeleton];
                 Source->SetSkeleton(pSkeleton);
             }
 
 			if (args.physicsBody >= 0)
 			{
-                check(program.m_constantPhysicsBodies.Num() > size_t(args.physicsBody));
-                Ptr<const PhysicsBody> pPhysicsBody = program.m_constantPhysicsBodies[args.physicsBody];
+                check(Program.m_constantPhysicsBodies.Num() > size_t(args.physicsBody));
+                Ptr<const PhysicsBody> pPhysicsBody = Program.m_constantPhysicsBodies[args.physicsBody];
                 Source->SetPhysicsBody(pPhysicsBody);
 			}
 
@@ -952,14 +959,12 @@ namespace mu
 
         case OP_TYPE::IM_CONSTANT:
         {
-			OP::ResourceConstantArgs args = pModel->GetPrivate()->m_program.GetOpArgs<OP::ResourceConstantArgs>(item.At);
+			OP::ResourceConstantArgs args = Program.GetOpArgs<OP::ResourceConstantArgs>(item.At);
             OP::ADDRESS cat = args.value;
 
-            const FProgram& program = pModel->GetPrivate()->m_program;
-
 			int32 MipsToSkip = item.ExecutionOptions;
             Ptr<const Image> Source;
-			program.GetConstant(cat, Source, MipsToSkip, [this](int32 x, int32 y, int32 m, EImageFormat f, EInitializationType i)
+			Program.GetConstant(cat, Source, MipsToSkip, [this](int32 x, int32 y, int32 m, EImageFormat f, EInitializationType i)
 				{
 					return CreateImage(x, y, m, f, i);
 				});
@@ -974,9 +979,7 @@ namespace mu
 
 		case OP_TYPE::ED_CONSTANT:
 		{
-			OP::ResourceConstantArgs Args = pModel->GetPrivate()->m_program.GetOpArgs<OP::ResourceConstantArgs>(item.At);
-
-            const FProgram& Program = pModel->GetPrivate()->m_program;
+			OP::ResourceConstantArgs Args = Program.GetOpArgs<OP::ResourceConstantArgs>(item.At);
 
 			// Assume the ROM has been loaded previously
 			ExtensionDataPtrConst SourceConst;
@@ -1003,14 +1006,16 @@ namespace mu
     {
 		MUTABLE_CPUPROFILER_SCOPE(RunCode_Mesh);
 
-		OP_TYPE type = pModel->GetPrivate()->m_program.GetOpType(item.At);
+		const FProgram& Program = m_pModel->GetPrivate()->m_program;
+
+		OP_TYPE type = Program.GetOpType(item.At);
 
         switch (type)
         {
 
         case OP_TYPE::ME_APPLYLAYOUT:
         {
-			OP::MeshApplyLayoutArgs args = pModel->GetPrivate()->m_program.GetOpArgs<OP::MeshApplyLayoutArgs>(item.At);
+			OP::MeshApplyLayoutArgs args = Program.GetOpArgs<OP::MeshApplyLayoutArgs>(item.At);
 			switch (item.Stage)
 			{
 			case 0:
@@ -1055,7 +1060,7 @@ namespace mu
 
         case OP_TYPE::ME_DIFFERENCE:
         {
-			const uint8* data = pModel->GetPrivate()->m_program.GetOpArgsPointer(item.At);
+			const uint8* data = Program.GetOpArgsPointer(item.At);
 
 			OP::ADDRESS BaseAt = 0;
 			FMemory::Memcpy(&BaseAt, data, sizeof(OP::ADDRESS)); 
@@ -1134,7 +1139,7 @@ namespace mu
 
         case OP_TYPE::ME_MORPH:
         {
-			const uint8* data = pModel->GetPrivate()->m_program.GetOpArgsPointer(item.At);
+			const uint8* data = Program.GetOpArgsPointer(item.At);
 
 			OP::ADDRESS FactorAt = 0;
 			FMemory::Memcpy(&FactorAt, data, sizeof(OP::ADDRESS)); 
@@ -1257,7 +1262,7 @@ namespace mu
 
         case OP_TYPE::ME_MERGE:
         {
-			OP::MeshMergeArgs args = pModel->GetPrivate()->m_program.GetOpArgs<OP::MeshMergeArgs>(item.At);
+			OP::MeshMergeArgs args = Program.GetOpArgs<OP::MeshMergeArgs>(item.At);
             switch (item.Stage)
             {
             case 0:
@@ -1335,7 +1340,7 @@ namespace mu
 
         case OP_TYPE::ME_INTERPOLATE:
         {
-			OP::MeshInterpolateArgs args = pModel->GetPrivate()->m_program.GetOpArgs<OP::MeshInterpolateArgs>(item.At);
+			OP::MeshInterpolateArgs args = Program.GetOpArgs<OP::MeshInterpolateArgs>(item.At);
             switch (item.Stage)
             {
             case 0:
@@ -1592,7 +1597,7 @@ namespace mu
 
         case OP_TYPE::ME_MASKCLIPMESH:
         {
-			OP::MeshMaskClipMeshArgs args = pModel->GetPrivate()->m_program.GetOpArgs<OP::MeshMaskClipMeshArgs>(item.At);
+			OP::MeshMaskClipMeshArgs args = Program.GetOpArgs<OP::MeshMaskClipMeshArgs>(item.At);
             switch (item.Stage)
             {
             case 0:
@@ -1648,7 +1653,7 @@ namespace mu
 
 		case OP_TYPE::ME_MASKCLIPUVMASK:
 		{
-			OP::MeshMaskClipUVMaskArgs args = pModel->GetPrivate()->m_program.GetOpArgs<OP::MeshMaskClipUVMaskArgs>(item.At);
+			OP::MeshMaskClipUVMaskArgs args = Program.GetOpArgs<OP::MeshMaskClipUVMaskArgs>(item.At);
 			switch (item.Stage)
 			{
 			case 0:
@@ -1704,7 +1709,7 @@ namespace mu
 
         case OP_TYPE::ME_MASKDIFF:
         {
-			OP::MeshMaskDiffArgs args = pModel->GetPrivate()->m_program.GetOpArgs<OP::MeshMaskDiffArgs>(item.At);
+			OP::MeshMaskDiffArgs args = Program.GetOpArgs<OP::MeshMaskDiffArgs>(item.At);
             switch (item.Stage)
             {
             case 0:
@@ -1761,7 +1766,7 @@ namespace mu
 
         case OP_TYPE::ME_FORMAT:
         {
-			OP::MeshFormatArgs args = pModel->GetPrivate()->m_program.GetOpArgs<OP::MeshFormatArgs>(item.At);
+			OP::MeshFormatArgs args = Program.GetOpArgs<OP::MeshFormatArgs>(item.At);
             switch (item.Stage)
             {
             case 0:
@@ -1845,7 +1850,7 @@ namespace mu
 
         case OP_TYPE::ME_EXTRACTLAYOUTBLOCK:
         {
-            const uint8* data = pModel->GetPrivate()->m_program.GetOpArgsPointer(item.At);
+            const uint8* data = Program.GetOpArgsPointer(item.At);
 
             OP::ADDRESS source;
             FMemory::Memcpy( &source, data, sizeof(OP::ADDRESS) );
@@ -1917,7 +1922,7 @@ namespace mu
 
         case OP_TYPE::ME_TRANSFORM:
         {
-			OP::MeshTransformArgs args = pModel->GetPrivate()->m_program.GetOpArgs<OP::MeshTransformArgs>(item.At);
+			OP::MeshTransformArgs args = Program.GetOpArgs<OP::MeshTransformArgs>(item.At);
             switch (item.Stage)
             {
             case 0:
@@ -1939,7 +1944,7 @@ namespace mu
             		
                 Ptr<const Mesh> Source = LoadMesh(FCacheAddress(args.source,item));
 
-                const mat4f& mat = pModel->GetPrivate()->m_program.m_constantMatrices[args.matrix];
+                const mat4f& mat = Program.m_constantMatrices[args.matrix];
 
 				Ptr<Mesh> Result = CreateMesh(Source ? Source->GetDataSize() : 0);
 
@@ -1968,7 +1973,7 @@ namespace mu
 
         case OP_TYPE::ME_CLIPMORPHPLANE:
         {
-			OP::MeshClipMorphPlaneArgs args = pModel->GetPrivate()->m_program.GetOpArgs<OP::MeshClipMorphPlaneArgs>(item.At);
+			OP::MeshClipMorphPlaneArgs args = Program.GetOpArgs<OP::MeshClipMorphPlaneArgs>(item.At);
             switch (item.Stage)
             {
             case 0:
@@ -1993,7 +1998,7 @@ namespace mu
                 check(args.morphShape < (uint32)pModel->GetPrivate()->m_program.m_constantShapes.Num());
 
                 // Should be an ellipse
-                const FShape& morphShape = pModel->GetPrivate()->m_program.m_constantShapes[args.morphShape];
+                const FShape& morphShape = Program.m_constantShapes[args.morphShape];
 
                 const mu::vec3f& origin = morphShape.position;
                 const mu::vec3f& normal = morphShape.up;
@@ -2003,7 +2008,7 @@ namespace mu
                     check(args.vertexSelectionShapeOrBone < (uint32)pModel->GetPrivate()->m_program.m_constantShapes.Num());
 
                     // Should be None or an axis aligned box
-                    const FShape& selectionShape = pModel->GetPrivate()->m_program.m_constantShapes[args.vertexSelectionShapeOrBone];
+                    const FShape& selectionShape = Program.m_constantShapes[args.vertexSelectionShapeOrBone];
 
 					Ptr<Mesh> Result = CreateMesh(Source ? Source->GetDataSize() : 0);
 
@@ -2079,7 +2084,7 @@ namespace mu
 
         case OP_TYPE::ME_CLIPWITHMESH:
         {
-			OP::MeshClipWithMeshArgs args = pModel->GetPrivate()->m_program.GetOpArgs<OP::MeshClipWithMeshArgs>(item.At);
+			OP::MeshClipWithMeshArgs args = Program.GetOpArgs<OP::MeshClipWithMeshArgs>(item.At);
             switch (item.Stage)
             {
             case 0:
@@ -2141,7 +2146,7 @@ namespace mu
         }
 		case OP_TYPE::ME_CLIPDEFORM:
 		{
-			OP::MeshClipDeformArgs args = pModel->GetPrivate()->m_program.GetOpArgs<OP::MeshClipDeformArgs>(item.At);
+			OP::MeshClipDeformArgs args = Program.GetOpArgs<OP::MeshClipDeformArgs>(item.At);
 			switch (item.Stage)
 			{
 			case 0:
@@ -2203,7 +2208,7 @@ namespace mu
 
         case OP_TYPE::ME_APPLYPOSE:
         {
-			OP::MeshApplyPoseArgs args = pModel->GetPrivate()->m_program.GetOpArgs<OP::MeshApplyPoseArgs>(item.At);
+			OP::MeshApplyPoseArgs args = Program.GetOpArgs<OP::MeshApplyPoseArgs>(item.At);
             switch (item.Stage)
             {
             case 0:
@@ -2266,7 +2271,7 @@ namespace mu
 
 		case OP_TYPE::ME_GEOMETRYOPERATION:
 		{
-			OP::MeshGeometryOperationArgs args = pModel->GetPrivate()->m_program.GetOpArgs<OP::MeshGeometryOperationArgs>(item.At);
+			OP::MeshGeometryOperationArgs args = Program.GetOpArgs<OP::MeshGeometryOperationArgs>(item.At);
 			switch (item.Stage)
 			{
 			case 0:
@@ -2325,8 +2330,8 @@ namespace mu
 
 		case OP_TYPE::ME_BINDSHAPE:
 		{
-			OP::MeshBindShapeArgs Args = pModel->GetPrivate()->m_program.GetOpArgs<OP::MeshBindShapeArgs>(item.At);
-			const uint8* Data = pModel->GetPrivate()->m_program.GetOpArgsPointer(item.At);
+			OP::MeshBindShapeArgs Args = Program.GetOpArgs<OP::MeshBindShapeArgs>(item.At);
+			const uint8* Data = Program.GetOpArgsPointer(item.At);
 
 			switch (item.Stage)
 			{
@@ -2445,7 +2450,7 @@ namespace mu
 
 		case OP_TYPE::ME_APPLYSHAPE:
 		{
-			OP::MeshApplyShapeArgs args = pModel->GetPrivate()->m_program.GetOpArgs<OP::MeshApplyShapeArgs>(item.At);
+			OP::MeshApplyShapeArgs args = Program.GetOpArgs<OP::MeshApplyShapeArgs>(item.At);
 			switch (item.Stage)
 			{
 			case 0:
@@ -2524,7 +2529,7 @@ namespace mu
 
 		case OP_TYPE::ME_MORPHRESHAPE:
 		{
-			OP::MeshMorphReshapeArgs Args = pModel->GetPrivate()->m_program.GetOpArgs<OP::MeshMorphReshapeArgs>(item.At);
+			OP::MeshMorphReshapeArgs Args = Program.GetOpArgs<OP::MeshMorphReshapeArgs>(item.At);
 			switch(item.Stage)
 			{
 			case 0:
@@ -2584,7 +2589,7 @@ namespace mu
 
         case OP_TYPE::ME_SETSKELETON:
         {
-			OP::MeshSetSkeletonArgs args = pModel->GetPrivate()->m_program.GetOpArgs<OP::MeshSetSkeletonArgs>(item.At);
+			OP::MeshSetSkeletonArgs args = Program.GetOpArgs<OP::MeshSetSkeletonArgs>(item.At);
             switch (item.Stage)
             {
             case 0:
@@ -2673,7 +2678,7 @@ namespace mu
         		
             // Decode op
             // TODO: Partial decode for each stage
-            const uint8* data = pModel->GetPrivate()->m_program.GetOpArgsPointer(item.At);
+            const uint8* data = Program.GetOpArgsPointer(item.At);
 
             OP::ADDRESS source;
             FMemory::Memcpy(&source,data,sizeof(OP::ADDRESS)); 
@@ -2820,7 +2825,7 @@ namespace mu
 
 			// Decode op
 			// TODO: Partial decode for each stage
-			const uint8* Data = pModel->GetPrivate()->m_program.GetOpArgsPointer(item.At);
+			const uint8* Data = Program.GetOpArgsPointer(item.At);
 
 			OP::ADDRESS Source;
 			FMemory::Memcpy(&Source, Data, sizeof(OP::ADDRESS));
@@ -2871,7 +2876,7 @@ namespace mu
 						Data += sizeof(OP::ADDRESS);
 
 						check(TagConstant < (uint32)pModel->GetPrivate()->m_program.m_constantStrings.Num());
-						const FString& Name = pModel->GetPrivate()->m_program.m_constantStrings[TagConstant];
+						const FString& Name = Program.m_constantStrings[TagConstant];
 						Result->m_tags[FirstMeshTagIndex+TagIndex] = Name;
 					}
 
@@ -2890,7 +2895,7 @@ namespace mu
 
         case OP_TYPE::ME_PROJECT:
         {
-			OP::MeshProjectArgs args = pModel->GetPrivate()->m_program.GetOpArgs<OP::MeshProjectArgs>(item.At);
+			OP::MeshProjectArgs args = Program.GetOpArgs<OP::MeshProjectArgs>(item.At);
             switch (item.Stage)
             {
             case 0:
@@ -2952,7 +2957,7 @@ namespace mu
 
 		case OP_TYPE::ME_OPTIMIZESKINNING:
 		{
-			OP::MeshOptimizeSkinningArgs args = pModel->GetPrivate()->m_program.GetOpArgs<OP::MeshOptimizeSkinningArgs>(item.At);
+			OP::MeshOptimizeSkinningArgs args = Program.GetOpArgs<OP::MeshOptimizeSkinningArgs>(item.At);
 			switch (item.Stage)
 			{
 			case 0:
@@ -3016,13 +3021,15 @@ namespace mu
 
 		FImageOperator ImOp = MakeImageOperator(this);
 
-		OP_TYPE type = pModel->GetPrivate()->m_program.GetOpType(item.At);
+		const FProgram& Program = m_pModel->GetPrivate()->m_program;
+
+		OP_TYPE type = Program.GetOpType(item.At);
 		switch (type)
         {
 
         case OP_TYPE::IM_LAYERCOLOUR:
         {
-			OP::ImageLayerColourArgs args = pModel->GetPrivate()->m_program.GetOpArgs<OP::ImageLayerColourArgs>(item.At);
+			OP::ImageLayerColourArgs args = Program.GetOpArgs<OP::ImageLayerColourArgs>(item.At);
             switch (item.Stage)
             {
             case 0:
@@ -3046,7 +3053,7 @@ namespace mu
 
         case OP_TYPE::IM_LAYER:
         {
-			OP::ImageLayerArgs args = pModel->GetPrivate()->m_program.GetOpArgs<OP::ImageLayerArgs>(item.At);
+			OP::ImageLayerArgs args = Program.GetOpArgs<OP::ImageLayerArgs>(item.At);
 
 			if (ExecutionStrategy == EExecutionStrategy::MinimizeMemory)
 			{
@@ -3099,7 +3106,7 @@ namespace mu
 
         case OP_TYPE::IM_MULTILAYER:
         {
-			OP::ImageMultiLayerArgs args = pModel->GetPrivate()->m_program.GetOpArgs<OP::ImageMultiLayerArgs>(item.At);
+			OP::ImageMultiLayerArgs args = Program.GetOpArgs<OP::ImageMultiLayerArgs>(item.At);
             switch (item.Stage)
             {
             case 0:
@@ -3402,7 +3409,7 @@ namespace mu
 
 		case OP_TYPE::IM_NORMALCOMPOSITE:
 		{
-			OP::ImageNormalCompositeArgs args = pModel->GetPrivate()->m_program.GetOpArgs<OP::ImageNormalCompositeArgs>(item.At);
+			OP::ImageNormalCompositeArgs args = Program.GetOpArgs<OP::ImageNormalCompositeArgs>(item.At);
 			switch (item.Stage)
 			{
 			case 0:
@@ -3411,7 +3418,7 @@ namespace mu
 					AddOp(FScheduledOp(item.At, item, 1),
 							FScheduledOp(args.base, item),
 							FScheduledOp(args.normal, item));
-			}
+				}
 				else
 				{
 					StoreImage(item, nullptr);
@@ -3458,7 +3465,7 @@ namespace mu
 
         case OP_TYPE::IM_PIXELFORMAT:
         {
-			OP::ImagePixelFormatArgs args = pModel->GetPrivate()->m_program.GetOpArgs<OP::ImagePixelFormatArgs>(item.At);
+			OP::ImagePixelFormatArgs args = Program.GetOpArgs<OP::ImagePixelFormatArgs>(item.At);
             switch (item.Stage)
             {
             case 0:
@@ -3481,7 +3488,7 @@ namespace mu
 
         case OP_TYPE::IM_MIPMAP:
         {
-			OP::ImageMipmapArgs args = pModel->GetPrivate()->m_program.GetOpArgs<OP::ImageMipmapArgs>(item.At);
+			OP::ImageMipmapArgs args = Program.GetOpArgs<OP::ImageMipmapArgs>(item.At);
             switch (item.Stage)
             {
             case 0:
@@ -3502,7 +3509,7 @@ namespace mu
 
         case OP_TYPE::IM_RESIZE:
         {
-			OP::ImageResizeArgs args = pModel->GetPrivate()->m_program.GetOpArgs<OP::ImageResizeArgs>(item.At);
+			OP::ImageResizeArgs args = Program.GetOpArgs<OP::ImageResizeArgs>(item.At);
             switch (item.Stage)
             {
             case 0:
@@ -3524,7 +3531,7 @@ namespace mu
 
         case OP_TYPE::IM_RESIZELIKE:
         {
-			OP::ImageResizeLikeArgs args = pModel->GetPrivate()->m_program.GetOpArgs<OP::ImageResizeLikeArgs>(item.At);
+			OP::ImageResizeLikeArgs args = Program.GetOpArgs<OP::ImageResizeLikeArgs>(item.At);
             switch (item.Stage)
             {
             case 0:
@@ -3586,7 +3593,7 @@ namespace mu
 
         case OP_TYPE::IM_RESIZEREL:
         {
-			OP::ImageResizeRelArgs args = pModel->GetPrivate()->m_program.GetOpArgs<OP::ImageResizeRelArgs>(item.At);
+			OP::ImageResizeRelArgs args = Program.GetOpArgs<OP::ImageResizeRelArgs>(item.At);
             switch (item.Stage)
             {
             case 0:
@@ -3609,7 +3616,7 @@ namespace mu
 
         case OP_TYPE::IM_BLANKLAYOUT:
         {
-			OP::ImageBlankLayoutArgs args = pModel->GetPrivate()->m_program.GetOpArgs<OP::ImageBlankLayoutArgs>(item.At);
+			OP::ImageBlankLayoutArgs args = Program.GetOpArgs<OP::ImageBlankLayoutArgs>(item.At);
             switch (item.Stage)
             {
             case 0:
@@ -3685,7 +3692,7 @@ namespace mu
 
         case OP_TYPE::IM_COMPOSE:
         {
-			OP::ImageComposeArgs Args = pModel->GetPrivate()->m_program.GetOpArgs<OP::ImageComposeArgs>(item.At);
+			OP::ImageComposeArgs Args = Program.GetOpArgs<OP::ImageComposeArgs>(item.At);
 
 			if (ExecutionStrategy == EExecutionStrategy::MinimizeMemory)
 			{
@@ -3781,7 +3788,7 @@ namespace mu
 
         case OP_TYPE::IM_INTERPOLATE:
         {
-			OP::ImageInterpolateArgs args = pModel->GetPrivate()->m_program.GetOpArgs<OP::ImageInterpolateArgs>(item.At);
+			OP::ImageInterpolateArgs args = Program.GetOpArgs<OP::ImageInterpolateArgs>(item.At);
             switch (item.Stage)
             {
             case 0:
@@ -3870,9 +3877,7 @@ namespace mu
                     Ptr<const Image> pMax = LoadImage( FCacheAddress(args.targets[max],item) );
 
                     if (pMin && pMax)
-                    {
-						int32 LevelCount = FMath::Max(pMin->GetLODCount(), pMax->GetLODCount());
-						
+                    {						
 						Ptr<Image> pNew = CloneOrTakeOver(pMin);
 
 						// Be defensive: ensure image sizes match.
@@ -3900,6 +3905,8 @@ namespace mu
 							pMax = Formatted;
 						}
 
+						int32 LevelCount = FMath::Max(pNew->GetLODCount(), pMax->GetLODCount());
+
 						if (pNew->GetLODCount() != LevelCount)
 						{
 							MUTABLE_CPUPROFILER_SCOPE(Mipmap_ForInterpolate);
@@ -3954,7 +3961,7 @@ namespace mu
 
         case OP_TYPE::IM_SATURATE:
         {
-			OP::ImageSaturateArgs args = pModel->GetPrivate()->m_program.GetOpArgs<OP::ImageSaturateArgs>(item.At);
+			OP::ImageSaturateArgs args = Program.GetOpArgs<OP::ImageSaturateArgs>(item.At);
             switch (item.Stage)
             {
             case 0:
@@ -3978,7 +3985,7 @@ namespace mu
 
         case OP_TYPE::IM_LUMINANCE:
         {
-			OP::ImageLuminanceArgs args = pModel->GetPrivate()->m_program.GetOpArgs<OP::ImageLuminanceArgs>(item.At);
+			OP::ImageLuminanceArgs args = Program.GetOpArgs<OP::ImageLuminanceArgs>(item.At);
             switch (item.Stage)
             {
             case 0:
@@ -4009,7 +4016,7 @@ namespace mu
 
         case OP_TYPE::IM_SWIZZLE:
         {
-			OP::ImageSwizzleArgs args = pModel->GetPrivate()->m_program.GetOpArgs<OP::ImageSwizzleArgs>(item.At);
+			OP::ImageSwizzleArgs args = Program.GetOpArgs<OP::ImageSwizzleArgs>(item.At);
             switch (item.Stage)
             {
             case 0:
@@ -4034,7 +4041,7 @@ namespace mu
 
         case OP_TYPE::IM_COLOURMAP:
         {
-			OP::ImageColourMapArgs args = pModel->GetPrivate()->m_program.GetOpArgs<OP::ImageColourMapArgs>(item.At);
+			OP::ImageColourMapArgs args = Program.GetOpArgs<OP::ImageColourMapArgs>(item.At);
             switch (item.Stage)
             {
             case 0:
@@ -4090,7 +4097,7 @@ namespace mu
 
         case OP_TYPE::IM_GRADIENT:
         {
-			OP::ImageGradientArgs args = pModel->GetPrivate()->m_program.GetOpArgs<OP::ImageGradientArgs>(item.At);
+			OP::ImageGradientArgs args = Program.GetOpArgs<OP::ImageGradientArgs>(item.At);
             switch (item.Stage)
             {
             case 0:
@@ -4122,7 +4129,7 @@ namespace mu
 
         case OP_TYPE::IM_BINARISE:
         {
-			OP::ImageBinariseArgs args = pModel->GetPrivate()->m_program.GetOpArgs<OP::ImageBinariseArgs>(item.At);
+			OP::ImageBinariseArgs args = Program.GetOpArgs<OP::ImageBinariseArgs>(item.At);
             switch (item.Stage)
             {
             case 0:
@@ -4156,7 +4163,7 @@ namespace mu
 
 		case OP_TYPE::IM_INVERT:
 		{
-			OP::ImageInvertArgs args = pModel->GetPrivate()->m_program.GetOpArgs<OP::ImageInvertArgs>(item.At);
+			OP::ImageInvertArgs args = Program.GetOpArgs<OP::ImageInvertArgs>(item.At);
 			switch (item.Stage)
 			{
 			case 0:
@@ -4178,7 +4185,7 @@ namespace mu
 
         case OP_TYPE::IM_PLAINCOLOUR:
         {
-			OP::ImagePlainColourArgs args = pModel->GetPrivate()->m_program.GetOpArgs<OP::ImagePlainColourArgs>(item.At);
+			OP::ImagePlainColourArgs args = Program.GetOpArgs<OP::ImagePlainColourArgs>(item.At);
             switch (item.Stage)
             {
             case 0:
@@ -4225,7 +4232,7 @@ namespace mu
 
 		case OP_TYPE::IM_REFERENCE:
 		{
-			OP::ResourceReferenceArgs Args = pModel->GetPrivate()->m_program.GetOpArgs<OP::ResourceReferenceArgs>(item.At);
+			OP::ResourceReferenceArgs Args = Program.GetOpArgs<OP::ResourceReferenceArgs>(item.At);
 			switch (item.Stage)
 			{
 			case 0:
@@ -4253,7 +4260,7 @@ namespace mu
 
         case OP_TYPE::IM_CROP:
         {
-			OP::ImageCropArgs args = pModel->GetPrivate()->m_program.GetOpArgs<OP::ImageCropArgs>(item.At);
+			OP::ImageCropArgs args = Program.GetOpArgs<OP::ImageCropArgs>(item.At);
             switch (item.Stage)
             {
             case 0:
@@ -4305,7 +4312,7 @@ namespace mu
         case OP_TYPE::IM_PATCH:
         {
 			// TODO: This is optimized for memory-usage but base and patch could be requested at the same time
-			OP::ImagePatchArgs args = pModel->GetPrivate()->m_program.GetOpArgs<OP::ImagePatchArgs>(item.At);
+			OP::ImagePatchArgs args = Program.GetOpArgs<OP::ImagePatchArgs>(item.At);
             switch (item.Stage)
             {
 			case 0:
@@ -4415,7 +4422,7 @@ namespace mu
 
         case OP_TYPE::IM_RASTERMESH:
         {
-			OP::ImageRasterMeshArgs args = pModel->GetPrivate()->m_program.GetOpArgs<OP::ImageRasterMeshArgs>(item.At);
+			OP::ImageRasterMeshArgs args = Program.GetOpArgs<OP::ImageRasterMeshArgs>(item.At);
             switch (item.Stage)
             {
             case 0:
@@ -4723,7 +4730,7 @@ namespace mu
 
         case OP_TYPE::IM_MAKEGROWMAP:
         {
-			OP::ImageMakeGrowMapArgs args = pModel->GetPrivate()->m_program.GetOpArgs<OP::ImageMakeGrowMapArgs>(item.At);
+			OP::ImageMakeGrowMapArgs args = Program.GetOpArgs<OP::ImageMakeGrowMapArgs>(item.At);
             switch (item.Stage)
             {
             case 0:
@@ -4740,7 +4747,6 @@ namespace mu
 
                 ImageMakeGrowMap(Result.get(), Mask.get(), args.border );
 				Result->m_flags |= Image::IF_CANNOT_BE_SCALED;
-				check(Result->GetData());
 
 				Release(Mask);
                 StoreImage( item, Result);
@@ -4756,7 +4762,7 @@ namespace mu
 
         case OP_TYPE::IM_DISPLACE:
         {
-			OP::ImageDisplaceArgs args = pModel->GetPrivate()->m_program.GetOpArgs<OP::ImageDisplaceArgs>(item.At);
+			OP::ImageDisplaceArgs args = Program.GetOpArgs<OP::ImageDisplaceArgs>(item.At);
             switch (item.Stage)
             {
             case 0:
@@ -4825,7 +4831,7 @@ namespace mu
 
         case OP_TYPE::IM_TRANSFORM:
         {
-            const OP::ImageTransformArgs Args = pModel->GetPrivate()->m_program.GetOpArgs<OP::ImageTransformArgs>(item.At);
+            const OP::ImageTransformArgs Args = Program.GetOpArgs<OP::ImageTransformArgs>(item.At);
 
             switch (item.Stage)
             {
@@ -5051,8 +5057,8 @@ namespace mu
             return nullptr;
         }
 
-        const FProgram& program = pModel->GetPrivate()->m_program;
-        const FParameterDesc& paramDesc = program.m_parameters[ parameterIndex ];
+		const FProgram& Program = m_pModel->GetPrivate()->m_program;
+		const FParameterDesc& paramDesc = Program.m_parameters[ parameterIndex ];
         for( size_t rangeIndexInParam=0;
              rangeIndexInParam<paramDesc.m_ranges.Num();
              ++rangeIndexInParam )
@@ -5072,14 +5078,14 @@ namespace mu
     {
         MUTABLE_CPUPROFILER_SCOPE(RunCode_Bool);
 
-        const FProgram& program = pModel->GetPrivate()->m_program;
-        OP_TYPE type = program.GetOpType(item.At);
+		const FProgram& Program = m_pModel->GetPrivate()->m_program;
+		OP_TYPE type = Program.GetOpType(item.At);
         switch (type)
         {
 
         case OP_TYPE::BO_CONSTANT:
         {
-			OP::BoolConstantArgs args = program.GetOpArgs<OP::BoolConstantArgs>(item.At);
+			OP::BoolConstantArgs args = Program.GetOpArgs<OP::BoolConstantArgs>(item.At);
             bool result = args.value;
             StoreBool( item, result );
             break;
@@ -5087,7 +5093,7 @@ namespace mu
 
         case OP_TYPE::BO_PARAMETER:
         {
-			OP::ParameterArgs args = program.GetOpArgs<OP::ParameterArgs>(item.At);
+			OP::ParameterArgs args = Program.GetOpArgs<OP::ParameterArgs>(item.At);
             bool result = false;
 			Ptr<RangeIndex> index = BuildCurrentOpRangeIndex( item, pParams, pModel, args.variable );
             result = pParams->GetBoolValue( args.variable, index );
@@ -5097,7 +5103,7 @@ namespace mu
 
         case OP_TYPE::BO_LESS:
         {
-			OP::BoolLessArgs args = program.GetOpArgs<OP::BoolLessArgs>(item.At);
+			OP::BoolLessArgs args = Program.GetOpArgs<OP::BoolLessArgs>(item.At);
             switch (item.Stage)
             {
             case 0:
@@ -5124,7 +5130,7 @@ namespace mu
 
         case OP_TYPE::BO_AND:
         {
-			OP::BoolBinaryArgs args = program.GetOpArgs<OP::BoolBinaryArgs>(item.At);
+			OP::BoolBinaryArgs args = Program.GetOpArgs<OP::BoolBinaryArgs>(item.At);
             switch (item.Stage)
             {
             case 0:
@@ -5190,7 +5196,7 @@ namespace mu
 
         case OP_TYPE::BO_OR:
         {
-			OP::BoolBinaryArgs args = program.GetOpArgs<OP::BoolBinaryArgs>(item.At);
+			OP::BoolBinaryArgs args = Program.GetOpArgs<OP::BoolBinaryArgs>(item.At);
             switch (item.Stage)
             {
             case 0:
@@ -5256,7 +5262,7 @@ namespace mu
 
         case OP_TYPE::BO_NOT:
         {
-			OP::BoolNotArgs args = program.GetOpArgs<OP::BoolNotArgs>(item.At);
+			OP::BoolNotArgs args = Program.GetOpArgs<OP::BoolNotArgs>(item.At);
             switch (item.Stage)
             {
             case 0:
@@ -5279,7 +5285,7 @@ namespace mu
 
         case OP_TYPE::BO_EQUAL_INT_CONST:
         {
-			OP::BoolEqualScalarConstArgs args = program.GetOpArgs<OP::BoolEqualScalarConstArgs>(item.At);
+			OP::BoolEqualScalarConstArgs args = Program.GetOpArgs<OP::BoolEqualScalarConstArgs>(item.At);
             switch (item.Stage)
             {
             case 0:
@@ -5312,13 +5318,15 @@ namespace mu
     {
         MUTABLE_CPUPROFILER_SCOPE(RunCode_Int);
 
-		OP_TYPE type = pModel->GetPrivate()->m_program.GetOpType(item.At);
+		const FProgram& Program = m_pModel->GetPrivate()->m_program;
+
+		OP_TYPE type = Program.GetOpType(item.At);
         switch (type)
         {
 
         case OP_TYPE::NU_CONSTANT:
         {
-			OP::IntConstantArgs args = pModel->GetPrivate()->m_program.GetOpArgs<OP::IntConstantArgs>(item.At);
+			OP::IntConstantArgs args = Program.GetOpArgs<OP::IntConstantArgs>(item.At);
             int result = args.value;
             StoreInt( item, result );
             break;
@@ -5326,7 +5334,7 @@ namespace mu
 
         case OP_TYPE::NU_PARAMETER:
         {
-			OP::ParameterArgs args = pModel->GetPrivate()->m_program.GetOpArgs<OP::ParameterArgs>(item.At);
+			OP::ParameterArgs args = Program.GetOpArgs<OP::ParameterArgs>(item.At);
 			Ptr<RangeIndex> index = BuildCurrentOpRangeIndex( item, pParams, pModel, args.variable );
             int result = pParams->GetIntValue( args.variable, index );
 
@@ -5366,13 +5374,15 @@ namespace mu
     {
         MUTABLE_CPUPROFILER_SCOPE(RunCode_Scalar);
 
-		OP_TYPE type = pModel->GetPrivate()->m_program.GetOpType(item.At);
+		const FProgram& Program = m_pModel->GetPrivate()->m_program;
+
+		OP_TYPE type = Program.GetOpType(item.At);
         switch (type)
         {
 
         case OP_TYPE::SC_CONSTANT:
         {
-			OP::ScalarConstantArgs args = pModel->GetPrivate()->m_program.GetOpArgs<OP::ScalarConstantArgs>(item.At);
+			OP::ScalarConstantArgs args = Program.GetOpArgs<OP::ScalarConstantArgs>(item.At);
             float result = args.value;
             StoreScalar( item, result );
             break;
@@ -5380,7 +5390,7 @@ namespace mu
 
         case OP_TYPE::SC_PARAMETER:
         {
-			OP::ParameterArgs args = pModel->GetPrivate()->m_program.GetOpArgs<OP::ParameterArgs>(item.At);
+			OP::ParameterArgs args = Program.GetOpArgs<OP::ParameterArgs>(item.At);
 			Ptr<RangeIndex> index = BuildCurrentOpRangeIndex( item, pParams, pModel, args.variable );
             float result = pParams->GetFloatValue( args.variable, index );
             StoreScalar( item, result );
@@ -5389,7 +5399,7 @@ namespace mu
 
         case OP_TYPE::SC_CURVE:
         {
-			OP::ScalarCurveArgs args = pModel->GetPrivate()->m_program.GetOpArgs<OP::ScalarCurveArgs>(item.At);
+			OP::ScalarCurveArgs args = Program.GetOpArgs<OP::ScalarCurveArgs>(item.At);
             switch (item.Stage)
             {
             case 0:
@@ -5401,7 +5411,7 @@ namespace mu
             {
                 float time = LoadScalar( FCacheAddress(args.time,item) );
 
-                const Curve& curve = pModel->GetPrivate()->m_program.m_constantCurves[args.curve];
+                const Curve& curve = Program.m_constantCurves[args.curve];
                 float result = EvalCurve(curve, time);
 
                 StoreScalar( item, result );
@@ -5421,7 +5431,7 @@ namespace mu
 
         case OP_TYPE::SC_ARITHMETIC:
         {
-			OP::ArithmeticArgs args = pModel->GetPrivate()->m_program.GetOpArgs<OP::ArithmeticArgs>(item.At);
+			OP::ArithmeticArgs args = Program.GetOpArgs<OP::ArithmeticArgs>(item.At);
             switch (item.Stage)
             {
             case 0:
@@ -5481,16 +5491,18 @@ namespace mu
     {
         MUTABLE_CPUPROFILER_SCOPE(RunCode_String );
 
-		OP_TYPE type = pModel->GetPrivate()->m_program.GetOpType( item.At );
+		const FProgram& Program = m_pModel->GetPrivate()->m_program;
+
+		OP_TYPE type = Program.GetOpType( item.At );
         switch ( type )
         {
 
         case OP_TYPE::ST_CONSTANT:
         {
-			OP::ResourceConstantArgs args = pModel->GetPrivate()->m_program.GetOpArgs<OP::ResourceConstantArgs>( item.At );
+			OP::ResourceConstantArgs args = Program.GetOpArgs<OP::ResourceConstantArgs>( item.At );
             check( args.value < (uint32)pModel->GetPrivate()->m_program.m_constantStrings.Num() );
 
-            const FString& result = pModel->GetPrivate()->m_program.m_constantStrings[args.value];
+            const FString& result = Program.m_constantStrings[args.value];
             StoreString( item, new String(result) );
 
             break;
@@ -5498,7 +5510,7 @@ namespace mu
 
         case OP_TYPE::ST_PARAMETER:
         {
-			OP::ParameterArgs args = pModel->GetPrivate()->m_program.GetOpArgs<OP::ParameterArgs>( item.At );
+			OP::ParameterArgs args = Program.GetOpArgs<OP::ParameterArgs>( item.At );
 			Ptr<RangeIndex> index = BuildCurrentOpRangeIndex( item, pParams, pModel, args.variable );
 			FString result;
 			pParams->GetStringValue(args.variable, result, index);
@@ -5518,16 +5530,16 @@ namespace mu
     {
 		MUTABLE_CPUPROFILER_SCOPE(RunCode_Colour);
 
-		OP_TYPE type = pModel->GetPrivate()->m_program.GetOpType(item.At);
+		const FProgram& Program = m_pModel->GetPrivate()->m_program;
 
-        const FProgram& program = pModel->GetPrivate()->m_program;
+		OP_TYPE type = Program.GetOpType(item.At);
 
         switch ( type )
         {
 
         case OP_TYPE::CO_CONSTANT:
         {
-			OP::ColourConstantArgs args = program.GetOpArgs<OP::ColourConstantArgs>(item.At);
+			OP::ColourConstantArgs args = Program.GetOpArgs<OP::ColourConstantArgs>(item.At);
 			FVector4f result;
             result[0] = args.value[0];
             result[1] = args.value[1];
@@ -5539,7 +5551,7 @@ namespace mu
 
         case OP_TYPE::CO_PARAMETER:
         {
-			OP::ParameterArgs args = pModel->GetPrivate()->m_program.GetOpArgs<OP::ParameterArgs>(item.At);
+			OP::ParameterArgs args = Program.GetOpArgs<OP::ParameterArgs>(item.At);
 			Ptr<RangeIndex> index = BuildCurrentOpRangeIndex( item, pParams, pModel, args.variable );
             float r=0.0f;
             float g=0.0f;
@@ -5551,7 +5563,7 @@ namespace mu
 
         case OP_TYPE::CO_SAMPLEIMAGE:
         {
-			OP::ColourSampleImageArgs args = pModel->GetPrivate()->m_program.GetOpArgs<OP::ColourSampleImageArgs>(item.At);
+			OP::ColourSampleImageArgs args = Program.GetOpArgs<OP::ColourSampleImageArgs>(item.At);
             switch (item.Stage)
             {
             case 0:
@@ -5601,7 +5613,7 @@ namespace mu
 
         case OP_TYPE::CO_SWIZZLE:
         {
-			OP::ColourSwizzleArgs args = pModel->GetPrivate()->m_program.GetOpArgs<OP::ColourSwizzleArgs>(item.At);
+			OP::ColourSwizzleArgs args = Program.GetOpArgs<OP::ColourSwizzleArgs>(item.At);
             switch (item.Stage)
             {
             case 0:
@@ -5638,7 +5650,7 @@ namespace mu
 
         case OP_TYPE::CO_FROMSCALARS:
         {
-			OP::ColourFromScalarsArgs args = pModel->GetPrivate()->m_program.GetOpArgs<OP::ColourFromScalarsArgs>(item.At);
+			OP::ColourFromScalarsArgs args = Program.GetOpArgs<OP::ColourFromScalarsArgs>(item.At);
             switch (item.Stage)
             {
             case 0:
@@ -5674,7 +5686,7 @@ namespace mu
 
         case OP_TYPE::CO_ARITHMETIC:
         {
-			OP::ArithmeticArgs args = program.GetOpArgs<OP::ArithmeticArgs>(item.At);
+			OP::ArithmeticArgs args = Program.GetOpArgs<OP::ArithmeticArgs>(item.At);
             switch (item.Stage)
             {
             case 0:
@@ -5685,10 +5697,10 @@ namespace mu
 
             case 1:
             {
-				OP_TYPE otype = program.GetOpType( args.a );
+				OP_TYPE otype = Program.GetOpType( args.a );
                 DATATYPE dtype = GetOpDataType( otype );
                 check( dtype == DT_COLOUR );
-                otype = program.GetOpType( args.b );
+                otype = Program.GetOpType( args.b );
                 dtype = GetOpDataType( otype );
                 check( dtype == DT_COLOUR );
 				FVector4f a = args.a ? LoadColor( FCacheAddress( args.a, item ) )
@@ -5743,27 +5755,28 @@ namespace mu
     {
         MUTABLE_CPUPROFILER_SCOPE(RunCode_Projector);
 
-        const FProgram& program = pModel->GetPrivate()->m_program;
-		OP_TYPE type = program.GetOpType(item.At);
+		const FProgram& Program = m_pModel->GetPrivate()->m_program;
+
+		OP_TYPE type = Program.GetOpType(item.At);
         switch (type)
         {
 
         case OP_TYPE::PR_CONSTANT:
         {
-			OP::ResourceConstantArgs args = pModel->GetPrivate()->m_program.GetOpArgs<OP::ResourceConstantArgs>(item.At);
-            FProjector Result = program.m_constantProjectors[args.value];
+			OP::ResourceConstantArgs args = Program.GetOpArgs<OP::ResourceConstantArgs>(item.At);
+            FProjector Result = Program.m_constantProjectors[args.value];
             StoreProjector( item, Result );
             break;
         }
 
         case OP_TYPE::PR_PARAMETER:
         {
-			OP::ParameterArgs args = pModel->GetPrivate()->m_program.GetOpArgs<OP::ParameterArgs>(item.At);
+			OP::ParameterArgs args = Program.GetOpArgs<OP::ParameterArgs>(item.At);
 			Ptr<RangeIndex> index = BuildCurrentOpRangeIndex( item, pParams, pModel, args.variable );
             FProjector Result = pParams->GetPrivate()->GetProjectorValue(args.variable,index);
 
             // The type cannot be changed, take it from the default value
-            const FProjector& def = program.m_parameters[args.variable].m_defaultValue.Get<ParamProjectorType>();
+            const FProjector& def = Program.m_parameters[args.variable].m_defaultValue.Get<ParamProjectorType>();
             Result.type = def.type;
 
             StoreProjector( item, Result );
@@ -5782,16 +5795,18 @@ namespace mu
     {
         //MUTABLE_CPUPROFILER_SCOPE(RunCode_Layout);
 
-		OP_TYPE type = pModel->GetPrivate()->m_program.GetOpType(item.At);
+		const FProgram& Program = m_pModel->GetPrivate()->m_program;
+
+		OP_TYPE type = Program.GetOpType(item.At);
         switch (type)
         {
 
         case OP_TYPE::LA_CONSTANT:
         {
-			OP::ResourceConstantArgs args = pModel->GetPrivate()->m_program.GetOpArgs<OP::ResourceConstantArgs>(item.At);
+			OP::ResourceConstantArgs args = Program.GetOpArgs<OP::ResourceConstantArgs>(item.At);
             check( args.value < (uint32)pModel->GetPrivate()->m_program.m_constantLayouts.Num() );
 
-            LayoutPtrConst pResult = pModel->GetPrivate()->m_program.m_constantLayouts
+            LayoutPtrConst pResult = Program.m_constantLayouts
                     [ args.value ];
             StoreLayout( item, pResult );
             break;
@@ -5799,7 +5814,7 @@ namespace mu
 
         case OP_TYPE::LA_MERGE:
         {
-			OP::LayoutMergeArgs args = pModel->GetPrivate()->m_program.GetOpArgs<OP::LayoutMergeArgs>(item.At);
+			OP::LayoutMergeArgs args = Program.GetOpArgs<OP::LayoutMergeArgs>(item.At);
             switch (item.Stage)
             {
             case 0:
@@ -5841,7 +5856,7 @@ namespace mu
 
         case OP_TYPE::LA_PACK:
         {
-			OP::LayoutPackArgs args = pModel->GetPrivate()->m_program.GetOpArgs<OP::LayoutPackArgs>(item.At);
+			OP::LayoutPackArgs args = Program.GetOpArgs<OP::LayoutPackArgs>(item.At);
             switch (item.Stage)
             {
             case 0:
@@ -5885,7 +5900,7 @@ namespace mu
 
 		case OP_TYPE::LA_FROMMESH:
 		{
-			OP::LayoutFromMeshArgs args = pModel->GetPrivate()->m_program.GetOpArgs<OP::LayoutFromMeshArgs>(item.At);
+			OP::LayoutFromMeshArgs args = Program.GetOpArgs<OP::LayoutFromMeshArgs>(item.At);
 			switch (item.Stage)
 			{
 			case 0:
@@ -5913,7 +5928,7 @@ namespace mu
 
 		case OP_TYPE::LA_REMOVEBLOCKS:
 		{
-			OP::LayoutRemoveBlocksArgs args = pModel->GetPrivate()->m_program.GetOpArgs<OP::LayoutRemoveBlocksArgs>(item.At);
+			OP::LayoutRemoveBlocksArgs args = Program.GetOpArgs<OP::LayoutRemoveBlocksArgs>(item.At);
 			switch (item.Stage)
 			{
 			case 0:
@@ -5964,7 +5979,10 @@ namespace mu
 		check( item.Type == FScheduledOp::EType::Full );
 
 		const Model* pModel = InModel.Get();
-		OP_TYPE type = pModel->GetPrivate()->m_program.GetOpType(item.At);
+		
+		const FProgram& Program = pModel->GetPrivate()->m_program;
+
+		OP_TYPE type = Program.GetOpType(item.At);
 		//UE_LOG(LogMutableCore, Log, TEXT("Running :%5d , %d, of type %d "), item.At, item.Stage, type);
 
 		// Very spammy, for debugging purposes.
@@ -6080,23 +6098,23 @@ namespace mu
 		}
 
 
-		const FProgram& program = pModel->GetPrivate()->m_program;
+		const FProgram& Program = m_pModel->GetPrivate()->m_program;
 
-		OP_TYPE type = program.GetOpType(item.At);
+		OP_TYPE type = Program.GetOpType(item.At);
 		switch (type)
 		{
 
 		case OP_TYPE::IM_CONSTANT:
 		{
 			check(item.Stage == 0);
-			OP::ResourceConstantArgs args = program.GetOpArgs<OP::ResourceConstantArgs>(item.At);
+			OP::ResourceConstantArgs args = Program.GetOpArgs<OP::ResourceConstantArgs>(item.At);
 			int32 ImageIndex = args.value;
 
 			FImageDesc& Result = m_heapImageDesc[item.CustomState];
-			Result.m_format = program.m_constantImages[ImageIndex].ImageFormat;
-			Result.m_size[0] = program.m_constantImages[ImageIndex].ImageSizeX;
-			Result.m_size[1] = program.m_constantImages[ImageIndex].ImageSizeY;
-			Result.m_lods = program.m_constantImages[ImageIndex].LODCount;
+			Result.m_format = Program.m_constantImages[ImageIndex].ImageFormat;
+			Result.m_size[0] = Program.m_constantImages[ImageIndex].ImageSizeX;
+			Result.m_size[1] = Program.m_constantImages[ImageIndex].ImageSizeY;
+			Result.m_lods = Program.m_constantImages[ImageIndex].LODCount;
 			StoreValidDesc(item);
 			break;
 		}
@@ -6104,7 +6122,7 @@ namespace mu
 		case OP_TYPE::IM_PARAMETER:
 		{
 			check(item.Stage == 0);
-			OP::ParameterArgs args = program.GetOpArgs<OP::ParameterArgs>(item.At);
+			OP::ParameterArgs args = Program.GetOpArgs<OP::ParameterArgs>(item.At);
 			FName Id = pParams->GetImageValue(args.variable);
 			uint8 MipsToSkip = item.ExecutionOptions;
 			m_heapImageDesc[item.CustomState] = GetExternalImageDesc(Id, MipsToSkip);
@@ -6115,7 +6133,7 @@ namespace mu
 		case OP_TYPE::IM_REFERENCE:
 		{
 			check(item.Stage == 0);
-			OP::ResourceReferenceArgs Args = program.GetOpArgs<OP::ResourceReferenceArgs>(item.At);
+			OP::ResourceReferenceArgs Args = Program.GetOpArgs<OP::ResourceReferenceArgs>(item.At);
 			FImageDesc& Result = m_heapImageDesc[item.CustomState];
 			Result = Args.ImageDesc;
 			StoreValidDesc(item);
@@ -6124,7 +6142,7 @@ namespace mu
 
 		case OP_TYPE::IM_CONDITIONAL:
 		{
-			OP::ConditionalArgs args = program.GetOpArgs<OP::ConditionalArgs>(item.At);
+			OP::ConditionalArgs args = Program.GetOpArgs<OP::ConditionalArgs>(item.At);
 			switch (item.Stage)
 			{
 			case 0:
@@ -6152,7 +6170,7 @@ namespace mu
 
 		case OP_TYPE::IM_SWITCH:
 		{
-			const uint8* data = program.GetOpArgsPointer(item.At);
+			const uint8* data = Program.GetOpArgsPointer(item.At);
 		
 			OP::ADDRESS VarAddress;
 			FMemory::Memcpy( &VarAddress, data, sizeof(OP::ADDRESS));
@@ -6221,7 +6239,7 @@ namespace mu
 
 		case OP_TYPE::IM_LAYERCOLOUR:
 		{
-			OP::ImageLayerColourArgs args = program.GetOpArgs<OP::ImageLayerColourArgs>(item.At);
+			OP::ImageLayerColourArgs args = Program.GetOpArgs<OP::ImageLayerColourArgs>(item.At);
 			switch (item.Stage)
 			{
 			case 0: AddOp(FScheduledOp(item.At, item, 1), FScheduledOp(args.base, item)); break;
@@ -6233,7 +6251,7 @@ namespace mu
 
 		case OP_TYPE::IM_LAYER:
 		{
-			OP::ImageLayerArgs args = program.GetOpArgs<OP::ImageLayerArgs>(item.At);
+			OP::ImageLayerArgs args = Program.GetOpArgs<OP::ImageLayerArgs>(item.At);
 			switch (item.Stage)
 			{
 			case 0: AddOp(FScheduledOp(item.At, item, 1), FScheduledOp(args.base, item)); break;
@@ -6245,7 +6263,7 @@ namespace mu
 
 		case OP_TYPE::IM_MULTILAYER:
 		{
-			OP::ImageMultiLayerArgs args = program.GetOpArgs<OP::ImageMultiLayerArgs>(item.At);
+			OP::ImageMultiLayerArgs args = Program.GetOpArgs<OP::ImageMultiLayerArgs>(item.At);
 			switch (item.Stage)
 			{
 			case 0: AddOp(FScheduledOp(item.At, item, 1), FScheduledOp(args.base, item)); break;
@@ -6257,7 +6275,7 @@ namespace mu
 
 		case OP_TYPE::IM_NORMALCOMPOSITE:
 		{
-			OP::ImageNormalCompositeArgs args = program.GetOpArgs<OP::ImageNormalCompositeArgs>(item.At);
+			OP::ImageNormalCompositeArgs args = Program.GetOpArgs<OP::ImageNormalCompositeArgs>(item.At);
 			switch (item.Stage)
 			{
 			case 0: AddOp(FScheduledOp(item.At, item, 1), FScheduledOp(args.base, item)); break;
@@ -6269,7 +6287,7 @@ namespace mu
 
 		case OP_TYPE::IM_PIXELFORMAT:
 		{
-			OP::ImagePixelFormatArgs args = program.GetOpArgs<OP::ImagePixelFormatArgs>(item.At);
+			OP::ImagePixelFormatArgs args = Program.GetOpArgs<OP::ImagePixelFormatArgs>(item.At);
 			switch (item.Stage)
 			{
 			case 0:
@@ -6301,7 +6319,7 @@ namespace mu
 
 		case OP_TYPE::IM_MIPMAP:
 		{
-			OP::ImageMipmapArgs args = program.GetOpArgs<OP::ImageMipmapArgs>(item.At);
+			OP::ImageMipmapArgs args = Program.GetOpArgs<OP::ImageMipmapArgs>(item.At);
 			switch (item.Stage)
 			{
 			case 0:
@@ -6343,7 +6361,7 @@ namespace mu
 
 		case OP_TYPE::IM_RESIZE:
 		{
-			OP::ImageResizeArgs args = program.GetOpArgs<OP::ImageResizeArgs>(item.At);
+			OP::ImageResizeArgs args = Program.GetOpArgs<OP::ImageResizeArgs>(item.At);
 			switch (item.Stage)
 			{
 			case 0:
@@ -6364,7 +6382,7 @@ namespace mu
 
 		case OP_TYPE::IM_RESIZELIKE:
 		{
-			OP::ImageResizeLikeArgs args = program.GetOpArgs<OP::ImageResizeLikeArgs>(item.At);
+			OP::ImageResizeLikeArgs args = Program.GetOpArgs<OP::ImageResizeLikeArgs>(item.At);
 			switch (item.Stage)
 			{
 			case 0:
@@ -6400,7 +6418,7 @@ namespace mu
 
 		case OP_TYPE::IM_RESIZEREL:
 		{
-			OP::ImageResizeRelArgs args = program.GetOpArgs<OP::ImageResizeRelArgs>(item.At);
+			OP::ImageResizeRelArgs args = Program.GetOpArgs<OP::ImageResizeRelArgs>(item.At);
 			switch (item.Stage)
 			{
 			case 0:
@@ -6427,7 +6445,7 @@ namespace mu
 
 		case OP_TYPE::IM_BLANKLAYOUT:
 		{
-			OP::ImageBlankLayoutArgs args = program.GetOpArgs<OP::ImageBlankLayoutArgs>(item.At);
+			OP::ImageBlankLayoutArgs args = Program.GetOpArgs<OP::ImageBlankLayoutArgs>(item.At);
 			switch (item.Stage)
 			{
 			case 0:
@@ -6476,7 +6494,7 @@ namespace mu
 
 		case OP_TYPE::IM_COMPOSE:
 		{
-			OP::ImageComposeArgs args = program.GetOpArgs<OP::ImageComposeArgs>(item.At);
+			OP::ImageComposeArgs args = Program.GetOpArgs<OP::ImageComposeArgs>(item.At);
 			switch (item.Stage)
 			{
 			case 0: AddOp(FScheduledOp(item.At, item, 1), FScheduledOp(args.base, item)); break;
@@ -6488,7 +6506,7 @@ namespace mu
 
 		case OP_TYPE::IM_INTERPOLATE:
 		{
-			OP::ImageInterpolateArgs args = program.GetOpArgs<OP::ImageInterpolateArgs>(item.At);
+			OP::ImageInterpolateArgs args = Program.GetOpArgs<OP::ImageInterpolateArgs>(item.At);
 			switch (item.Stage)
 			{
 			case 0: AddOp(FScheduledOp(item.At, item, 1), FScheduledOp(args.targets[0], item)); break;
@@ -6500,7 +6518,7 @@ namespace mu
 
 		case OP_TYPE::IM_SATURATE:
 		{
-			OP::ImageSaturateArgs args = program.GetOpArgs<OP::ImageSaturateArgs>(item.At);
+			OP::ImageSaturateArgs args = Program.GetOpArgs<OP::ImageSaturateArgs>(item.At);
 			switch (item.Stage)
 			{
 			case 0: AddOp(FScheduledOp(item.At, item, 1), FScheduledOp(args.base, item)); break;
@@ -6512,7 +6530,7 @@ namespace mu
 
 		case OP_TYPE::IM_LUMINANCE:
 		{
-			OP::ImageLuminanceArgs args = program.GetOpArgs<OP::ImageLuminanceArgs>(item.At);
+			OP::ImageLuminanceArgs args = Program.GetOpArgs<OP::ImageLuminanceArgs>(item.At);
 			switch (item.Stage)
 			{
 			case 0:
@@ -6533,7 +6551,7 @@ namespace mu
 
 		case OP_TYPE::IM_SWIZZLE:
 		{
-			OP::ImageSwizzleArgs args = program.GetOpArgs<OP::ImageSwizzleArgs>(item.At);
+			OP::ImageSwizzleArgs args = Program.GetOpArgs<OP::ImageSwizzleArgs>(item.At);
 			switch (item.Stage)
 			{
 			case 0:
@@ -6554,7 +6572,7 @@ namespace mu
 
 		case OP_TYPE::IM_COLOURMAP:
 		{
-			OP::ImageColourMapArgs args = program.GetOpArgs<OP::ImageColourMapArgs>(item.At);
+			OP::ImageColourMapArgs args = Program.GetOpArgs<OP::ImageColourMapArgs>(item.At);
 			switch (item.Stage)
 			{
 			case 0: AddOp(FScheduledOp(item.At, item, 1), FScheduledOp(args.base, item)); break;
@@ -6566,7 +6584,7 @@ namespace mu
 
 		case OP_TYPE::IM_GRADIENT:
 		{
-			OP::ImageGradientArgs args = program.GetOpArgs<OP::ImageGradientArgs>(item.At);
+			OP::ImageGradientArgs args = Program.GetOpArgs<OP::ImageGradientArgs>(item.At);
 			m_heapImageDesc[item.CustomState].m_size[0] = args.size[0];
 			m_heapImageDesc[item.CustomState].m_size[1] = args.size[1];
 			m_heapImageDesc[item.CustomState].m_lods = 1;
@@ -6577,7 +6595,7 @@ namespace mu
 
 		case OP_TYPE::IM_BINARISE:
 		{
-			OP::ImageBinariseArgs args = program.GetOpArgs<OP::ImageBinariseArgs>(item.At);
+			OP::ImageBinariseArgs args = Program.GetOpArgs<OP::ImageBinariseArgs>(item.At);
 			switch (item.Stage)
 			{
 			case 0:
@@ -6597,7 +6615,7 @@ namespace mu
 
 		case OP_TYPE::IM_INVERT:
 		{
-			OP::ImageInvertArgs args = program.GetOpArgs<OP::ImageInvertArgs>(item.At);
+			OP::ImageInvertArgs args = Program.GetOpArgs<OP::ImageInvertArgs>(item.At);
 			switch (item.Stage)
 			{
 			case 0: AddOp(FScheduledOp(item.At, item, 1), FScheduledOp(args.base, item)); break;
@@ -6609,7 +6627,7 @@ namespace mu
 
 		case OP_TYPE::IM_PLAINCOLOUR:
 		{
-			OP::ImagePlainColourArgs args = program.GetOpArgs<OP::ImagePlainColourArgs>(item.At);
+			OP::ImagePlainColourArgs args = Program.GetOpArgs<OP::ImagePlainColourArgs>(item.At);
 			m_heapImageDesc[item.CustomState].m_size[0] = args.size[0];
 			m_heapImageDesc[item.CustomState].m_size[1] = args.size[1];
 			m_heapImageDesc[item.CustomState].m_lods = args.LODs;
@@ -6620,7 +6638,7 @@ namespace mu
 
 		case OP_TYPE::IM_CROP:
 		{
-			OP::ImageCropArgs args = program.GetOpArgs<OP::ImageCropArgs>(item.At);
+			OP::ImageCropArgs args = Program.GetOpArgs<OP::ImageCropArgs>(item.At);
 			switch (item.Stage)
 			{
 			case 0:
@@ -6642,7 +6660,7 @@ namespace mu
 
 		case OP_TYPE::IM_PATCH:
 		{
-			OP::ImagePatchArgs args = program.GetOpArgs<OP::ImagePatchArgs>(item.At);
+			OP::ImagePatchArgs args = Program.GetOpArgs<OP::ImagePatchArgs>(item.At);
 			switch (item.Stage)
 			{
 			case 0: AddOp(FScheduledOp(item.At, item, 1), FScheduledOp(args.base, item)); break;
@@ -6654,7 +6672,7 @@ namespace mu
 
 		case OP_TYPE::IM_RASTERMESH:
 		{
-			OP::ImageRasterMeshArgs args = program.GetOpArgs<OP::ImageRasterMeshArgs>(item.At);
+			OP::ImageRasterMeshArgs args = Program.GetOpArgs<OP::ImageRasterMeshArgs>(item.At);
 			m_heapImageDesc[item.CustomState].m_size[0] = args.sizeX;
 			m_heapImageDesc[item.CustomState].m_size[1] = args.sizeY;
 			m_heapImageDesc[item.CustomState].m_lods = 1;
@@ -6665,7 +6683,7 @@ namespace mu
 
 		case OP_TYPE::IM_MAKEGROWMAP:
 		{
-			OP::ImageMakeGrowMapArgs args = program.GetOpArgs<OP::ImageMakeGrowMapArgs>(item.At);
+			OP::ImageMakeGrowMapArgs args = Program.GetOpArgs<OP::ImageMakeGrowMapArgs>(item.At);
 			switch (item.Stage)
 			{
 			case 0:
@@ -6687,7 +6705,7 @@ namespace mu
 
 		case OP_TYPE::IM_DISPLACE:
 		{
-			OP::ImageDisplaceArgs args = program.GetOpArgs<OP::ImageDisplaceArgs>(item.At);
+			OP::ImageDisplaceArgs args = Program.GetOpArgs<OP::ImageDisplaceArgs>(item.At);
 			switch (item.Stage)
 			{
 			case 0: AddOp(FScheduledOp(item.At, item, 1), FScheduledOp(args.source, item)); break;
@@ -6700,7 +6718,7 @@ namespace mu
         case OP_TYPE::IM_TRANSFORM:
         {
 
-			OP::ImageTransformArgs Args = program.GetOpArgs<OP::ImageTransformArgs>(item.At);
+			OP::ImageTransformArgs Args = Program.GetOpArgs<OP::ImageTransformArgs>(item.At);
 
             switch (item.Stage)
             {
diff --git a/Engine/Plugins/Experimental/Mutable/Source/MutableRuntime/Private/MuR/CodeRunnerTasked.cpp b/Engine/Plugins/Experimental/Mutable/Source/MutableRuntime/Private/MuR/CodeRunnerTasked.cpp
--- a/Engine/Plugins/Experimental/Mutable/Source/MutableRuntime/Private/MuR/CodeRunnerTasked.cpp
+++ b/Engine/Plugins/Experimental/Mutable/Source/MutableRuntime/Private/MuR/CodeRunnerTasked.cpp
@@ -1435,6 +1435,26 @@ namespace mu
 
 		int32 ResultLODs = Sources[0]->GetLODCount();
 
+		// Be defensive: ensure formats are uncompressed
+		for (int i = 0; i < MUTABLE_OP_MAX_SWIZZLE_CHANNELS; ++i)
+		{
+			if (Sources[i] && Sources[i]->GetFormat()!=GetUncompressedFormat(Sources[i]->GetFormat()))
+			{
+				MUTABLE_CPUPROFILER_SCOPE(ImageFormat_ForSwizzle);
+
+				EImageFormat UncompressedFormat = GetUncompressedFormat(Sources[i]->GetFormat());
+				Ptr<Image> Formatted = Runner->CreateImage(Sources[i]->GetSizeX(), Sources[i]->GetSizeY(), 1, UncompressedFormat, EInitializationType::NotInitialized);
+				bool bSuccess = false;
+				int32 ImageCompressionQuality = 4; // TODO
+				ImOp.ImagePixelFormat(bSuccess, ImageCompressionQuality, Formatted.get(), Sources[i].get());
+				check(bSuccess); // Decompression cannot fail
+				Runner->Release(Sources[i]);
+				Sources[i] = Formatted;
+				ResultLODs = 1;
+			}
+		}
+
+
 		// Be defensive: ensure image sizes match.
 		for (int i = 1; i < MUTABLE_OP_MAX_SWIZZLE_CHANNELS; ++i)
 		{
diff --git a/Engine/Plugins/Experimental/Mutable/Source/MutableRuntime/Private/MuR/CodeVisitor.h b/Engine/Plugins/Experimental/Mutable/Source/MutableRuntime/Private/MuR/CodeVisitor.h
--- a/Engine/Plugins/Experimental/Mutable/Source/MutableRuntime/Private/MuR/CodeVisitor.h
+++ b/Engine/Plugins/Experimental/Mutable/Source/MutableRuntime/Private/MuR/CodeVisitor.h
@@ -16,7 +16,6 @@
 
 namespace mu
 {
-class Parameters;
 
     //---------------------------------------------------------------------------------------------
     //! Decide what operations are an "add resource" since they are handled differently sometimes.
@@ -28,7 +27,6 @@ class Parameters;
     }
 
 
-
     //---------------------------------------------------------------------------------------------
     //! Code visitor that:
     //! - is top-down
@@ -37,7 +35,7 @@ class Parameters;
     //! - Its iterative
     //---------------------------------------------------------------------------------------------
     template<typename STATE=int>
-    class UniqueConstCodeVisitorIterative : public Base
+    class UniqueConstCodeVisitorIterative
     {
     public:
 
@@ -152,7 +150,7 @@ class Parameters;
         //! States found so far
 		TArray<STATE> m_states;
 
-        //! Index of the current state, from the m_states vector.
+        //! Index of the current state, from the m_states array.
         int m_currentState;
 
         //! If true, operations adding resources (meshes or images) will only
@@ -160,7 +158,7 @@ class Parameters;
         bool m_skipResources;
 
         //! Array of states visited for each operation.
-        //! Empty vector means operation not visited at all.
+        //! Empty array means operation not visited at all.
 		TArray<TArray<int>> m_visited;
 
         //! Process all the pending operations and visit all children if necessary
@@ -210,7 +208,7 @@ class Parameters;
     //! - Its iterative
     //---------------------------------------------------------------------------------------------
     template<class STATE=int>
-    class RepeatConstCodeVisitorIterative : public Base
+    class RepeatConstCodeVisitorIterative
     {
     public:
 
@@ -329,7 +327,7 @@ class Parameters;
         //! States found so far
 		TArray<STATE> m_states;
 
-        //! Index of the current state, from the m_states vector.
+        //! Index of the current state, from the m_states array.
         int m_currentState;
 
         //! If true, operations adding resources (meshes or images) will only
@@ -653,7 +651,7 @@ class Parameters;
     //! Calculate all the parameters found under a particular operation
     //! It has an internal cache, so don't reuse objects of this class if the program changes.
     //---------------------------------------------------------------------------------------------
-    class MUTABLERUNTIME_API SubtreeParametersVisitor : public Base
+    class MUTABLERUNTIME_API SubtreeParametersVisitor
     {
     public:
 
diff --git a/Engine/Plugins/Experimental/Mutable/Source/MutableRuntime/Private/MuR/ImagePrivate.h b/Engine/Plugins/Experimental/Mutable/Source/MutableRuntime/Private/MuR/ImagePrivate.h
--- a/Engine/Plugins/Experimental/Mutable/Source/MutableRuntime/Private/MuR/ImagePrivate.h
+++ b/Engine/Plugins/Experimental/Mutable/Source/MutableRuntime/Private/MuR/ImagePrivate.h
@@ -22,11 +22,21 @@ namespace mu
     //---------------------------------------------------------------------------------------------
 	inline EImageFormat GetMostGenericFormat(EImageFormat FormatA, EImageFormat FormatB)
     {
-        if (FormatA == FormatB) 
-		{	
+		if (FormatA == FormatB)
+		{
+			return FormatA;
+		}
+
+		if (FormatA == EImageFormat::IF_NONE)
+		{
 			return FormatA;
 		}
 
+		if (FormatB == EImageFormat::IF_NONE)
+		{
+			return FormatB;
+		}
+
         if (GetImageFormatData(FormatA).Channels > GetImageFormatData(FormatB).Channels)
 		{	
 			return FormatA;
@@ -37,12 +47,15 @@ namespace mu
 			return FormatB;
 		}
 
-        if (FormatA == EImageFormat::IF_BC2 || FormatA == EImageFormat::IF_BC3 || FormatA == EImageFormat::IF_ASTC_4x4_RGBA_LDR)
+        if (FormatA == EImageFormat::IF_BC2 || FormatA == EImageFormat::IF_BC3 
+			|| FormatA == EImageFormat::IF_ASTC_4x4_RGBA_LDR || FormatA == EImageFormat::IF_ASTC_6x6_RGBA_LDR || FormatA == EImageFormat::IF_ASTC_8x8_RGBA_LDR || FormatA == EImageFormat::IF_ASTC_10x10_RGBA_LDR)
 		{	
 			return FormatA;
 		}
 
-        if (FormatB == EImageFormat::IF_BC2 || FormatB == EImageFormat::IF_BC3 || FormatB == EImageFormat::IF_ASTC_4x4_RGBA_LDR)
+        if (FormatB == EImageFormat::IF_BC2 || FormatB == EImageFormat::IF_BC3
+			|| FormatB == EImageFormat::IF_ASTC_4x4_RGBA_LDR || FormatB == EImageFormat::IF_ASTC_6x6_RGBA_LDR || FormatB == EImageFormat::IF_ASTC_8x8_RGBA_LDR || FormatB == EImageFormat::IF_ASTC_10x10_RGBA_LDR)
+
 		{
 			return FormatB;
 		}
diff --git a/Engine/Plugins/Experimental/Mutable/Source/MutableRuntime/Private/MuR/InstancePrivate.h b/Engine/Plugins/Experimental/Mutable/Source/MutableRuntime/Private/MuR/InstancePrivate.h
--- a/Engine/Plugins/Experimental/Mutable/Source/MutableRuntime/Private/MuR/InstancePrivate.h
+++ b/Engine/Plugins/Experimental/Mutable/Source/MutableRuntime/Private/MuR/InstancePrivate.h
@@ -123,7 +123,7 @@ namespace mu
 		FName Name;
 	};
 
-	class Instance::Private : public Base
+	class Instance::Private
 	{
 	public:
 
diff --git a/Engine/Plugins/Experimental/Mutable/Source/MutableRuntime/Private/MuR/OpImagePixelFormat.cpp b/Engine/Plugins/Experimental/Mutable/Source/MutableRuntime/Private/MuR/OpImagePixelFormat.cpp
--- a/Engine/Plugins/Experimental/Mutable/Source/MutableRuntime/Private/MuR/OpImagePixelFormat.cpp
+++ b/Engine/Plugins/Experimental/Mutable/Source/MutableRuntime/Private/MuR/OpImagePixelFormat.cpp
@@ -11,6 +11,7 @@
 #include "MuR/ImageRLE.h"
 #include "MuR/MutableMath.h"
 #include "MuR/MutableTrace.h"
+#include "MuR/MutableRuntimeModule.h"
 #include "MuR/Ptr.h"
 
 
@@ -1226,8 +1227,20 @@ void FImageOperator::ImagePixelFormat( bool& bOutSuccess, int32 CompressionQuali
             }
 
             default:
-                // Case not implemented
-                check( false );
+			{
+				UE_LOG(LogMutableCore, Log, TEXT("Image format conversion not implemented. Expensive generic one used. "));
+
+				// Case not implemented, try a generic more expensive operatino.
+				ImagePtr Temp = ImagePixelFormat(CompressionQuality, Base, EImageFormat::IF_RGBA_UBYTE);
+				for (int m = 0; m < resultLODCount; ++m)
+				{
+					FIntVector2 mipSize = pResult->CalculateMipSize(m);
+					miro::RGBA_to_BC1(
+						mipSize[0], mipSize[1], Base->GetMipData(baseLOD + m),
+						pResult->GetMipData(m), CompressionQuality);
+				}
+				ReleaseImage(Temp);
+			}
 
             }
 
@@ -1249,8 +1262,20 @@ void FImageOperator::ImagePixelFormat( bool& bOutSuccess, int32 CompressionQuali
                 break;
 
             default:
-                // Case not implemented
-                check( false );
+			{
+				UE_LOG(LogMutableCore, Log, TEXT("Image format conversion not implemented. Expensive generic one used. "));
+
+				// Case not implemented, try a generic more expensive operatino.
+				ImagePtr Temp = ImagePixelFormat(CompressionQuality, Base, EImageFormat::IF_RGBA_UBYTE);
+				for (int m = 0; m < resultLODCount; ++m)
+				{
+					FIntVector2 mipSize = pResult->CalculateMipSize(m);
+					miro::RGBA_to_BC2(
+						mipSize[0], mipSize[1], Base->GetMipData(baseLOD + m),
+						pResult->GetMipData(m), CompressionQuality);
+				}
+				ReleaseImage(Temp);
+			}
 
             }
 
@@ -1296,8 +1321,20 @@ void FImageOperator::ImagePixelFormat( bool& bOutSuccess, int32 CompressionQuali
             }
 
             default:
-                // Case not implemented
-                check( false );
+			{
+				UE_LOG(LogMutableCore, Log, TEXT("Image format conversion not implemented. Expensive generic one used. "));
+
+				// Case not implemented, try a generic more expensive operatino.
+				ImagePtr Temp = ImagePixelFormat(CompressionQuality, Base, EImageFormat::IF_RGBA_UBYTE);
+				for (int m = 0; m < resultLODCount; ++m)
+				{
+					FIntVector2 mipSize = pResult->CalculateMipSize(m);
+					miro::RGBA_to_BC3(
+						mipSize[0], mipSize[1], Base->GetMipData(baseLOD + m),
+						pResult->GetMipData(m), CompressionQuality);
+				}
+				ReleaseImage(Temp);
+			}               
 
             }
 
