[mutable] Redesign code generator to not rely on the visitor pattern and reducing RTTI usage.

diff --git a/Engine/Plugins/Experimental/Mutable/Source/MutableRuntime/Private/MuR/OpLayoutPack.h b/Engine/Plugins/Experimental/Mutable/Source/MutableRuntime/Private/MuR/OpLayoutPack.h
--- a/Engine/Plugins/Experimental/Mutable/Source/MutableRuntime/Private/MuR/OpLayoutPack.h
+++ b/Engine/Plugins/Experimental/Mutable/Source/MutableRuntime/Private/MuR/OpLayoutPack.h
@@ -11,136 +11,16 @@
 
 namespace mu
 {
-	//---------------------------------------------------------------------------------------------
-    // inline void SimpleLayoutPack( Layout* pResult, const Layout* pSource )
-	// {
-    //     check( pResult->GetBlockCount() == pSource->GetBlockCount() );
-
-	// 	int blockCount = pSource->GetBlockCount();
-	// 	vector< vec2<int> > blocks( blockCount );
-
-	// 	// Look for the maximum block sizes on the layout and the total area
-	// 	int maxX = 0;
-	// 	int maxY = 0;
-	// 	int area = 0;
-	// 	for ( int index=0; index<blockCount; ++index )
-	// 	{
-	// 		box< vec2<int> > b;
-	// 		pSource->GetBlock( index, &b.min[0], &b.min[1], &b.size[0], &b.size[1] );
-
-	// 		maxX = FMath::Max( maxX, b.size[0] );
-	// 		maxY = FMath::Max( maxY, b.size[1] );
-
-	// 		area += b.size[0] * b.size[1];
-
-	// 		blocks[index] = b.size;
-	// 	}
-
-	// 	// Grow until the area is big enough to fit all blocks. We always grow X first, because
-	// 	// in case we cannot pack everything, we will grow Y with the current horizon algorithm.
-	// 	maxX = ceilPow2( maxX );
-	// 	maxY = ceilPow2( maxY );
-	// 	while ( maxX*maxY<area )
-	// 	{
-	// 		if (maxX>maxY)
-	// 		{
-	// 			maxY*=2;
-	// 		}
-	// 		else
-	// 		{
-	// 			maxX*=2;
-	// 		}
-	// 	}
-
-	// 	// Sort by height, area
-	// 	vector<int> sorted;
-	// 	sorted.reserve( blockCount );
-	// 	for ( int index=0; index<blockCount; ++index )
-	// 	{
-	// 		int thisHeight = blocks[index][1];
-	// 		int thisArea = blocks[index][0] * blocks[index][1];
-
-	// 		int p;
-	// 		for ( p=0; p<(int)sorted.Num(); ++p )
-	// 		{
-	// 			int pos = sorted[p];
-	// 			int posHeight = blocks[pos][1];
-	// 			int posArea = blocks[pos][0] * blocks[pos][1];
-
-	// 			if ( posHeight<thisHeight || ( posHeight==thisHeight && posArea<thisArea ))
-	// 			{
-	// 				break;
-	// 			}
-	// 		}
-
-	// 		sorted.insert( sorted.begin()+p, index );
-	// 	}
-
-	// 	// Pack with fixed horizontal size
-	// 	vector<int> horizon( maxX, 0 );
-	// 	vector< vec2<int> > positions( blockCount );
-	// 	maxY = 0;
-
-	// 	for ( int p=0; p<(int)sorted.Num(); ++p )
-	// 	{
-	// 		int index = sorted[p];
-
-	// 		// Seek for the lowest span where the block fits
-	// 		int currentLevel = std::numeric_limits<int>::max();
-	// 		int currentX = -1;
-	// 		for ( int x=0; x<=maxX-blocks[index][0]; ++x )
-	// 		{
-	// 			int level = 0;
-	// 			for( int xs=x; xs<x+blocks[index][0]; ++xs )
-	// 			{
-	// 				level = FMath::Max( level, horizon[xs] );
-	// 			}
-
-	// 			if (level<currentLevel)
-	// 			{
-	// 				currentLevel = level;
-	// 				currentX = x;
-	// 			}
-
-	// 		}
-
-	// 		check( currentX>=0 && currentX<=maxX-blocks[index][0] );
-
-	// 		// Update horizon
-	// 		for( int xs=currentX; xs<currentX+blocks[index][0]; ++xs )
-	// 		{
-	// 			horizon[xs] = currentLevel+blocks[index][1];
-	// 		}
-
-	// 		// Store
-	// 		positions[ index ] = vec2<int>( currentX, currentLevel );
-	// 		maxY = FMath::Max( maxY, currentLevel+blocks[index][1] );
-	// 	}
-
-    //     // Set data in the result
-	// 	maxY = ceilPow2( maxY );
-	// 	pResult->SetGridSize( maxX, maxY );
-	// 	for ( int index=0; index<blockCount; ++index )
-	// 	{
-	// 		pResult->SetBlock
-	// 			(
-	// 				index,
-	// 				positions[index][0], positions[index][1],
-    //                 blocks[index][0], blocks[index][1]
-	// 			);
-	// 	}
-	// }
-	
 
 	//---------------------------------------------------------------------------------------------
-	struct LAY_BLOCK
+	struct FLayoutBlock
 	{
-		LAY_BLOCK()
+		FLayoutBlock()
 		{
 			index = -1;
 		}
 
-		LAY_BLOCK( int32 i, UE::Math::TIntVector2<uint16> s, int32 p=0, bool bInReduceBothAxes = false, bool bInReduceByTwo = false )
+		FLayoutBlock( int32 i, UE::Math::TIntVector2<uint16> s, int32 p=0, bool bInReduceBothAxes = false, bool bInReduceByTwo = false )
 		{
 			index = i;
 			size = s;
@@ -156,7 +36,7 @@ namespace mu
 		bool bReduceByTwo;
 	};
 
-    inline bool CompareBlocks( const LAY_BLOCK& a, const LAY_BLOCK& b )
+    inline bool CompareBlocks( const FLayoutBlock& a, const FLayoutBlock& b )
     {
         if ( a.size[1]>b.size[1] )
         {
@@ -186,7 +66,7 @@ namespace mu
         }
     }
 
-	inline bool CompareBlocksPriority(const LAY_BLOCK& a, const LAY_BLOCK& b)
+	inline bool CompareBlocksPriority(const FLayoutBlock& a, const FLayoutBlock& b)
 	{
 		if (a.priority == b.priority)
 		{
@@ -204,7 +84,7 @@ namespace mu
     struct SCRATCH_LAYOUT_PACK
     {
         TArray< UE::Math::TIntVector2<uint16> > blocks;
-		TArray< LAY_BLOCK > sorted;
+		TArray< FLayoutBlock > sorted;
 		TArray< vec2<int> > positions;
 		TArray< int > priorities;
 		TArray< vec2<int> > reductions;
@@ -282,7 +162,7 @@ namespace mu
     //     check( (int)scratch->sorted.Num()==blockCount );
     //     for ( int index=0; index<blockCount; ++index )
     //     {
-    //         scratch->sorted[index] = LAY_BLOCK( index, scratch->blocks[index], scratch->priorities[index] );
+    //         scratch->sorted[index] = FLayoutBlock( index, scratch->blocks[index], scratch->priorities[index] );
     //     }
     //     std::sort( scratch->sorted.begin(), scratch->sorted.end(), CompareBlocks );
 
@@ -831,7 +711,7 @@ namespace mu
         check( (int)scratch->sorted.Num()==blockCount );
         for ( int index=0; index<blockCount; ++index )
         {
-            scratch->sorted[index] = LAY_BLOCK( index, scratch->blocks[index], scratch->priorities[index], (bool)scratch->ReduceBothAxes[index], (bool)scratch->ReduceByTwo[index]);
+            scratch->sorted[index] = FLayoutBlock( index, scratch->blocks[index], scratch->priorities[index], (bool)scratch->ReduceBothAxes[index], (bool)scratch->ReduceByTwo[index]);
         }
 
 		// Sort blocks by height, area
diff --git a/Engine/Plugins/Experimental/Mutable/Source/MutableRuntime/Private/MuR/Serialisation.cpp b/Engine/Plugins/Experimental/Mutable/Source/MutableRuntime/Private/MuR/Serialisation.cpp
--- a/Engine/Plugins/Experimental/Mutable/Source/MutableRuntime/Private/MuR/Serialisation.cpp
+++ b/Engine/Plugins/Experimental/Mutable/Source/MutableRuntime/Private/MuR/Serialisation.cpp
@@ -253,7 +253,7 @@ namespace mu
     //---------------------------------------------------------------------------------------------
     //---------------------------------------------------------------------------------------------
     //---------------------------------------------------------------------------------------------
-    class InputArchiveWithProxies::Private : public Base
+    class InputArchiveWithProxies::Private
     {
     public:
         TArray< Ptr<ResourceProxy<Image>> > m_proxyHistory;
diff --git a/Engine/Plugins/Experimental/Mutable/Source/MutableRuntime/Private/MuR/System.cpp b/Engine/Plugins/Experimental/Mutable/Source/MutableRuntime/Private/MuR/System.cpp
--- a/Engine/Plugins/Experimental/Mutable/Source/MutableRuntime/Private/MuR/System.cpp
+++ b/Engine/Plugins/Experimental/Mutable/Source/MutableRuntime/Private/MuR/System.cpp
@@ -725,7 +725,7 @@ namespace mu
 		int32 Result = 0;
 		if (!bUnrecoverableError)
 		{
-			Result = WorkingMemoryManager.CurrentInstanceCache->GetInt(FCacheAddress(at, 0, 0));;
+			Result = WorkingMemoryManager.CurrentInstanceCache->GetInt(FCacheAddress(at, 0, 0));
 		}
 
 		WorkingMemoryManager.EndRunnerThread();
diff --git a/Engine/Plugins/Experimental/Mutable/Source/MutableRuntime/Public/MuR/ExtensionDataStreamer.h b/Engine/Plugins/Experimental/Mutable/Source/MutableRuntime/Public/MuR/ExtensionDataStreamer.h
--- a/Engine/Plugins/Experimental/Mutable/Source/MutableRuntime/Public/MuR/ExtensionDataStreamer.h
+++ b/Engine/Plugins/Experimental/Mutable/Source/MutableRuntime/Public/MuR/ExtensionDataStreamer.h
@@ -35,7 +35,7 @@ public:
 //! All functions must be safe to call from the Mutable thread.
 //! 
 //! \ingroup runtime
-class MUTABLERUNTIME_API ExtensionDataStreamer : public Base
+class MUTABLERUNTIME_API ExtensionDataStreamer
 {
 public:
 
diff --git a/Engine/Plugins/Experimental/Mutable/Source/MutableRuntime/Public/MuR/MeshBufferSet.h b/Engine/Plugins/Experimental/Mutable/Source/MutableRuntime/Public/MuR/MeshBufferSet.h
--- a/Engine/Plugins/Experimental/Mutable/Source/MutableRuntime/Public/MuR/MeshBufferSet.h
+++ b/Engine/Plugins/Experimental/Mutable/Source/MutableRuntime/Public/MuR/MeshBufferSet.h
@@ -206,7 +206,7 @@ namespace mu
 
 
 	//! Set of buffers storing mesh element data. Elements can be vertices, indices or faces.
-	class MUTABLERUNTIME_API FMeshBufferSet : public Base
+	class MUTABLERUNTIME_API FMeshBufferSet
 	{
 	public:
 
diff --git a/Engine/Plugins/Experimental/Mutable/Source/MutableRuntime/Public/MuR/MutableMemory.h b/Engine/Plugins/Experimental/Mutable/Source/MutableRuntime/Public/MuR/MutableMemory.h
--- a/Engine/Plugins/Experimental/Mutable/Source/MutableRuntime/Public/MuR/MutableMemory.h
+++ b/Engine/Plugins/Experimental/Mutable/Source/MutableRuntime/Public/MuR/MutableMemory.h
@@ -4,50 +4,16 @@
 
 #include "CoreTypes.h"
 
-#include <new>
-
 namespace mu
 {
-	//! \defgroup runtime Run-time classes
-    //! This module contains the only classes necessary to integrate the system into an application
-    //! loading models and building instances by setting parameter values.
-
-    //! \brief Set the memory management methods for the library.
-    //! The library will always manage the heap memory using these methods.
-	//! This should be called only before starting to use the library or any of its data.
-	//! If this method is not called, the standard C library malloc and free will be used.
-	//! \todo: Notes about thread safety
-	//! \param customMalloc Function pointer to a function taking the size in bytes to allocate.
-    //!        The second parameter is the required alignment in bytes.
-	//! 	   This function must return a valid heap pointer or 0 in case of allocation failure.
-	//! \param customFree Function pointer to a function taking a pointer obtained reviously with
-	//! 	   customMalloc, that needs to be freed.
-	//! \ingroup runtime
+	/** This should be called only before starting to use the library or any of its data. */
 	MUTABLERUNTIME_API extern void Initialize();
 
-    //! \brief Shutdown everything related to mutable.
-    //! This should be called when no other mutable objects will be used, created or destroyed,
-    //! and all threads involving mutable tasks have been terminated.
+    /** Shutdown everything related to mutable.
+     * This should be called when no other mutable objects will be used, created or destroyed,
+     * and all threads involving mutable tasks have been terminated.
+	 */
 	MUTABLERUNTIME_API extern void Finalize();
 
-
-    //! \brief %Base class from which all library classes inherit to control memory allocation.
-	//! For internal use only.
-	//! \ingroup runtime
-	class MUTABLERUNTIME_API Base
-	{
-	public:
-
-		//inline void* operator new(size_t s)
-		//{
-		//	return mutable_malloc(s);
-		//}
-
-		//inline void operator delete(void* p, size_t size)
-		//{
-		//	mutable_free(p, size);
-		//}
-    };
-
 }
 
diff --git a/Engine/Plugins/Experimental/Mutable/Source/MutableRuntime/Public/MuR/RefCounted.h b/Engine/Plugins/Experimental/Mutable/Source/MutableRuntime/Public/MuR/RefCounted.h
--- a/Engine/Plugins/Experimental/Mutable/Source/MutableRuntime/Public/MuR/RefCounted.h
+++ b/Engine/Plugins/Experimental/Mutable/Source/MutableRuntime/Public/MuR/RefCounted.h
@@ -20,7 +20,7 @@ namespace mu
 	* Any subclass of this class can be managed using smart pointers through the Ptr<T> template.
 	* \warning This base allow multi-threaded manipulation of smart pointers, since the count increments and decrements are atomic.
 	*/
-	class MUTABLERUNTIME_API RefCounted : public Base
+	class MUTABLERUNTIME_API RefCounted
 	{
 	public:
 
diff --git a/Engine/Plugins/Experimental/Mutable/Source/MutableRuntime/Public/MuR/Serialisation.h b/Engine/Plugins/Experimental/Mutable/Source/MutableRuntime/Public/MuR/Serialisation.h
--- a/Engine/Plugins/Experimental/Mutable/Source/MutableRuntime/Public/MuR/Serialisation.h
+++ b/Engine/Plugins/Experimental/Mutable/Source/MutableRuntime/Public/MuR/Serialisation.h
@@ -74,7 +74,7 @@ namespace mu
 
 
 	/** */
-	class MUTABLERUNTIME_API ModelReader : public Base
+	class MUTABLERUNTIME_API ModelReader
     {
     public:
 
@@ -122,7 +122,7 @@ namespace mu
 
 
 	/** */
-	class MUTABLERUNTIME_API ModelWriter : public Base
+	class MUTABLERUNTIME_API ModelWriter
 	{
 	public:
 
@@ -167,7 +167,7 @@ namespace mu
 
     //! Interface for any input stream to be use with InputArchives.
     //! \ingroup tools
-    class MUTABLERUNTIME_API InputStream : public Base
+    class MUTABLERUNTIME_API InputStream
     {
     public:
 
@@ -183,7 +183,7 @@ namespace mu
 
     //! Interface for any output stream to be used with OutputArchives
     //! \ingroup tools
-    class MUTABLERUNTIME_API OutputStream : public Base
+    class MUTABLERUNTIME_API OutputStream
     {
     public:
 
@@ -200,7 +200,7 @@ namespace mu
 
     //! Archive containing data to be deserialised.
     //! \ingroup tools
-    class MUTABLERUNTIME_API InputArchive : public Base
+    class MUTABLERUNTIME_API InputArchive
     {
     public:
 
@@ -228,7 +228,7 @@ namespace mu
 
     //! Archive where data can be serialised to.
     //! \ingroup tools
-    class MUTABLERUNTIME_API OutputArchive : public Base
+    class MUTABLERUNTIME_API OutputArchive
     {
     public:
 
diff --git a/Engine/Plugins/Experimental/Mutable/Source/MutableRuntime/Public/MuR/System.h b/Engine/Plugins/Experimental/Mutable/Source/MutableRuntime/Public/MuR/System.h
--- a/Engine/Plugins/Experimental/Mutable/Source/MutableRuntime/Public/MuR/System.h
+++ b/Engine/Plugins/Experimental/Mutable/Source/MutableRuntime/Public/MuR/System.h
@@ -84,7 +84,7 @@ namespace mu
 
 
     /** Interface to request external images used as parameters. */
-    class MUTABLERUNTIME_API ImageParameterGenerator : public Base
+    class MUTABLERUNTIME_API ImageParameterGenerator
     {
     public:
 
diff --git a/Engine/Plugins/Experimental/Mutable/Source/MutableTools/Private/MuT/AST.h b/Engine/Plugins/Experimental/Mutable/Source/MutableTools/Private/MuT/AST.h
--- a/Engine/Plugins/Experimental/Mutable/Source/MutableTools/Private/MuT/AST.h
+++ b/Engine/Plugins/Experimental/Mutable/Source/MutableTools/Private/MuT/AST.h
@@ -797,7 +797,7 @@ namespace mu
     inline Ptr<DERIVED> Clone( const ASTOp* s )
     {
 		ASTOp::MapChildFunc Identity = [](const Ptr<ASTOp>& o) {return o; };
-		auto c = s->Clone(Identity);
+		Ptr<ASTOp> c = s->Clone(Identity);
 		Ptr<DERIVED> t = dynamic_cast<DERIVED*>(c.get());
         check(t);
         return t;
@@ -807,7 +807,7 @@ namespace mu
     inline Ptr<DERIVED> Clone( const Ptr<const ASTOp>& s )
     {
 		ASTOp::MapChildFunc Identity = [](const Ptr<ASTOp>& o) {return o; };
-        auto c = s->Clone(Identity);
+		Ptr<ASTOp> c = s->Clone(Identity);
         Ptr<DERIVED> t = dynamic_cast<DERIVED*>(c.get());
         check(t);
         return t;
@@ -817,7 +817,7 @@ namespace mu
     //!
     //---------------------------------------------------------------------------------------------
     template<typename STATE>
-    class Visitor_TopDown_Unique_Const : public Base
+    class Visitor_TopDown_Unique_Const
     {
     private:
 
@@ -970,7 +970,7 @@ namespace mu
     //! Once an instruction has changed, all the chain of instructions up to the root will be
     //! cloned, referencing the new instruction.
     //---------------------------------------------------------------------------------------------
-    class Visitor_TopDown_Unique_Cloning : public Base
+    class Visitor_TopDown_Unique_Cloning
     {
     public:
 
diff --git a/Engine/Plugins/Experimental/Mutable/Source/MutableTools/Private/MuT/ASTOpImageCompose.cpp b/Engine/Plugins/Experimental/Mutable/Source/MutableTools/Private/MuT/ASTOpImageCompose.cpp
--- a/Engine/Plugins/Experimental/Mutable/Source/MutableTools/Private/MuT/ASTOpImageCompose.cpp
+++ b/Engine/Plugins/Experimental/Mutable/Source/MutableTools/Private/MuT/ASTOpImageCompose.cpp
@@ -9,15 +9,13 @@
 #include "MuR/ModelPrivate.h"
 #include "MuR/MutableMath.h"
 #include "MuR/MutableTrace.h"
+#include "MuR/ImagePrivate.h"
 #include "MuR/RefCounted.h"
 #include "MuR/Types.h"
 #include "MuT/ASTOpConstantResource.h"
 #include "MuT/ASTOpImagePixelFormat.h"
 #include "MuT/StreamsPrivate.h"
 
-#include <memory>
-#include <utility>
-
 
 namespace mu
 {
@@ -144,8 +142,13 @@ FImageDesc ASTOpImageCompose::GetImageDesc( bool returnBestOption, FGetImageDesc
         res = Base->GetImageDesc( returnBestOption, context );
     }
 
+	if (BlockImage)
+	{
+		FImageDesc BlockDesc = BlockImage->GetImageDesc(returnBestOption, context);
+		res.m_format = GetMostGenericFormat(res.m_format,BlockDesc.m_format);
+	}
 
-    // Cache the result
+    // Cache th result
     if (context)
     {
 		context->m_results.Add(this, res);
diff --git a/Engine/Plugins/Experimental/Mutable/Source/MutableTools/Private/MuT/ASTOpImagePixelFormat.cpp b/Engine/Plugins/Experimental/Mutable/Source/MutableTools/Private/MuT/ASTOpImagePixelFormat.cpp
--- a/Engine/Plugins/Experimental/Mutable/Source/MutableTools/Private/MuT/ASTOpImagePixelFormat.cpp
+++ b/Engine/Plugins/Experimental/Mutable/Source/MutableTools/Private/MuT/ASTOpImagePixelFormat.cpp
@@ -20,9 +20,6 @@
 #include "MuT/ASTOpSwitch.h"
 #include "MuT/StreamsPrivate.h"
 
-#include <memory>
-#include <utility>
-
 namespace mu
 {
 
diff --git a/Engine/Plugins/Experimental/Mutable/Source/MutableTools/Private/MuT/ASTOpImageSwizzle.cpp b/Engine/Plugins/Experimental/Mutable/Source/MutableTools/Private/MuT/ASTOpImageSwizzle.cpp
--- a/Engine/Plugins/Experimental/Mutable/Source/MutableTools/Private/MuT/ASTOpImageSwizzle.cpp
+++ b/Engine/Plugins/Experimental/Mutable/Source/MutableTools/Private/MuT/ASTOpImageSwizzle.cpp
@@ -723,8 +723,9 @@ namespace mu
 							|| Typed->UncroppedSizeY != FirstRasterMesh->UncroppedSizeY
 							|| Typed->CropMinX != FirstRasterMesh->CropMinX
 							|| Typed->CropMinY != FirstRasterMesh->CropMinY
-							|| Typed->bIsRGBFadingEnabled != FirstRasterMesh->bIsRGBFadingEnabled
-							|| Typed->bIsAlphaFadingEnabled != FirstRasterMesh->bIsAlphaFadingEnabled
+							// Also ignore the fading flags. They are dealt with below.
+							//|| Typed->bIsRGBFadingEnabled != FirstRasterMesh->bIsRGBFadingEnabled
+							//|| Typed->bIsAlphaFadingEnabled != FirstRasterMesh->bIsAlphaFadingEnabled
 							)
 						{
 							bAreAllRasterMeshesCompatible = false;
@@ -759,7 +760,8 @@ namespace mu
 						}
 						else if (Sources[3] && SourceChannels[3] < 3)
 						{
-							NewRaster->bIsAlphaFadingEnabled = NewRaster->bIsRGBFadingEnabled;
+							const ASTOpImageRasterMesh* ChannelRaster = static_cast<const ASTOpImageRasterMesh*>(Sources[3].child().get());
+							NewRaster->bIsAlphaFadingEnabled = ChannelRaster->bIsRGBFadingEnabled;
 						}
 					}
 
diff --git a/Engine/Plugins/Experimental/Mutable/Source/MutableTools/Private/MuT/ASTOpMeshBindShape.cpp b/Engine/Plugins/Experimental/Mutable/Source/MutableTools/Private/MuT/ASTOpMeshBindShape.cpp
--- a/Engine/Plugins/Experimental/Mutable/Source/MutableTools/Private/MuT/ASTOpMeshBindShape.cpp
+++ b/Engine/Plugins/Experimental/Mutable/Source/MutableTools/Private/MuT/ASTOpMeshBindShape.cpp
@@ -223,8 +223,8 @@ namespace mu
 			case OP_TYPE::ME_SWITCH:
 			{
 				// If the switch variable and structure is the same
-				const ASTOpSwitch* MeshSwitch = reinterpret_cast<const ASTOpSwitch*>(MeshAt.get());
-				const ASTOpSwitch* ShapeSwitch = reinterpret_cast<const ASTOpSwitch*>(ShapeAt.get());
+				const ASTOpSwitch* MeshSwitch = static_cast<const ASTOpSwitch*>(MeshAt.get());
+				const ASTOpSwitch* ShapeSwitch = static_cast<const ASTOpSwitch*>(ShapeAt.get());
 				bool bIsSimilarSwitch = MeshSwitch->IsCompatibleWith(ShapeSwitch);
 				if (!bIsSimilarSwitch)
 				{
@@ -260,8 +260,8 @@ namespace mu
 
 			case OP_TYPE::ME_CONDITIONAL:
 			{
-				const ASTOpConditional* MeshConditional = reinterpret_cast<const ASTOpConditional*>(MeshAt.get());
-				const ASTOpConditional* ShapeConditional = reinterpret_cast<const ASTOpConditional*>(ShapeAt.get());
+				const ASTOpConditional* MeshConditional = static_cast<const ASTOpConditional*>(MeshAt.get());
+				const ASTOpConditional* ShapeConditional = static_cast<const ASTOpConditional*>(ShapeAt.get());
 				bool bIsSimilar = MeshConditional->condition == ShapeConditional->condition;
 				if (!bIsSimilar)
 				{
diff --git a/Engine/Plugins/Experimental/Mutable/Source/MutableTools/Private/MuT/ASTOpMeshDifference.cpp b/Engine/Plugins/Experimental/Mutable/Source/MutableTools/Private/MuT/ASTOpMeshDifference.cpp
--- a/Engine/Plugins/Experimental/Mutable/Source/MutableTools/Private/MuT/ASTOpMeshDifference.cpp
+++ b/Engine/Plugins/Experimental/Mutable/Source/MutableTools/Private/MuT/ASTOpMeshDifference.cpp
@@ -133,8 +133,8 @@ namespace mu
 			case OP_TYPE::ME_SWITCH:
 			{
 				// If the switch variable and structure is the same
-				const ASTOpSwitch* BaseSwitch = reinterpret_cast<const ASTOpSwitch*>(BaseAt.get());
-				const ASTOpSwitch* TargetSwitch = reinterpret_cast<const ASTOpSwitch*>(TargetAt.get());
+				const ASTOpSwitch* BaseSwitch = static_cast<const ASTOpSwitch*>(BaseAt.get());
+				const ASTOpSwitch* TargetSwitch = static_cast<const ASTOpSwitch*>(TargetAt.get());
 				bool bIsSimilarSwitch = BaseSwitch->IsCompatibleWith(TargetSwitch);
 				if (!bIsSimilarSwitch)
 				{
@@ -170,8 +170,8 @@ namespace mu
 
 			case OP_TYPE::ME_CONDITIONAL:
 			{
-				const ASTOpConditional* BaseConditional = reinterpret_cast<const ASTOpConditional*>(BaseAt.get());
-				const ASTOpConditional* TargetConditional = reinterpret_cast<const ASTOpConditional*>(TargetAt.get());
+				const ASTOpConditional* BaseConditional = static_cast<const ASTOpConditional*>(BaseAt.get());
+				const ASTOpConditional* TargetConditional = static_cast<const ASTOpConditional*>(TargetAt.get());
 				bool bIsSimilar = BaseConditional->condition == TargetConditional->condition;
 				if (!bIsSimilar)
 				{
diff --git a/Engine/Plugins/Experimental/Mutable/Source/MutableTools/Private/MuT/CodeGenerator.cpp b/Engine/Plugins/Experimental/Mutable/Source/MutableTools/Private/MuT/CodeGenerator.cpp
--- a/Engine/Plugins/Experimental/Mutable/Source/MutableTools/Private/MuT/CodeGenerator.cpp
+++ b/Engine/Plugins/Experimental/Mutable/Source/MutableTools/Private/MuT/CodeGenerator.cpp
@@ -109,12 +109,12 @@ namespace mu
 	}
 
 	//---------------------------------------------------------------------------------------------
-	void CodeGenerator::GenerateRoot(const NodePtrConst pNode)
+	void CodeGenerator::GenerateRoot(const Ptr<const Node> pNode)
 	{
 		MUTABLE_CPUPROFILER_SCOPE(Generate);
 
 		// First pass
-		m_firstPass.Generate(m_pErrorLog, pNode->GetBasePrivate(), m_compilerOptions->bIgnoreStates, this);
+		m_firstPass.Generate(m_pErrorLog, pNode.get(), m_compilerOptions->bIgnoreStates, this);
 
 		// Second pass
 		SecondPassGenerator SecondPass(&m_firstPass, m_compilerOptions);
@@ -128,86 +128,73 @@ namespace mu
 		{
 			MUTABLE_CPUPROFILER_SCOPE(MainPass);
 
-            m_currentStateIndex = 0;
+            int32 CurrentStateIndex = 0;
 			for (const TPair<FObjectState, const Node::Private*>& s : m_firstPass.m_states)
 			{
 				MUTABLE_CPUPROFILER_SCOPE(MainPassState);
 
-				Ptr<ASTOp> stateRoot = Generate(pNode);
+				FGenericGenerationOptions Options;
+				Options.State = CurrentStateIndex;
+
+				Ptr<ASTOp> stateRoot = Generate(pNode, Options);
 				m_states.Emplace(s.Key, stateRoot);
 
-				++m_currentStateIndex;
+				++CurrentStateIndex;
 			}
 		}
 
 		// Free caches
-		m_compiled.Reset();
-		m_constantMeshes.Empty();
-		m_generatedLayouts.Empty();
-		m_nodeVariables.clear();
+		GeneratedGenericNodes.Reset();
+		m_constantMeshes.SetNum(0,false);
+		GeneratedLayouts.Reset();
+		NodeVariables.Reset();
 		m_generatedMeshes.Reset();
-		m_generatedProjectors.Reset();
-		m_generatedRanges.Reset();
-		m_generatedTables.clear();
+		GeneratedProjectors.Reset();
+		GeneratedRanges.Reset();
+		GeneratedStrings.Reset();
+		GeneratedTables.Reset();
 		m_firstPass = FirstPassGenerator();
-		m_currentBottomUpState = FBottomUpState();
-		m_currentParents.Empty();
-		m_currentObject.Empty();
-		m_additionalComponents.clear();
+		m_currentParents.Reset();
+		m_currentObject.Reset();
+		AdditionalComponents.Reset();
 	}
 
 
-	//---------------------------------------------------------------------------------------------
-	Ptr<ASTOp> CodeGenerator::Generate(const NodePtrConst pNode)
+	Ptr<ASTOp> CodeGenerator::Generate(const Ptr<const Node> pNode, const FGenericGenerationOptions& Options )
 	{
 		if (!pNode)
 		{
 			return nullptr;
 		}
 
-		// Clear bottom-up state
-		m_currentBottomUpState.m_address = nullptr;
-
-        // Temp by-passes while we remove the visitor pattern
+        // Type-specific generation
 		if (const NodeScalar* ScalarNode = dynamic_cast<const NodeScalar*>(pNode.get()))
 		{
 			FScalarGenerationResult ScalarResult;
-			GenerateScalar(ScalarResult, ScalarNode);
+			GenerateScalar(ScalarResult, Options, ScalarNode);
 			return ScalarResult.op;
 		}
 
-		if (const NodeColour* ColorNode = dynamic_cast<const NodeColour*>(pNode.get()))
+		else if (const NodeColour* ColorNode = dynamic_cast<const NodeColour*>(pNode.get()))
 		{
 			FColorGenerationResult Result;
-			GenerateColor(Result, ColorNode);
+			GenerateColor(Result, Options, ColorNode);
 			return Result.op;
 		}
 
-		if (const NodeImage* ImageNode = dynamic_cast<const NodeImage*>(pNode.get()))
-		{
-			// This should never be called for images. Use GenerateImage.
-			check(false);
-		}
-
-		if (const NodeMesh* MeshNode = dynamic_cast<const NodeMesh*>(pNode.get()))
-		{
-			// This should never be called for meshes. Use GenerateMesh.
-			check(false);
-		}
-
-		if (const NodeProjector* projNode = dynamic_cast<const NodeProjector*>(pNode.get()))
+		else if (const NodeProjector* projNode = dynamic_cast<const NodeProjector*>(pNode.get()))
 		{
 			FProjectorGenerationResult ProjResult;
-			GenerateProjector(ProjResult, projNode);
+			GenerateProjector(ProjResult, Options, projNode);
 			return ProjResult.op;
 		}
 
-		if (const NodeSurfaceNew* surfNode = dynamic_cast<const NodeSurfaceNew*>(pNode.get()))
+		else if (const NodeSurfaceNew* surfNode = dynamic_cast<const NodeSurfaceNew*>(pNode.get()))
 		{
 			// This happens only if we generate a node graph that has a NodeSurfaceNew at the root.
 			FSurfaceGenerationResult surfResult;
 			const TArray<FirstPassGenerator::FSurface::FEdit> edits;
-			GenerateSurface(surfResult, surfNode, edits);
+			GenerateSurface(surfResult, Options, surfNode, edits);
 			return surfResult.surfaceOp;
 		}
 
@@ -219,7 +206,7 @@ namespace mu
 
 		else if (dynamic_cast<const NodeSurfaceSwitch*>(pNode.get()))
 		{
-			// This happens only if we generate a node graph that has a NodeSurfaceVariation at the root.
+			// This happens only if we generate a node graph that has a NodeSurfaceSwitch at the root.
 			return nullptr;
 		}
 
@@ -236,38 +223,61 @@ namespace mu
 		}
 
 
-		Ptr<ASTOp> result;
+		Ptr<ASTOp> ResultOp;
 
 		// See if it was already generated
-		FVisitedKeyMap key = GetCurrentCacheKey(pNode);
-		VisitedMap::ValueType* it = m_compiled.Find(key);
+
+		FGeneratedCacheKey Key;
+		Key.Node = pNode;
+		Key.Options = Options;
+		FGeneratedGenericNodesMap::ValueType* it = GeneratedGenericNodes.Find(Key);
 		if (it)
 		{
-			m_currentBottomUpState = *it;
-			result = m_currentBottomUpState.m_address;
+			ResultOp = it->op;
 		}
 		else
 		{
-			result = pNode->GetBasePrivate()->Accept(*this);
-			m_currentBottomUpState.m_address = result;
-			m_compiled.Add(key, m_currentBottomUpState);
+			FGenericGenerationResult Result;
+
+			// Generate for each different type of node
+			if (pNode->GetType() == NodeComponentNew::GetStaticType())
+			{
+				Generate_ComponentNew(Options, Result, static_cast<const NodeComponentNew*>(pNode.get()));
+			}
+			else if (pNode->GetType() == NodeComponentEdit::GetStaticType())
+			{
+				// Nothing to do: processed elsewhere.
+			}
+			else if (pNode->GetType() == NodeObjectNew::GetStaticType())
+			{
+				Generate_ObjectNew(Options, Result, static_cast<const NodeObjectNew*>(pNode.get()));
+			}
+			else if (pNode->GetType()==NodeObjectGroup::GetStaticType())
+			{
+				Generate_ObjectGroup(Options, Result, static_cast<const NodeObjectGroup*>(pNode.get()));
+			}
+			else
+			{
+				check(false);
+			}
+
+			ResultOp = Result.op;
+			GeneratedGenericNodes.Add(Key, Result);
 		}
 
 		// debug: expensive check of all code generation
-//        if (result)
+//        if (ResultOp)
 //        {
 //            ASTOpList roots;
-//            roots.push_back(result);
+//            roots.push_back(ResultOp);
 //            ASTOp::FullAssert(roots);
 //        }
 
-		return result;
+		return ResultOp;
 	}
 
 
-	//---------------------------------------------------------------------------------------------
-	void CodeGenerator::GenerateRange(FRangeGenerationResult& Result,
-		NodeRangePtrConst Untyped)
+	void CodeGenerator::GenerateRange(FRangeGenerationResult& Result, const FGenericGenerationOptions& Options, Ptr<const NodeRange> Untyped)
 	{
 		if (!Untyped)
 		{
@@ -276,8 +286,10 @@ namespace mu
 		}
 
 		// See if it was already generated
-		FVisitedKeyMap Key = GetCurrentCacheKey(Untyped);
-		GeneratedRangeMap::ValueType* it = m_generatedRanges.Find(Key);
+		FGeneratedCacheKey Key;
+		Key.Node = Untyped;
+		Key.Options = Options;
+		FGeneratedRangeMap::ValueType* it = GeneratedRanges.Find(Key);
 		if (it)
 		{
 			Result = *it;
@@ -290,7 +302,10 @@ namespace mu
 		{
 			Result = FRangeGenerationResult();
 			Result.rangeName = FromScalar->GetName();
-			Result.sizeOp = Generate(FromScalar->GetSize());
+
+			FScalarGenerationResult ChildResult;
+			GenerateScalar(ChildResult, Options, FromScalar->GetSize());
+			Result.sizeOp = ChildResult.op;
 		}
 		else
 		{
@@ -299,11 +314,10 @@ namespace mu
 
 
 		// Cache the result
-		m_generatedRanges.Add(Key, Result);
+		GeneratedRanges.Add(Key, Result);
 	}
 
 
-	//---------------------------------------------------------------------------------------------
 	Ptr<ASTOp> CodeGenerator::GenerateTableVariable(TablePtr pTable, const FString& strName)
 	{
 		Ptr<ASTOp> result;
@@ -370,13 +384,12 @@ namespace mu
 	}
 
 
-	//---------------------------------------------------------------------------------------------
 	Ptr<const Layout> CodeGenerator::AddLayout(Ptr<const Layout> SourceLayout)
 	{
 		// The layout we are adding must be a source layout, without block ids yet.
 		check(SourceLayout->m_blocks.IsEmpty() || SourceLayout->m_blocks[0].m_id == -1);
 
-		Ptr<const Layout>* it = m_generatedLayouts.Find(SourceLayout.get());
+		Ptr<const Layout>* it = GeneratedLayouts.Find(SourceLayout.get());
 
 		if (it)
 		{
@@ -393,13 +406,12 @@ namespace mu
 		}
 		check(SourceLayout->m_blocks.Num() == ClonedLayout->m_blocks.Num());
 		check(ClonedLayout->m_blocks.IsEmpty() || ClonedLayout->m_blocks[0].m_id != -1);
-		m_generatedLayouts.Add(SourceLayout.get(), ClonedLayout);
+		GeneratedLayouts.Add(SourceLayout.get(), ClonedLayout);
 
 		return ClonedLayout;
 	}
 
 
-	//---------------------------------------------------------------------------------------------
 	Ptr<ASTOp> CodeGenerator::GenerateImageBlockPatch(Ptr<ASTOp> blockAd,
 		const NodePatchImage* pPatch,
 		Ptr<ASTOp> conditionAd,
@@ -483,7 +495,6 @@ namespace mu
 	}
 
 
-	//---------------------------------------------------------------------------------------------
 	const NodeMeshConstant* FindSourceMesh(const Node* pNode)
 	{
 		const NodeMeshConstant* pResult = nullptr;
@@ -581,44 +592,10 @@ namespace mu
 	}
 
 
-	//---------------------------------------------------------------------------------------------
-	Ptr<ASTOp> CodeGenerator::Visit(const NodePatchImage::Private& node)
+	void CodeGenerator::Generate_LOD(const FGenericGenerationOptions& Options, FGenericGenerationResult& Result, const NodeLOD* InNode)
 	{
-		// Get the parent component layout
-		const NodeObjectNew::Private* pParent = nullptr;
-		if (m_currentParents.Num() > 2)
-		{
-			pParent = m_currentParents[m_currentParents.Num() - 2].m_pObject;
-		}
-
-		const NodeLayout* pNodeLayout = nullptr;
-		if (pParent)
-		{
-			pNodeLayout = pParent->GetLayout
-			(
-				m_currentParents.Last().m_lod,
-				m_currentParents.Last().m_component,
-				m_currentParents.Last().m_surface,
-				m_currentParents.Last().m_texture
-			).get();
-		}
-
-		if (!pNodeLayout)
-		{
-			FString Msg = FString::Printf(TEXT("In object [%s] NodePatchImage couldn't find the layout in parent."),
-				*m_currentParents.Last().m_pObject->m_name
-			);
-
-			m_pErrorLog->GetPrivate()->Add(Msg, ELMT_ERROR, node.m_errorContext);
-		}
-
-		return 0;
-	}
-
+		const NodeLOD::Private& node = *InNode->GetPrivate();
 
-	//---------------------------------------------------------------------------------------------
-	Ptr<ASTOp> CodeGenerator::Visit(const NodeLOD::Private& node)
-	{
 		// Build a series of operations to assemble all the LOD components
 		Ptr<ASTOp> lastCompOp;
 
@@ -630,7 +607,7 @@ namespace mu
 			{
 				m_currentParents.Last().m_component = (int)t;
 
-				Ptr<ASTOp> componentOp = Generate(pComponentNode);
+				Ptr<ASTOp> componentOp = Generate(pComponentNode, Options);
 
 				if (componentOp)
 				{
@@ -650,10 +627,10 @@ namespace mu
 		FAdditionalComponentKey thisKey;
 		thisKey.m_lod = m_currentParents.Last().m_lod;
 		thisKey.m_pObject = m_currentParents.Last().m_pObject;
-		auto addIt = m_additionalComponents.find(thisKey);
-		if (addIt != m_additionalComponents.end())
+		TArray<Ptr<ASTOp>>* addIt = AdditionalComponents.Find(thisKey);
+		if (addIt)
 		{
-			for (const auto& cop : addIt->second)
+			for (const auto& cop : *addIt)
 			{
 				// Add the additional components after the main ones, this means higher up in the
 				// op tree.
@@ -683,14 +660,13 @@ namespace mu
 			FAdditionalComponentKey parentKey;
 			parentKey.m_lod = m_currentParents.Last().m_lod;
 			parentKey.m_pObject = parentObjectKey.m_pObject;
-			m_additionalComponents[parentKey].Add(lastCompOp);
+			AdditionalComponents.FindOrAdd(parentKey).Add(lastCompOp);
 		}
 
-		return lastCompOp;
+		Result.op = lastCompOp;
 	}
 
 
-	//---------------------------------------------------------------------------------------------
 	Ptr<ASTOp> CodeGenerator::ApplyTiling(Ptr<ASTOp> Source, UE::Math::TIntVector2<int32> Size, EImageFormat Format)
 	{
 		// For now always apply tiling
@@ -748,8 +724,9 @@ namespace mu
 
 
     //---------------------------------------------------------------------------------------------
-    void CodeGenerator::GenerateSurface( FSurfaceGenerationResult& result,
-                                         NodeSurfaceNewPtrConst surfaceNode,
+    void CodeGenerator::GenerateSurface( FSurfaceGenerationResult& result, 
+										 const FGenericGenerationOptions& Options,
+                                         Ptr<const NodeSurfaceNew> surfaceNode,
                                          const TArray<FirstPassGenerator::FSurface::FEdit>& edits )
     {
         MUTABLE_CPUPROFILER_SCOPE(GenerateSurface);
@@ -795,9 +772,9 @@ namespace mu
             {
                 // Check state conditions
                 bool surfaceValidForThisState =
-                    m_currentStateIndex >= its.stateCondition.Num()
+                    Options.State >= its.stateCondition.Num()
                     ||
-                    its.stateCondition[m_currentStateIndex];
+                    its.stateCondition[Options.State];
 
                 if (surfaceValidForThisState)
                 {
@@ -840,7 +817,7 @@ namespace mu
 			FMeshGenerationOptions MeshOptions;
 			MeshOptions.bUniqueVertexIDs = true;
 			MeshOptions.bLayouts = true;
-			MeshOptions.State = m_currentStateIndex;
+			MeshOptions.State = Options.State;
 			MeshOptions.ActiveTags = node.m_tags;
 
 			const FMeshGenerationResult* SharedMeshResults = nullptr;
@@ -891,7 +868,7 @@ namespace mu
 						MergedMeshOptions.bLayouts = true;
 						MergedMeshOptions.bClampUVIslands = bShareSurface && bNormalizeUVs;
 						MergedMeshOptions.bNormalizeUVs = bNormalizeUVs;
-						MergedMeshOptions.State = m_currentStateIndex;
+						MergedMeshOptions.State = Options.State;
 						MergedMeshOptions.ActiveTags = e.node->m_tags;
 
 						if (SharedMeshResults)
@@ -909,8 +886,9 @@ namespace mu
 
 						// Apply the modifier for the post-normal operations stage to the added mesh
 						bool bModifiersForBeforeOperations = false;
-						lastMeshOp = ApplyMeshModifiers(lastMeshOp, e.node->m_tags,
-							bModifiersForBeforeOperations, node.m_errorContext);
+						FGenericGenerationOptions ModifierOptions(Options);
+						ModifierOptions.ActiveTags = e.node->m_tags;
+						lastMeshOp = ApplyMeshModifiers(Options, lastMeshOp, bModifiersForBeforeOperations, node.m_errorContext);
 
                         FMeshGenerationResult::FExtraLayouts data;
 						data.GeneratedLayouts = addResults.GeneratedLayouts;
@@ -959,7 +937,7 @@ namespace mu
 						FMeshGenerationOptions RemoveMeshOptions;
 						RemoveMeshOptions.bUniqueVertexIDs = false;
 						RemoveMeshOptions.bLayouts = false;
-						RemoveMeshOptions.State = m_currentStateIndex;
+						RemoveMeshOptions.State = Options.State;
 						RemoveMeshOptions.ActiveTags = e.node->m_tags;
 
                         GenerateMesh(RemoveMeshOptions, removeResults, pRemove );
@@ -1012,7 +990,7 @@ namespace mu
 					FMeshGenerationOptions MorphTargetMeshOptions;
 					MorphTargetMeshOptions.bUniqueVertexIDs = false;
 					MorphTargetMeshOptions.bLayouts = false;
-					MorphTargetMeshOptions.State = m_currentStateIndex;
+					MorphTargetMeshOptions.State = Options.State;
 
                     FMeshGenerationResult morphResult;
                     GenerateMesh(MorphTargetMeshOptions, morphResult, pMorph );
@@ -1040,15 +1018,18 @@ namespace mu
 						// Factor
 						if (e.node->m_pFactor)
 						{
-							op->Factor = Generate(e.node->m_pFactor);
+							FScalarGenerationResult ChildResult;
+							GenerateScalar(ChildResult, Options, e.node->m_pFactor);
+							op->Factor = ChildResult.op;
 						}
 						else
 						{
 							NodeScalarConstantPtr auxNode = new NodeScalarConstant();
 							auxNode->SetValue(1.0f);
-							Ptr<ASTOp> resultNode = Generate(auxNode);
 
-							op->Factor = resultNode;
+							FScalarGenerationResult ChildResult;
+							GenerateScalar(ChildResult, Options, auxNode);
+							op->Factor = ChildResult.op;
 						}
 
 						// Base		
@@ -1078,8 +1059,9 @@ namespace mu
 
 			// Apply the modifier for the post-normal operations stage.
 			bool bModifiersForBeforeOperations = false;
-			lastMeshOp = ApplyMeshModifiers( lastMeshOp, node.m_tags, 
-				bModifiersForBeforeOperations, node.m_errorContext);
+			FGenericGenerationOptions ModifierOptions(Options);
+			ModifierOptions.ActiveTags = node.m_tags;
+			lastMeshOp = ApplyMeshModifiers(Options, lastMeshOp, bModifiersForBeforeOperations, node.m_errorContext);
 
             // Layouts
             for ( int32 LayoutIndex=0; LayoutIndex <meshResults.GeneratedLayouts.Num(); ++LayoutIndex)
@@ -1249,22 +1231,6 @@ namespace mu
 						}
 					}
 
-					// Find out the size of the image
-					// TODO: What if the image is empty and everything is added?
-					//       Look in the extending images.
-					FImageDesc desc;
-					if (pImageNode)
-					{
-						MUTABLE_CPUPROFILER_SCOPE(CalculateImageDesc);
-						desc = CalculateImageDesc(*pImageNode->GetBasePrivate());
-					}
-
-					// If the image format doesn't come bottom-up, it may come top-down
-					if (desc.m_format == EImageFormat::IF_NONE && formatNode)
-					{
-						desc.m_format = formatNode->GetFormat();
-					}
-
 					const int LayoutIndex = node.m_images[t].m_layoutIndex;
 
 					// If the layout index has been set to negative, it means we should ignore the layout for this image.
@@ -1273,34 +1239,16 @@ namespace mu
 						: CompilerOptions::TextureLayoutStrategy::Pack
 						;
 
-					if (desc.m_size[0] == 0 || desc.m_size[1] == 0)
+					if (ImageLayoutStrategy == CompilerOptions::TextureLayoutStrategy::None)
 					{
-						int currentLOD = m_currentParents.Last().m_lod;
-						FString Msg = FString::Printf( TEXT("An image for [%s] [%s] [%s] at lod [%d] has zero size and will not be generated. "),
-							*node.m_images[t].m_name,
-							*node.m_images[t].m_materialName,
-							*node.m_images[t].m_materialParameterName,
-							currentLOD
-						);
-						m_pErrorLog->GetPrivate()->Add(Msg, ELMT_INFO, node.m_errorContext);
-					}
-
-					else if (desc.m_format == EImageFormat::IF_NONE)
-					{
-						FString Msg = FString::Printf(TEXT("An image [%s] has an unidentified pixel format. "), *node.m_images[t].m_name);
-						m_pErrorLog->GetPrivate()->Add(Msg, ELMT_ERROR, node.m_errorContext);
-					}
-
-					else if (ImageLayoutStrategy == CompilerOptions::TextureLayoutStrategy::None)
-					{
-						check(desc.m_format != EImageFormat::IF_NONE);
+						//check(desc.m_format != EImageFormat::IF_NONE);
 
 						// Generate the image
 						FImageGenerationOptions ImageOptions;
-						ImageOptions.CurrentStateIndex = m_currentStateIndex;
+						ImageOptions.State = Options.State;
 						ImageOptions.ImageLayoutStrategy = ImageLayoutStrategy;
 						ImageOptions.ActiveTags = node.m_tags;
-						ImageOptions.RectSize = UE::Math::TIntVector2<int32>(desc.m_size);
+						ImageOptions.RectSize = { 0, 0 };
 						FImageGenerationResult Result;
 						GenerateImage(ImageOptions, Result, pImageNode);
 						Ptr<ASTOp> imageAd = Result.op;
@@ -1392,85 +1340,101 @@ namespace mu
 							// Size of a layout block in pixels
 							FIntPoint grid = pLayout->GetGridSize();
 
-							check(desc.m_format != EImageFormat::IF_NONE);
+							EImageFormat FinalFormat = EImageFormat::IF_NONE;
+							int32 BlockPixelsX = 0;
+							int32 BlockPixelsY = 0;
+							bool bBlocksHaveMips = false;
+
+							bool bImageSizeWarning = false;
 
-							// If the image is too small or not a multiple of the the layout size, 
-							// resize it, but raise a warning.
-							if ((desc.m_size[0] % grid[0] != 0) || (desc.m_size[1] % grid[1] != 0))
+							// Start with a blank image. It will be completed later with the blockSize, format and mips information
+							Ptr<ASTOpFixed> BlankImageOp;
+							Ptr<ASTOp> imageAd;
 							{
-								FImageSize oldSize = desc.m_size;
-								desc.m_size[0] = grid[0] * FMath::Max(1, desc.m_size[0] / grid[0]);
-								desc.m_size[1] = grid[1] * FMath::Max(1, desc.m_size[1] / grid[1]);
-
-								int currentLOD = m_currentParents.Last().m_lod;
-								FString Msg = FString::Printf( TEXT("A texture [%s] for material [%s] parameter [%s] in LOD [%d] has been resized from [%d x %d] to [%d x %d] because it didn't fit the layout [%d x %d]. "),
-									*node.m_images[t].m_name,
-									*node.m_images[t].m_materialName,
-									*node.m_images[t].m_materialParameterName,
-									currentLOD,
-									oldSize[0], oldSize[1], desc.m_size[0], desc.m_size[1], grid[0], grid[1]);
-								m_pErrorLog->GetPrivate()->Add(Msg, ELMT_INFO, node.m_errorContext);
+								BlankImageOp = new ASTOpFixed();
+								BlankImageOp->op.type = OP_TYPE::IM_BLANKLAYOUT;
+								BlankImageOp->SetChild(BlankImageOp->op.args.ImageBlankLayout.layout, meshResults.layoutOps[LayoutIndex]);
+								// The rest ok the op will be completed below
+								BlankImageOp->op.args.ImageBlankLayout.mipmapCount = 0;
+								imageAd = BlankImageOp;
 							}
 
-							int blockSizeX = FMath::Max(1, desc.m_size[0] / grid[0]);
-							int blockSizeY = FMath::Max(1, desc.m_size[1] / grid[1]);
-
-							bool bBlocksHaveMips = desc.m_lods > 1;
-
-							// Start with a blank image
-							Ptr<ASTOp> imageAd;
+							auto UpdateBlockSize = [&BlockPixelsX, &BlockPixelsY, &FinalFormat, &bBlocksHaveMips, &bImageSizeWarning, &formatNode, &BlankImageOp, &node, &t, this]( FImageDesc BlockDesc, UE::Math::TIntVector2<uint16> LayoutCellSize )
 							{
-								Ptr<ASTOpFixed> bop = new ASTOpFixed();
-								bop->op.type = OP_TYPE::IM_BLANKLAYOUT;
-								bop->SetChild(bop->op.args.ImageBlankLayout.layout, meshResults.layoutOps[LayoutIndex]);
-								bop->op.args.ImageBlankLayout.blockSize[0] = uint16(blockSizeX);
-								bop->op.args.ImageBlankLayout.blockSize[1] = uint16(blockSizeY);
-								// We support block compression directly here, but not non-block compression
-								if (GetImageFormatData(desc.m_format).PixelsPerBlockX == 0)
+								if (BlockPixelsX == 0)
 								{
-									// It's something like RLE
-									bop->op.args.ImageBlankLayout.format = GetUncompressedFormat(desc.m_format);
-								}
-								else
-								{
-									// Directly supported
-									bop->op.args.ImageBlankLayout.format = desc.m_format;
+									if (!bImageSizeWarning)
+									{
+										// If the block pixels is not a multiple of the block layout cells
+										if ((BlockDesc.m_size[0] % LayoutCellSize[0] != 0) || (BlockDesc.m_size[1] % LayoutCellSize[1] != 0))
+										{
+											bImageSizeWarning = true;
+
+											int currentLOD = m_currentParents.Last().m_lod;
+											FString Msg = FString::Printf(TEXT("A texture [%s] for material [%s] parameter [%s] in LOD [%d] has been resized because it didn't fit the layout. "),
+												*node.m_images[t].m_name,
+												*node.m_images[t].m_materialName,
+												*node.m_images[t].m_materialParameterName,
+												currentLOD);
+											m_pErrorLog->GetPrivate()->Add(Msg, ELMT_INFO, node.m_errorContext);
+										}
+									}
+
+									BlockPixelsX = FMath::Max(1, BlockDesc.m_size[0] / LayoutCellSize[0]);
+									BlockPixelsY = FMath::Max(1, BlockDesc.m_size[1] / LayoutCellSize[1]);
+									bBlocksHaveMips = BlockDesc.m_lods > 1;
+
+									FinalFormat = BlockDesc.m_format;
+									if (formatNode)
+									{
+										FinalFormat = formatNode->GetPrivate()->m_formatIfAlpha;
+										if (FinalFormat == EImageFormat::IF_NONE)
+										{
+											FinalFormat = formatNode->GetPrivate()->m_format;
+										}
+									}
+
+									// Complete the base op
+									BlankImageOp->op.args.ImageBlankLayout.blockSize[0] = uint16(BlockPixelsX);
+									BlankImageOp->op.args.ImageBlankLayout.blockSize[1] = uint16(BlockPixelsY);
+									BlankImageOp->op.args.ImageBlankLayout.format = GetUncompressedFormat(FinalFormat);
+									BlankImageOp->op.args.ImageBlankLayout.generateMipmaps = bBlocksHaveMips;
+									BlankImageOp->op.args.ImageBlankLayout.mipmapCount = 0;
 								}
-								bop->op.args.ImageBlankLayout.generateMipmaps = bBlocksHaveMips;
-								bop->op.args.ImageBlankLayout.mipmapCount = 0;
-								imageAd = bop;
-							}
+							};
 
 							for (int b = 0; b < pLayout->GetBlockCount(); ++b)
 							{
-								// Block in layout grid units
-								box< UE::Math::TIntVector2<uint16> > rectInCells;
-								pLayout->GetBlock
-								(
-									b,
-									&rectInCells.min[0], &rectInCells.min[1],
-									&rectInCells.size[0], &rectInCells.size[1]
-								);
-
-								// Transform to pixels
-								box< UE::Math::TIntVector2<uint16> > rect = rectInCells;
-								rect.min[0] *= blockSizeX;
-								rect.min[1] *= blockSizeY;
-								rect.size[0] *= blockSizeX;
-								rect.size[1] *= blockSizeY;
-
 								// Generate the image
 								FImageGenerationOptions ImageOptions;
-								ImageOptions.CurrentStateIndex = m_currentStateIndex;
+								ImageOptions.State = Options.State;
 								ImageOptions.ImageLayoutStrategy = ImageLayoutStrategy;
+								ImageOptions.RectSize = { 0,0 };
 								ImageOptions.ActiveTags = node.m_tags;
-								ImageOptions.RectSize = UE::Math::TIntVector2<int32>(rect.size);
 								ImageOptions.LayoutToApply = pLayout;
 								ImageOptions.LayoutBlockId = pLayout->m_blocks[b].m_id;
 								FImageGenerationResult Result;
 								GenerateImage(ImageOptions, Result, pImageNode);
 								Ptr<ASTOp> blockAd = Result.op;
 
+								// Calculate the desc of the generated block
+								FImageDesc BlockDesc = blockAd->GetImageDesc();
+
+								// Block in layout grid units (cells)
+								box< UE::Math::TIntVector2<uint16> > rectInCells;
+								pLayout->GetBlock
+								(
+									b,
+									&rectInCells.min[0], &rectInCells.min[1],
+									&rectInCells.size[0], &rectInCells.size[1]
+								);
+
+								// If we don't know the size of a layout block in pixels, calculate it
+								UpdateBlockSize(BlockDesc, rectInCells.size);
+
+								// Even if we force the size afterwards, we need some size hint in some cases, like image projections.
+								ImageOptions.RectSize = UE::Math::TIntVector2<int32>(BlockDesc.m_size);
+
 								// Look for patches to this block
 								for (int32 editIndex = 0; editIndex < edits.Num(); ++editIndex)
 								{
@@ -1488,17 +1452,21 @@ namespace mu
 									}
 								}
 
+								// Enforce block size and optimizations
+								blockAd = GenerateImageSize(blockAd, FIntVector2(BlockDesc.m_size));
+
 								EImageFormat baseFormat = imageAd->GetImageDesc().m_format;
-								Ptr<ASTOp> FormattedBlock = GenerateImageFormat(blockAd, baseFormat);
+								// Actually don't do it, it will be propagated from the top format operation.
+								//Ptr<ASTOp> blockAd = GenerateImageFormat(blockAd, baseFormat);
 
 								// Apply tiling to avoid generating chunks of image that are too big.
-								FormattedBlock = ApplyTiling(FormattedBlock, ImageOptions.RectSize, desc.m_format);
+								blockAd = ApplyTiling(blockAd, ImageOptions.RectSize, FinalFormat);
 
 								// Compose layout operation
 								Ptr<ASTOpImageCompose> composeOp = new ASTOpImageCompose();
 								composeOp->Layout = meshResults.layoutOps[LayoutIndex];
 								composeOp->Base = imageAd;
-								composeOp->BlockImage = FormattedBlock;
+								composeOp->BlockImage = blockAd;
 
 								// Set the absolute block index.
 								check(pLayout->m_blocks[b].m_id >= 0);
@@ -1531,11 +1499,6 @@ namespace mu
 										{
 											Ptr<const Layout> pExtendLayout = meshResults.extraMeshLayouts[editIndex].GeneratedLayouts[LayoutIndex];
 
-											// Find out the size of the image
-											// TODO: What if the image is empty and everything is added?
-											//       Look in the extending images.
-											FImageDesc extendDesc = CalculateImageDesc(*pExtend->GetBasePrivate());
-
 											// Size of a layout block in pixels
 											FIntPoint extlayout = pExtendLayout->GetGridSize();
 
@@ -1543,44 +1506,43 @@ namespace mu
 
 											for (int b = 0; b < pExtendLayout->GetBlockCount(); ++b)
 											{
-												// Block in layout grid units
-												box< UE::Math::TIntVector2<uint16> > blockRect;
-												pExtendLayout->GetBlock
-												(
-													b,
-													&blockRect.min[0], &blockRect.min[1],
-													&blockRect.size[0], &blockRect.size[1]
-												);
-
-												// Transform to pixels
-												box< UE::Math::TIntVector2<int32> > rect;
-												rect.min[0] = (blockRect.min[0] * extendDesc.m_size[0]) / extlayout[0];
-												rect.min[1] = (blockRect.min[1] * extendDesc.m_size[1]) / extlayout[1];
-												rect.size[0] = (blockRect.size[0] * extendDesc.m_size[0]) / extlayout[0];
-												rect.size[1] = (blockRect.size[1] * extendDesc.m_size[1]) / extlayout[1];
-
 												// Generate the image block
 												FImageGenerationOptions ImageOptions;
-												ImageOptions.CurrentStateIndex = m_currentStateIndex;
+												ImageOptions.State = Options.State;
 												ImageOptions.ImageLayoutStrategy = ImageLayoutStrategy;
 												ImageOptions.ActiveTags = node.m_tags;
-												ImageOptions.RectSize = UE::Math::TIntVector2<int32>(extendDesc.m_size);
+												ImageOptions.RectSize = { 0,0 };
 												ImageOptions.LayoutToApply = pExtendLayout;
 												ImageOptions.LayoutBlockId = pExtendLayout->m_blocks[b].m_id;
 												FImageGenerationResult ExtendResult;
 												GenerateImage(ImageOptions, ExtendResult, pExtend);
 												Ptr<ASTOp> fragmentAd = ExtendResult.op;
 
+												// Block in layout grid units
+												box< UE::Math::TIntVector2<uint16> > rectInCells;
+												pExtendLayout->GetBlock
+												(
+													b,
+													&rectInCells.min[0], &rectInCells.min[1],
+													&rectInCells.size[0], &rectInCells.size[1]
+												);
+
+												FImageDesc ExtendDesc = fragmentAd->GetImageDesc();
+
+												// If we don't know the size of a layout block in pixels, calculate it
+												UpdateBlockSize(ExtendDesc, rectInCells.size);
+
 												// Adjust the format and size of the block to be added
-												Ptr<ASTOp> formatted = GenerateImageFormat(fragmentAd, GetUncompressedFormat(desc.m_format));
+												// Actually don't do it, it will be propagated from the top format operation.
+												//fragmentAd = GenerateImageFormat(fragmentAd, FinalFormat);
+
 												UE::Math::TIntVector2<int32> expectedSize;
-												expectedSize[0] = blockSizeX * blockRect.size[0];
-												expectedSize[1] = blockSizeY * blockRect.size[1];
-												formatted = GenerateImageSize(formatted, expectedSize);
-												fragmentAd = formatted;
+												expectedSize[0] = BlockPixelsX * rectInCells.size[0];
+												expectedSize[1] = BlockPixelsY * rectInCells.size[1];
+												fragmentAd = GenerateImageSize(fragmentAd, expectedSize);
 
 												// Apply tiling to avoid generating chunks of image that are too big.
-												fragmentAd = ApplyTiling(fragmentAd, expectedSize, desc.m_format);
+												fragmentAd = ApplyTiling(fragmentAd, expectedSize, FinalFormat);
 
 												// Compose operation
 												Ptr<ASTOpImageCompose> composeOp = new ASTOpImageCompose();
@@ -1646,10 +1608,10 @@ namespace mu
 
 								// We have to avoid mips smaller than the image format block size, so
 								// we will devide the layout block by the format block
-								const FImageFormatData& finfo = GetImageFormatData(desc.m_format);
+								const FImageFormatData& finfo = GetImageFormatData(FinalFormat);
 
-								int32 mipsX = FMath::CeilLogTwo(blockSizeX / finfo.PixelsPerBlockX);
-								int32 mipsY = FMath::CeilLogTwo(blockSizeY / finfo.PixelsPerBlockY);
+								int32 mipsX = FMath::CeilLogTwo(BlockPixelsX / finfo.PixelsPerBlockX);
+								int32 mipsY = FMath::CeilLogTwo(BlockPixelsY / finfo.PixelsPerBlockY);
 								mop->BlockLevels = (uint8)FMath::Max(mipsX, mipsY);
 
 								mop->AddressMode = mipmapNode->GetPrivate()->m_settings.m_addressMode;
@@ -1673,17 +1635,17 @@ namespace mu
 
 								// We have to avoid mips smaller than the image format block size, so
 								// we will devide the layout block by the format block
-								const FImageFormatData& finfo = GetImageFormatData(desc.m_format);
+								const FImageFormatData& finfo = GetImageFormatData(FinalFormat);
 
-								int mipsX = (int)ceilf(logf((float)blockSizeX / finfo.PixelsPerBlockX) / logf(2.0f));
-								int mipsY = (int)ceilf(logf((float)blockSizeY / finfo.PixelsPerBlockY) / logf(2.0f));
+								int mipsX = (int)ceilf(logf((float)BlockPixelsX / finfo.PixelsPerBlockX) / logf(2.0f));
+								int mipsY = (int)ceilf(logf((float)BlockPixelsY / finfo.PixelsPerBlockY) / logf(2.0f));
 								mop->BlockLevels = (uint8_t)FMath::Max(mipsX, mipsY);
 
-							// Not important for the end of the mip tail?
-							mop->AddressMode = EAddressMode::ClampToEdge;
-							mop->FilterType = EMipmapFilterType::MFT_SimpleAverage;
-							mop->SharpenFactor = 0;
-							mop->DitherMipmapAlpha = false;
+								// Not important for the end of the mip tail?
+								mop->AddressMode = EAddressMode::ClampToEdge;
+								mop->FilterType = EMipmapFilterType::MFT_SimpleAverage;
+								mop->SharpenFactor = 0;
+								mop->DitherMipmapAlpha = false;
 
 								imageAd = mop;
 							}
@@ -1728,8 +1690,9 @@ namespace mu
 					op->instance = lastSurfOp;
 
 					// Vector
-					Ptr<ASTOp> vectorAd = Generate(pVectorNode);
-					op->value = vectorAd;
+					FColorGenerationResult VectorResult;
+					GenerateColor(VectorResult, Options, pVectorNode);
+					op->value = VectorResult.op;
 
 					// Name
 					op->name = node.m_vectors[t].m_name;
@@ -1751,8 +1714,9 @@ namespace mu
 					op->instance = lastSurfOp;
 
 					// Scalar
-					Ptr<ASTOp> scalarAd = Generate(pScalarNode);
-					op->value = scalarAd;
+					FScalarGenerationResult ScalarResult;
+					GenerateScalar(ScalarResult, Options, pScalarNode);
+					op->value = ScalarResult.op;
 
 					// Name
 					op->name = node.m_scalars[t].m_name;
@@ -1771,8 +1735,9 @@ namespace mu
 					op->type = OP_TYPE::IN_ADDSTRING;
 					op->instance = lastSurfOp;
 
-					Ptr<ASTOp> stringAd = Generate(pStringNode);
-					op->value = stringAd;
+					FStringGenerationResult StringResult;
+					GenerateString(StringResult, Options, pStringNode);
+					op->value = StringResult.op;
 
 					// Name
 					op->name = node.m_strings[t].m_name;
@@ -1795,9 +1760,12 @@ namespace mu
 
 
     //---------------------------------------------------------------------------------------------
-    Ptr<ASTOp> CodeGenerator::Visit( const NodeComponentNew::Private& node )
-    {
-        MUTABLE_CPUPROFILER_SCOPE(NodeComponentNew);
+    //Ptr<ASTOp> CodeGenerator::Visit( const NodeComponentNew::Private& node )
+	void CodeGenerator::Generate_ComponentNew(const FGenericGenerationOptions& Options, FGenericGenerationResult& Result, const NodeComponentNew* InNode)
+	{
+		const NodeComponentNew::Private& node = *InNode->GetPrivate();
+
+		MUTABLE_CPUPROFILER_SCOPE(NodeComponentNew);
 
 		// Build a series of operations to assemble the component
         Ptr<ASTOp> lastCompOp;
@@ -1818,12 +1786,12 @@ namespace mu
                 // Apply state conditions: only generate it if it enabled in this state
                 {
                     bool enabledInThisState = true;
-                    if (its.stateCondition.Num() && m_currentStateIndex >= 0)
+                    if (its.stateCondition.Num() && Options.State >= 0)
                     {
                         enabledInThisState =
-                                ( m_currentStateIndex < its.stateCondition.Num() )
+                                (Options.State < its.stateCondition.Num())
                                 &&
-                                ( its.stateCondition[m_currentStateIndex] );
+                                ( its.stateCondition[Options.State] );
                     }
                     if (!enabledInThisState)
                     {
@@ -1837,7 +1805,7 @@ namespace mu
                 sop->instance = lastCompOp;
 
 				FSurfaceGenerationResult surfaceGenerationResult;
-                GenerateSurface( surfaceGenerationResult, its.node, its.edits );
+                GenerateSurface( surfaceGenerationResult, Options, its.node, its.edits );
                 sop->value = surfaceGenerationResult.surfaceOp;
 
                 sop->id = surfaceID;
@@ -1932,26 +1900,17 @@ namespace mu
     	iop->id = node.m_id;
         lastCompOp = iop;
 
-        return lastCompOp;
+        Result.op = lastCompOp;
     }
 
 
     //---------------------------------------------------------------------------------------------
-    Ptr<ASTOp> CodeGenerator::Visit( const NodeComponentEdit::Private& )
-    {
-        MUTABLE_CPUPROFILER_SCOPE(NodeComponentEdit);
-
-        // Nothing to do. Surface information will be already collected in the suitable
-        // parent components during the first and second passes.
-
-        return nullptr;
-    }
-
-
-    //---------------------------------------------------------------------------------------------
-    Ptr<ASTOp> CodeGenerator::Visit( const NodeObjectNew::Private& node )
-    {
-        MUTABLE_CPUPROFILER_SCOPE(NodeObjectNew);
+   // Ptr<ASTOp> CodeGenerator::Visit( const NodeObjectNew::Private& node )
+	void CodeGenerator::Generate_ObjectNew(const FGenericGenerationOptions& Options, FGenericGenerationResult& Result, const NodeObjectNew* InNode)
+	{
+		const NodeObjectNew::Private& node = *InNode->GetPrivate();
+		
+		MUTABLE_CPUPROFILER_SCOPE(NodeObjectNew);
 
         m_currentParents.Add( FParentKey() );
         m_currentParents.Last().m_pObject = &node;
@@ -1978,13 +1937,13 @@ namespace mu
                     paramOp = op;
                 }
 
-                OBJECT_GENERATION_DATA data;
+				FObjectGenerationData data;
                 data.m_condition = paramOp;
                 m_currentObject.Add( data );
 
                 // This op is ignored: everything is stored as patches to apply to the parent when
                 // it is compiled.
-                Generate( pChildNode );
+                Generate( pChildNode, Options );
 
                 m_currentObject.Pop();
             }
@@ -1998,9 +1957,9 @@ namespace mu
             {
                 m_currentParents.Last().m_lod = t;
 
-                Ptr<ASTOp> lodOp = Generate( pLODNode );
-
-                lodsOp->lods.Emplace( lodsOp, lodOp );
+				FGenericGenerationResult LODResult;
+				Generate_LOD(Options, LODResult, pLODNode);
+                lodsOp->lods.Emplace( lodsOp, LODResult.op);
             }
         }
         Ptr<ASTOp> rootOp = lodsOp;
@@ -2017,27 +1976,27 @@ namespace mu
 			// Name must be valid
 			check(NamedNode.Name.Len() > 0);
 
-			FExtensionDataGenerationResult Result;
-			GenerateExtensionData(Result, NamedNode.Node);
+			FExtensionDataGenerationResult ChildResult;
+			GenerateExtensionData(ChildResult, Options, NamedNode.Node);
 
-			if (!Result.Op.get())
+			if (!ChildResult.Op.get())
 			{
 				// Failed to generate anything for this node
 				continue;
 			}
 
-			FConditionalExtensionDataOp& SavedOp = m_conditionalExtensionDataOps.AddDefaulted_GetRef();
+			FConditionalExtensionDataOp& SavedOp = ConditionalExtensionDataOps.AddDefaulted_GetRef();
 			if (m_currentObject.Num() > 0)
 			{
 				SavedOp.Condition = m_currentObject.Last().m_condition;
 			}
-			SavedOp.ExtensionDataOp = Result.Op;
+			SavedOp.ExtensionDataOp = ChildResult.Op;
 			SavedOp.ExtensionDataName = NamedNode.Name;
 		}
 
 		if (m_currentObject.Num() == 0)
 		{
-			for (const FConditionalExtensionDataOp& SavedOp : m_conditionalExtensionDataOps)
+			for (const FConditionalExtensionDataOp& SavedOp : ConditionalExtensionDataOps)
 			{
 				Ptr<ASTOpAddExtensionData> ExtensionPinOp = new ASTOpAddExtensionData();
 				ExtensionPinOp->Instance = ASTChild(ExtensionPinOp, rootOp);
@@ -2063,13 +2022,16 @@ namespace mu
 
         m_currentParents.Pop();
 
-        return rootOp;
+        Result.op = rootOp;
     }
 
 
     //---------------------------------------------------------------------------------------------
-    Ptr<ASTOp> CodeGenerator::Visit( const NodeObjectGroup::Private& node )
-    {
+    //Ptr<ASTOp> CodeGenerator::Visit( const NodeObjectGroup::Private& node )
+	void CodeGenerator::Generate_ObjectGroup(const FGenericGenerationOptions& Options, FGenericGenerationResult& Result, const NodeObjectGroup* InNode)
+	{
+		const NodeObjectGroup::Private& node = *InNode->GetPrivate();
+
 		TArray<FString> usedNames;
 
         // Parse the child objects first, which will accumulate operations in the patching lists
@@ -2095,13 +2057,13 @@ namespace mu
                 // It may happen with partial compilations?
                 // check(found);
 
-                OBJECT_GENERATION_DATA data;
+				FObjectGenerationData data;
                 data.m_condition = conditionOp;
                 m_currentObject.Add( data );
 
                 // This op is ignored: everything is stored as patches to apply to the parent when
                 // it is compiled.
-                Generate( pChildNode );
+                Generate( pChildNode, Options );
 
                 m_currentObject.Pop();
 
@@ -2118,26 +2080,24 @@ namespace mu
 				}
             }
         }
-
-        return 0;
     }
 
+
     //---------------------------------------------------------------------------------------------
-    Ptr<ASTOp> CodeGenerator::GenerateMissingBoolCode(const TCHAR* strWhere,
-                                                      bool value,
-                                                      const void* errorContext )
+    Ptr<ASTOp> CodeGenerator::GenerateMissingBoolCode(const TCHAR* Where, bool Value, const void* ErrorContext )
     {
         // Log a warning
-		FString Msg = FString::Printf(TEXT("Required connection not found: %s"), strWhere );
-        m_pErrorLog->GetPrivate()->Add( Msg, ELMT_ERROR, errorContext );
+		FString Msg = FString::Printf(TEXT("Required connection not found: %s"), Where);
+        m_pErrorLog->GetPrivate()->Add( Msg, ELMT_ERROR, ErrorContext);
 
         // Create a constant node
         Ptr<NodeBoolConstant> pNode = new NodeBoolConstant();
-        pNode->SetValue( value );
+        pNode->SetValue(Value);
 
-        Ptr<ASTOp> result = Generate( pNode );
-
-        return result;
+		FBoolGenerationResult ChildResult;
+		FGenericGenerationOptions Options;
+        GenerateBool(ChildResult,Options, pNode );
+		return ChildResult.op;
     }
 
 
@@ -2220,8 +2180,8 @@ namespace mu
 
 	//---------------------------------------------------------------------------------------------
 	Ptr<ASTOp> CodeGenerator::ApplyMeshModifiers(
+		const FGenericGenerationOptions& Options,
 		const Ptr<ASTOp>& sourceOp,
-		const TArray<FString>& SurfaceTags,
 		bool bModifiersForBeforeOperations,
 		const void* errorContext )
 	{
@@ -2231,7 +2191,7 @@ namespace mu
 		TArray<FirstPassGenerator::FModifier> modifiers;
 
 		int currentLOD = m_currentParents.Last().m_lod;
-		GetModifiersFor(SurfaceTags, currentLOD, bModifiersForBeforeOperations, modifiers);
+		GetModifiersFor(Options.ActiveTags, currentLOD, bModifiersForBeforeOperations, modifiers);
 
 		Ptr<ASTOp> preModifiersMesh = lastMeshOp;
 
@@ -2250,7 +2210,7 @@ namespace mu
 				FMeshGenerationOptions ClipOptions;
 				ClipOptions.bUniqueVertexIDs = false;
 				ClipOptions.bLayouts = false;
-				ClipOptions.State = m_currentStateIndex;
+				ClipOptions.State = Options.State;
 
 				FMeshGenerationResult clipResult;
 				GenerateMesh(ClipOptions, clipResult, TypedClipNode->ClipMesh);
@@ -2294,7 +2254,7 @@ namespace mu
 				FImageGenerationOptions ClipOptions;
 				ClipOptions.ImageLayoutStrategy = CompilerOptions::TextureLayoutStrategy::None;
 				ClipOptions.LayoutBlockId = -1;
-				ClipOptions.CurrentStateIndex = m_currentStateIndex;
+				ClipOptions.State = Options.State;
 
 				FImageGenerationResult ClipMaskResult;
 				GenerateImage(ClipOptions, ClipMaskResult, TypedClipNode->ClipMask);
@@ -2421,7 +2381,7 @@ namespace mu
 				FMeshGenerationOptions ClipOptions;
 				ClipOptions.bUniqueVertexIDs = false;
 				ClipOptions.bLayouts = false;
-				ClipOptions.State = m_currentStateIndex;
+				ClipOptions.State = Options.State;
 
 				FMeshGenerationResult ClipShapeResult;
 				GenerateMesh(ClipOptions, ClipShapeResult, TypedClipNode->ClipMesh);
@@ -2485,7 +2445,10 @@ namespace mu
 			mu::Ptr<mu::NodeColourConstant> pNode = new NodeColourConstant();
 			pNode->SetValue(mu::DefaultMutableColorValue);
 
-			return Generate(pNode);
+			FColorGenerationResult ChildResult;
+			FGenericGenerationOptions Options;
+			GenerateColor(ChildResult, Options, pNode);
+			return ChildResult.op;
 		}
 		case mu::ETableColumnType::Image:
 		{
