[mutable] Redesign code generator to not rely on the visitor pattern and reducing RTTI usage.

diff --git a/Engine/Plugins/Experimental/Mutable/Source/MutableTools/Private/MuT/CodeGenerator.h b/Engine/Plugins/Experimental/Mutable/Source/MutableTools/Private/MuT/CodeGenerator.h
--- a/Engine/Plugins/Experimental/Mutable/Source/MutableTools/Private/MuT/CodeGenerator.h
+++ b/Engine/Plugins/Experimental/Mutable/Source/MutableTools/Private/MuT/CodeGenerator.h
@@ -38,11 +38,8 @@
 #include "MuT/NodeSurfaceNew.h"
 #include "MuT/Table.h"
 #include "MuT/TablePrivate.h"
-#include "MuT/Visitor.h"
 #include "Templates/TypeHash.h"
 
-#include <utility>
-
 
 namespace mu
 {
@@ -59,7 +56,6 @@ namespace mu
 	class NodeImageColourMap;
 	class NodeImageConditional;
 	class NodeImageConstant;
-	class NodeImageReference;
 	class NodeImageFormat;
 	class NodeImageGradient;
 	class NodeImageInterpolate;
@@ -112,15 +108,7 @@ namespace mu
     //---------------------------------------------------------------------------------------------
     //! Code generator
     //---------------------------------------------------------------------------------------------
-    class CodeGenerator : public Base,
-                          public BaseVisitor,
-
-                          public Visitor<NodeComponentNew::Private, Ptr<ASTOp>, true>,
-                          public Visitor<NodeComponentEdit::Private, Ptr<ASTOp>, true>,
-                          public Visitor<NodeLOD::Private, Ptr<ASTOp>, true>,
-                          public Visitor<NodeObjectNew::Private, Ptr<ASTOp>, true>,
-                          public Visitor<NodeObjectGroup::Private, Ptr<ASTOp>, true>,
-                          public Visitor<NodePatchImage::Private, Ptr<ASTOp>, true>
+    class CodeGenerator
     {
 		
 		friend class FirstPassGenerator;
@@ -130,90 +118,88 @@ namespace mu
         CodeGenerator( CompilerOptions::Private* options );
 
         //! Data will be stored in m_states
-        void GenerateRoot( const NodePtrConst pNode );
-
-	protected:
-
-        Ptr<ASTOp> Generate(const NodePtrConst pNode);
+        void GenerateRoot( const Ptr<const Node> );
 
 	public:
 
-        Ptr<ASTOp> Visit( const NodeComponentNew::Private& ) override;
-        Ptr<ASTOp> Visit( const NodeComponentEdit::Private& ) override;
-        Ptr<ASTOp> Visit( const NodeLOD::Private& ) override;
-        Ptr<ASTOp> Visit( const NodeObjectNew::Private& ) override;
-        Ptr<ASTOp> Visit( const NodeObjectGroup::Private& ) override;
-        Ptr<ASTOp> Visit( const NodePatchImage::Private& ) override;
+		// Generic top-level nodes
+		struct FGenericGenerationOptions
+		{
+			friend FORCEINLINE uint32 GetTypeHash(const FGenericGenerationOptions& InKey)
+			{
+				uint32 KeyHash = 0;
+				KeyHash = HashCombineFast(KeyHash, ::GetTypeHash(InKey.State));
+				KeyHash = HashCombineFast(KeyHash, ::GetTypeHash(InKey.ActiveTags.Num()));
+				return KeyHash;
+			}
 
-    public:
+			bool operator==(const FGenericGenerationOptions& InKey) const
+			{
+				if (State != InKey.State) return false;
+				if (ActiveTags != InKey.ActiveTags) return false;
+				return true;
+			}
 
-        //! Settings
-        CompilerOptions::Private* m_compilerOptions = nullptr;
+			int32 State = -1;
+			TArray<FString> ActiveTags;
+		};
 
-		//!
-		FirstPassGenerator m_firstPass;
+		struct FGenericGenerationResult
+		{
+			Ptr<ASTOp> op;
+		};
 
-        struct FVisitedKeyMap
-        {
-            FVisitedKeyMap()
-            {
-            }
+		struct FGeneratedCacheKey
+		{
+			Ptr<const Node> Node;
+			FGenericGenerationOptions Options;
 
-			friend FORCEINLINE uint32 GetTypeHash(const FVisitedKeyMap& InKey)
+			friend FORCEINLINE uint32 GetTypeHash(const FGeneratedCacheKey& InKey)
 			{
 				uint32 KeyHash = 0;
-				KeyHash = HashCombine(KeyHash, ::GetTypeHash(InKey.pNode.get()));
-				KeyHash = HashCombine(KeyHash, ::GetTypeHash(InKey.state));
-				KeyHash = HashCombine(KeyHash, ::GetTypeHash((uint64)InKey.activeTags.Num()));
+				KeyHash = HashCombineFast(KeyHash, ::GetTypeHash(InKey.Node.get()));
+				KeyHash = HashCombineFast(KeyHash, GetTypeHash(InKey.Options));
 				return KeyHash;
 			}
 
-			bool operator==(const FVisitedKeyMap& InKey) const
+			FORCEINLINE bool operator==(const FGeneratedCacheKey& Other) const
 			{
-				if (pNode != InKey.pNode) return false;
-				if (state != InKey.state) return false;
-				if (activeTags != InKey.activeTags) return false;
-				return true;
+				return Node == Other.Node && Options == Other.Options;
 			}
+		};
 
-            // This reference has to be the smart pointer to avoid memory aliasing, keeping
-            // processed nodes alive.
-            Ptr<const Node> pNode;
-            int state = -1;
-			TArray<FString> activeTags;
-        };
+		typedef TMap<FGeneratedCacheKey, FGenericGenerationResult> FGeneratedGenericNodesMap;
+		FGeneratedGenericNodesMap GeneratedGenericNodes;
 
-        //! This struct contains additional state propagated from bottom to top of the object node graph.
-        //! It is stored for every visited node, and restored when the cache is used.
-        struct FBottomUpState
-        {
-            //! Generated root address for the node.
-            Ptr<ASTOp> m_address;
-        };
-		FBottomUpState m_currentBottomUpState;
+		Ptr<ASTOp> Generate(const Ptr<const Node>, const FGenericGenerationOptions& );
+		void Generate_ComponentNew(const FGenericGenerationOptions&, FGenericGenerationResult&, const NodeComponentNew*);
+		void Generate_LOD(const FGenericGenerationOptions&, FGenericGenerationResult&, const NodeLOD*);
+		void Generate_ObjectNew(const FGenericGenerationOptions&, FGenericGenerationResult&, const NodeObjectNew*);
+		void Generate_ObjectGroup(const FGenericGenerationOptions&, FGenericGenerationResult&, const NodeObjectGroup*);
 
-        typedef TMap<FVisitedKeyMap, FBottomUpState> VisitedMap;
-        VisitedMap m_compiled;
+    public:
+
+        //! Settings
+        CompilerOptions::Private* m_compilerOptions = nullptr;
+
+		//!
+		FirstPassGenerator m_firstPass;
 
         //!
         ErrorLogPtr m_pErrorLog;
 
-        //! While generating code, this contains the index of the state being generated. This
-        //! can only be used with the state data in m_firstPass.
-        int m_currentStateIndex = -1;
-
         //! After the entire code generation this contains the information about all the states
-        typedef TArray< std::pair<FObjectState, Ptr<ASTOp>> > StateList;
+        typedef TArray< TPair<FObjectState, Ptr<ASTOp>> > StateList;
         StateList m_states;
 
     private:
 
         //! List of meshes generated to be able to reuse them
-		TArray<MeshPtr> m_constantMeshes;
+		TArray<Ptr<Mesh>> m_constantMeshes;
 
         //! List of image resources for every image formata that have been generated so far as
         //! palceholders for missing images.
-        ImagePtr m_missingImage[size_t(EImageFormat::IF_COUNT)];
+        Ptr<Image> m_missingImage[size_t(EImageFormat::IF_COUNT)];
 
         //! First free index for a layout block
         int32 m_absoluteLayoutIndex = 0;
@@ -259,29 +245,37 @@ namespace mu
             }
 
             const NodeObjectNew::Private* m_pObject;
-            int m_lod;
+            int32 m_lod;
 
-            inline bool operator<(const FAdditionalComponentKey& o) const
-            {
-                if (m_pObject < o.m_pObject) return true;
-                if (m_pObject > o.m_pObject) return false;
-                return m_lod < o.m_lod;
-            }
-        };
-        std::map< FAdditionalComponentKey, TArray<Ptr<ASTOp>> > m_additionalComponents;
+			FORCEINLINE bool operator==(const FAdditionalComponentKey& Other) const
+			{
+				return m_pObject == Other.m_pObject
+					&&
+					m_lod == m_lod;
+			}
 
+			friend FORCEINLINE uint32 GetTypeHash(const FAdditionalComponentKey& InKey)
+			{
+				uint32 KeyHash = 0;
+				KeyHash = HashCombineFast(KeyHash, ::GetTypeHash(InKey.m_pObject));
+				KeyHash = HashCombineFast(KeyHash, ::GetTypeHash(InKey.m_lod));
+				return KeyHash;
+			}
+		};
+        TMap< FAdditionalComponentKey, TArray<Ptr<ASTOp>> > AdditionalComponents;
 
-        struct OBJECT_GENERATION_DATA
+
+        struct FObjectGenerationData
         {
             // Condition that enables a specific object
             Ptr<ASTOp> m_condition;
         };
-		TArray< OBJECT_GENERATION_DATA > m_currentObject;
+		TArray<FObjectGenerationData> m_currentObject;
 
-        map< std::pair<TablePtr, FString>, std::pair<TablePtr,Ptr<ASTOp>> > m_generatedTables;
+		TMap< TPair<Ptr<Table>, FString>, TPair<Ptr<Table>,Ptr<ASTOp>> > GeneratedTables;
 
         //! Variables added for every node
-        map< Ptr<const Node>, Ptr<ASTOpParameter> > m_nodeVariables;
+		TMap< Ptr<const Node>, Ptr<ASTOpParameter> > NodeVariables;
 
 		struct FConditionalExtensionDataOp
 		{
@@ -290,7 +284,7 @@ namespace mu
 			FString ExtensionDataName;
 		};
 
-		TArray<FConditionalExtensionDataOp> m_conditionalExtensionDataOps;
+		TArray<FConditionalExtensionDataOp> ConditionalExtensionDataOps;
 
 		//-----------------------------------------------------------------------------------------
 
@@ -299,12 +293,9 @@ namespace mu
 			bool bModifiersForBeforeOperations, TArray<FirstPassGenerator::FModifier>& OutModifiers);
 
 		// Apply the required mesh modifiers to the given operation.
-		Ptr<ASTOp> ApplyMeshModifiers( const Ptr<ASTOp>& sourceOp, const TArray<FString>& SurfaceTags,
+		Ptr<ASTOp> ApplyMeshModifiers(const FGenericGenerationOptions&, const Ptr<ASTOp>& SourceOp,
 			bool bModifiersForBeforeOperations, const void* errorContext);
 
-		// Get the modifiers that have to be applied to elements with a specific tag.
-        //void GetSurfacesWithTag(const string& tag, vector<FirstPassGenerator::SURFACE>& surfaces);
-
         //-----------------------------------------------------------------------------------------
         //!
         Ptr<ASTOp> GenerateTableVariable( TablePtr pTable, const FString& strName );
@@ -317,52 +308,30 @@ namespace mu
 		Ptr<ASTOp> GenerateTableSwitch( const NODE_TABLE_PRIVATE& node, F&& GenerateOption );
 
 
-    private:
-
-		//! Generate the key with all the relevant state that is used in generation of operations for a node.
-		FVisitedKeyMap GetCurrentCacheKey(const NodePtrConst& InNode) const
-		{
-			FVisitedKeyMap key;
-			key.pNode = InNode;
-			key.state = m_currentStateIndex;
-			if (!m_activeTags.IsEmpty())
-			{
-				key.activeTags = m_activeTags.Last();
-			}
-			return key;
-		}
-
-
 		//-----------------------------------------------------------------------------------------
 		// Images
 		
 		/** Options that affect the generation of images. It is like list of what required data we want while parsing down the image node graph. */
-		struct FImageGenerationOptions
+		struct FImageGenerationOptions : public FGenericGenerationOptions
 		{
 			/** */
 			CompilerOptions::TextureLayoutStrategy ImageLayoutStrategy = CompilerOptions::TextureLayoutStrategy::None;
 
-			/** This is used to introduce additional image generation safety. \TODO: Move this "safety" to optimization? */
+			/** If different than {0,0} this is the mandatory size of the image that needs to be generated. */
 			UE::Math::TIntVector2<int32> RectSize = {0, 0};
 
-			/** */
-			int32 CurrentStateIndex = -1;
-
 			/** Layout block that we are trying to generate if any. */
 			int32 LayoutBlockId = -1;
 			Ptr<const Layout> LayoutToApply;
 
-			/** Tags that are active at this point of the generation. */
-			TArray<FString> ActiveTags;
-
 			friend FORCEINLINE uint32 GetTypeHash(const FImageGenerationOptions& InKey)
 			{
 				uint32 KeyHash = 0;
-				KeyHash = HashCombine(KeyHash, ::GetTypeHash(InKey.ImageLayoutStrategy));
-				KeyHash = HashCombine(KeyHash, GetTypeHash(InKey.RectSize));
-				KeyHash = HashCombine(KeyHash, ::GetTypeHash(InKey.CurrentStateIndex));
-				KeyHash = HashCombine(KeyHash, ::GetTypeHash(InKey.LayoutBlockId));
-				KeyHash = HashCombine(KeyHash, ::GetTypeHash(InKey.LayoutToApply.get()));
+				KeyHash = HashCombineFast(KeyHash, ::GetTypeHash(InKey.ImageLayoutStrategy));
+				KeyHash = HashCombineFast(KeyHash, GetTypeHash(InKey.RectSize));
+				KeyHash = HashCombineFast(KeyHash, ::GetTypeHash(InKey.State));
+				KeyHash = HashCombineFast(KeyHash, ::GetTypeHash(InKey.LayoutBlockId));
+				KeyHash = HashCombineFast(KeyHash, ::GetTypeHash(InKey.LayoutToApply.get()));
 				return KeyHash;
 			}
 
@@ -372,7 +341,7 @@ namespace mu
 					&&
 					RectSize == Other.RectSize
 					&&
-					CurrentStateIndex == Other.CurrentStateIndex
+					State == Other.State
 					&&
 					LayoutBlockId == Other.LayoutBlockId
 					&&
@@ -398,8 +367,8 @@ namespace mu
 			friend FORCEINLINE uint32 GetTypeHash(const FGeneratedImageCacheKey& InKey)
 			{
 				uint32 KeyHash = 0;
-				KeyHash = HashCombine(KeyHash, ::GetTypeHash(InKey.Node.get()));
-				KeyHash = HashCombine(KeyHash, GetTypeHash(InKey.Options));
+				KeyHash = HashCombineFast(KeyHash, ::GetTypeHash(InKey.Node.get()));
+				KeyHash = HashCombineFast(KeyHash, GetTypeHash(InKey.Options));
 				return KeyHash;
 			}
 
@@ -414,7 +383,6 @@ namespace mu
 
 		void GenerateImage(const FImageGenerationOptions&, FImageGenerationResult& result, const NodeImagePtrConst& node);
 		void GenerateImage_Constant(const FImageGenerationOptions&, FImageGenerationResult&, const NodeImageConstant*);
-		void GenerateImage_Reference(const FImageGenerationOptions&, FImageGenerationResult&, const NodeImageReference*);
 		void GenerateImage_Interpolate(const FImageGenerationOptions&, FImageGenerationResult&, const NodeImageInterpolate*);
 		void GenerateImage_Saturate(const FImageGenerationOptions&, FImageGenerationResult&, const NodeImageSaturate*);
 		void GenerateImage_Table(const FImageGenerationOptions&, FImageGenerationResult&, const NodeImageTable*);
@@ -457,9 +425,6 @@ namespace mu
 		//!
 		Ptr<ASTOp> GenerateImageSize(Ptr<ASTOp>, UE::Math::TIntVector2<int32>);
 
-		//!
-		FImageDesc CalculateImageDesc(const Node::Private&);
-
 		/** Evaluate if the image to generate is big enough to be split in separate operations and tiled afterwards. */
 		Ptr<ASTOp> ApplyTiling(Ptr<ASTOp> Source, UE::Math::TIntVector2<int32> Size, EImageFormat Format);
 
@@ -472,14 +437,8 @@ namespace mu
 		/** Options that affect the generation of meshes. It is like list of what required data we want
 		* while parsing down the mesh node graph.
 		*/
-		struct FMeshGenerationOptions
+		struct FMeshGenerationOptions : public FGenericGenerationOptions
 		{
-			/** TODO: Review and document. */
-			int32 State = 0;
-
-			/** Tags that are active at this point of the generation. */
-			TArray<FString> ActiveTags;
-
 			/** Whatever mesh we reach at the leaves of the graph will need to have unique ids for its vertices.
 			* This is used to track mesh removal indices, morph data in other nodes, clothing data, etc.
 			*/
@@ -509,11 +468,11 @@ namespace mu
 			friend FORCEINLINE uint32 GetTypeHash(const FMeshGenerationOptions& InKey)
 			{
 				uint32 KeyHash = 0;
-				KeyHash = HashCombine(KeyHash, ::GetTypeHash(InKey.bUniqueVertexIDs));
-				KeyHash = HashCombine(KeyHash, ::GetTypeHash(InKey.bLayouts));
-				KeyHash = HashCombine(KeyHash, ::GetTypeHash(InKey.OverrideLayouts.Num()));
-				KeyHash = HashCombine(KeyHash, ::GetTypeHash(InKey.State));
-				KeyHash = HashCombine(KeyHash, ::GetTypeHash(InKey.ActiveTags.Num()));
+				KeyHash = HashCombineFast(KeyHash, ::GetTypeHash(InKey.bUniqueVertexIDs));
+				KeyHash = HashCombineFast(KeyHash, ::GetTypeHash(InKey.bLayouts));
+				KeyHash = HashCombineFast(KeyHash, ::GetTypeHash(InKey.OverrideLayouts.Num()));
+				KeyHash = HashCombineFast(KeyHash, ::GetTypeHash(InKey.State));
+				KeyHash = HashCombineFast(KeyHash, ::GetTypeHash(InKey.ActiveTags.Num()));
 				return KeyHash;
 			}
 
@@ -559,8 +518,8 @@ namespace mu
 			friend FORCEINLINE uint32 GetTypeHash(const FGeneratedMeshCacheKey& InKey)
 			{
 				uint32 KeyHash = 0;
-				KeyHash = HashCombine(KeyHash, ::GetTypeHash(InKey.Node.get()));
-				KeyHash = HashCombine(KeyHash, GetTypeHash(InKey.Options));
+				KeyHash = HashCombineFast(KeyHash, ::GetTypeHash(InKey.Node.get()));
+				KeyHash = HashCombineFast(KeyHash, GetTypeHash(InKey.Options));
 				return KeyHash;
 			}
 
@@ -581,7 +540,7 @@ namespace mu
 		//! Map of layouts found in the code already generated. The map is from the source layout
 		//! pointer of the layouts in the node meshes to the cloned and modified layout. 
 		//! The cloned layout will have absolute block ids assigned.
-		TMap<Ptr<const Layout>, Ptr<const Layout>> m_generatedLayouts;
+		TMap<Ptr<const Layout>, Ptr<const Layout>> GeneratedLayouts;
 
         void GenerateMesh(const FMeshGenerationOptions&, FMeshGenerationResult& result, const NodeMeshPtrConst&);
         void GenerateMesh_Constant(const FMeshGenerationOptions&, FMeshGenerationResult&, const NodeMeshConstant* );
@@ -602,9 +561,9 @@ namespace mu
 		void GenerateMesh_ClipDeform(const FMeshGenerationOptions&, FMeshGenerationResult&, const NodeMeshClipDeform*);
 
 		//-----------------------------------------------------------------------------------------
-		void PrepareForLayout(LayoutPtrConst GeneratedLayout,
-			MeshPtr currentLayoutMesh,
-			size_t currentLayoutChannel,
+		void PrepareForLayout( Ptr<const Layout> GeneratedLayout,
+			Ptr<Mesh> currentLayoutMesh,
+			int32 currentLayoutChannel,
 			const void* errorContext,
 			const FMeshGenerationOptions& MeshOptions);
 
@@ -617,13 +576,13 @@ namespace mu
 		};
 
 		typedef const NodeExtensionData* FGeneratedExtensionDataCacheKey;
-		typedef TMap<FGeneratedExtensionDataCacheKey, FExtensionDataGenerationResult> GeneratedExtensionDataMap;
-		GeneratedExtensionDataMap m_generatedExtensionData;
+		typedef TMap<FGeneratedExtensionDataCacheKey, FExtensionDataGenerationResult> FGeneratedExtensionDataMap;
+		FGeneratedExtensionDataMap GeneratedExtensionData;
 
-		void GenerateExtensionData(FExtensionDataGenerationResult& OutResult, const NodeExtensionDataPtrConst& InUntypedNode);
-		void GenerateExtensionData_Constant(FExtensionDataGenerationResult& OutResult, const class NodeExtensionDataConstant* Constant);
-		void GenerateExtensionData_Switch(FExtensionDataGenerationResult& OutResult, const class NodeExtensionDataSwitch* Switch);
-		void GenerateExtensionData_Variation(FExtensionDataGenerationResult& OutResult, const class NodeExtensionDataVariation* Variation);
+		void GenerateExtensionData(FExtensionDataGenerationResult& OutResult, const FGenericGenerationOptions&, const NodeExtensionDataPtrConst&);
+		void GenerateExtensionData_Constant(FExtensionDataGenerationResult& OutResult, const FGenericGenerationOptions&, const class NodeExtensionDataConstant*);
+		void GenerateExtensionData_Switch(FExtensionDataGenerationResult& OutResult, const FGenericGenerationOptions&, const class NodeExtensionDataSwitch*);
+		void GenerateExtensionData_Variation(FExtensionDataGenerationResult& OutResult, const FGenericGenerationOptions&, const class NodeExtensionDataVariation*);
 		Ptr<ASTOp> GenerateMissingExtensionDataCode(const TCHAR* StrWhere, const void* ErrorContext);
 
         //-----------------------------------------------------------------------------------------
@@ -634,12 +593,12 @@ namespace mu
             PROJECTOR_TYPE type;
         };
 
-        typedef TMap<FVisitedKeyMap,FProjectorGenerationResult> GeneratedProjectorsMap;
-        GeneratedProjectorsMap m_generatedProjectors;
+        typedef TMap<FGeneratedCacheKey,FProjectorGenerationResult> FGeneratedProjectorsMap;
+        FGeneratedProjectorsMap GeneratedProjectors;
 
-        void GenerateProjector( FProjectorGenerationResult&, const NodeProjectorPtrConst& );
-        void GenerateProjector_Constant( FProjectorGenerationResult&, const Ptr<const NodeProjectorConstant>& );
-        void GenerateProjector_Parameter( FProjectorGenerationResult&, const Ptr<const NodeProjectorParameter>& );
+        void GenerateProjector( FProjectorGenerationResult&, const FGenericGenerationOptions&, const Ptr<const NodeProjector>& );
+        void GenerateProjector_Constant( FProjectorGenerationResult&, const FGenericGenerationOptions&, const Ptr<const NodeProjectorConstant>& );
+        void GenerateProjector_Parameter( FProjectorGenerationResult&, const FGenericGenerationOptions&, const Ptr<const NodeProjectorParameter>& );
         void GenerateMissingProjectorCode( FProjectorGenerationResult&, const void* errorContext );
 
 		//-----------------------------------------------------------------------------------------
@@ -649,15 +608,14 @@ namespace mu
 			Ptr<ASTOp> op;
 		};
 
-		typedef TMap<FVisitedKeyMap, FBoolGenerationResult> GeneratedBoolsMap;
-		GeneratedBoolsMap m_generatedBools;
+		typedef TMap<FGeneratedCacheKey, FBoolGenerationResult> FGeneratedBoolsMap;
+		FGeneratedBoolsMap GeneratedBools;
 
-		void GenerateBool(FBoolGenerationResult&, const Ptr<const NodeBool>&);
-		void GenerateBool_Constant(FBoolGenerationResult&, const Ptr<const NodeBoolConstant>&);
-		void GenerateBool_Parameter(FBoolGenerationResult&, const Ptr<const NodeBoolParameter>&);
-		void GenerateBool_IsNull(FBoolGenerationResult&, const Ptr<const NodeBoolIsNull>&);
-		void GenerateBool_Not(FBoolGenerationResult&, const Ptr<const NodeBoolNot>&);
-		void GenerateBool_And(FBoolGenerationResult&, const Ptr<const NodeBoolAnd>&);
+		void GenerateBool(FBoolGenerationResult&, const FGenericGenerationOptions&, const Ptr<const NodeBool>&);
+		void GenerateBool_Constant(FBoolGenerationResult&, const FGenericGenerationOptions&, const Ptr<const NodeBoolConstant>&);
+		void GenerateBool_Parameter(FBoolGenerationResult&, const FGenericGenerationOptions&, const Ptr<const NodeBoolParameter>&);
+		void GenerateBool_Not(FBoolGenerationResult&, const FGenericGenerationOptions&, const Ptr<const NodeBoolNot>&);
+		void GenerateBool_And(FBoolGenerationResult&, const FGenericGenerationOptions&, const Ptr<const NodeBoolAnd>&);
 
 		//-----------------------------------------------------------------------------------------
 		// Scalars
@@ -666,18 +624,18 @@ namespace mu
 			Ptr<ASTOp> op;
 		};
 
-		typedef TMap<FVisitedKeyMap, FScalarGenerationResult> GeneratedScalarsMap;
-		GeneratedScalarsMap m_generatedScalars;
-
-		void GenerateScalar(FScalarGenerationResult&, const Ptr<const NodeScalar>&);
-		void GenerateScalar_Constant(FScalarGenerationResult&, const Ptr<const NodeScalarConstant>&);
-		void GenerateScalar_Parameter(FScalarGenerationResult&, const Ptr<const NodeScalarParameter>&);
-		void GenerateScalar_Switch(FScalarGenerationResult&, const Ptr<const NodeScalarSwitch>&);
-		void GenerateScalar_EnumParameter(FScalarGenerationResult&, const Ptr<const NodeScalarEnumParameter>&);
-		void GenerateScalar_Curve(FScalarGenerationResult&, const Ptr<const NodeScalarCurve>&);
-		void GenerateScalar_Arithmetic(FScalarGenerationResult&, const Ptr<const NodeScalarArithmeticOperation>&);
-		void GenerateScalar_Variation(FScalarGenerationResult&, const Ptr<const NodeScalarVariation>&);
-		void GenerateScalar_Table(FScalarGenerationResult&, const Ptr<const NodeScalarTable>&);
+		typedef TMap<FGeneratedCacheKey, FScalarGenerationResult> FGeneratedScalarsMap;
+		FGeneratedScalarsMap GeneratedScalars;
+
+		void GenerateScalar(FScalarGenerationResult&, const FGenericGenerationOptions&, const Ptr<const NodeScalar>&);
+		void GenerateScalar_Constant(FScalarGenerationResult&, const FGenericGenerationOptions&, const Ptr<const NodeScalarConstant>&);
+		void GenerateScalar_Parameter(FScalarGenerationResult&, const FGenericGenerationOptions&, const Ptr<const NodeScalarParameter>&);
+		void GenerateScalar_Switch(FScalarGenerationResult&, const FGenericGenerationOptions&, const Ptr<const NodeScalarSwitch>&);
+		void GenerateScalar_EnumParameter(FScalarGenerationResult&, const FGenericGenerationOptions&, const Ptr<const NodeScalarEnumParameter>&);
+		void GenerateScalar_Curve(FScalarGenerationResult&, const FGenericGenerationOptions&, const Ptr<const NodeScalarCurve>&);
+		void GenerateScalar_Arithmetic(FScalarGenerationResult&, const FGenericGenerationOptions&, const Ptr<const NodeScalarArithmeticOperation>&);
+		void GenerateScalar_Variation(FScalarGenerationResult&, const FGenericGenerationOptions&, const Ptr<const NodeScalarVariation>&);
+		void GenerateScalar_Table(FScalarGenerationResult&, const FGenericGenerationOptions&, const Ptr<const NodeScalarTable>&);
 		Ptr<ASTOp> GenerateMissingScalarCode(const TCHAR* strWhere, float value, const void* errorContext);
 
 		//-----------------------------------------------------------------------------------------
@@ -687,18 +645,18 @@ namespace mu
 			Ptr<ASTOp> op;
 		};
 
-		typedef TMap<FVisitedKeyMap, FColorGenerationResult> GeneratedColorsMap;
-		GeneratedColorsMap m_generatedColors;
-
-		void GenerateColor(FColorGenerationResult&, const NodeColourPtrConst&);
-		void GenerateColor_Constant(FColorGenerationResult&, const Ptr<const NodeColourConstant>&);
-		void GenerateColor_Parameter(FColorGenerationResult&, const Ptr<const NodeColourParameter>&);
-		void GenerateColor_Switch(FColorGenerationResult&, const Ptr<const NodeColourSwitch>&);
-		void GenerateColor_SampleImage(FColorGenerationResult&, const Ptr<const NodeColourSampleImage>&);
-		void GenerateColor_FromScalars(FColorGenerationResult&, const Ptr<const NodeColourFromScalars>&);
-		void GenerateColor_Arithmetic(FColorGenerationResult&, const Ptr<const NodeColourArithmeticOperation>&);
-		void GenerateColor_Variation(FColorGenerationResult&, const Ptr<const NodeColourVariation>&);
-		void GenerateColor_Table(FColorGenerationResult&, const Ptr<const NodeColourTable>&);
+		typedef TMap<FGeneratedCacheKey, FColorGenerationResult> FGeneratedColorsMap;
+		FGeneratedColorsMap GeneratedColors;
+
+		void GenerateColor(FColorGenerationResult&, const FGenericGenerationOptions&, const Ptr<const NodeColour>&);
+		void GenerateColor_Constant(FColorGenerationResult&, const FGenericGenerationOptions&, const Ptr<const NodeColourConstant>&);
+		void GenerateColor_Parameter(FColorGenerationResult&, const FGenericGenerationOptions&, const Ptr<const NodeColourParameter>&);
+		void GenerateColor_Switch(FColorGenerationResult&, const FGenericGenerationOptions&, const Ptr<const NodeColourSwitch>&);
+		void GenerateColor_SampleImage(FColorGenerationResult&, const FGenericGenerationOptions&, const Ptr<const NodeColourSampleImage>&);
+		void GenerateColor_FromScalars(FColorGenerationResult&, const FGenericGenerationOptions&, const Ptr<const NodeColourFromScalars>&);
+		void GenerateColor_Arithmetic(FColorGenerationResult&, const FGenericGenerationOptions&, const Ptr<const NodeColourArithmeticOperation>&);
+		void GenerateColor_Variation(FColorGenerationResult&, const FGenericGenerationOptions&, const Ptr<const NodeColourVariation>&);
+		void GenerateColor_Table(FColorGenerationResult&, const FGenericGenerationOptions&, const Ptr<const NodeColourTable>&);
 		Ptr<ASTOp> GenerateMissingColourCode(const TCHAR* strWhere, const void* errorContext);
 
 		//-----------------------------------------------------------------------------------------
@@ -708,12 +666,12 @@ namespace mu
 			Ptr<ASTOp> op;
 		};
 
-		typedef TMap<FVisitedKeyMap, FStringGenerationResult> GeneratedStringsMap;
-		GeneratedStringsMap m_generatedStrings;
+		typedef TMap<FGeneratedCacheKey, FStringGenerationResult> FGeneratedStringsMap;
+		FGeneratedStringsMap GeneratedStrings;
 
-		void GenerateString(FStringGenerationResult&, const NodeStringPtrConst&);
-		void GenerateString_Constant(FStringGenerationResult&, const Ptr<const NodeStringConstant>&);
-		void GenerateString_Parameter(FStringGenerationResult&, const Ptr<const NodeStringParameter>&);
+		void GenerateString(FStringGenerationResult&, const FGenericGenerationOptions&, const Ptr<const NodeString>&);
+		void GenerateString_Constant(FStringGenerationResult&, const FGenericGenerationOptions& Options, const Ptr<const NodeStringConstant>&);
+		void GenerateString_Parameter(FStringGenerationResult&, const FGenericGenerationOptions& Options, const Ptr<const NodeStringParameter>&);
 
         //-----------------------------------------------------------------------------------------
         // Ranges
@@ -729,10 +687,10 @@ namespace mu
 			FString rangeUID;
         };
 
-        typedef TMap<FVisitedKeyMap,FRangeGenerationResult> GeneratedRangeMap;
-        GeneratedRangeMap m_generatedRanges;
+        typedef TMap<FGeneratedCacheKey,FRangeGenerationResult> FGeneratedRangeMap;
+        FGeneratedRangeMap GeneratedRanges;
 
-        void GenerateRange( FRangeGenerationResult& result, Ptr<const NodeRange> node);
+        void GenerateRange(FRangeGenerationResult&, const FGenericGenerationOptions&, Ptr<const NodeRange>);
 
 
         //-----------------------------------------------------------------------------------------
@@ -741,47 +699,28 @@ namespace mu
             Ptr<ASTOp> surfaceOp;
         };
 
-        void GenerateSurface( FSurfaceGenerationResult& result,
-                              NodeSurfaceNewPtrConst node,
-                              const TArray<FirstPassGenerator::FSurface::FEdit>& edits );
+        void GenerateSurface( FSurfaceGenerationResult&, const FGenericGenerationOptions&,
+                              Ptr<const NodeSurfaceNew>,
+                              const TArray<FirstPassGenerator::FSurface::FEdit>& Edits );
 
 		//-----------------------------------------------------------------------------------------
 		//Default Table Parameters
 		Ptr<ASTOp> GenerateDefaultTableValue(ETableColumnType NodeType);
     };
 
-
-    //---------------------------------------------------------------------------------------------
-    //! Analyse the code trying to guess the descriptor of the image genereated by the instruction
-    //! address.
-    //! \param returnBestOption If true, try to resolve ambiguities returning some value.
-    //---------------------------------------------------------------------------------------------
-    extern FImageDesc GetImageDesc( const FProgram& program, OP::ADDRESS at,
-                                    bool returnBestOption = false,
-                                    class FGetImageDescContext* context=nullptr );
-
-    //!
-    extern void PartialOptimise( Ptr<ASTOp>& op, const CompilerOptions* options );
-    
-
 	
     //---------------------------------------------------------------------------------------------
     template<class NODE_TABLE_PRIVATE, ETableColumnType TYPE, OP_TYPE OPTYPE, typename F>
-    Ptr<ASTOp> CodeGenerator::GenerateTableSwitch
-        (
-            const NODE_TABLE_PRIVATE& node, F&& GenerateOption
-        )
+    Ptr<ASTOp> CodeGenerator::GenerateTableSwitch( const NODE_TABLE_PRIVATE& node, F&& GenerateOption )
     {
-        TablePtr pTable;
+        Ptr<Table> pTable;
         Ptr<ASTOp> variable;
 
-        map< std::pair<TablePtr,FString>, std::pair<TablePtr,Ptr<ASTOp>> >::iterator it
-                = m_generatedTables.find
-                ( std::pair<TablePtr,FString>(node.m_pTable,node.m_parameterName) );
-        if ( it!=m_generatedTables.end() )
+        TPair<TablePtr,Ptr<ASTOp>>* it = GeneratedTables.Find( TPair<TablePtr,FString>(node.m_pTable,node.m_parameterName) );
+        if ( it )
         {
-            pTable = it->second.first;
-            variable = it->second.second;
+            pTable = it->Key;
+            variable = it->Value;
         }
 
         if ( !pTable )
@@ -790,12 +729,11 @@ namespace mu
             pTable = node.m_pTable;
             variable = GenerateTableVariable( pTable, node.m_parameterName );
 
-            m_generatedTables[ std::pair<TablePtr, FString>(node.m_pTable,node.m_parameterName) ] =
-                    std::pair<TablePtr,Ptr<ASTOp>>( pTable, variable );
+            GeneratedTables.Add( TPair<TablePtr, FString>(node.m_pTable,node.m_parameterName), TPair<TablePtr,Ptr<ASTOp>>( pTable, variable ) );
         }
 
         // Verify that the table column is the right type
-        int colIndex = pTable->FindColumn( node.m_columnName );
+        int32 colIndex = pTable->FindColumn( node.m_columnName );
         if ( colIndex<0 )
         {
             m_pErrorLog->GetPrivate()->Add("Table column not found.", ELMT_ERROR, node.m_errorContext);
diff --git a/Engine/Plugins/Experimental/Mutable/Source/MutableTools/Private/MuT/CodeGenerator_Bool.cpp b/Engine/Plugins/Experimental/Mutable/Source/MutableTools/Private/MuT/CodeGenerator_Bool.cpp
--- a/Engine/Plugins/Experimental/Mutable/Source/MutableTools/Private/MuT/CodeGenerator_Bool.cpp
+++ b/Engine/Plugins/Experimental/Mutable/Source/MutableTools/Private/MuT/CodeGenerator_Bool.cpp
@@ -15,54 +15,46 @@
 #include "MuT/NodeBool.h"
 #include "MuT/NodeBoolPrivate.h"
 #include "MuT/NodeRange.h"
-#include "map"
-
-#include <memory>
-#include <utility>
 
 
 namespace mu
 {
 
-
-	//-------------------------------------------------------------------------------------------------
-	void CodeGenerator::GenerateBool(FBoolGenerationResult& result, const Ptr<const NodeBool>& untyped)
+	void CodeGenerator::GenerateBool(FBoolGenerationResult& Result, const FGenericGenerationOptions& Options, const Ptr<const NodeBool>& Untyped)
 	{
-		if (!untyped)
+		if (!Untyped)
 		{
-			result = FBoolGenerationResult();
+			Result = FBoolGenerationResult();
 			return;
 		}
 
 		// See if it was already generated
-		FVisitedKeyMap key = GetCurrentCacheKey(untyped);
-		GeneratedBoolsMap::ValueType* it = m_generatedBools.Find(key);
+		FGeneratedCacheKey Key;
+		Key.Node = Untyped;
+		Key.Options = Options;
+		FGeneratedBoolsMap::ValueType* it = GeneratedBools.Find(Key);
 		if (it)
 		{
-			result = *it;
+			Result = *it;
 			return;
 		}
 
 		// Generate for each different type of node
-		if (auto Constant = dynamic_cast<const NodeBoolConstant*>(untyped.get()))
-		{
-			GenerateBool_Constant(result, Constant);
-		}
-		else if (auto Param = dynamic_cast<const NodeBoolParameter*>(untyped.get()))
+		if (auto Constant = dynamic_cast<const NodeBoolConstant*>(Untyped.get()))
 		{
-			GenerateBool_Parameter(result, Param);
+			GenerateBool_Constant(Result, Options, Constant);
 		}
-		else if (auto Switch = dynamic_cast<const NodeBoolIsNull*>(untyped.get()))
+		else if (auto Param = dynamic_cast<const NodeBoolParameter*>(Untyped.get()))
 		{
-			GenerateBool_IsNull(result, Switch);
+			GenerateBool_Parameter(Result, Options, Param);
 		}
-		else if (auto Sample = dynamic_cast<const NodeBoolNot*>(untyped.get()))
+		else if (auto Sample = dynamic_cast<const NodeBoolNot*>(Untyped.get()))
 		{
-			GenerateBool_Not(result, Sample);
+			GenerateBool_Not(Result, Options, Sample);
 		}
-		else if (auto From = dynamic_cast<const NodeBoolAnd*>(untyped.get()))
+		else if (auto From = dynamic_cast<const NodeBoolAnd*>(Untyped.get()))
 		{
-			GenerateBool_And(result, From);
+			GenerateBool_And(Result, Options, From);
 		}
 		else
 		{
@@ -70,12 +62,11 @@ namespace mu
 		}
 
 		// Cache the result
-		m_generatedBools.Add(key, result);
+		GeneratedBools.Add(Key, Result);
 	}
 
 
-	//-------------------------------------------------------------------------------------------------
-	void CodeGenerator::GenerateBool_Constant(FBoolGenerationResult& result, const Ptr<const NodeBoolConstant>& Typed)
+	void CodeGenerator::GenerateBool_Constant(FBoolGenerationResult& result, const FGenericGenerationOptions& Options, const Ptr<const NodeBoolConstant>& Typed)
 	{
 		const NodeBoolConstant::Private& node = *Typed->GetPrivate();
 
@@ -86,15 +77,14 @@ namespace mu
 	}
 
 
-	//-------------------------------------------------------------------------------------------------
-	void CodeGenerator::GenerateBool_Parameter(FBoolGenerationResult& result, const Ptr<const NodeBoolParameter>& Typed)
+	void CodeGenerator::GenerateBool_Parameter(FBoolGenerationResult& result, const FGenericGenerationOptions& Options, const Ptr<const NodeBoolParameter>& Typed)
 	{
 		const NodeBoolParameter::Private& node = *Typed->GetPrivate();
 
 		Ptr<ASTOpParameter> op;
 
-		auto it = m_nodeVariables.find(node.m_pNode);
-		if (it == m_nodeVariables.end())
+		Ptr<ASTOpParameter>* it = NodeVariables.Find(node.m_pNode);
+		if (!it)
 		{
 			FParameterDesc param;
 			param.m_name = node.m_name;
@@ -110,56 +100,49 @@ namespace mu
 			for (int32 a = 0; a < node.m_ranges.Num(); ++a)
 			{
 				FRangeGenerationResult rangeResult;
-				GenerateRange(rangeResult, node.m_ranges[a]);
+				GenerateRange(rangeResult, Options, node.m_ranges[a]);
 				op->ranges.Emplace(op.get(), rangeResult.sizeOp, rangeResult.rangeName, rangeResult.rangeUID);
 			}
 
-			m_nodeVariables[node.m_pNode] = op;
+			NodeVariables.Add(node.m_pNode,op);
 		}
 		else
 		{
-			op = it->second;
+			op = *it;
 		}
 
 		result.op = op;
 	}
 
 
-	//-------------------------------------------------------------------------------------------------
-	void CodeGenerator::GenerateBool_IsNull(FBoolGenerationResult& result, const Ptr<const NodeBoolIsNull>& Typed)
-	{
-		const NodeBoolIsNull::Private& node = *Typed->GetPrivate();
-
-		Ptr<ASTOpConstantBool> op = new ASTOpConstantBool();
-		Ptr<ASTOp> source = Generate(node.m_pSource);
-		op->value = !source;
-
-		result.op = op;
-	}
-
-
-	//-------------------------------------------------------------------------------------------------
-	void CodeGenerator::GenerateBool_Not(FBoolGenerationResult& result, const Ptr<const NodeBoolNot>& Typed)
+	void CodeGenerator::GenerateBool_Not(FBoolGenerationResult& result, const FGenericGenerationOptions& Options, const Ptr<const NodeBoolNot>& Typed)
 	{
 		const NodeBoolNot::Private& node = *Typed->GetPrivate();
 
 		Ptr<ASTOpFixed> op = new ASTOpFixed();
 		op->op.type = OP_TYPE::BO_NOT;
-		op->SetChild(op->op.args.BoolNot.source, Generate(node.m_pSource));
+
+		FBoolGenerationResult ChildResult;
+		GenerateBool(ChildResult, Options, node.m_pSource);
+		op->SetChild(op->op.args.BoolNot.source, ChildResult.op);
 
 		result.op = op;
 	}
 
 
-	//-------------------------------------------------------------------------------------------------
-	void CodeGenerator::GenerateBool_And(FBoolGenerationResult& result, const Ptr<const NodeBoolAnd>& Typed)
+	void CodeGenerator::GenerateBool_And(FBoolGenerationResult& result, const FGenericGenerationOptions& Options, const Ptr<const NodeBoolAnd>& Typed)
 	{
 		const NodeBoolAnd::Private& node = *Typed->GetPrivate();
 
 		Ptr<ASTOpFixed> op = new ASTOpFixed();
 		op->op.type = OP_TYPE::BO_AND;
-		op->SetChild(op->op.args.BoolBinary.a, Generate(node.m_pA));
-		op->SetChild(op->op.args.BoolBinary.b, Generate(node.m_pB));
+
+		FBoolGenerationResult ChildResult;
+		GenerateBool(ChildResult, Options, node.m_pA);
+		op->SetChild(op->op.args.BoolBinary.a, ChildResult.op);
+
+		GenerateBool(ChildResult, Options, node.m_pB);
+		op->SetChild(op->op.args.BoolBinary.b, ChildResult.op);
 
 		result.op = op;
 	}
diff --git a/Engine/Plugins/Experimental/Mutable/Source/MutableTools/Private/MuT/CodeGenerator_Color.cpp b/Engine/Plugins/Experimental/Mutable/Source/MutableTools/Private/MuT/CodeGenerator_Color.cpp
--- a/Engine/Plugins/Experimental/Mutable/Source/MutableTools/Private/MuT/CodeGenerator_Color.cpp
+++ b/Engine/Plugins/Experimental/Mutable/Source/MutableTools/Private/MuT/CodeGenerator_Color.cpp
@@ -57,55 +57,57 @@ namespace mu
 
 
 	//-------------------------------------------------------------------------------------------------
-	void CodeGenerator::GenerateColor(FColorGenerationResult& result, const NodeColourPtrConst& untyped)
+	void CodeGenerator::GenerateColor(FColorGenerationResult& Result, const FGenericGenerationOptions& Options, const Ptr<const NodeColour>& Untyped)
 	{
-		if (!untyped)
+		if (!Untyped)
 		{
-			result = FColorGenerationResult();
+			Result = FColorGenerationResult();
 			return;
 		}
 
 		// See if it was already generated
-		FVisitedKeyMap key = GetCurrentCacheKey(untyped);
-		GeneratedColorsMap::ValueType* it = m_generatedColors.Find(key);
+		FGeneratedCacheKey Key;
+		Key.Node = Untyped;
+		Key.Options = Options;
+		FGeneratedColorsMap::ValueType* it = GeneratedColors.Find(Key);
 		if (it)
 		{
-			result = *it;
+			Result = *it;
 			return;
 		}
 
 		// Generate for each different type of node
-		if (auto Constant = dynamic_cast<const NodeColourConstant*>(untyped.get()))
+		if (auto Constant = dynamic_cast<const NodeColourConstant*>(Untyped.get()))
 		{
-			GenerateColor_Constant(result, Constant);
+			GenerateColor_Constant(Result, Options, Constant);
 		}
-		else if (auto Param = dynamic_cast<const NodeColourParameter*>(untyped.get()))
+		else if (auto Param = dynamic_cast<const NodeColourParameter*>(Untyped.get()))
 		{
-			GenerateColor_Parameter(result, Param);
+			GenerateColor_Parameter(Result, Options, Param);
 		}
-		else if (auto Switch = dynamic_cast<const NodeColourSwitch*>(untyped.get()))
+		else if (auto Switch = dynamic_cast<const NodeColourSwitch*>(Untyped.get()))
 		{
-			GenerateColor_Switch(result, Switch);
+			GenerateColor_Switch(Result, Options, Switch);
 		}
-		else if (auto Sample = dynamic_cast<const NodeColourSampleImage*>(untyped.get()))
+		else if (auto Sample = dynamic_cast<const NodeColourSampleImage*>(Untyped.get()))
 		{
-			GenerateColor_SampleImage(result, Sample);
+			GenerateColor_SampleImage(Result, Options, Sample);
 		}
-		else if (auto From = dynamic_cast<const NodeColourFromScalars*>(untyped.get()))
+		else if (auto From = dynamic_cast<const NodeColourFromScalars*>(Untyped.get()))
 		{
-			GenerateColor_FromScalars(result, From);
+			GenerateColor_FromScalars(Result, Options, From);
 		}
-		else if (auto Arithmetic = dynamic_cast<const NodeColourArithmeticOperation*>(untyped.get()))
+		else if (auto Arithmetic = dynamic_cast<const NodeColourArithmeticOperation*>(Untyped.get()))
 		{
-			GenerateColor_Arithmetic(result, Arithmetic);
+			GenerateColor_Arithmetic(Result, Options, Arithmetic);
 		}
-		else if (auto Variation = dynamic_cast<const NodeColourVariation*>(untyped.get()))
+		else if (auto Variation = dynamic_cast<const NodeColourVariation*>(Untyped.get()))
 		{
-			GenerateColor_Variation(result, Variation);
+			GenerateColor_Variation(Result, Options, Variation);
 		}
-		else if (auto Table = dynamic_cast<const NodeColourTable*>(untyped.get()))
+		else if (auto Table = dynamic_cast<const NodeColourTable*>(Untyped.get()))
 		{
-			GenerateColor_Table(result, Table);
+			GenerateColor_Table(Result, Options, Table);
 		}
 		else
 		{
@@ -113,12 +115,12 @@ namespace mu
 		}
 
 		// Cache the result
-		m_generatedColors.Add(key, result);
+		GeneratedColors.Add(Key, Result);
 	}
 
 
 	//-------------------------------------------------------------------------------------------------
-	void CodeGenerator::GenerateColor_Constant(FColorGenerationResult& result, const Ptr<const NodeColourConstant>& Typed)
+	void CodeGenerator::GenerateColor_Constant(FColorGenerationResult& result, const FGenericGenerationOptions& Options, const Ptr<const NodeColourConstant>& Typed)
 	{
 		const NodeColourConstant::Private& node = *Typed->GetPrivate();
 
@@ -134,15 +136,15 @@ namespace mu
 
 
 	//-------------------------------------------------------------------------------------------------
-	void CodeGenerator::GenerateColor_Parameter(FColorGenerationResult& result, const Ptr<const NodeColourParameter>& Typed)
+	void CodeGenerator::GenerateColor_Parameter(FColorGenerationResult& result, const FGenericGenerationOptions& Options, const Ptr<const NodeColourParameter>& Typed)
 	{
 		const NodeColourParameter::Private& node = *Typed->GetPrivate();
 
 		Ptr<ASTOpParameter> op;
 
-		auto it = m_nodeVariables.find(node.m_pNode);
+		Ptr<ASTOpParameter>* it = NodeVariables.Find(node.m_pNode);
 
-		if (it == m_nodeVariables.end())
+		if (!it)
 		{
 			FParameterDesc param;
 			param.m_name = node.m_name;
@@ -164,15 +166,15 @@ namespace mu
 			for (int32 a = 0; a < node.m_ranges.Num(); ++a)
 			{
 				FRangeGenerationResult rangeResult;
-				GenerateRange(rangeResult, node.m_ranges[a]);
+				GenerateRange(rangeResult, Options, node.m_ranges[a]);
 				op->ranges.Emplace(op.get(), rangeResult.sizeOp, rangeResult.rangeName, rangeResult.rangeUID);
 			}
 
-			m_nodeVariables[node.m_pNode] = op;
+			NodeVariables.Add(node.m_pNode, op);
 		}
 		else
 		{
-			op = it->second;
+			op = *it;
 		}
 
 		result.op = op;
@@ -180,7 +182,7 @@ namespace mu
 
 
 	//-------------------------------------------------------------------------------------------------
-	void CodeGenerator::GenerateColor_Switch(FColorGenerationResult& result, const Ptr<const NodeColourSwitch>& Typed)
+	void CodeGenerator::GenerateColor_Switch(FColorGenerationResult& result, const FGenericGenerationOptions& Options, const Ptr<const NodeColourSwitch>& Typed)
 	{
 		const NodeColourSwitch::Private& node = *Typed->GetPrivate();
 
@@ -200,7 +202,7 @@ namespace mu
 		// Variable value
 		if (node.m_pParameter)
 		{
-			op->variable = Generate(node.m_pParameter.get());
+			op->variable = Generate(node.m_pParameter.get(), Options);
 		}
 		else
 		{
@@ -214,7 +216,7 @@ namespace mu
 			Ptr<ASTOp> branch;
 			if (node.m_options[t])
 			{
-				branch = Generate(node.m_options[t].get());
+				branch = Generate(node.m_options[t].get(), Options);
 			}
 			else
 			{
@@ -229,7 +231,7 @@ namespace mu
 
 
 	//-------------------------------------------------------------------------------------------------
-	void CodeGenerator::GenerateColor_Variation(FColorGenerationResult& result, const Ptr<const NodeColourVariation>& Typed)
+	void CodeGenerator::GenerateColor_Variation(FColorGenerationResult& result, const FGenericGenerationOptions& Options, const Ptr<const NodeColourVariation>& Typed)
 	{
 		const NodeColourVariation::Private& node = *Typed->GetPrivate();
 
@@ -239,7 +241,7 @@ namespace mu
 		if (node.m_defaultColour)
 		{
 			FMeshGenerationResult branchResults;
-			currentOp = Generate(node.m_defaultColour);
+			currentOp = Generate(node.m_defaultColour, Options);
 		}
 
 		// Process variations in reverse order, since conditionals are built bottom-up.
@@ -265,7 +267,7 @@ namespace mu
 			Ptr<ASTOp> variationOp;
 			if (node.m_variations[t].m_colour)
 			{
-				variationOp = Generate(node.m_variations[t].m_colour);
+				variationOp = Generate(node.m_variations[t].m_colour, Options);
 			}
 			else
 			{
@@ -288,7 +290,7 @@ namespace mu
 
 
 	//-------------------------------------------------------------------------------------------------
-	void CodeGenerator::GenerateColor_SampleImage(FColorGenerationResult& result, const Ptr<const NodeColourSampleImage>& Typed)
+	void CodeGenerator::GenerateColor_SampleImage(FColorGenerationResult& result, const FGenericGenerationOptions& Options, const Ptr<const NodeColourSampleImage>& Typed)
 	{
 		const NodeColourSampleImage::Private& node = *Typed->GetPrivate();
 
@@ -298,7 +300,7 @@ namespace mu
 
 		// Source image
 		FImageGenerationOptions ImageOptions;
-		ImageOptions.CurrentStateIndex = m_currentStateIndex;
+		ImageOptions.State = Options.State;
 		if (!m_activeTags.IsEmpty())
 		{
 			ImageOptions.ActiveTags = m_activeTags.Last();
@@ -320,32 +322,37 @@ namespace mu
 		base = GenerateImageFormat(base, EImageFormat::IF_RGB_UBYTE);
 		op->SetChild(op->op.args.ColourSampleImage.image, base);
 
+		FScalarGenerationResult ChildResult;
 
 		// X
-		if (Node* pX = node.m_pX.get())
+		if (NodeScalar* pX = node.m_pX.get())
 		{
-			op->SetChild(op->op.args.ColourSampleImage.x, Generate(pX));
+			GenerateScalar(ChildResult, Options, pX);
+			op->SetChild(op->op.args.ColourSampleImage.x, ChildResult.op);
 		}
 		else
 		{
 			// Set a constant 0.5 value
 			NodeScalarConstantPtr pNode = new NodeScalarConstant();
 			pNode->SetValue(0.5f);
-			op->SetChild(op->op.args.ColourSampleImage.x, Generate(pNode));
+			GenerateScalar(ChildResult, Options, pNode);
+			op->SetChild(op->op.args.ColourSampleImage.x, ChildResult.op);
 		}
 
 
 		// Y
-		if (Node* pY = node.m_pY.get())
+		if (NodeScalar* pY = node.m_pY.get())
 		{
-			op->SetChild(op->op.args.ColourSampleImage.y, Generate(pY));
+			GenerateScalar(ChildResult, Options, pY);
+			op->SetChild(op->op.args.ColourSampleImage.y, ChildResult.op);
 		}
 		else
 		{
 			// Set a constant 0.5 value
 			NodeScalarConstantPtr pNode = new NodeScalarConstant();
 			pNode->SetValue(0.5f);
-			op->SetChild(op->op.args.ColourSampleImage.y, Generate(pNode));
+			GenerateScalar(ChildResult, Options, pNode);
+			op->SetChild(op->op.args.ColourSampleImage.y, ChildResult.op);
 		}
 
 		// TODO
@@ -356,59 +363,69 @@ namespace mu
 
 
 	//-------------------------------------------------------------------------------------------------
-	void CodeGenerator::GenerateColor_FromScalars(FColorGenerationResult& result, const Ptr<const NodeColourFromScalars>& Typed)
+	void CodeGenerator::GenerateColor_FromScalars(FColorGenerationResult& result, const FGenericGenerationOptions& Options, const Ptr<const NodeColourFromScalars>& Typed)
 	{
 		const NodeColourFromScalars::Private& node = *Typed->GetPrivate();
 
 		Ptr<ASTOpFixed> op = new ASTOpFixed();
 		op->op.type = OP_TYPE::CO_FROMSCALARS;
 
+		FScalarGenerationResult ChildResult;
+
 		// X
-		if (Node* pX = node.m_pX.get())
+		if (NodeScalar* pX = node.m_pX.get())
 		{
-			op->SetChild(op->op.args.ColourFromScalars.v[0], Generate(pX));
+			GenerateScalar(ChildResult, Options, pX);
+			op->SetChild(op->op.args.ColourFromScalars.v[0], ChildResult.op );
 		}
 		else
 		{
 			NodeScalarConstantPtr pNode = new NodeScalarConstant();
 			pNode->SetValue(1.0f);
-			op->SetChild(op->op.args.ColourFromScalars.v[0], Generate(pNode));
+			GenerateScalar(ChildResult, Options, pNode);
+			op->SetChild(op->op.args.ColourFromScalars.v[0], ChildResult.op);
 		}
 
 		// Y
-		if (Node* pY = node.m_pY.get())
+		if (NodeScalar* pY = node.m_pY.get())
 		{
-			op->SetChild(op->op.args.ColourFromScalars.v[1], Generate(pY));
+			GenerateScalar(ChildResult, Options, pY);
+			op->SetChild(op->op.args.ColourFromScalars.v[1], ChildResult.op);
 		}
 		else
 		{
 			NodeScalarConstantPtr pNode = new NodeScalarConstant();
 			pNode->SetValue(1.0f);
-			op->SetChild(op->op.args.ColourFromScalars.v[1], Generate(pNode));
+			GenerateScalar(ChildResult, Options, pNode);
+			op->SetChild(op->op.args.ColourFromScalars.v[1], ChildResult.op);
 		}
 
 		// Z
-		if (Node* pZ = node.m_pZ.get())
+		if (NodeScalar* pZ = node.m_pZ.get())
 		{
-			op->SetChild(op->op.args.ColourFromScalars.v[2], Generate(pZ));
+			GenerateScalar(ChildResult, Options, pZ);
+			op->SetChild(op->op.args.ColourFromScalars.v[2], ChildResult.op);
 		}
 		else
 		{
 			NodeScalarConstantPtr pNode = new NodeScalarConstant();
 			pNode->SetValue(1.0f);
-			op->SetChild(op->op.args.ColourFromScalars.v[2], Generate(pNode));
+			GenerateScalar(ChildResult, Options, pNode);
+			op->SetChild(op->op.args.ColourFromScalars.v[2], ChildResult.op);
 		}
 
 		// W
-		if (Node* pW = node.m_pW.get())
+		if (NodeScalar* pW = node.m_pW.get())
 		{
-			op->SetChild(op->op.args.ColourFromScalars.v[3], Generate(pW));
+			GenerateScalar(ChildResult, Options, pW);
+			op->SetChild(op->op.args.ColourFromScalars.v[3], ChildResult.op);
 		}
 		else
 		{
 			NodeScalarConstantPtr pNode = new NodeScalarConstant();
 			pNode->SetValue(1.0f);
-			op->SetChild(op->op.args.ColourFromScalars.v[3], Generate(pNode));
+			GenerateScalar(ChildResult, Options, pNode);
+			op->SetChild(op->op.args.ColourFromScalars.v[3], ChildResult.op);
 		}
 
 		result.op = op;
@@ -416,7 +433,7 @@ namespace mu
 
 
 	//-------------------------------------------------------------------------------------------------
-	void CodeGenerator::GenerateColor_Arithmetic(FColorGenerationResult& result, const Ptr<const NodeColourArithmeticOperation>& Typed)
+	void CodeGenerator::GenerateColor_Arithmetic(FColorGenerationResult& result, const FGenericGenerationOptions& Options, const Ptr<const NodeColourArithmeticOperation>& Typed)
 	{
 		const NodeColourArithmeticOperation::Private& node = *Typed->GetPrivate();
 
@@ -435,10 +452,13 @@ namespace mu
 			break;
 		}
 
+		FColorGenerationResult ChildResult;
+
 		// A
-		if (Node* pA = node.m_pA.get())
+		if (NodeColour* pA = node.m_pA.get())
 		{
-			op->SetChild(op->op.args.ColourArithmetic.a, Generate(pA));
+			GenerateColor(ChildResult, Options, pA );
+			op->SetChild(op->op.args.ColourArithmetic.a, ChildResult.op);
 		}
 		else
 		{
@@ -447,9 +467,10 @@ namespace mu
 		}
 
 		// B
-		if (Node* pB = node.m_pB.get())
+		if (NodeColour* pB = node.m_pB.get())
 		{
-			op->SetChild(op->op.args.ColourArithmetic.b, Generate(pB));
+			GenerateColor(ChildResult, Options, pB);
+			op->SetChild(op->op.args.ColourArithmetic.b, ChildResult.op);
 		}
 		else
 		{
@@ -462,17 +483,17 @@ namespace mu
 
 
 	//-------------------------------------------------------------------------------------------------
-	void CodeGenerator::GenerateColor_Table(FColorGenerationResult& result, const Ptr<const NodeColourTable>& Typed)
+	void CodeGenerator::GenerateColor_Table(FColorGenerationResult& result, const FGenericGenerationOptions& Options, const Ptr<const NodeColourTable>& Typed)
 	{
 		const NodeColourTable::Private& node = *Typed->GetPrivate();
 
 		result.op = GenerateTableSwitch<NodeColourTable::Private, ETableColumnType::Color, OP_TYPE::CO_SWITCH>(node,
-			[this](const NodeColourTable::Private& node, int colIndex, int row, ErrorLog* pErrorLog)
+			[this, &Options](const NodeColourTable::Private& node, int colIndex, int row, ErrorLog* pErrorLog)
 			{
 				NodeColourConstantPtr CellData = new NodeColourConstant();
 				FVector4f Colour = node.m_pTable->GetPrivate()->Rows[row].Values[colIndex].Color;
 				CellData->SetValue(Colour);
-				return Generate(CellData);
+				return Generate(CellData, Options);
 			});
 	}
 
@@ -489,7 +510,8 @@ namespace mu
 		pNode->SetValue(FVector4f(1, 1, 0, 1));
 
 		FColorGenerationResult Result;
-		GenerateColor(Result, pNode);
+		FGenericGenerationOptions Options;
+		GenerateColor(Result, Options, pNode);
 
 		return Result.op;
 	}
diff --git a/Engine/Plugins/Experimental/Mutable/Source/MutableTools/Private/MuT/CodeGenerator_ExtensionData.cpp b/Engine/Plugins/Experimental/Mutable/Source/MutableTools/Private/MuT/CodeGenerator_ExtensionData.cpp
--- a/Engine/Plugins/Experimental/Mutable/Source/MutableTools/Private/MuT/CodeGenerator_ExtensionData.cpp
+++ b/Engine/Plugins/Experimental/Mutable/Source/MutableTools/Private/MuT/CodeGenerator_ExtensionData.cpp
@@ -33,7 +33,7 @@ namespace mu
 class Node;
 
 	//---------------------------------------------------------------------------------------------
-	void CodeGenerator::GenerateExtensionData(FExtensionDataGenerationResult& OutResult, const NodeExtensionDataPtrConst& InUntypedNode)
+	void CodeGenerator::GenerateExtensionData(FExtensionDataGenerationResult& OutResult, const FGenericGenerationOptions& Options, const NodeExtensionDataPtrConst& InUntypedNode)
 	{
 		if (!InUntypedNode)
 		{
@@ -41,12 +41,9 @@ class Node;
 			return;
 		}
 
-		// Clear bottom-up state
-		m_currentBottomUpState.m_address = nullptr;
-
 		// See if it was already generated
 		const FGeneratedExtensionDataCacheKey Key = InUntypedNode.get();
-		GeneratedExtensionDataMap::ValueType* CachedResult = m_generatedExtensionData.Find(Key);
+		FGeneratedExtensionDataMap::ValueType* CachedResult = GeneratedExtensionData.Find(Key);
 		if (CachedResult)
 		{
 			OutResult = *CachedResult;
@@ -58,18 +55,18 @@ class Node;
 		// Generate for each different type of node
 		switch (Node->GetExtensionDataNodeType())
 		{
-			case NodeExtensionData::EType::Constant:  GenerateExtensionData_Constant(OutResult, static_cast<const NodeExtensionDataConstant*>(Node)); break;
-			case NodeExtensionData::EType::Switch:    GenerateExtensionData_Switch(OutResult, static_cast<const NodeExtensionDataSwitch*>(Node)); break;
-			case NodeExtensionData::EType::Variation: GenerateExtensionData_Variation(OutResult, static_cast<const NodeExtensionDataVariation*>(Node)); break;
+			case NodeExtensionData::EType::Constant:  GenerateExtensionData_Constant(OutResult, Options, static_cast<const NodeExtensionDataConstant*>(Node)); break;
+			case NodeExtensionData::EType::Switch:    GenerateExtensionData_Switch(OutResult, Options, static_cast<const NodeExtensionDataSwitch*>(Node)); break;
+			case NodeExtensionData::EType::Variation: GenerateExtensionData_Variation(OutResult, Options, static_cast<const NodeExtensionDataVariation*>(Node)); break;
 			case NodeExtensionData::EType::None: check(false);
 		}
 
 		// Cache the result
-		m_generatedExtensionData.Add(Key, OutResult);
+		GeneratedExtensionData.Add(Key, OutResult);
 	}
 
 	//---------------------------------------------------------------------------------------------
-	void CodeGenerator::GenerateExtensionData_Constant(FExtensionDataGenerationResult& OutResult, const NodeExtensionDataConstant* Constant)
+	void CodeGenerator::GenerateExtensionData_Constant(FExtensionDataGenerationResult& OutResult, const FGenericGenerationOptions& Options, const NodeExtensionDataConstant* Constant)
 	{
 		NodeExtensionDataConstant::Private& Node = *Constant->GetPrivate();
 
@@ -90,7 +87,7 @@ class Node;
 	}
 
 	//---------------------------------------------------------------------------------------------
-	void CodeGenerator::GenerateExtensionData_Switch(FExtensionDataGenerationResult& OutResult, const class NodeExtensionDataSwitch* Switch)
+	void CodeGenerator::GenerateExtensionData_Switch(FExtensionDataGenerationResult& OutResult, const FGenericGenerationOptions& Options, const class NodeExtensionDataSwitch* Switch)
 	{
 		NodeExtensionDataSwitch::Private& Node = *Switch->GetPrivate();
 
@@ -109,7 +106,7 @@ class Node;
 		// Variable name
 		if (Node.Parameter)
 		{
-			Op->variable = Generate(Node.Parameter.get());
+			Op->variable = Generate(Node.Parameter.get(), Options);
 		}
 		else
 		{
@@ -123,10 +120,10 @@ class Node;
 			Ptr<ASTOp> Branch;
 			if (Node.Options[OptionIndex])
 			{
-				NodeExtensionDataPtr Option = Node.Options[OptionIndex];
+				NodeExtensionDataPtr SwitchOption = Node.Options[OptionIndex];
 
 				FExtensionDataGenerationResult OptionResult;
-				GenerateExtensionData(OptionResult, Option);
+				GenerateExtensionData(OptionResult, Options, SwitchOption);
 
 				Branch = OptionResult.Op;
 			}
@@ -142,7 +139,7 @@ class Node;
 	}
 
 	//---------------------------------------------------------------------------------------------
-	void CodeGenerator::GenerateExtensionData_Variation(FExtensionDataGenerationResult& OutResult, const class NodeExtensionDataVariation* Variation)
+	void CodeGenerator::GenerateExtensionData_Variation(FExtensionDataGenerationResult& OutResult, const FGenericGenerationOptions& Options, const class NodeExtensionDataVariation* Variation)
 	{
 		const NodeExtensionDataVariation::Private& Node = *Variation->GetPrivate();
 
@@ -152,7 +149,7 @@ class Node;
 		if (Node.DefaultValue)
 		{
 			FExtensionDataGenerationResult DefaultResult;
-			GenerateExtensionData(DefaultResult, Node.DefaultValue);
+			GenerateExtensionData(DefaultResult, Options, Node.DefaultValue);
 
 			CurrentOp = DefaultResult.Op;
 		}
@@ -177,7 +174,7 @@ class Node;
 			if (NodeExtensionDataPtr VariationValue = Node.Variations[VariationIndex].Value)
 			{
 				FExtensionDataGenerationResult VariationResult;
-				GenerateExtensionData(VariationResult, VariationValue);
+				GenerateExtensionData(VariationResult, Options, VariationValue);
 
 				VariationOp = VariationResult.Op;
 			}
@@ -209,7 +206,8 @@ class Node;
 		NodeExtensionDataConstantPtrConst Node = new NodeExtensionDataConstant;
 
 		FExtensionDataGenerationResult Result;
-		GenerateExtensionData_Constant(Result, Node.get());
+		FGenericGenerationOptions Options;
+		GenerateExtensionData_Constant(Result, Options, Node.get());
 
 		return Result.Op;
 	}
diff --git a/Engine/Plugins/Experimental/Mutable/Source/MutableTools/Private/MuT/CodeGenerator_FirstPass.cpp b/Engine/Plugins/Experimental/Mutable/Source/MutableTools/Private/MuT/CodeGenerator_FirstPass.cpp
--- a/Engine/Plugins/Experimental/Mutable/Source/MutableTools/Private/MuT/CodeGenerator_FirstPass.cpp
+++ b/Engine/Plugins/Experimental/Mutable/Source/MutableTools/Private/MuT/CodeGenerator_FirstPass.cpp
@@ -49,38 +49,22 @@ namespace mu
 
 
 	//---------------------------------------------------------------------------------------------
-    void FirstPassGenerator::Generate( ErrorLogPtr pErrorLog,
-                                       const Node::Private* root,
-                                       bool ignoreStates,
+    void FirstPassGenerator::Generate( Ptr<ErrorLog> InErrorLog,
+                                       const Node* Root,
+                                       bool bIgnoreStates,
 									   CodeGenerator* InGenerator )
 	{
 		MUTABLE_CPUPROFILER_SCOPE(FirstPassGenerate);
 
 		Generator = InGenerator;
-		m_pErrorLog = pErrorLog;
-        m_ignoreStates = ignoreStates;
+		m_pErrorLog = InErrorLog;
+        m_ignoreStates = bIgnoreStates;
 
 		// Step 1: collect all objects, surfaces and object conditions
-        if ( root )
+        if (Root)
 		{
-            bool isFirstPassNode =
-                    dynamic_cast<const NodeSurface*>(root->m_pNode)
-                    ||
-                    dynamic_cast<const NodeComponent*>(root->m_pNode)
-                    ||
-                    dynamic_cast<const NodeObject*>(root->m_pNode)
-                    ||
-                    dynamic_cast<const NodeLOD*>(root->m_pNode)
-                    ||
-                    dynamic_cast<const NodePatchMesh*>(root->m_pNode)
-                    ||
-                    dynamic_cast<const NodeModifier*>(root->m_pNode);
-
-            if (isFirstPassNode)
-            {
-                root->Accept(*this);
-            }
-		}
+ 			Generate_Generic(Root);
+ 		}
 
 		// Step 2: Collect all tags and a list of the surfaces that activate them
 		for (int32 s=0; s<surfaces.Num(); ++s)
@@ -148,9 +132,8 @@ namespace mu
 
 		}
 
-
         // Step 3: Create default state if necessary
-        if ( ignoreStates )
+        if ( bIgnoreStates )
         {
             m_states.Empty();
         }
@@ -159,109 +142,105 @@ namespace mu
         {
             FObjectState data;
             data.m_name = "Default";
-            m_states.Emplace( data, root );
+            m_states.Emplace( data, Root->GetBasePrivate() );
         }
 	}
 
 
-	//---------------------------------------------------------------------------------------------
-    Ptr<ASTOp> FirstPassGenerator::Visit(const NodeModifierMeshClipMorphPlane::Private& node)
+	void FirstPassGenerator::Generate_Generic(const Node* Root)
 	{
-        // Add the data about this modifier
-        FModifier thisData;
-        thisData.node = &node;
-        thisData.objectCondition = m_currentCondition.Last().objectCondition;
-        thisData.stateCondition = m_currentStateCondition.Last();
-        thisData.lod = m_currentLOD;
-        thisData.positiveTags = m_currentPositiveTags;
-        thisData.negativeTags = m_currentNegativeTags;
-        modifiers.Add(thisData);
-
-        return nullptr;
-    }
-
+		if (Root->GetType()==NodeSurfaceNew::GetStaticType())
+		{
+			Generate_SurfaceNew(static_cast<const NodeSurfaceNew*>(Root));
+		}
+		else if (Root->GetType() == NodeSurfaceEdit::GetStaticType())
+		{
+			Generate_SurfaceEdit(static_cast<const NodeSurfaceEdit*>(Root));
+		}
+		else if (Root->GetType() == NodeSurfaceVariation::GetStaticType())
+		{
+			Generate_SurfaceVariation(static_cast<const NodeSurfaceVariation*>(Root));
+		}
+		else if (Root->GetType() == NodeSurfaceSwitch::GetStaticType())
+		{
+			Generate_SurfaceSwitch(static_cast<const NodeSurfaceSwitch*>(Root));
+		}
+		else if (Root->GetType() == NodeComponentNew::GetStaticType())
+		{
+			Generate_ComponentNew(static_cast<const NodeComponentNew*>(Root));
+		}
+		else if (Root->GetType() == NodeComponentEdit::GetStaticType())
+		{
+			Generate_ComponentEdit(static_cast<const NodeComponentEdit*>(Root));
+		}
+		else if (Root->GetType() == NodeObjectNew::GetStaticType())
+		{
+			Generate_ObjectNew(static_cast<const NodeObjectNew*>(Root));
+		}
+		else if (Root->GetType() == NodeObjectGroup::GetStaticType())
+		{
+			Generate_ObjectGroup(static_cast<const NodeObjectGroup*>(Root));
+		}
+		else if (Root->GetType() == NodeLOD::GetStaticType())
+		{
+			Generate_LOD(static_cast<const NodeLOD*>(Root));
+		}
+		else if (Root->GetType() == NodeModifier::GetStaticType())
+		{
+			Generate_Modifier(static_cast<const NodeModifier*>(Root));
+		}
+		else
+		{
+			check(false);
+		}
+	}
 
 	//---------------------------------------------------------------------------------------------
-	Ptr<ASTOp> FirstPassGenerator::Visit(const NodeModifierMeshClipWithMesh::Private& node)
+	void FirstPassGenerator::Generate_Modifier(const NodeModifier* InNode)
 	{
 		// Add the data about this modifier
 		FModifier thisData;
-		thisData.node = &node;
+		thisData.node = static_cast<const NodeModifier::Private*>(InNode->GetBasePrivate());
 		thisData.objectCondition = m_currentCondition.Last().objectCondition;
 		thisData.stateCondition = m_currentStateCondition.Last();
 		thisData.lod = m_currentLOD;
 		thisData.positiveTags = m_currentPositiveTags;
 		thisData.negativeTags = m_currentNegativeTags;
 		modifiers.Add(thisData);
-
-		return nullptr;
 	}
 
 
 	//---------------------------------------------------------------------------------------------
-	Ptr<ASTOp> FirstPassGenerator::Visit(const NodeModifierMeshClipWithUVMask::Private& node)
+	void FirstPassGenerator::Generate_SurfaceNew(const NodeSurfaceNew* InNode)
 	{
-		// Add the data about this modifier
-		FModifier thisData;
-		thisData.node = &node;
+		// Add the data about this surface
+		FSurface thisData;
+		thisData.node = InNode;
+		thisData.component = m_currentComponent;
 		thisData.objectCondition = m_currentCondition.Last().objectCondition;
 		thisData.stateCondition = m_currentStateCondition.Last();
-		thisData.lod = m_currentLOD;
 		thisData.positiveTags = m_currentPositiveTags;
 		thisData.negativeTags = m_currentNegativeTags;
-		modifiers.Add(thisData);
-
-		return nullptr;
+		surfaces.Add(thisData);
 	}
 
-	
-	//---------------------------------------------------------------------------------------------
-    Ptr<ASTOp> FirstPassGenerator::Visit(const NodeModifierMeshClipDeform::Private& node)
-	{
-        // Add the data about this modifier
-		FModifier thisData;
-		thisData.node = &node;
-        thisData.objectCondition = m_currentCondition.Last().objectCondition;
-        thisData.stateCondition = m_currentStateCondition.Last();
-        thisData.lod = m_currentLOD;
-        thisData.positiveTags = m_currentPositiveTags;
-        thisData.negativeTags = m_currentNegativeTags;
-        modifiers.Add(thisData);
-
-        return nullptr;
-	}
 
 	//---------------------------------------------------------------------------------------------
-    Ptr<ASTOp> FirstPassGenerator::Visit(const NodeSurfaceNew::Private& node)
+    void FirstPassGenerator::Generate_SurfaceEdit(const NodeSurfaceEdit* InNode)
 	{
-		// Add the data about this surface
-		FSurface thisData;
-		thisData.node = dynamic_cast<const NodeSurfaceNew*>(node.m_pNode);
-		thisData.component = m_currentComponent;
-        thisData.objectCondition = m_currentCondition.Last().objectCondition;
-        thisData.stateCondition = m_currentStateCondition.Last();
-        thisData.positiveTags = m_currentPositiveTags;
-        thisData.negativeTags = m_currentNegativeTags;
-        surfaces.Add(thisData);
-
-        return nullptr;
-	}
-
+		const NodeSurfaceEdit::Private* Private = InNode->GetPrivate();
 
-	//---------------------------------------------------------------------------------------------
-    Ptr<ASTOp> FirstPassGenerator::Visit(const NodeSurfaceEdit::Private& node)
-	{
 		// Store a reference to this node in the surface data for the surface that this node is
 		// editing.
-		auto its = surfaces.FindByPredicate([&node](const FirstPassGenerator::FSurface& s)
+		FSurface* Surface = surfaces.FindByPredicate([&Private](const FSurface& s)
         {
             // Are we editing the main surface node of this surface?
-            if (s.node.get() == node.m_pParent.get()) return true;
+            if (s.node.get() == Private->m_pParent.get()) return true;
 
             // Are we editing an edit node modifying this surface?
             for (const auto& e: s.edits)
             {
-                if (node.m_pParent && e.node==node.m_pParent->GetBasePrivate())
+                if (Private->m_pParent && e.node==Private->m_pParent->GetBasePrivate())
                 {
                     return true;
                 }
@@ -272,69 +251,66 @@ namespace mu
 		
 		// The surface could be missing if the parent is not in the hierarchy. This could happen
 		// with wrong input or in case of partial models for preview.
-		if (its)
+		if (Surface)
 		{
-			FSurface& surface = *its;
-
 			FSurface::FEdit edit;
-            edit.node = &node;
+            edit.node = Private;
             edit.condition = m_currentCondition.Last().objectCondition;
-            surface.edits.Add(edit);
+			Surface->edits.Add(edit);
 		}
 		else
 		{
-			m_pErrorLog->GetPrivate()->Add("Missing parent object for edit node.",
-				ELMT_WARNING, node.m_errorContext);
+			m_pErrorLog->GetPrivate()->Add("Missing parent object for edit node.", ELMT_WARNING, Private->m_errorContext);
 		}
-
-        return nullptr;
 	}
 
 
 	//---------------------------------------------------------------------------------------------
-    Ptr<ASTOp> FirstPassGenerator::Visit(const NodeSurfaceVariation::Private& node)
+	void FirstPassGenerator::Generate_SurfaceVariation(const NodeSurfaceVariation* InNode)
 	{
-        switch(node.m_type)
+		const NodeSurfaceVariation::Private* Private = InNode->GetPrivate();
+
+        switch(Private->m_type)
         {
 
         case NodeSurfaceVariation::VariationType::Tag:
         {
             // Any of the tags in the variations would prevent the default surface
             auto oldNegativeTags = m_currentNegativeTags;
-            for (int32 v=0; v<node.m_variations.Num(); ++v)
+            for (int32 v=0; v< Private->m_variations.Num(); ++v)
             {
-                m_currentNegativeTags.Add(node.m_variations[v].m_tag);
+                m_currentNegativeTags.Add(Private->m_variations[v].m_tag);
             }
 
-            for(const auto& n:node.m_defaultSurfaces)
+            for(const auto& n: Private->m_defaultSurfaces)
             {
-                n->GetBasePrivate()->Accept(*this);
-            }
-            for(const auto& n:node.m_defaultModifiers)
+				Generate_Generic(n.get());
+			}
+            for(const auto& n: Private->m_defaultModifiers)
             {
-                n->GetBasePrivate()->Accept(*this);
+				Generate_Modifier(n.get());
             }
 
             m_currentNegativeTags = oldNegativeTags;
 
-            for (int32 v=0; v<node.m_variations.Num(); ++v)
+            for (int32 v=0; v< Private->m_variations.Num(); ++v)
             {
-                m_currentPositiveTags.Add(node.m_variations[v].m_tag);
-                for (const auto& s : node.m_variations[v].m_surfaces)
+                m_currentPositiveTags.Add(Private->m_variations[v].m_tag);
+                for (const auto& s : Private->m_variations[v].m_surfaces)
                 {
-                    s->GetBasePrivate()->Accept(*this);
+					Generate_Generic(s.get());
                 }
 
-                for (const auto& s : node.m_variations[v].m_modifiers)
+                for (const auto& s : Private->m_variations[v].m_modifiers)
                 {
-                    s->GetBasePrivate()->Accept(*this);
-                }
+					Generate_Modifier(s.get());
+				}
 
                 m_currentPositiveTags.Pop();
 
                 // Tags have an order in a variation node: the current tag should prevent any following
                 // variation surface
-                m_currentNegativeTags.Add(node.m_variations[v].m_tag);
+                m_currentNegativeTags.Add(Private->m_variations[v].m_tag);
             }
 
             m_currentNegativeTags = oldNegativeTags;
@@ -358,7 +334,7 @@ namespace mu
                             ? AllTrue
                             : m_currentStateCondition.Last();
 
-                    for (const auto& v:node.m_variations)
+                    for (const auto& v: Private->m_variations)
                     {
                         for( size_t s=0; s<stateCount; ++s )
                         {
@@ -373,21 +349,20 @@ namespace mu
 
                 m_currentStateCondition.Add(defaultStates);
 
-                for (const auto& n : node.m_defaultSurfaces)
+                for (const auto& n : Private->m_defaultSurfaces)
                 {
-                    n->GetBasePrivate()->Accept(*this);
-                }
-                for (const auto& n : node.m_defaultModifiers)
+					Generate_Generic(n.get());
+				}
+                for (const auto& n : Private->m_defaultModifiers)
                 {
-                    n->GetBasePrivate()->Accept(*this);
-                }
+					Generate_Modifier(n.get());
+				}
 
                 m_currentStateCondition.Pop();
             }
 
-
             // Variation branches
-            for (const auto& v:node.m_variations)
+            for (const auto& v: Private->m_variations)
             {
                 // Store the states for this variation here
 				StateCondition variationStates;
@@ -405,12 +380,12 @@ namespace mu
 
                 for (const auto& n : v.m_surfaces)
                 {
-                    n->GetBasePrivate()->Accept(*this);
-                }
+					Generate_Generic(n.get());
+				}
                 for (const auto& n : v.m_modifiers)
                 {
-                    n->GetBasePrivate()->Accept(*this);
-                }
+					Generate_Modifier(n.get());
+				}
 
                 m_currentStateCondition.Pop();
             }
@@ -418,44 +393,40 @@ namespace mu
             break;
         }
 
-
         default:
             // Case not implemented.
             check(false);
             break;
         }
-
-        return nullptr;
 	}
 
 
 	//---------------------------------------------------------------------------------------------
-	Ptr<ASTOp> FirstPassGenerator::Visit(const NodeSurfaceSwitch::Private& node)
+	void FirstPassGenerator::Generate_SurfaceSwitch(const NodeSurfaceSwitch* InNode)
 	{
-		if (node.Options.Num() == 0)
+		const NodeSurfaceSwitch::Private* Private = InNode->GetPrivate();
+
+		if (Private->Options.Num() == 0)
 		{
 			// No options in the switch!
-			return nullptr;
+			return;
 		}
 
 		// Prepare the enumeration parameter
+		CodeGenerator::FGenericGenerationOptions Options;
 		CodeGenerator::FScalarGenerationResult ScalarResult;
-		if (node.Parameter)
+		if (Private->Parameter)
 		{
-			Generator->GenerateScalar( ScalarResult, node.Parameter );
-
-			// Not really necessary but it currently should always be this type.
-			Ptr<ASTOpParameter> EnumOp = dynamic_cast<ASTOpParameter*>(ScalarResult.op.get());
-			check(EnumOp);
+			Generator->GenerateScalar( ScalarResult, Options, Private->Parameter );
 		}
 		else
 		{
 			// This argument is required
-			ScalarResult.op = Generator->GenerateMissingScalarCode(TEXT("Switch variable"), 0.0f, node.m_errorContext);
+			ScalarResult.op = Generator->GenerateMissingScalarCode(TEXT("Switch variable"), 0.0f, Private->m_errorContext);
 		}
 
 		// Parse the options
-		for (int32 t = 0; t < node.Options.Num(); ++t)
+		for (int32 t = 0; t < Private->Options.Num(); ++t)
 		{
 			// Create a comparison operation as the boolean parameter for the child
 			Ptr<ASTOpFixed> ParamOp = new ASTOpFixed();
@@ -477,144 +448,144 @@ namespace mu
 			data.objectCondition = ParamOp;
 			m_currentCondition.Push(data);
 
-			if (node.Options[t])
+			if (Private->Options[t])
 			{
-				node.Options[t]->GetBasePrivate()->Accept(*this);
+				Generate_Generic(Private->Options[t].get());
 			}
 
 			m_currentCondition.Pop();
 		}
-
-		return nullptr;
 	}
 
 
 	//---------------------------------------------------------------------------------------------
-    Ptr<ASTOp> FirstPassGenerator::Visit(const NodeComponentNew::Private& node)
+	void FirstPassGenerator::Generate_ComponentNew(const NodeComponentNew* InNode)
 	{
-        m_currentComponent = node.GetParentComponentNew();
+		const NodeComponentNew::Private* Private = InNode->GetPrivate();
+
+        m_currentComponent = Private->GetParentComponentNew();
 
-		for (const auto& c : node.m_surfaces)
+		for (const Ptr<NodeSurface>& c : Private->m_surfaces)
 		{
 			if (c)
 			{
-				c->GetBasePrivate()->Accept(*this);
+				Generate_Generic(c.get());
 			}
 		}
 
 		m_currentComponent = nullptr;
-
-        return nullptr;
 	}
 
 
 	//---------------------------------------------------------------------------------------------
-    Ptr<ASTOp> FirstPassGenerator::Visit(const NodeComponentEdit::Private& node)
+	void FirstPassGenerator::Generate_ComponentEdit(const NodeComponentEdit* InNode)
 	{
-		m_currentComponent = node.GetParentComponentNew();
+		const NodeComponentEdit::Private* Private = InNode->GetPrivate();
 
-		for (const auto& c : node.m_surfaces)
+		m_currentComponent = Private->GetParentComponentNew();
+
+		for (const auto& c : Private->m_surfaces)
 		{
 			if (c)
 			{
-                c->GetBasePrivate()->Accept(*this);
+				Generate_Generic(c.get());
 			}
 		}
 
 		m_currentComponent = nullptr;
-
-        return nullptr;
 	}
 
 
 	//---------------------------------------------------------------------------------------------
-    Ptr<ASTOp> FirstPassGenerator::Visit(const NodeLOD::Private& node)
+	void FirstPassGenerator::Generate_LOD(const NodeLOD* InNode)
 	{
-		for (const auto& c : node.m_components)
+		const NodeLOD::Private* Private = InNode->GetPrivate();
+
+		for (const auto& c : Private->m_components)
 		{
 			if (c)
 			{
-				c->GetBasePrivate()->Accept(*this);
+				Generate_Generic(c.get());
 			}
 		}
-		for (const auto& c : node.m_modifiers)
+		for (const auto& c : Private->m_modifiers)
 		{
 			if (c)
 			{
-				c->GetBasePrivate()->Accept(*this);
+				Generate_Modifier(c.get());
 			}
 		}
-
-        return nullptr;
 	}
 
 
 	//---------------------------------------------------------------------------------------------
-    Ptr<ASTOp> FirstPassGenerator::Visit(const NodeObjectNew::Private& node)
+	void FirstPassGenerator::Generate_ObjectNew(const NodeObjectNew* InNode)
 	{
+		const NodeObjectNew::Private* Private = InNode->GetPrivate();
+
 		// Add the data about this object
 		FObject thisData;
-		thisData.node = &node;
+		thisData.node = Private;
         thisData.condition = m_currentCondition.Last().objectCondition;
 		objects.Add(thisData);
 
         // Accumulate the model states
-        for ( const auto& s: node.m_states )
+        for ( const auto& s: Private->m_states )
         {
-            m_states.Emplace( s, &node );
+            m_states.Emplace( s, Private );
 
             if ( s.m_runtimeParams.Num() > MUTABLE_MAX_RUNTIME_PARAMETERS_PER_STATE )
             {
                 FString Msg = FString::Printf( TEXT("State [%s] has more than %d runtime parameters. Their update may fail."), 
 					*s.m_name,
                     MUTABLE_MAX_RUNTIME_PARAMETERS_PER_STATE);
-                m_pErrorLog->GetPrivate()->Add(Msg, ELMT_ERROR, node.m_errorContext );
+                m_pErrorLog->GetPrivate()->Add(Msg, ELMT_ERROR, Private->m_errorContext );
             }
         }
 
 		// Process the lods
 		int i = 0;
-		for (const auto& l : node.m_lods)
+		for (const auto& l : Private->m_lods)
 		{
 			if (l)
 			{
                 m_currentLOD = i++;
-				l->GetBasePrivate()->Accept(*this);
+				Generate_Generic(l.get());
 			}
 		}
 
 		m_currentLOD = -1;
 
 		// Process the children
-		for (const auto& c : node.m_children)
+		for (const auto& c : Private->m_children)
 		{
 			if (c)
 			{
-				c->GetBasePrivate()->Accept(*this);
+				Generate_Generic(c.get());
 			}
 		}
-
-        return nullptr;
 	}
 
 
 	//---------------------------------------------------------------------------------------------
-    Ptr<ASTOp> FirstPassGenerator::Visit(const NodeObjectGroup::Private& node)
+	void FirstPassGenerator::Generate_ObjectGroup(const NodeObjectGroup* InNode)
 	{
-       // Prepare the enumeration parameter if necessary
+		const NodeObjectGroup::Private* Private = InNode->GetPrivate();
+
+		// Prepare the enumeration parameter if necessary
         Ptr<ASTOpParameter> enumOp;
-        if ( node.m_type==NodeObjectGroup::CS_ALWAYS_ONE ||
-             node.m_type==NodeObjectGroup::CS_ONE_OR_NONE )
+        if (Private->m_type==NodeObjectGroup::CS_ALWAYS_ONE ||
+			Private->m_type==NodeObjectGroup::CS_ONE_OR_NONE )
         {
             Ptr<ASTOpParameter> op = new ASTOpParameter();
             op->type = OP_TYPE::NU_PARAMETER;
 
-            op->parameter.m_name = node.Name;
-            op->parameter.m_uid = node.Uid;
+            op->parameter.m_name = Private->Name;
+            op->parameter.m_uid = Private->Uid;
             op->parameter.m_type = PARAMETER_TYPE::T_INT;
             op->parameter.m_defaultValue.Set<ParamIntType>(-1);
 
-            if ( node.m_type==NodeObjectGroup::CS_ONE_OR_NONE )
+            if (Private->m_type==NodeObjectGroup::CS_ONE_OR_NONE )
             {
                 FParameterDesc::FIntValueDesc nullValue;
                 nullValue.m_value = -1;
@@ -628,13 +599,13 @@ namespace mu
 
 
         // Parse the child objects
-		for ( int32 t=0; t<node.m_children.Num(); ++t )
+		for ( int32 t=0; t< Private->m_children.Num(); ++t )
         {
-            if ( const NodeObject* pChildNode = node.m_children[t].get() )
+            if ( const NodeObject* pChildNode = Private->m_children[t].get() )
             {
                 // Overwrite the implicit condition
                 Ptr<ASTOp> paramOp = 0;
-                switch ( node.m_type )
+                switch (Private->m_type )
                 {
                     case NodeObjectGroup::CS_TOGGLE_EACH:
                     {
@@ -706,21 +677,12 @@ namespace mu
                 data.objectCondition = paramOp;
                 m_currentCondition.Add( data );
 
-				pChildNode->GetBasePrivate()->Accept(*this);
+				Generate_Generic(pChildNode);
 
                 m_currentCondition.Pop();
             }
         }
-
-        return nullptr;
  	}
 
-
-	//---------------------------------------------------------------------------------------------
-    Ptr<ASTOp> FirstPassGenerator::Visit(const NodePatchMesh::Private&)
-	{
-        return nullptr;
-	}
-
 }
 
diff --git a/Engine/Plugins/Experimental/Mutable/Source/MutableTools/Private/MuT/CodeGenerator_FirstPass.h b/Engine/Plugins/Experimental/Mutable/Source/MutableTools/Private/MuT/CodeGenerator_FirstPass.h
--- a/Engine/Plugins/Experimental/Mutable/Source/MutableTools/Private/MuT/CodeGenerator_FirstPass.h
+++ b/Engine/Plugins/Experimental/Mutable/Source/MutableTools/Private/MuT/CodeGenerator_FirstPass.h
@@ -23,47 +23,36 @@
 #include "MuT/NodeSurfaceNew.h"
 #include "MuT/NodeSurfaceVariation.h"
 #include "MuT/NodeSurfaceSwitch.h"
-#include "MuT/Visitor.h"
-
-#include <stdint.h>
-#include <memory>
-#include <utility>
 
 namespace mu
 {
 	class Layout;
 	struct FObjectState;
 
-	//---------------------------------------------------------------------------------------------
-	//! First pass of the code generation process.
-	//! It collects data about the object hierarchy, the conditions for each object and the global
-	//! modifiers.
-	//---------------------------------------------------------------------------------------------
-	class FirstPassGenerator :
-		public Base,
-		public BaseVisitor,
-
-        public Visitor<NodeSurfaceNew::Private, Ptr<ASTOp>, true>,
-        public Visitor<NodeSurfaceEdit::Private, Ptr<ASTOp>, true>,
-		public Visitor<NodeSurfaceVariation::Private, Ptr<ASTOp>, true>,
-		public Visitor<NodeSurfaceSwitch::Private, Ptr<ASTOp>, true>,
-		public Visitor<NodeComponentNew::Private, Ptr<ASTOp>, true>,
-        public Visitor<NodeComponentEdit::Private, Ptr<ASTOp>, true>,
-        public Visitor<NodeLOD::Private, Ptr<ASTOp>, true>,
-        public Visitor<NodeObjectNew::Private, Ptr<ASTOp>, true>,
-        public Visitor<NodeObjectGroup::Private, Ptr<ASTOp>, true>,
-        public Visitor<NodePatchMesh::Private, Ptr<ASTOp>, true>,
-
-        public Visitor<NodeModifierMeshClipMorphPlane::Private, Ptr<ASTOp>, true>,
-		public Visitor<NodeModifierMeshClipWithMesh::Private, Ptr<ASTOp>, true>,
-		public Visitor<NodeModifierMeshClipWithUVMask::Private, Ptr<ASTOp>, true>,
-		public Visitor<NodeModifierMeshClipDeform::Private, Ptr<ASTOp>, true>
+	/** First pass of the code generation process.
+	 * It collects data about the object hierarchy, the conditions for each object and the global modifiers.
+	 */
+	class FirstPassGenerator
 	{
 	public:
 
 		FirstPassGenerator();
 
-        void Generate(ErrorLogPtr pErrorLog, const Node::Private* root, bool ignoreStates, class CodeGenerator*);
+        void Generate(Ptr<ErrorLog>, const Node* Root, bool bIgnoreStates, class CodeGenerator*);
+
+	private:
+
+		void Generate_Generic(const Node*);
+		void Generate_Modifier(const NodeModifier*);
+		void Generate_SurfaceNew(const NodeSurfaceNew*);
+		void Generate_SurfaceEdit(const NodeSurfaceEdit*);
+		void Generate_SurfaceSwitch(const NodeSurfaceSwitch*);
+		void Generate_SurfaceVariation(const NodeSurfaceVariation*);
+		void Generate_ComponentNew(const NodeComponentNew*);
+		void Generate_ComponentEdit(const NodeComponentEdit*);
+		void Generate_LOD(const NodeLOD*);
+		void Generate_ObjectNew(const NodeObjectNew*);
+		void Generate_ObjectGroup(const NodeObjectGroup*);
 
 	public:
 
@@ -82,7 +71,7 @@ namespace mu
         //! This is the state mask for the states in which this surface must be added. If it
         //! is empty it means the surface is valid for all states. Otherwise it is only valid
         //! for the states whose index is true.
-        using StateCondition = TArray<uint8_t>;
+        using StateCondition = TArray<uint8>;
 
 		//! Store information about every surface including
 		//! - the component it may be added to
@@ -184,23 +173,6 @@ namespace mu
         typedef TArray< TPair<FObjectState, const Node::Private*> > StateList;
         StateList m_states;
 
-	public:
-
-        Ptr<ASTOp> Visit(const NodeSurfaceNew::Private&) override;
-        Ptr<ASTOp> Visit(const NodeSurfaceEdit::Private&) override;
-		Ptr<ASTOp> Visit(const NodeSurfaceVariation::Private&) override;
-		Ptr<ASTOp> Visit(const NodeSurfaceSwitch::Private&) override;
-		Ptr<ASTOp> Visit(const NodeComponentNew::Private&) override;
-        Ptr<ASTOp> Visit(const NodeComponentEdit::Private&) override;
-        Ptr<ASTOp> Visit(const NodeLOD::Private&) override;
-        Ptr<ASTOp> Visit(const NodeObjectNew::Private& node) override;
-        Ptr<ASTOp> Visit(const NodeObjectGroup::Private&) override;
-        Ptr<ASTOp> Visit(const NodeModifierMeshClipMorphPlane::Private&) override;
-        Ptr<ASTOp> Visit(const NodeModifierMeshClipWithMesh::Private&) override;
-		Ptr<ASTOp> Visit(const NodeModifierMeshClipWithUVMask::Private&) override;
-		Ptr<ASTOp> Visit(const NodeModifierMeshClipDeform::Private&) override;
-        Ptr<ASTOp> Visit(const NodePatchMesh::Private&) override;
-
 	private:
 
         struct FConditionContext
