[mutable] Redesign code generator to not rely on the visitor pattern and reducing RTTI usage.

diff --git a/Engine/Plugins/Experimental/Mutable/Source/MutableTools/Private/MuT/CodeGenerator_Image.cpp b/Engine/Plugins/Experimental/Mutable/Source/MutableTools/Private/MuT/CodeGenerator_Image.cpp
--- a/Engine/Plugins/Experimental/Mutable/Source/MutableTools/Private/MuT/CodeGenerator_Image.cpp
+++ b/Engine/Plugins/Experimental/Mutable/Source/MutableTools/Private/MuT/CodeGenerator_Image.cpp
@@ -41,7 +41,6 @@
 #include "MuT/CompilerPrivate.h"
 #include "MuT/ErrorLog.h"
 #include "MuT/ErrorLogPrivate.h"
-#include "MuT/ImageDescGenerator.h"
 #include "MuT/Node.h"
 #include "MuT/NodeBool.h"
 #include "MuT/NodeColour.h"
@@ -267,8 +266,8 @@ namespace mu
 
         Ptr<ASTOpParameter> op;
 
-        auto it = m_nodeVariables.find( node.m_pNode );
-        if ( it == m_nodeVariables.end() )
+		Ptr<ASTOpParameter>* it = NodeVariables.Find( node.m_pNode );
+        if ( !it )
         {
             op = new ASTOpParameter();
             op->type = OP_TYPE::IM_PARAMETER;
@@ -282,15 +281,15 @@ namespace mu
 			for (int32 a = 0; a < node.m_ranges.Num(); ++a)
 			{
 				FRangeGenerationResult rangeResult;
-				GenerateRange(rangeResult, node.m_ranges[a]);
+				GenerateRange(rangeResult, Options, node.m_ranges[a]);
 				op->ranges.Emplace(op.get(), rangeResult.sizeOp, rangeResult.rangeName, rangeResult.rangeUID);
 			}
 
-            m_nodeVariables[node.m_pNode] = op;
-        }
+			NodeVariables.Add(node.m_pNode, op);
+		}
         else
         {
-            op = it->second;
+            op = *it;
         }
 
 		Result.op = op;
@@ -324,8 +323,15 @@ namespace mu
 
 		UE::Math::TIntVector2<int32> TargetSize = Options.RectSize;
 
-		EImageFormat baseFormat = base->GetImageDesc( true ).m_format;
-        //base = GenerateImageFormat( base, EImageFormat::IF_RGB_UBYTE );
+		FImageDesc BaseDesc = base->GetImageDesc(true);
+
+		// If no target size was specified in the generation options, get the base size to apply it to the mask and blend.
+		if (TargetSize == UE::Math::TIntVector2<int32>(0, 0))
+		{
+			TargetSize = UE::Math::TIntVector2<int32>(BaseDesc.m_size);
+		}
+
+		EImageFormat baseFormat = BaseDesc.m_format;
         base = GenerateImageSize( base, TargetSize);
         op->base = base;
 
@@ -334,7 +340,9 @@ namespace mu
         if ( node.m_pMask )
         {
 			FImageGenerationResult MaskResult;
-			GenerateImage(Options, MaskResult, node.m_pMask);
+			FImageGenerationOptions MaskOptions(Options);
+			MaskOptions.RectSize = TargetSize;
+			GenerateImage(MaskOptions, MaskResult, node.m_pMask);
 			mask = MaskResult.op;
 
             mask = GenerateImageFormat( mask, EImageFormat::IF_L_UBYTE );
@@ -347,7 +355,9 @@ namespace mu
         if ( node.m_pBlended )
         {
 			FImageGenerationResult BlendedResult;
-			GenerateImage(Options, BlendedResult, node.m_pBlended);
+			FImageGenerationOptions BlendOptions(Options);
+			BlendOptions.RectSize = TargetSize;
+			GenerateImage(BlendOptions, BlendedResult, node.m_pBlended);
 			blended = BlendedResult.op;
         }
         else
@@ -355,7 +365,6 @@ namespace mu
             // This argument is required
             blended = GeneratePlainImageCode(FVector4f( 1,1,0,1 ), Options );
         }
-        //blended = GenerateImageFormat( blended, EImageFormat::IF_RGB_UBYTE );
         blended = GenerateImageFormat( blended, baseFormat );
         blended = GenerateImageSize( blended, TargetSize);
         op->blend = blended;
@@ -387,8 +396,18 @@ namespace mu
             // This argument is required
             base = GenerateMissingImageCode(TEXT("Layer base image"), EImageFormat::IF_RGB_UBYTE, node.m_errorContext, Options );
         }
+
+		UE::Math::TIntVector2<int32> TargetSize = Options.RectSize;
+
+		// If no target size was specified in the generation options, get the base size to apply it to the mask and blend.
+		if (TargetSize == UE::Math::TIntVector2<int32>(0, 0))
+		{
+			FImageDesc BaseDesc = base->GetImageDesc(true);
+			TargetSize = UE::Math::TIntVector2<int32>(BaseDesc.m_size);
+		}
+
         base = GenerateImageFormat( base, EImageFormat::IF_RGB_UBYTE );
-        base = GenerateImageSize( base,Options.RectSize );
+        base = GenerateImageSize( base, TargetSize);
         op->base = base;
 
         // Mask of the effect
@@ -400,7 +419,7 @@ namespace mu
 			mask = MaskResult.op;
 			
 			mask = GenerateImageFormat( mask, EImageFormat::IF_L_UBYTE );
-            mask = GenerateImageSize( mask, Options.RectSize);
+            mask = GenerateImageSize( mask, TargetSize);
         }
         op->mask = mask;
 
@@ -409,7 +428,7 @@ namespace mu
         if ( node.m_pColour )
         {
 			FColorGenerationResult ColorResult;
-            GenerateColor(ColorResult, node.m_pColour);
+            GenerateColor(ColorResult, Options, node.m_pColour);
 			colour = ColorResult.op;
         }
         else
@@ -448,8 +467,18 @@ namespace mu
             base = GenerateMissingImageCode(TEXT("Image MultiLayer base"), EImageFormat::IF_RGB_UBYTE, node.m_errorContext, Options);
         }
 
-		EImageFormat baseFormat = base->GetImageDesc().m_format;
-        base = GenerateImageSize( base, Options.RectSize);
+		FImageDesc BaseDesc = base->GetImageDesc(true);
+
+		UE::Math::TIntVector2<int32> TargetSize = Options.RectSize;
+
+		// If no target size was specified in the generation options, get the base size to apply it to the mask and blend.
+		if (TargetSize == UE::Math::TIntVector2<int32>(0, 0))
+		{
+			TargetSize = UE::Math::TIntVector2<int32>(BaseDesc.m_size);
+		}
+
+		EImageFormat baseFormat = BaseDesc.m_format;
+        base = GenerateImageSize( base, TargetSize);
         op->base = base;
 
         // Mask of the effect
@@ -461,7 +490,7 @@ namespace mu
 			mask = MaskResult.op;
 
 			mask = GenerateImageFormat( mask, EImageFormat::IF_L_UBYTE );
-            mask = GenerateImageSize( mask, Options.RectSize);
+            mask = GenerateImageSize( mask, TargetSize);
         }
         op->mask = mask;
 
@@ -479,14 +508,14 @@ namespace mu
             blended = GeneratePlainImageCode(FVector4f( 1,1,0,1 ), Options);
         }
         blended = GenerateImageFormat( blended, baseFormat );
-        blended = GenerateImageSize( blended, Options.RectSize);
+        blended = GenerateImageSize( blended, TargetSize);
         op->blend = blended;
 
         // Range of iteration
         if ( node.m_pRange )
         {
             FRangeGenerationResult rangeResult;
-            GenerateRange( rangeResult, node.m_pRange );
+            GenerateRange( rangeResult, Options, node.m_pRange );
 
             op->range.rangeSize = rangeResult.sizeOp;
             op->range.rangeName = rangeResult.rangeName;
@@ -523,8 +552,19 @@ namespace mu
             base = GenerateMissingImageCode(TEXT("Image Composite Base"), EImageFormat::IF_RGB_UBYTE, node.m_errorContext, Options);
         }
 
-		EImageFormat baseFormat = base->GetImageDesc().m_format;
-        base = GenerateImageSize( base, Options.RectSize);
+		FImageDesc BaseDesc = base->GetImageDesc(true);
+
+		UE::Math::TIntVector2<int32> TargetSize = Options.RectSize;
+
+		// If no target size was specified in the generation options, get the base size to apply it to the mask and blend.
+		if (TargetSize == UE::Math::TIntVector2<int32>(0, 0))
+		{
+			TargetSize = UE::Math::TIntVector2<int32>(BaseDesc.m_size);
+		}
+
+
+		EImageFormat baseFormat = BaseDesc.m_format;
+        base = GenerateImageSize( base, TargetSize);
         op->Base = base;
 
         Ptr<ASTOp> normal;
@@ -535,7 +575,6 @@ namespace mu
 			normal = BaseResult.op;
 
             normal = GenerateImageFormat( normal, EImageFormat::IF_RGB_UBYTE );
-            normal = GenerateImageSize( normal, Options.RectSize);
         }
 		else
 		{
@@ -543,6 +582,8 @@ namespace mu
             normal = GenerateMissingImageCode(TEXT("Image Composite Normal"), EImageFormat::IF_RGB_UBYTE, node.m_errorContext, Options);
 		}
 
+		normal = GenerateImageSize(normal, TargetSize);
+
         op->Normal = normal;
         
 		Result.op = op;
@@ -559,31 +600,31 @@ namespace mu
 		Ptr<ASTOp> OffsetX;
 		if (Node.m_pOffsetX)
 		{
-			OffsetX = Generate(Node.m_pOffsetX);
+			OffsetX = Generate(Node.m_pOffsetX, Options);
 		}
 
 		Ptr<ASTOp> OffsetY;
 		if (Node.m_pOffsetY)
 		{
-			OffsetY = Generate(Node.m_pOffsetY);
+			OffsetY = Generate(Node.m_pOffsetY, Options);
 		}
 	
 		Ptr<ASTOp> ScaleX;
 		if (Node.m_pScaleX)
 		{
-			ScaleX = Generate(Node.m_pScaleX);
+			ScaleX = Generate(Node.m_pScaleX, Options);
 		}
 	
 		Ptr<ASTOp> ScaleY;
 		if (Node.m_pScaleY)
 		{
-			ScaleY = Generate(Node.m_pScaleY);
+			ScaleY = Generate(Node.m_pScaleY, Options);
 		}
 
 		Ptr<ASTOp> Rotation;
 		if (Node.m_pRotation)
 		{
-			Rotation = Generate(Node.m_pRotation);
+			Rotation = Generate(Node.m_pRotation, Options);
 		}
 
 		// If one of the inputs (offset or scale) is missig assume unifrom translation/scaling 
@@ -639,7 +680,7 @@ namespace mu
         // Factor
         if ( Node* pFactor = node.m_pFactor.get() )
         {
-            op->SetChild( op->op.args.ImageInterpolate.factor, Generate( pFactor ));
+            op->SetChild( op->op.args.ImageInterpolate.factor, Generate( pFactor, Options));
         }
         else
         {
@@ -651,19 +692,29 @@ namespace mu
         // Target images
         int numTargets = 0;
 
+		UE::Math::TIntVector2<int32> FinalRectSize = Options.RectSize;
+
         for ( std::size_t t=0
             ; t< node.m_targets.Num() && numTargets<MUTABLE_OP_MAX_INTERPOLATE_COUNT
             ; ++t )
         {
             if ( node.m_targets[t] )
             {
+				FImageGenerationOptions ChildOptions = Options;
+				ChildOptions.RectSize = FinalRectSize;
 				FImageGenerationResult BaseResult;
-				GenerateImage(Options, BaseResult, node.m_targets[t]);
+				GenerateImage(ChildOptions, BaseResult, node.m_targets[t]);
 				Ptr<ASTOp> target = BaseResult.op;
 
+				if (FinalRectSize[0] == 0)
+				{
+					FImageDesc ChildDesc = target->GetImageDesc();
+					FinalRectSize = UE::Math::TIntVector2<int32>(ChildDesc.m_size);
+				}
+
                 // TODO: Support other formats
                 target = GenerateImageFormat( target, EImageFormat::IF_RGB_UBYTE );
-                target = GenerateImageSize( target, Options.RectSize);
+                target = GenerateImageSize( target, FinalRectSize);
 
                 op->SetChild( op->op.args.ImageInterpolate.targets[numTargets], target);
                 numTargets++;
@@ -738,6 +789,7 @@ namespace mu
 
 		// First source, for reference in the size
         Ptr<ASTOp> first;
+		FImageDesc FirstDesc;
 		for (int32 t = 0; t<node.m_sources.Num(); ++t)
 		{
 			if (node.m_sources[t])
@@ -755,18 +807,15 @@ namespace mu
 				}
 
                 Ptr<ASTOp> sizedSource;
-				if (first)
+				if (first && FirstDesc.m_size[0])
 				{
-                    Ptr<ASTOpFixed> sop = new ASTOpFixed();
-                    sop->op.type = OP_TYPE::IM_RESIZELIKE;
-                    sop->SetChild( sop->op.args.ImageResizeLike.source, source);
-                    sop->SetChild( sop->op.args.ImageResizeLike.sizeSource, first);
-                    sizedSource = sop;
+					sizedSource = GenerateImageSize(source, FIntVector2(FirstDesc.m_size));
 				}
 				else
 				{
 					first = source;
 					sizedSource = source;
+					FirstDesc = first->GetImageDesc();
 				}
 
                 op->Sources[t] = sizedSource;
@@ -855,7 +904,7 @@ namespace mu
         // Factor
         if ( Node* pFactor = node.m_pFactor.get() )
         {
-            op->SetChild( op->op.args.ImageSaturate.factor, Generate( pFactor ));
+            op->SetChild( op->op.args.ImageSaturate.factor, Generate( pFactor, Options));
         }
         else
         {
@@ -975,7 +1024,7 @@ namespace mu
         Ptr<ASTOp> colour0 = 0;
         if ( Node* pColour0 = node.m_pColour0.get() )
         {
-            colour0 = Generate( pColour0 );
+            colour0 = Generate( pColour0, Options);
         }
         else
         {
@@ -988,7 +1037,7 @@ namespace mu
         Ptr<ASTOp> colour1 = 0;
         if ( Node* pColour1 = node.m_pColour1.get() )
         {
-            colour1 = Generate( pColour1 );
+            colour1 = Generate( pColour1, Options);
         }
         else
         {
@@ -1033,7 +1082,7 @@ namespace mu
         Ptr<ASTOp> b = 0;
         if ( Node* pScalar = node.m_pThreshold.get() )
         {
-            b = Generate( pScalar );
+            b = Generate( pScalar, Options);
         }
         else
         {
@@ -1089,10 +1138,35 @@ namespace mu
         }
         else
         {
+			// Apply the layout block to the rect size
+			UE::Math::TIntVector2<int32> FinalImageSize = { int32(node.m_sizeX), int32(node.m_sizeY) };
+			if (Options.LayoutToApply)
+			{
+				int32 BlockIndex = Options.LayoutToApply->FindBlock(Options.LayoutBlockId);
+				check(BlockIndex >= 0);
+
+				// Block in layout grid units
+				box< UE::Math::TIntVector2<uint16> > RectInCells;
+				Options.LayoutToApply->GetBlock
+				(
+					BlockIndex,
+					&RectInCells.min[0], &RectInCells.min[1],
+					&RectInCells.size[0], &RectInCells.size[1]
+				);
+
+				FIntPoint grid = Options.LayoutToApply->GetGridSize();
+				grid[0] = FMath::Max(1, grid[0]);
+				grid[1] = FMath::Max(1, grid[1]);
+
+				// Transform to pixels
+				FinalImageSize[0] = (RectInCells.size[0] * FinalImageSize[0]) / grid[0];
+				FinalImageSize[1] = (RectInCells.size[1] * FinalImageSize[1]) / grid[1];
+			}
+
             Ptr<ASTOpFixed> op = new ASTOpFixed();
             op->op.type = OP_TYPE::IM_RESIZE;
-            op->op.args.ImageResize.size[0] = (uint16)node.m_sizeX;
-            op->op.args.ImageResize.size[1] = (uint16)node.m_sizeY;
+            op->op.args.ImageResize.size[0] = (uint16)FinalImageSize[0];
+            op->op.args.ImageResize.size[1] = (uint16)FinalImageSize[1];
             op->SetChild( op->op.args.ImageResize.source, base);
             at = op;
         }
@@ -1110,7 +1184,7 @@ namespace mu
         Ptr<ASTOp> base = 0;
         if ( node.m_pColour )
         {
-            base = Generate( node.m_pColour.get() );
+            base = Generate( node.m_pColour.get(), Options);
         }
         else
         {
@@ -1118,12 +1192,48 @@ namespace mu
             base = GenerateMissingColourCode(TEXT("Image plain colour base"), node.m_errorContext );
         }
 
+		UE::Math::TIntVector2<int32> FinalImageSize = { 0, 0 };
+
+		if (Options.RectSize.X > 0)
+		{
+			FinalImageSize = Options.RectSize;
+		}
+		else
+		{
+			FinalImageSize = { node.m_sizeX, node.m_sizeY };
+
+			// Apply the layout block to the rect size
+			if (Options.LayoutToApply)
+			{
+				int32 BlockIndex = Options.LayoutToApply->FindBlock(Options.LayoutBlockId);
+				check(BlockIndex >= 0);
+
+				// Block in layout grid units
+				box< UE::Math::TIntVector2<uint16> > RectInCells;
+				Options.LayoutToApply->GetBlock
+				(
+					BlockIndex,
+					&RectInCells.min[0], &RectInCells.min[1],
+					&RectInCells.size[0], &RectInCells.size[1]
+				);
+
+				FIntPoint grid = Options.LayoutToApply->GetGridSize();
+				grid[0] = FMath::Max(1, grid[0]);
+				grid[1] = FMath::Max(1, grid[1]);
+
+				// Transform to pixels
+				FinalImageSize[0] = (RectInCells.size[0] * FinalImageSize[0]) / grid[0];
+				FinalImageSize[1] = (RectInCells.size[1] * FinalImageSize[1]) / grid[1];
+			}
+		}
+
+
         Ptr<ASTOpFixed> op = new ASTOpFixed();
         op->op.type = OP_TYPE::IM_PLAINCOLOUR;
         op->SetChild( op->op.args.ImagePlainColour.colour, base);
 		op->op.args.ImagePlainColour.format = node.Format;
-        op->op.args.ImagePlainColour.size[0] = (Options.RectSize[0] > 0) ? Options.RectSize[0] : node.m_sizeX;
-		op->op.args.ImagePlainColour.size[1] = (Options.RectSize[1] > 0) ? Options.RectSize[1] : node.m_sizeY;
+		op->op.args.ImagePlainColour.size[0] = FinalImageSize[0];
+		op->op.args.ImagePlainColour.size[1] = FinalImageSize[1];
 		op->op.args.ImagePlainColour.LODs = 1;
 
         Result.op = op;
@@ -1151,7 +1261,7 @@ namespace mu
 		// Variable value
 		if ( node.m_pParameter )
 		{
-            op->variable = Generate( node.m_pParameter.get() );
+            op->variable = Generate( node.m_pParameter.get(), Options);
 		}
 		else
 		{
@@ -1225,7 +1335,7 @@ namespace mu
         // Condition
         if ( node.m_parameter )
         {
-            op->condition = Generate( node.m_parameter.get() );
+            op->condition = Generate( node.m_parameter.get(), Options);
         }
         else
         {
@@ -1264,7 +1374,7 @@ namespace mu
         FProjectorGenerationResult projectorResult;
         if ( node.m_pProjector )
         {
-            GenerateProjector( projectorResult, node.m_pProjector );
+            GenerateProjector( projectorResult, Options, node.m_pProjector );
             //projectorAt = Generate( node.m_pProjector.get() );
         }
         else
@@ -1288,7 +1398,7 @@ namespace mu
 			// TODO: This will probably Result in a duplicated mesh subgraph, with the original mesh but new layout block ids.
 			// See if it can be optimized and try to reuse the existing layout block ids instead of generating new ones.
 			FMeshGenerationOptions MeshOptions;
-			MeshOptions.State = Options.CurrentStateIndex;
+			MeshOptions.State = Options.State;
 			MeshOptions.ActiveTags = Options.ActiveTags;
 			MeshOptions.bLayouts = true;			// We need the layout that we will use to render
 			MeshOptions.bNormalizeUVs = true;		// We need normalized UVs for the projection
@@ -1395,6 +1505,42 @@ namespace mu
 		ImageRasterOp->mesh = LastMeshOp;
 		ImageRasterOp->projector = projectorResult.op;
 
+		// Calculate size of image to raster:
+		// The full image is:
+		// 0) The hint value in the image options passed down.
+		// 1) whatever is specified in the node attributes.
+		// 2) if that is 0, the size of the mask
+		// 3) if still 0, take the size of the image to project (which is not necessarily related, but often)
+		// 4) if still 0, a default value bigger than 0
+		// then if we are applying a layout a layout block rect need to be calculated of that size, like in image constants.
+		UE::Math::TIntVector2<int32> RasterImageSize = Options.RectSize;
+		bool bApplyLayoutToSize = false;
+
+		if (RasterImageSize.X == 0)
+		{
+			RasterImageSize = UE::Math::TIntVector2<int32>(node.m_imageSize);
+			bApplyLayoutToSize = true;
+		}
+
+		// Target mask
+		if (node.m_pMask)
+		{
+			FImageGenerationResult MaskResult;
+			GenerateImage(Options, MaskResult, node.m_pMask);
+			Ptr<ASTOp> mask = MaskResult.op;
+
+			mask = GenerateImageFormat(mask, EImageFormat::IF_L_UBYTE);
+
+			ImageRasterOp->mask = GenerateImageSize(mask, RasterImageSize);
+
+			if (RasterImageSize.X == 0)
+			{
+				FImageDesc MaskDesc = ImageRasterOp->mask->GetImageDesc();
+				RasterImageSize = UE::Math::TIntVector2<int32>(MaskDesc.m_size);
+				bApplyLayoutToSize = true;
+			}
+		}
+
         // Image
         if ( node.m_pImage )
         {
@@ -1403,15 +1549,21 @@ namespace mu
 			NewOptions.ImageLayoutStrategy = CompilerOptions::TextureLayoutStrategy::None;
 			NewOptions.LayoutToApply = nullptr;
 			NewOptions.LayoutBlockId = -1;
-			NewOptions.RectSize = {};
+			NewOptions.RectSize = { 0,0 };
 
 			FImageGenerationResult ImageResult;
 			GenerateImage(NewOptions, ImageResult, node.m_pImage);
 			ImageRasterOp->image = ImageResult.op;
 
-			FImageDesc desc = CalculateImageDesc(*node.m_pImage->GetBasePrivate());
+			FImageDesc desc = ImageRasterOp->image->GetImageDesc();
 			ImageRasterOp->SourceSizeX = desc.m_size[0];
 			ImageRasterOp->SourceSizeY = desc.m_size[1];
+
+			if (RasterImageSize.X == 0)
+			{
+				RasterImageSize = UE::Math::TIntVector2<int32>(desc.m_size);
+				bApplyLayoutToSize = true;
+			}
         }
         else
         {
@@ -1419,9 +1571,39 @@ namespace mu
 			ImageRasterOp->image = GenerateMissingImageCode(TEXT("Projector image"), EImageFormat::IF_RGB_UBYTE, node.m_errorContext, Options);
         }
 
+		if (RasterImageSize.X == 0)
+		{
+			// Last resort
+			RasterImageSize = { 256, 256 };
+		}
+
+		// Apply the layout block to the rect size
+		if (bApplyLayoutToSize && Options.LayoutToApply)
+		{
+			int32 BlockIndex = Options.LayoutToApply->FindBlock(Options.LayoutBlockId);
+			check(BlockIndex >= 0);
+
+			// Block in layout grid units
+			box< UE::Math::TIntVector2<uint16> > RectInCells;
+			Options.LayoutToApply->GetBlock
+			(
+				BlockIndex,
+				&RectInCells.min[0], &RectInCells.min[1],
+				&RectInCells.size[0], &RectInCells.size[1]
+			);
+
+			FIntPoint grid = Options.LayoutToApply->GetGridSize();
+			grid[0] = FMath::Max(1, grid[0]);
+			grid[1] = FMath::Max(1, grid[1]);
+
+			// Transform to pixels
+			RasterImageSize[0] = (RectInCells.size[0] * RasterImageSize[0]) / grid[0];
+			RasterImageSize[1] = (RectInCells.size[1] * RasterImageSize[1]) / grid[1];
+		}
+
         // Image size, from the current block being generated
-		ImageRasterOp->SizeX = Options.RectSize[0];
-		ImageRasterOp->SizeY = Options.RectSize[1];
+		ImageRasterOp->SizeX = RasterImageSize[0];
+		ImageRasterOp->SizeY = RasterImageSize[1];
 		ImageRasterOp->BlockId = GeneratedLayoutBlockId;
 		ImageRasterOp->LayoutIndex = node.m_layout;
 
@@ -1444,20 +1626,9 @@ namespace mu
 			if (node.m_pAngleFadeEnd) pPropsNode->SetY(node.m_pAngleFadeEnd);
 			else pPropsNode->SetY(pDefaultFade);
 
-			ImageRasterOp->angleFadeProperties = Generate(pPropsNode);
+			ImageRasterOp->angleFadeProperties = Generate(pPropsNode, Options);
 		}
 
-        // Target mask
-        if ( node.m_pMask )
-        {
-			FImageGenerationResult MaskResult;
-			GenerateImage(Options, MaskResult, node.m_pMask);
-			Ptr<ASTOp> mask = MaskResult.op;
-
-            mask = GenerateImageFormat( mask, EImageFormat::IF_L_UBYTE );
-			ImageRasterOp->mask = GenerateImageSize( mask, Options.RectSize);
-        }
-
         // Seam correction operations
         //------------------------------
 		if (node.bEnableTextureSeamCorrection)
@@ -1746,7 +1917,7 @@ namespace mu
         if (!m_missingImage[size_t(format)])
         {
             // Make a checkered debug image
-            const FImageSize size = MUTABLE_MISSING_IMAGE_DESC.m_size;
+            const FImageSize size(16,16);
 
             ImagePtr pImage = new Image( size[0], size[1], 1, format, EInitializationType::NotInitialized);
 
@@ -1933,26 +2104,5 @@ namespace mu
         return Result;
     }
 
-
-    //---------------------------------------------------------------------------------------------
-    FImageDesc CodeGenerator::CalculateImageDesc( const Node::Private& node )
-    {
-        ImageDescGenerator imageDescGenerator;
-        imageDescGenerator.Generate( node );
-        return imageDescGenerator.m_desc;
-    }
-
-
-
-    //---------------------------------------------------------------------------------------------
-    //! This class contains the support data to accelerate the GetImageDesc recursive function.
-    //! _If none is provided in the call, one will be created at that level and used from there on.
-    class FGetImageDescContext
-    {
-    public:
-        vector<bool> m_visited;
-        vector<FImageDesc> m_results;
-    };
-
  
 }
diff --git a/Engine/Plugins/Experimental/Mutable/Source/MutableTools/Private/MuT/CodeGenerator_Mesh.cpp b/Engine/Plugins/Experimental/Mutable/Source/MutableTools/Private/MuT/CodeGenerator_Mesh.cpp
--- a/Engine/Plugins/Experimental/Mutable/Source/MutableTools/Private/MuT/CodeGenerator_Mesh.cpp
+++ b/Engine/Plugins/Experimental/Mutable/Source/MutableTools/Private/MuT/CodeGenerator_Mesh.cpp
@@ -265,8 +265,8 @@ namespace mu
 
     //---------------------------------------------------------------------------------------------
 	void CodeGenerator::PrepareForLayout(Ptr<const Layout> GeneratedLayout,
-		MeshPtr currentLayoutMesh,
-		size_t currentLayoutChannel,
+		Ptr<Mesh> currentLayoutMesh,
+		int32 currentLayoutChannel,
 		const void* errorContext,
 		const FMeshGenerationOptions& MeshOptions
 		)
@@ -692,9 +692,6 @@ namespace mu
             return;
         }
 
-		// Clear bottom-up state
-		m_currentBottomUpState.m_address = nullptr;
-
         // See if it was already generated
 		FGeneratedMeshCacheKey Key;
 		Key.Node = InUntypedNode;
@@ -749,7 +746,7 @@ namespace mu
         // Factor
         if ( node.Factor )
         {
-            OpMorph->Factor = Generate( node.Factor.get() );
+            OpMorph->Factor = Generate( node.Factor.get(), InOptions );
         }
         else
         {
@@ -979,7 +976,7 @@ namespace mu
         // Factor
         if ( Node* pFactor = node.m_pFactor.get() )
         {
-            op->SetChild( op->op.args.MeshInterpolate.factor, Generate( pFactor ) );
+            op->SetChild( op->op.args.MeshInterpolate.factor, Generate( pFactor, InOptions ) );
         }
         else
         {
@@ -1071,7 +1068,7 @@ namespace mu
         // Factor
         if ( node.m_pParameter )
         {
-            op->variable = Generate( node.m_pParameter.get() );
+            op->variable = Generate( node.m_pParameter.get(), InOptions);
         }
         else
         {
@@ -1377,18 +1374,22 @@ namespace mu
 					// Apply whatever transform is necessary for every layout
 					for (int32 LayoutIndex = 0; LayoutIndex < node.m_layouts.Num(); ++LayoutIndex)
 					{
-						NodeLayoutPtr pLayoutNode = node.m_layouts[LayoutIndex];
-						// TODO: In a cleanup of the design of the layouts, we should remove this cast.
-						const NodeLayoutBlocks* TypedNode = dynamic_cast<NodeLayoutBlocks*>(pLayoutNode.get());
-						if (TypedNode)
+						Ptr<NodeLayout> pLayoutNode = node.m_layouts[LayoutIndex];
+						if (!pLayoutNode)
 						{
-							Ptr<const Layout> SourceLayout = TypedNode->GetPrivate()->m_pLayout;
-							Ptr<const Layout> GeneratedLayout = AddLayout( SourceLayout );
-							const void* Context = InOptions.OverrideContext.Get(node.m_errorContext);
-							PrepareForLayout(GeneratedLayout, pCloned, LayoutIndex, Context, InOptions);
-
-							OutResult.GeneratedLayouts.Add(GeneratedLayout);
+							continue;
 						}
+
+						// TODO: In a cleanup of the design of the layouts, we should remove this cast.
+						check(pLayoutNode->GetType()==NodeLayoutBlocks::GetStaticType() );
+						const NodeLayoutBlocks* TypedNode = static_cast<const NodeLayoutBlocks*>(pLayoutNode.get());
+
+						Ptr<const Layout> SourceLayout = TypedNode->GetPrivate()->m_pLayout;
+						Ptr<const Layout> GeneratedLayout = AddLayout( SourceLayout );
+						const void* Context = InOptions.OverrideContext.Get(node.m_errorContext);
+						PrepareForLayout(GeneratedLayout, pCloned, LayoutIndex, Context, InOptions);
+
+						OutResult.GeneratedLayouts.Add(GeneratedLayout);
 					}
 				}
 				else
@@ -1460,13 +1461,8 @@ namespace mu
 		}
 
 		// Apply the modifier for the pre-normal operations stage.
-		FBottomUpState TempState = m_currentBottomUpState;
-
 		bool bModifiersForBeforeOperations = true;
-		OutResult.meshOp = ApplyMeshModifiers(LastMeshOp, InOptions.ActiveTags, bModifiersForBeforeOperations, node.m_errorContext);
-
-		m_currentBottomUpState = TempState;
-
+		OutResult.meshOp = ApplyMeshModifiers(InOptions, LastMeshOp, bModifiersForBeforeOperations, node.m_errorContext);
     }
 
 
@@ -1796,8 +1792,8 @@ namespace mu
 			op->meshB = bResult.meshOp;
 		}
 
-		op->scalarA = Generate(node.m_pScalarA);
-		op->scalarB = Generate(node.m_pScalarB);
+		op->scalarA = Generate(node.m_pScalarA, InOptions);
+		op->scalarB = Generate(node.m_pScalarB, InOptions);
 
 		OutResult.meshOp = op;
 	}
diff --git a/Engine/Plugins/Experimental/Mutable/Source/MutableTools/Private/MuT/CodeGenerator_Projector.cpp b/Engine/Plugins/Experimental/Mutable/Source/MutableTools/Private/MuT/CodeGenerator_Projector.cpp
--- a/Engine/Plugins/Experimental/Mutable/Source/MutableTools/Private/MuT/CodeGenerator_Projector.cpp
+++ b/Engine/Plugins/Experimental/Mutable/Source/MutableTools/Private/MuT/CodeGenerator_Projector.cpp
@@ -51,17 +51,19 @@ namespace mu
 
 
 	//-------------------------------------------------------------------------------------------------
-	void CodeGenerator::GenerateProjector(FProjectorGenerationResult& result, const NodeProjectorPtrConst& untyped)
+	void CodeGenerator::GenerateProjector(FProjectorGenerationResult& result, const FGenericGenerationOptions& Options, const Ptr<const NodeProjector>& Untyped)
 	{
-		if (!untyped)
+		if (!Untyped)
 		{
 			result = FProjectorGenerationResult();
 			return;
 		}
 
 		// See if it was already generated
-		FVisitedKeyMap key = GetCurrentCacheKey(untyped);
-		GeneratedProjectorsMap::ValueType* it = m_generatedProjectors.Find(key);
+		FGeneratedCacheKey Key;
+		Key.Node = Untyped;
+		Key.Options = Options;
+		FGeneratedProjectorsMap::ValueType* it = GeneratedProjectors.Find(Key);
 		if (it)
 		{
 			result = *it;
@@ -69,13 +71,13 @@ namespace mu
 		}
 
 		// Generate for each different type of node
-		if (auto Constant = dynamic_cast<const NodeProjectorConstant*>(untyped.get()))
+		if ( Untyped->GetType()== NodeProjectorConstant::GetStaticType() )
 		{
-			GenerateProjector_Constant(result, Constant);
+			GenerateProjector_Constant(result, Options, static_cast<const NodeProjectorConstant*>(Untyped.get()));
 		}
-		else if (auto Param = dynamic_cast<const NodeProjectorParameter*>(untyped.get()))
+		else if (Untyped->GetType() == NodeProjectorParameter::GetStaticType() )
 		{
-			GenerateProjector_Parameter(result, Param);
+			GenerateProjector_Parameter(result, Options, static_cast<const NodeProjectorParameter*>(Untyped.get()));
 		}
 		else
 		{
@@ -84,12 +86,12 @@ namespace mu
 
 
 		// Cache the result
-		m_generatedProjectors.Add(key, result);
+		GeneratedProjectors.Add(Key, result);
 	}
 
 
 	//-------------------------------------------------------------------------------------------------
-	void CodeGenerator::GenerateProjector_Constant(FProjectorGenerationResult& result,
+	void CodeGenerator::GenerateProjector_Constant(FProjectorGenerationResult& result, const FGenericGenerationOptions& Options,
 		const Ptr<const NodeProjectorConstant>& constant)
 	{
 		const NodeProjectorConstant::Private& node = *constant->GetPrivate();
@@ -103,15 +105,15 @@ namespace mu
 
 
 	//-------------------------------------------------------------------------------------------------
-	void CodeGenerator::GenerateProjector_Parameter(FProjectorGenerationResult& result,
+	void CodeGenerator::GenerateProjector_Parameter(FProjectorGenerationResult& result, const FGenericGenerationOptions& Options,
 		const Ptr<const NodeProjectorParameter>& paramn)
 	{
 		const NodeProjectorParameter::Private& node = *paramn->GetPrivate();
 
 		Ptr<ASTOpParameter> op;
 
-		auto it = m_nodeVariables.find(node.m_pNode);
-		if (it == m_nodeVariables.end())
+		Ptr<ASTOpParameter>* it = NodeVariables.Find(node.m_pNode);
+		if (!it)
 		{
 			FParameterDesc param;
 			param.m_name = node.m_name;
@@ -129,15 +131,15 @@ namespace mu
 			for (int32 a = 0; a < node.m_ranges.Num(); ++a)
 			{
 				FRangeGenerationResult rangeResult;
-				GenerateRange(rangeResult, node.m_ranges[a]);
+				GenerateRange(rangeResult, Options, node.m_ranges[a]);
 				op->ranges.Emplace(op.get(), rangeResult.sizeOp, rangeResult.rangeName, rangeResult.rangeUID);
 			}
 
-			m_nodeVariables[node.m_pNode] = op;
+			NodeVariables.Add(node.m_pNode, op);
 		}
 		else
 		{
-			op = it->second;
+			op = *it;
 		}
 
 		const FProjector& Projector = op->parameter.m_defaultValue.Get<ParamProjectorType>();
diff --git a/Engine/Plugins/Experimental/Mutable/Source/MutableTools/Private/MuT/CodeGenerator_Scalar.cpp b/Engine/Plugins/Experimental/Mutable/Source/MutableTools/Private/MuT/CodeGenerator_Scalar.cpp
--- a/Engine/Plugins/Experimental/Mutable/Source/MutableTools/Private/MuT/CodeGenerator_Scalar.cpp
+++ b/Engine/Plugins/Experimental/Mutable/Source/MutableTools/Private/MuT/CodeGenerator_Scalar.cpp
@@ -44,27 +44,23 @@
 #include "MuT/Table.h"
 #include "MuT/TablePrivate.h"
 
-#include <memory>
-#include <utility>
-
 
 namespace mu
 {
-	class Node;
-
 
-	//-------------------------------------------------------------------------------------------------
-	void CodeGenerator::GenerateScalar(FScalarGenerationResult& result, const Ptr<const NodeScalar>& untyped)
+	void CodeGenerator::GenerateScalar(FScalarGenerationResult& result, const FGenericGenerationOptions& Options, const Ptr<const NodeScalar>& Untyped)
 	{
-		if (!untyped)
+		if (!Untyped)
 		{
 			result = FScalarGenerationResult();
 			return;
 		}
 
 		// See if it was already generated
-		FVisitedKeyMap key = GetCurrentCacheKey(untyped);
-		GeneratedScalarsMap::ValueType* it = m_generatedScalars.Find(key);
+		FGeneratedCacheKey Key;
+		Key.Node = Untyped;
+		Key.Options = Options;
+		FGeneratedScalarsMap::ValueType* it = GeneratedScalars.Find(Key);
 		if (it)
 		{
 			result = *it;
@@ -73,37 +69,37 @@ namespace mu
 
 
 		// Generate for each different type of node
-		if (auto Constant = dynamic_cast<const NodeScalarConstant*>(untyped.get()))
+		if (auto Constant = dynamic_cast<const NodeScalarConstant*>(Untyped.get()))
 		{
-			GenerateScalar_Constant(result, Constant);
+			GenerateScalar_Constant(result, Options, Constant);
 		}
-		else if (auto Param = dynamic_cast<const NodeScalarParameter*>(untyped.get()))
+		else if (auto Param = dynamic_cast<const NodeScalarParameter*>(Untyped.get()))
 		{
-			GenerateScalar_Parameter(result, Param);
+			GenerateScalar_Parameter(result, Options, Param);
 		}
-		else if (auto Switch = dynamic_cast<const NodeScalarSwitch*>(untyped.get()))
+		else if (auto Switch = dynamic_cast<const NodeScalarSwitch*>(Untyped.get()))
 		{
-			GenerateScalar_Switch(result, Switch);
+			GenerateScalar_Switch(result, Options, Switch);
 		}
-		else if (auto EnumParam = dynamic_cast<const NodeScalarEnumParameter*>(untyped.get()))
+		else if (auto EnumParam = dynamic_cast<const NodeScalarEnumParameter*>(Untyped.get()))
 		{
-			GenerateScalar_EnumParameter(result, EnumParam);
+			GenerateScalar_EnumParameter(result, Options, EnumParam);
 		}
-		else if (auto Curve = dynamic_cast<const NodeScalarCurve*>(untyped.get()))
+		else if (auto Curve = dynamic_cast<const NodeScalarCurve*>(Untyped.get()))
 		{
-			GenerateScalar_Curve(result, Curve);
+			GenerateScalar_Curve(result, Options, Curve);
 		}
-		else if (auto Arithmetic = dynamic_cast<const NodeScalarArithmeticOperation*>(untyped.get()))
+		else if (auto Arithmetic = dynamic_cast<const NodeScalarArithmeticOperation*>(Untyped.get()))
 		{
-			GenerateScalar_Arithmetic(result, Arithmetic);
+			GenerateScalar_Arithmetic(result, Options, Arithmetic);
 		}
-		else if (auto Variation = dynamic_cast<const NodeScalarVariation*>(untyped.get()))
+		else if (auto Variation = dynamic_cast<const NodeScalarVariation*>(Untyped.get()))
 		{
-			GenerateScalar_Variation(result, Variation);
+			GenerateScalar_Variation(result, Options, Variation);
 		}
-		else if (auto Table = dynamic_cast<const NodeScalarTable*>(untyped.get()))
+		else if (auto Table = dynamic_cast<const NodeScalarTable*>(Untyped.get()))
 		{
-			GenerateScalar_Table(result, Table);
+			GenerateScalar_Table(result, Options, Table);
 		}
 		else
 		{
@@ -112,12 +108,12 @@ namespace mu
 		}
 
 		// Cache the result
-		m_generatedScalars.Add(key, result);
+		GeneratedScalars.Add(Key, result);
 	}
 
 
 	//-------------------------------------------------------------------------------------------------
-	void CodeGenerator::GenerateScalar_Constant(FScalarGenerationResult& result, const Ptr<const NodeScalarConstant>& Typed)
+	void CodeGenerator::GenerateScalar_Constant(FScalarGenerationResult& result, const FGenericGenerationOptions& Options, const Ptr<const NodeScalarConstant>& Typed)
 	{
 		const NodeScalarConstant::Private& node = *Typed->GetPrivate();
 
@@ -130,14 +126,14 @@ namespace mu
 
 
 	//-------------------------------------------------------------------------------------------------
-	void CodeGenerator::GenerateScalar_Parameter(FScalarGenerationResult& result, const Ptr<const NodeScalarParameter>& Typed)
+	void CodeGenerator::GenerateScalar_Parameter(FScalarGenerationResult& result, const FGenericGenerationOptions& Options, const Ptr<const NodeScalarParameter>& Typed)
 	{
 		const NodeScalarParameter::Private& node = *Typed->GetPrivate();
 
 		Ptr<ASTOpParameter> op;
 
-		auto it = m_nodeVariables.find(node.m_pNode);
-		if (it == m_nodeVariables.end())
+		Ptr<ASTOpParameter>* it = NodeVariables.Find(node.m_pNode);
+		if (!it)
 		{
 			FParameterDesc param;
 			param.m_name = node.m_name;
@@ -153,15 +149,15 @@ namespace mu
 			for (int32 a = 0; a < node.m_ranges.Num(); ++a)
 			{
 				FRangeGenerationResult rangeResult;
-				GenerateRange(rangeResult, node.m_ranges[a]);
+				GenerateRange(rangeResult, Options, node.m_ranges[a]);
 				op->ranges.Emplace(op.get(), rangeResult.sizeOp, rangeResult.rangeName, rangeResult.rangeUID);
 			}
 
-			m_nodeVariables[node.m_pNode] = op;
+			NodeVariables.Add(node.m_pNode, op);
 		}
 		else
 		{
-			op = it->second;
+			op = *it;
 		}
 
 		result.op = op;
@@ -169,14 +165,14 @@ namespace mu
 
 
 	//-------------------------------------------------------------------------------------------------
-	void CodeGenerator::GenerateScalar_EnumParameter(FScalarGenerationResult& result, const Ptr<const NodeScalarEnumParameter>& Typed)
+	void CodeGenerator::GenerateScalar_EnumParameter(FScalarGenerationResult& result, const FGenericGenerationOptions& Options, const Ptr<const NodeScalarEnumParameter>& Typed)
 	{
 		const NodeScalarEnumParameter::Private& node = *Typed->GetPrivate();
 
 		Ptr<ASTOpParameter> op;
 
-		auto it = m_nodeVariables.find(node.m_pNode);
-		if (it == m_nodeVariables.end())
+		Ptr<ASTOpParameter>* it = NodeVariables.Find(node.m_pNode);
+		if (!it)
 		{
 			FParameterDesc param;
 			param.m_name = node.m_name;
@@ -199,16 +195,15 @@ namespace mu
 			for (int32 a = 0; a < node.m_ranges.Num(); ++a)
 			{
 				FRangeGenerationResult rangeResult;
-				GenerateRange(rangeResult, node.m_ranges[a]);
+				GenerateRange(rangeResult, Options, node.m_ranges[a]);
 				op->ranges.Emplace(op.get(), rangeResult.sizeOp, rangeResult.rangeName, rangeResult.rangeUID);
 			}
 
-			m_nodeVariables[node.m_pNode] = op;
-
+			NodeVariables.Add(node.m_pNode, op);
 		}
 		else
 		{
-			op = it->second;
+			op = *it;
 		}
 
 		result.op = op;
@@ -216,7 +211,7 @@ namespace mu
 
 
 	//-------------------------------------------------------------------------------------------------
-	void CodeGenerator::GenerateScalar_Switch(FScalarGenerationResult& result, const Ptr<const NodeScalarSwitch>& Typed)
+	void CodeGenerator::GenerateScalar_Switch(FScalarGenerationResult& result, const FGenericGenerationOptions& Options, const Ptr<const NodeScalarSwitch>& Typed)
 	{
 		const NodeScalarSwitch::Private& node = *Typed->GetPrivate();
 
@@ -237,7 +232,7 @@ namespace mu
 		if (node.m_pParameter)
 		{
 			FScalarGenerationResult ChildResult;
-			GenerateScalar(ChildResult, node.m_pParameter.get());
+			GenerateScalar(ChildResult, Options, node.m_pParameter.get());
 			op->variable = ChildResult.op;
 		}
 		else
@@ -253,7 +248,7 @@ namespace mu
 			if (node.m_options[t])
 			{
 				FScalarGenerationResult ChildResult;
-				GenerateScalar(ChildResult, node.m_options[t].get());
+				GenerateScalar(ChildResult, Options, node.m_options[t].get());
 				branch = ChildResult.op;
 			}
 			else
@@ -269,7 +264,7 @@ namespace mu
 
 
 	//-------------------------------------------------------------------------------------------------
-	void CodeGenerator::GenerateScalar_Variation(FScalarGenerationResult& result, const Ptr<const NodeScalarVariation>& Typed)
+	void CodeGenerator::GenerateScalar_Variation(FScalarGenerationResult& result, const FGenericGenerationOptions& Options, const Ptr<const NodeScalarVariation>& Typed)
 	{
 		const NodeScalarVariation::Private& node = *Typed->GetPrivate();
 
@@ -281,7 +276,7 @@ namespace mu
 			FMeshGenerationResult branchResults;
 
 			FScalarGenerationResult ChildResult;
-			GenerateScalar(ChildResult, node.m_defaultScalar);
+			GenerateScalar(ChildResult, Options, node.m_defaultScalar);
 			op = ChildResult.op;
 		}
 
@@ -310,7 +305,7 @@ namespace mu
 			if (node.m_variations[t].m_scalar)
 			{
 				FScalarGenerationResult ChildResult;
-				GenerateScalar(ChildResult, node.m_variations[t].m_scalar);
+				GenerateScalar(ChildResult, Options, node.m_variations[t].m_scalar);
 				variationOp = ChildResult.op;
 			}
 			else
@@ -334,7 +329,7 @@ namespace mu
 
 
 	//-------------------------------------------------------------------------------------------------
-	void CodeGenerator::GenerateScalar_Curve(FScalarGenerationResult& result, const Ptr<const NodeScalarCurve>& Typed)
+	void CodeGenerator::GenerateScalar_Curve(FScalarGenerationResult& result, const FGenericGenerationOptions& Options, const Ptr<const NodeScalarCurve>& Typed)
 	{
 		const NodeScalarCurve::Private& node = *Typed->GetPrivate();
 
@@ -343,7 +338,7 @@ namespace mu
 		// T
 		if (Node* pA = node.m_input_scalar.get())
 		{
-			op->time = Generate(pA);
+			op->time = Generate(pA, Options);
 		}
 		else
 		{
@@ -357,7 +352,7 @@ namespace mu
 
 
 	//-------------------------------------------------------------------------------------------------
-	void CodeGenerator::GenerateScalar_Arithmetic(FScalarGenerationResult& result, const Ptr<const NodeScalarArithmeticOperation>& Typed)
+	void CodeGenerator::GenerateScalar_Arithmetic(FScalarGenerationResult& result, const FGenericGenerationOptions& Options, const Ptr<const NodeScalarArithmeticOperation>& Typed)
 	{
 		const NodeScalarArithmeticOperation::Private& node = *Typed->GetPrivate();
 
@@ -379,7 +374,7 @@ namespace mu
 		// A
 		if (Node* pA = node.m_pA.get())
 		{
-			op->SetChild(op->op.args.ScalarArithmetic.a, Generate(pA));
+			op->SetChild(op->op.args.ScalarArithmetic.a, Generate(pA, Options));
 		}
 		else
 		{
@@ -391,7 +386,7 @@ namespace mu
 		// B
 		if (Node* pB = node.m_pB.get())
 		{
-			op->SetChild(op->op.args.ScalarArithmetic.b, Generate(pB));
+			op->SetChild(op->op.args.ScalarArithmetic.b, Generate(pB, Options));
 		}
 		else
 		{
@@ -405,17 +400,17 @@ namespace mu
 
 
 	//-------------------------------------------------------------------------------------------------
-	void CodeGenerator::GenerateScalar_Table(FScalarGenerationResult& result, const Ptr<const NodeScalarTable>& Typed)
+	void CodeGenerator::GenerateScalar_Table(FScalarGenerationResult& result, const FGenericGenerationOptions& Options, const Ptr<const NodeScalarTable>& Typed)
 	{
 		const NodeScalarTable::Private& node = *Typed->GetPrivate();
 
 		Ptr<ASTOp> Op = GenerateTableSwitch<NodeScalarTable::Private, ETableColumnType::Scalar, OP_TYPE::SC_SWITCH>(node,
-			[this](const NodeScalarTable::Private& node, int colIndex, int row, ErrorLog* pErrorLog)
+			[this,&Options](const NodeScalarTable::Private& node, int colIndex, int row, ErrorLog* pErrorLog)
 			{
 				NodeScalarConstantPtr pCell = new NodeScalarConstant();
 				float scalar = node.m_pTable->GetPrivate()->Rows[row].Values[colIndex].Scalar;
 				pCell->SetValue(scalar);
-				return Generate(pCell);
+				return Generate(pCell, Options);
 			});
 
 		result.op = Op;
@@ -433,7 +428,8 @@ namespace mu
 		NodeScalarConstantPtr pNode = new NodeScalarConstant();
 		pNode->SetValue(value);
 
-		Ptr<ASTOp> result = Generate(pNode);
+		FGenericGenerationOptions Options;
+		Ptr<ASTOp> result = Generate(pNode, Options);
 
 		return result;
 	}
diff --git a/Engine/Plugins/Experimental/Mutable/Source/MutableTools/Private/MuT/CodeGenerator_SecondPass.h b/Engine/Plugins/Experimental/Mutable/Source/MutableTools/Private/MuT/CodeGenerator_SecondPass.h
--- a/Engine/Plugins/Experimental/Mutable/Source/MutableTools/Private/MuT/CodeGenerator_SecondPass.h
+++ b/Engine/Plugins/Experimental/Mutable/Source/MutableTools/Private/MuT/CodeGenerator_SecondPass.h
@@ -20,11 +20,10 @@
 namespace mu
 {
 
-	//---------------------------------------------------------------------------------------------
-	//! Second pass of the code generation process.
-    //! Solves surface and modifier conditions from tags and variations
-	//---------------------------------------------------------------------------------------------
-    class SecondPassGenerator : public Base
+	/** Second pass of the code generation process.
+     * It solves surface and modifier conditions from tags and variations
+	 */
+    class SecondPassGenerator
 	{
 	public:
 
diff --git a/Engine/Plugins/Experimental/Mutable/Source/MutableTools/Private/MuT/CodeGenerator_String.cpp b/Engine/Plugins/Experimental/Mutable/Source/MutableTools/Private/MuT/CodeGenerator_String.cpp
--- a/Engine/Plugins/Experimental/Mutable/Source/MutableTools/Private/MuT/CodeGenerator_String.cpp
+++ b/Engine/Plugins/Experimental/Mutable/Source/MutableTools/Private/MuT/CodeGenerator_String.cpp
@@ -19,28 +19,25 @@
 #include "MuT/NodeStringConstantPrivate.h"
 #include "MuT/NodeStringParameter.h"
 #include "MuT/NodeStringParameterPrivate.h"
-#include "map"
-
-#include <memory>
-#include <utility>
-
 
 namespace mu
 {
 
 
 	//-------------------------------------------------------------------------------------------------
-	void CodeGenerator::GenerateString(FStringGenerationResult& result, const NodeStringPtrConst& untyped)
+	void CodeGenerator::GenerateString(FStringGenerationResult& result, const FGenericGenerationOptions& Options, const Ptr<const NodeString>& Untyped)
 	{
-		if (!untyped)
+		if (!Untyped)
 		{
 			result = FStringGenerationResult();
 			return;
 		}
 
 		// See if it was already generated
-		FVisitedKeyMap key = GetCurrentCacheKey(untyped);
-		GeneratedStringsMap::ValueType* it = m_generatedStrings.Find(key);
+		FGeneratedCacheKey Key;
+		Key.Node = Untyped;
+		Key.Options = Options;
+		FGeneratedStringsMap::ValueType* it = GeneratedStrings.Find(Key);
 		if (it)
 		{
 			result = *it;
@@ -48,13 +45,13 @@ namespace mu
 		}
 
 		// Generate for each different type of node
-		if (auto Constant = dynamic_cast<const NodeStringConstant*>(untyped.get()))
+		if (Untyped->GetType()==NodeStringConstant::GetStaticType())
 		{
-			GenerateString_Constant(result, Constant);
+			GenerateString_Constant(result, Options, static_cast<const NodeStringConstant*>(Untyped.get()));
 		}
-		else if (auto Param = dynamic_cast<const NodeStringParameter*>(untyped.get()))
+		else if (Untyped->GetType() == NodeStringParameter::GetStaticType())
 		{
-			GenerateString_Parameter(result, Param);
+			GenerateString_Parameter(result, Options, static_cast<const NodeStringParameter*>(Untyped.get()));
 		}
 		else
 		{
@@ -62,12 +59,12 @@ namespace mu
 		}
 
 		// Cache the result
-		m_generatedStrings.Add(key, result);
+		GeneratedStrings.Add(Key, result);
 	}
 
 
 	//-------------------------------------------------------------------------------------------------
-	void CodeGenerator::GenerateString_Constant(FStringGenerationResult& result, const Ptr<const NodeStringConstant>& Typed)
+	void CodeGenerator::GenerateString_Constant(FStringGenerationResult& result, const FGenericGenerationOptions& Options, const Ptr<const NodeStringConstant>& Typed)
 	{
 		const NodeStringConstant::Private& node = *Typed->GetPrivate();
 
@@ -79,14 +76,14 @@ namespace mu
 
 
 	//-------------------------------------------------------------------------------------------------
-	void CodeGenerator::GenerateString_Parameter(FStringGenerationResult& result, const Ptr<const NodeStringParameter>& Typed)
+	void CodeGenerator::GenerateString_Parameter(FStringGenerationResult& result, const FGenericGenerationOptions& Options, const Ptr<const NodeStringParameter>& Typed)
 	{
 		const NodeStringParameter::Private& node = *Typed->GetPrivate();
 
 		Ptr<ASTOpParameter> op;
 
-		auto it = m_nodeVariables.find(node.m_pNode);
-		if (it == m_nodeVariables.end())
+		Ptr<ASTOpParameter>* it = NodeVariables.Find(node.m_pNode);
+		if (!it)
 		{
 			FParameterDesc param;
 			param.m_name = node.m_name;
@@ -102,15 +99,15 @@ namespace mu
 			for (int32 a = 0; a < node.m_ranges.Num(); ++a)
 			{
 				FRangeGenerationResult rangeResult;
-				GenerateRange(rangeResult, node.m_ranges[a]);
+				GenerateRange(rangeResult, Options, node.m_ranges[a]);
 				op->ranges.Emplace(op.get(), rangeResult.sizeOp, rangeResult.rangeName, rangeResult.rangeUID);
 			}
 
-			m_nodeVariables[node.m_pNode] = op;
+			NodeVariables.Add(node.m_pNode, op);
 		}
 		else
 		{
-			op = it->second;
+			op = *it;
 		}
 
 		result.op = op;
diff --git a/Engine/Plugins/Experimental/Mutable/Source/MutableTools/Private/MuT/CodeOptimiser.cpp b/Engine/Plugins/Experimental/Mutable/Source/MutableTools/Private/MuT/CodeOptimiser.cpp
--- a/Engine/Plugins/Experimental/Mutable/Source/MutableTools/Private/MuT/CodeOptimiser.cpp
+++ b/Engine/Plugins/Experimental/Mutable/Source/MutableTools/Private/MuT/CodeOptimiser.cpp
@@ -371,23 +371,22 @@ namespace mu
 		bool found = false;
 
 		// If the program added new operations, we haven't visited them.
-		m_visited.resize( program.m_opAddress.Num(), 0 );
+		m_visited.SetNumZeroed( program.m_opAddress.Num() );
 
-		m_pending.clear();
-		m_pending.reserve( program.m_opAddress.Num()/4 );
-		m_pending.push_back( std::make_pair(false,rootAt) );
+		m_pending.Reset();
+		m_pending.Reserve( program.m_opAddress.Num()/4 );
+		m_pending.Add({ false,rootAt });
 
 		// Don't early out to be able to complete parent op cached flags
-		while ( !m_pending.empty() )
+		while ( !m_pending.IsEmpty() )
 		{
-			std::pair<bool,int> item = m_pending.back();
-			m_pending.pop_back();
-			OP::ADDRESS at = item.second;
+			TPair<bool,int> item = m_pending.Pop();
+			OP::ADDRESS at = item.Value;
 
 			// Not cached?
 			if (m_visited[at]<=1)
 			{
-				if (item.first)
+				if (item.Key)
 				{
 					// Item indicating we finished with all the children of a parent
 					check(m_visited[at]==1);
@@ -402,9 +401,7 @@ namespace mu
 				}
 				else if (!found)
 				{
-					if ( std::find( m_typesToFind.begin(), m_typesToFind.end(), program.GetOpType(at))
-						 !=
-						 m_typesToFind.end() )
+					if (m_typesToFind.Find(program.GetOpType(at)) != INDEX_NONE )
 					{
 						m_visited[at] = 3; // visited, ops found
 						found = true;
@@ -414,13 +411,13 @@ namespace mu
 						check(m_visited[at]==0);
 
 						m_visited[at] = 1;
-						m_pending.push_back( std::make_pair(true,at) );
+						m_pending.Add({ true,at });
 
 						ForEachReference( program, at, [&](OP::ADDRESS ref)
 						{
 							if (ref && m_visited[ref]==0)
 							{
-								m_pending.push_back( std::make_pair(false,ref) );
+								m_pending.Add({ false,ref });
 							}
 						});
 					}
@@ -442,14 +439,14 @@ namespace mu
 	//---------------------------------------------------------------------------------------------
 	IsConstantVisitor::IsConstantVisitor()
 	{
-		vector<OP_TYPE> parameterTypes;
-		parameterTypes.push_back( OP_TYPE::BO_PARAMETER );
-		parameterTypes.push_back( OP_TYPE::NU_PARAMETER );
-		parameterTypes.push_back( OP_TYPE::SC_PARAMETER );
-		parameterTypes.push_back( OP_TYPE::CO_PARAMETER );
-		parameterTypes.push_back( OP_TYPE::PR_PARAMETER );
-		parameterTypes.push_back( OP_TYPE::IM_PARAMETER );
-		m_findOpTypeVisitor= std::make_unique<FindOpTypeVisitor>(parameterTypes);
+		TArray<OP_TYPE> parameterTypes;
+		parameterTypes.Add( OP_TYPE::BO_PARAMETER );
+		parameterTypes.Add( OP_TYPE::NU_PARAMETER );
+		parameterTypes.Add( OP_TYPE::SC_PARAMETER );
+		parameterTypes.Add( OP_TYPE::CO_PARAMETER );
+		parameterTypes.Add( OP_TYPE::PR_PARAMETER );
+		parameterTypes.Add( OP_TYPE::IM_PARAMETER );
+		m_findOpTypeVisitor = MakeUnique<FindOpTypeVisitor>(parameterTypes);
 	}
 
 
@@ -823,7 +820,7 @@ namespace mu
 	//-------------------------------------------------------------------------------------------------
 	//-------------------------------------------------------------------------------------------------
 	//-------------------------------------------------------------------------------------------------
-	CodeOptimiser::CodeOptimiser(Ptr<CompilerOptions> options, vector<STATE_COMPILATION_DATA>& states )
+	CodeOptimiser::CodeOptimiser(Ptr<CompilerOptions> options, TArray<FStateCompilationData>& states )
 		: m_states( states )
 	{
 		m_options = options;
@@ -1244,7 +1241,7 @@ namespace mu
 			// like the grow-map generation.
 			bool modified = true;
 			int numIterations = 0;
-			while (modified && (!m_optimizeIterationsMax || (m_optimizeIterationsLeft>0) || !numIterations) )
+			while (modified)
 			{
 				MUTABLE_CPUPROFILER_SCOPE(FirstStage);
 
@@ -1306,7 +1303,7 @@ namespace mu
 			});
 
 			// Make sure we didn't lose track of pointers
-			for ( size_t s=0;  s<m_states.size(); ++s )
+			for ( int32 s=0;  s<m_states.Num(); ++s )
 			{
 				check( roots.Contains( m_states[s].root ) );
 			}
@@ -1347,7 +1344,7 @@ namespace mu
 				FullOptimiseAST(roots, 2);
 			}
 
-			for ( size_t s=0;  s<m_states.size(); ++s )
+			for ( int32 s=0;  s<m_states.Num(); ++s )
 			{
 				UE_LOG(LogMutableCore, Verbose, TEXT(" - constant generator"));
 				ConstantGeneratorAST( m_options->GetPrivate(), m_states[s].root );
@@ -1363,7 +1360,7 @@ namespace mu
 			//AXE_INT_VALUE("Mutable", Verbose, "ast size", (int64_t)ASTOp::CountNodes(roots));
 
 			// Make sure we didn't lose track of pointers
-			for ( size_t s=0;  s<m_states.size(); ++s )
+			for ( int32 s=0;  s<m_states.Num(); ++s )
 			{
 				check( roots.Contains( m_states[s].root ) );
 			}
diff --git a/Engine/Plugins/Experimental/Mutable/Source/MutableTools/Private/MuT/CodeOptimiser.h b/Engine/Plugins/Experimental/Mutable/Source/MutableTools/Private/MuT/CodeOptimiser.h
--- a/Engine/Plugins/Experimental/Mutable/Source/MutableTools/Private/MuT/CodeOptimiser.h
+++ b/Engine/Plugins/Experimental/Mutable/Source/MutableTools/Private/MuT/CodeOptimiser.h
@@ -9,18 +9,16 @@
 #include "MuT/CompilerPrivate.h"
 #include "MuT/AST.h"
 
-#include <memory>
-
 namespace mu
 {
 
     //---------------------------------------------------------------------------------------------
     //! Code optimiser
     //---------------------------------------------------------------------------------------------
-    class CodeOptimiser : public Base
+    class CodeOptimiser
     {
     public:
-        CodeOptimiser( Ptr<CompilerOptions> options, vector<STATE_COMPILATION_DATA>& states );
+        CodeOptimiser( Ptr<CompilerOptions> options, TArray<FStateCompilationData>& states );
 
         //! Optimise the virtual machine code, using several transforms.
         void OptimiseAST();
@@ -30,7 +28,7 @@ namespace mu
 		Ptr<CompilerOptions> m_options;
 
         //!
-        vector<STATE_COMPILATION_DATA>& m_states;
+		TArray<FStateCompilationData>& m_states;
 
         //! The max number of optimize iterations is shared across several stages now.
         //! This is how many are left
@@ -56,7 +54,7 @@ namespace mu
     {
     public:
 
-        FindOpTypeVisitor(const vector<OP_TYPE>& types)
+        FindOpTypeVisitor(const TArray<OP_TYPE>& types)
             : m_typesToFind(types)
         {
         }
@@ -66,15 +64,15 @@ namespace mu
 
     private:
 
-        vector<OP_TYPE> m_typesToFind;
+        TArray<OP_TYPE> m_typesToFind;
 
-        vector< std::pair<bool,OP::ADDRESS> > m_pending;
+		TArray< TPair<bool,OP::ADDRESS> > m_pending;
 
         // 0 not visited
         // 1 children pending
         // 2 visited and not found
         // 3 visited and found
-        vector<uint8_t> m_visited;
+		TArray<uint8> m_visited;
     };
 
 
@@ -92,7 +90,7 @@ namespace mu
 
     private:
 
-        std::unique_ptr<FindOpTypeVisitor> m_findOpTypeVisitor;
+        TUniquePtr<FindOpTypeVisitor> m_findOpTypeVisitor;
     };
 
 
@@ -163,13 +161,13 @@ namespace mu
     {
     public:
 
-        RuntimeParameterVisitorAST(const STATE_COMPILATION_DATA* pState);
+        RuntimeParameterVisitorAST(const FStateCompilationData* pState);
 
         bool HasAny( const Ptr<ASTOp>& root );
 
     private:
 
-        const STATE_COMPILATION_DATA* m_pState;
+        const FStateCompilationData* m_pState;
 
         //!
         struct PENDING_ITEM
@@ -215,7 +213,7 @@ namespace mu
     {
     public:
 
-        RuntimeTextureCompressionRemoverAST( STATE_COMPILATION_DATA* state, bool bInAlwaysUncompress );
+        RuntimeTextureCompressionRemoverAST(FStateCompilationData* state, bool bInAlwaysUncompress );
 
     protected:
 
@@ -236,12 +234,12 @@ namespace mu
     {
     private:
 
-        STATE_COMPILATION_DATA& m_stateProps;
+		FStateCompilationData& m_stateProps;
         // unused const GPU_PLATFORM_PROPS& m_gpuPlatformProps;
 
     public:
 
-        ParameterOptimiserAST( STATE_COMPILATION_DATA &s,
+        ParameterOptimiserAST(FStateCompilationData&s,
                                const FModelOptimizationOptions& optimisationOptions );
 
         bool Apply();
@@ -266,11 +264,6 @@ namespace mu
     //---------------------------------------------------------------------------------------------
     extern Ptr<ASTOp> EnsureValidMask( Ptr<ASTOp> mask, Ptr<ASTOp> base );
 
-    //---------------------------------------------------------------------------------------------
-    //! Return true if two non-zero pixels of the masks overlap.
-    //---------------------------------------------------------------------------------------------
-    extern bool AreMasksOverlapping( const FProgram& program, OP::ADDRESS a, OP::ADDRESS b );
-
 
     //---------------------------------------------------------------------------------------------
     //! Calculate all the parameters found relevant under a particular operation. This may not
@@ -278,7 +271,7 @@ namespace mu
     //! relevant)
     //! It has an internal cache, so don't reuse if the program changes.
     //---------------------------------------------------------------------------------------------
-    class SubtreeRelevantParametersVisitorAST : public Base
+    class SubtreeRelevantParametersVisitorAST
     {
     public:
 
@@ -289,31 +282,32 @@ namespace mu
 
     private:
 
-        struct STATE
+        struct FState
         {
             Ptr<ASTOp> op;
-            bool onlyLayoutIsRelevant=false;
+            bool bOnlyLayoutIsRelevant=false;
 
-            STATE( Ptr<ASTOp> o=nullptr, bool l=false) : op(o), onlyLayoutIsRelevant(l) {}
+			FState( Ptr<ASTOp> o=nullptr, bool l=false) : op(o), bOnlyLayoutIsRelevant(l) {}
 
-            bool operator==(const STATE& o) const
+            bool operator==(const FState& o) const
             {
                 return  op == o.op &&
-                        onlyLayoutIsRelevant == o.onlyLayoutIsRelevant;
+					bOnlyLayoutIsRelevant == o.bOnlyLayoutIsRelevant;
             }
-        };
 
-        struct state_hash
-        {
-            std::size_t operator()(const STATE& k) const
-            {
-                return std::hash<const void*>()(k.op.get());
-            }
-        };
+			friend FORCEINLINE uint32 GetTypeHash(const FState& InKey)
+			{
+				uint32 KeyHash = 0;
+				KeyHash = HashCombineFast(KeyHash, ::GetTypeHash(InKey.op.get()));
+				KeyHash = HashCombineFast(KeyHash, ::GetTypeHash(InKey.bOnlyLayoutIsRelevant));
+				return KeyHash;
+			}
+		};
+
 
         // Result cache
         // \todo optimise by storing unique lists separately and an index here.
-        std::unordered_map< STATE, TSet<FString>, state_hash > m_resultCache;
+        TMap< FState, TSet<FString> > m_resultCache;
     };
 
 }
diff --git a/Engine/Plugins/Experimental/Mutable/Source/MutableTools/Private/MuT/CodeOptimiser_Parameters.cpp b/Engine/Plugins/Experimental/Mutable/Source/MutableTools/Private/MuT/CodeOptimiser_Parameters.cpp
--- a/Engine/Plugins/Experimental/Mutable/Source/MutableTools/Private/MuT/CodeOptimiser_Parameters.cpp
+++ b/Engine/Plugins/Experimental/Mutable/Source/MutableTools/Private/MuT/CodeOptimiser_Parameters.cpp
@@ -45,7 +45,7 @@ namespace mu
     //---------------------------------------------------------------------------------------------
     //---------------------------------------------------------------------------------------------
     //---------------------------------------------------------------------------------------------
-    RuntimeParameterVisitorAST::RuntimeParameterVisitorAST(const STATE_COMPILATION_DATA* pState)
+    RuntimeParameterVisitorAST::RuntimeParameterVisitorAST(const FStateCompilationData* pState)
         : m_pState(pState)
     {
     }
@@ -271,7 +271,7 @@ namespace mu
     //---------------------------------------------------------------------------------------------
     //---------------------------------------------------------------------------------------------
     ParameterOptimiserAST::ParameterOptimiserAST(
-            STATE_COMPILATION_DATA& s,
+		FStateCompilationData& s,
             const FModelOptimizationOptions& optimisationOptions
             )
             : m_stateProps(s)
@@ -1621,10 +1621,10 @@ namespace mu
     void SubtreeRelevantParametersVisitorAST::Run( Ptr<ASTOp> root )
     {
         // Cached?
-		std::unordered_map< STATE, TSet<FString>, state_hash >::iterator it = m_resultCache.find( STATE(root,false) );
-        if (it!=m_resultCache.end())
+		TSet<FString>* it = m_resultCache.Find( FState(root,false) );
+        if (it)
         {
-            m_params = it->second;
+            m_params = *it;
             return;
         }
 
@@ -1706,7 +1706,7 @@ namespace mu
                 return true;
             });
 
-            m_resultCache[STATE(root,false)] = m_params;
+            m_resultCache.Add( FState(root,false), m_params );
         }
     }
 
@@ -1723,7 +1723,7 @@ namespace mu
     {
     public:
 
-        StateCacheDetectorAST( STATE_COMPILATION_DATA* pState )
+        StateCacheDetectorAST(FStateCompilationData* pState )
             : m_hasRuntimeParamVisitor( pState )
         {
             ASTOpList roots;
@@ -1851,7 +1851,7 @@ namespace mu
     {
     public:
 
-        StateCacheFormatOptimiserAST( STATE_COMPILATION_DATA& state,
+        StateCacheFormatOptimiserAST(FStateCompilationData& state,
                                    const AccumulateAllImageFormatsOpAST& opFormats )
             : m_state(state)
             , m_opFormats(opFormats)
@@ -1900,7 +1900,7 @@ namespace mu
 
     private:
 
-        STATE_COMPILATION_DATA& m_state;
+		FStateCompilationData& m_state;
 
         const AccumulateAllImageFormatsOpAST& m_opFormats;
     };
@@ -1910,7 +1910,7 @@ namespace mu
     //---------------------------------------------------------------------------------------------
     //---------------------------------------------------------------------------------------------
     RuntimeTextureCompressionRemoverAST::RuntimeTextureCompressionRemoverAST(
-		STATE_COMPILATION_DATA* pState,
+		FStateCompilationData* pState,
 		bool bInAlwaysUncompress
 	)
         : m_hasRuntimeParamVisitor(pState)
@@ -2006,7 +2006,7 @@ namespace mu
     {
         MUTABLE_CPUPROFILER_SCOPE(OptimiseStatesAST);
 
-         for ( size_t s=0; s<m_states.size(); ++s )
+         for ( int32 s=0; s<m_states.Num(); ++s )
          {
             // Remove the unnecessary lods
             if (m_states[s].nodeState.m_optimisation.bOnlyFirstLOD)
@@ -2101,7 +2101,7 @@ namespace mu
         // necessary at this stage before GPU optimisation.
         {
             TArray<Ptr<ASTOp>> roots;
-            for(const STATE_COMPILATION_DATA& s:m_states)
+            for(const FStateCompilationData& s:m_states)
             {
                 roots.Add(s.root);
             }
@@ -2109,7 +2109,7 @@ namespace mu
             AccumulateAllImageFormatsOpAST opFormats;
             opFormats.Run(roots);
 
-            for (STATE_COMPILATION_DATA& s: m_states )
+            for (FStateCompilationData& s: m_states )
             {
                 {
                     UE_LOG(LogMutableCore, Verbose, TEXT(" - state cache"));
@@ -2133,7 +2133,7 @@ namespace mu
             while (modified && (!m_optimizeIterationsMax || m_optimizeIterationsLeft>0 || !numIterations ))
             {
                 TArray<Ptr<ASTOp>> roots;
-                for(const STATE_COMPILATION_DATA& s:m_states)
+                for(const FStateCompilationData& s:m_states)
                 {
                     roots.Add(s.root);
                 }
@@ -2156,14 +2156,14 @@ namespace mu
 				UE_LOG(LogMutableCore, Verbose, TEXT("(int) %s : %ld"), TEXT("ast size"), int64(ASTOp::CountNodes(roots)));
 			}
 
-            for(STATE_COMPILATION_DATA& s:m_states)
+            for(FStateCompilationData& s:m_states)
             {
                 UE_LOG(LogMutableCore, Verbose, TEXT(" - constant optimiser"));
                 ConstantGeneratorAST( m_options->GetPrivate(), s.root );
             }
 
             TArray<Ptr<ASTOp>> roots;
-            for(const STATE_COMPILATION_DATA& s:m_states)
+            for(const FStateCompilationData& s:m_states)
             {
                 roots.Add(s.root);
             }
@@ -2180,7 +2180,7 @@ namespace mu
 
         // Gather all the current roots
         TArray<Ptr<ASTOp>> roots;
-        for(const STATE_COMPILATION_DATA& s:m_states)
+        for(const FStateCompilationData& s:m_states)
         {
             roots.Add(s.root);
         }
@@ -2197,7 +2197,7 @@ namespace mu
 
             // Update the marks for the instructions that don't depend on runtime parameters to be
             // cached.
-            for (STATE_COMPILATION_DATA& s:m_states)
+            for (FStateCompilationData& s:m_states)
             {
                 StateCacheDetectorAST c( &s );
             }
diff --git a/Engine/Plugins/Experimental/Mutable/Source/MutableTools/Private/MuT/CodeOptimiser_Semantic.cpp b/Engine/Plugins/Experimental/Mutable/Source/MutableTools/Private/MuT/CodeOptimiser_Semantic.cpp
--- a/Engine/Plugins/Experimental/Mutable/Source/MutableTools/Private/MuT/CodeOptimiser_Semantic.cpp
+++ b/Engine/Plugins/Experimental/Mutable/Source/MutableTools/Private/MuT/CodeOptimiser_Semantic.cpp
@@ -416,6 +416,36 @@ namespace mu
 			break;
 		}
 
+		//-------------------------------------------------------------------------------------
+		case OP_TYPE::ME_APPLYLAYOUT:
+		{
+			Ptr<ASTOp> Base = children[op.args.MeshApplyLayout.mesh].child();
+			OP_TYPE BaseType = Base->GetOpType();
+			switch (BaseType)
+			{
+			case OP_TYPE::ME_ADDTAGS:
+			{
+				// Add the tags after layout
+				Ptr<ASTOpMeshAddTags> NewAddTags = mu::Clone<ASTOpMeshAddTags>(Base);
+
+				if (NewAddTags->Source)
+				{
+					Ptr<ASTOpFixed> NewAt = mu::Clone<ASTOpFixed>(this);
+					NewAt->SetChild(NewAt->op.args.MeshApplyLayout.mesh, NewAddTags->Source);
+					NewAddTags->Source = NewAt;
+				}
+
+				at = NewAddTags;
+				break;
+			}
+
+			default:
+				break;
+			}
+
+			break;
+		}
+
         default:
             break;
         }
@@ -530,12 +560,17 @@ namespace mu
 
 			case OP_TYPE::ME_ADDTAGS:
 			{
-				// Ignore tags for projection
-				const ASTOpMeshAddTags* Add = dynamic_cast<const ASTOpMeshAddTags*>(sourceAt.get());
+				// Add the tags after layout
+				Ptr<ASTOpMeshAddTags> NewAddTags = mu::Clone<ASTOpMeshAddTags>(sourceAt);
+
+				if (NewAddTags->Source)
+				{
+					Ptr<ASTOpFixed> NewAt = mu::Clone<ASTOpFixed>(this);
+					NewAt->SetChild(NewAt->op.args.MeshProject.mesh, NewAddTags->Source);
+					NewAddTags->Source = NewAt;
+				}
 
-				Ptr<ASTOpFixed> NewProject = mu::Clone<ASTOpFixed>(this);
-				NewProject->SetChild(NewProject->op.args.MeshProject.mesh, Add->Source);
-				at = NewProject;
+				at = NewAddTags;
 				break;
 			}
 
diff --git a/Engine/Plugins/Experimental/Mutable/Source/MutableTools/Private/MuT/Compiler.cpp b/Engine/Plugins/Experimental/Mutable/Source/MutableTools/Private/MuT/Compiler.cpp
--- a/Engine/Plugins/Experimental/Mutable/Source/MutableTools/Private/MuT/Compiler.cpp
+++ b/Engine/Plugins/Experimental/Mutable/Source/MutableTools/Private/MuT/Compiler.cpp
@@ -196,7 +196,7 @@ namespace mu
     {
         MUTABLE_CPUPROFILER_SCOPE(Compile);
 
-        vector< STATE_COMPILATION_DATA > states;
+        TArray< FStateCompilationData > states;
         Ptr<ErrorLog> genErrorLog;
         {
             CodeGenerator gen( m_pD->m_options->GetPrivate() );
@@ -207,23 +207,23 @@ namespace mu
 
             for ( const auto& s: gen.m_states )
             {
-                STATE_COMPILATION_DATA data;
-                data.nodeState = s.first;
-                data.root = s.second;
-                data.state.Name = s.first.m_name;
-                states.push_back( data );
+                FStateCompilationData data;
+                data.nodeState = s.Key;
+                data.root = s.Value;
+                data.state.Name = s.Key.m_name;
+                states.Add( data );
             }
 
             genErrorLog = gen.m_pErrorLog;
         }
 
-        //vector<Ptr<ASTOp>> roots;
-        //for( const STATE_COMPILATION_DATA& s: states)
-        //{
-        //    roots.push_back(s.root);
-        //}
 
         // Slow AST code verification for debugging.
+        //TArray<Ptr<ASTOp>> roots;
+        //for( const FStateCompilationData& s: states)
+        //{
+        //    roots.Add(s.root);
+        //}
         //ASTOp::FullAssert(roots);
 
         // Optimize the generated code
@@ -246,7 +246,7 @@ namespace mu
 		FImageOperator ImOp = FImageOperator::GetDefault(m_pD->m_options->GetPrivate()->ImageFormatFunc);
 		FLinkerOptions LinkerOptions(ImOp);
 
-		for( STATE_COMPILATION_DATA& s: states )
+		for(FStateCompilationData& s: states )
         {
 			LinkerOptions.MinTextureResidentMipCount = m_pD->m_options->GetPrivate()->MinTextureResidentMipCount;
 
@@ -265,7 +265,7 @@ namespace mu
 		program.m_opAddress.Shrink();
 
         // Set the runtime parameter indices.
-        for( STATE_COMPILATION_DATA& s: states )
+        for(FStateCompilationData& s: states )
         {
             for ( int32 p=0; p<s.nodeState.m_runtimeParams.Num(); ++p )
             {
diff --git a/Engine/Plugins/Experimental/Mutable/Source/MutableTools/Private/MuT/CompilerPrivate.h b/Engine/Plugins/Experimental/Mutable/Source/MutableTools/Private/MuT/CompilerPrivate.h
--- a/Engine/Plugins/Experimental/Mutable/Source/MutableTools/Private/MuT/CompilerPrivate.h
+++ b/Engine/Plugins/Experimental/Mutable/Source/MutableTools/Private/MuT/CompilerPrivate.h
@@ -82,7 +82,7 @@ namespace mu
 
 
     //!
-    class CompilerOptions::Private : public Base
+    class CompilerOptions::Private
     {
     public:
 
@@ -162,7 +162,7 @@ namespace mu
 
 
     //!
-    struct STATE_COMPILATION_DATA
+    struct FStateCompilationData
     {
         FObjectState nodeState;
         Ptr<ASTOp> root;
@@ -179,7 +179,7 @@ namespace mu
 
 
     //!
-    class Compiler::Private : public Base
+    class Compiler::Private
     {
     public:
 
diff --git a/Engine/Plugins/Experimental/Mutable/Source/MutableTools/Private/MuT/ErrorLogPrivate.h b/Engine/Plugins/Experimental/Mutable/Source/MutableTools/Private/MuT/ErrorLogPrivate.h
--- a/Engine/Plugins/Experimental/Mutable/Source/MutableTools/Private/MuT/ErrorLogPrivate.h
+++ b/Engine/Plugins/Experimental/Mutable/Source/MutableTools/Private/MuT/ErrorLogPrivate.h
@@ -1,6 +1,5 @@
 // Copyright Epic Games, Inc. All Rights Reserved.
 
-
 #pragma once
 
 #include "MuT/ErrorLog.h"
@@ -8,18 +7,12 @@
 #include "MuR/ModelPrivate.h"
 #include "MuR/Operations.h"
 
-#include <memory>
-
 namespace mu
 {
-	class ErrorLog::Private : public Base
+	class ErrorLog::Private
 	{
 	public:
 
-		Private()
-		{
-		}
-
         struct FErrorData
         {
             TArray< float > m_unassignedUVs;
@@ -37,8 +30,6 @@ namespace mu
 		TArray<FMessage> m_messages;
 
 
-		//-----------------------------------------------------------------------------------------
-
 		//!
 		void Add(const FString& Message, ErrorLogMessageType Type, const void* Context, ErrorLogMessageSpamBin SpamBin = ELMSB_ALL);
 
diff --git a/Engine/Plugins/Experimental/Mutable/Source/MutableTools/Private/MuT/NodeBool.cpp b/Engine/Plugins/Experimental/Mutable/Source/MutableTools/Private/MuT/NodeBool.cpp
--- a/Engine/Plugins/Experimental/Mutable/Source/MutableTools/Private/MuT/NodeBool.cpp
+++ b/Engine/Plugins/Experimental/Mutable/Source/MutableTools/Private/MuT/NodeBool.cpp
@@ -44,7 +44,7 @@ namespace mu
 		{
 		case 0 :  return NodeBoolConstant::StaticUnserialise( arch ); break;
 		case 1 :  return NodeBoolParameter::StaticUnserialise( arch ); break;
-		case 2 :  return NodeBoolIsNull::StaticUnserialise( arch ); break;
+		//case 2 :  return NodeBoolIsNull::StaticUnserialise( arch ); break;
 		case 3 :  return NodeBoolNot::StaticUnserialise( arch ); break;
 		case 4 :  return NodeBoolAnd::StaticUnserialise( arch ); break;
 		default : check(false);
diff --git a/Engine/Plugins/Experimental/Mutable/Source/MutableTools/Private/MuT/NodeBoolPrivate.h b/Engine/Plugins/Experimental/Mutable/Source/MutableTools/Private/MuT/NodeBoolPrivate.h
--- a/Engine/Plugins/Experimental/Mutable/Source/MutableTools/Private/MuT/NodeBoolPrivate.h
+++ b/Engine/Plugins/Experimental/Mutable/Source/MutableTools/Private/MuT/NodeBoolPrivate.h
@@ -21,10 +21,6 @@ namespace mu
 
 	class NodeBoolConstant::Private : public NodeBool::Private
 	{
-	public:
-
-		MUTABLE_DEFINE_CONST_VISITABLE();
-
 	public:
 
 		static NODE_TYPE s_type;
@@ -54,10 +50,6 @@ namespace mu
 
 	class NodeBoolParameter::Private : public NodeBool::Private
 	{
-	public:
-
-		MUTABLE_DEFINE_CONST_VISITABLE();
-
 	public:
 
 		static NODE_TYPE s_type;
@@ -107,45 +99,8 @@ namespace mu
 	};
 
 
-	class NodeBoolIsNull::Private : public NodeBool::Private
-	{
-	public:
-
-		MUTABLE_DEFINE_CONST_VISITABLE();
-
-	public:
-
-		static NODE_TYPE s_type;
-
-		NodePtr m_pSource;
-
-		//!
-		void Serialise( OutputArchive& arch ) const
-		{
-            uint32 ver = 0;
-			arch << ver;
-
-			arch << m_pSource;
-		}
-
-		//!
-		void Unserialise( InputArchive& arch )
-		{
-            uint32 ver;
-			arch >> ver;
-            check(ver<=0);
-
-			arch >> m_pSource;
-		}
-	};
-
-
 	class NodeBoolNot::Private : public NodeBool::Private
 	{
-	public:
-
-		MUTABLE_DEFINE_CONST_VISITABLE();
-
 	public:
 
 		static NODE_TYPE s_type;
@@ -175,10 +130,6 @@ namespace mu
 
 	class NodeBoolAnd::Private : public NodeBool::Private
 	{
-	public:
-
-		MUTABLE_DEFINE_CONST_VISITABLE();
-
 	public:
 
 		static NODE_TYPE s_type;
diff --git a/Engine/Plugins/Experimental/Mutable/Source/MutableTools/Private/MuT/NodeColourArithmeticOperationPrivate.h b/Engine/Plugins/Experimental/Mutable/Source/MutableTools/Private/MuT/NodeColourArithmeticOperationPrivate.h
--- a/Engine/Plugins/Experimental/Mutable/Source/MutableTools/Private/MuT/NodeColourArithmeticOperationPrivate.h
+++ b/Engine/Plugins/Experimental/Mutable/Source/MutableTools/Private/MuT/NodeColourArithmeticOperationPrivate.h
@@ -11,10 +11,6 @@ namespace mu
 
 	class NodeColourArithmeticOperation::Private : public NodeColour::Private
 	{
-	public:
-
-		MUTABLE_DEFINE_CONST_VISITABLE()
-
 	public:
 
 		static NODE_TYPE s_type;
diff --git a/Engine/Plugins/Experimental/Mutable/Source/MutableTools/Private/MuT/NodeColourConstantPrivate.h b/Engine/Plugins/Experimental/Mutable/Source/MutableTools/Private/MuT/NodeColourConstantPrivate.h
--- a/Engine/Plugins/Experimental/Mutable/Source/MutableTools/Private/MuT/NodeColourConstantPrivate.h
+++ b/Engine/Plugins/Experimental/Mutable/Source/MutableTools/Private/MuT/NodeColourConstantPrivate.h
@@ -12,10 +12,6 @@ namespace mu
 
 	class NodeColourConstant::Private : public NodeColour::Private
 	{
-	public:
-
-		MUTABLE_DEFINE_CONST_VISITABLE()
-
 	public:
 
 		static NODE_TYPE s_type;
diff --git a/Engine/Plugins/Experimental/Mutable/Source/MutableTools/Private/MuT/NodeColourVariationPrivate.h b/Engine/Plugins/Experimental/Mutable/Source/MutableTools/Private/MuT/NodeColourVariationPrivate.h
--- a/Engine/Plugins/Experimental/Mutable/Source/MutableTools/Private/MuT/NodeColourVariationPrivate.h
+++ b/Engine/Plugins/Experimental/Mutable/Source/MutableTools/Private/MuT/NodeColourVariationPrivate.h
@@ -14,9 +14,6 @@ namespace mu
     //---------------------------------------------------------------------------------------------
     class NodeColourVariation::Private : public Node::Private
     {
-    public:
-        MUTABLE_DEFINE_CONST_VISITABLE()
-
     public:
         Private() {}
 
diff --git a/Engine/Plugins/Experimental/Mutable/Source/MutableTools/Private/MuT/NodeComponentEdit.cpp b/Engine/Plugins/Experimental/Mutable/Source/MutableTools/Private/MuT/NodeComponentEdit.cpp
--- a/Engine/Plugins/Experimental/Mutable/Source/MutableTools/Private/MuT/NodeComponentEdit.cpp
+++ b/Engine/Plugins/Experimental/Mutable/Source/MutableTools/Private/MuT/NodeComponentEdit.cpp
@@ -76,6 +76,21 @@ namespace mu
         m_pD->m_surfaces[ index ] = pNode;
     }
 
+
+	//---------------------------------------------------------------------------------------------
+	const NodeComponentNew::Private* NodeComponentEdit::Private::GetParentComponentNew() const
+	{
+		const NodeComponentNew::Private* parent = nullptr;
+		if (m_pParent)
+		{
+			NodeComponent::Private* ParentPrivate = dynamic_cast<NodeComponent::Private*>(m_pParent->GetBasePrivate());
+			parent = ParentPrivate->GetParentComponentNew();
+		}
+
+		check(parent);
+		return parent;
+	}
+
 }
 
 
diff --git a/Engine/Plugins/Experimental/Mutable/Source/MutableTools/Private/MuT/NodeComponentEditPrivate.h b/Engine/Plugins/Experimental/Mutable/Source/MutableTools/Private/MuT/NodeComponentEditPrivate.h
--- a/Engine/Plugins/Experimental/Mutable/Source/MutableTools/Private/MuT/NodeComponentEditPrivate.h
+++ b/Engine/Plugins/Experimental/Mutable/Source/MutableTools/Private/MuT/NodeComponentEditPrivate.h
@@ -6,7 +6,6 @@
 #include "MuT/NodeComponentPrivate.h"
 #include "MuT/NodeComponentEdit.h"
 #include "MuT/NodePatchImagePrivate.h"
-
 #include "MuT/NodeMesh.h"
 #include "MuT/NodePatchMesh.h"
 #include "MuT/NodeImage.h"
@@ -22,10 +21,6 @@ namespace mu
 	//---------------------------------------------------------------------------------------------
 	class NodeComponentEdit::Private : public NodeComponent::Private
 	{
-	public:
-
-		MUTABLE_DEFINE_CONST_VISITABLE()
-
 	public:
 
 		static NODE_TYPE s_type;
@@ -57,18 +52,7 @@ namespace mu
 
 
 		// NodeComponent::Private interface
-        const NodeComponentNew::Private* GetParentComponentNew() const override
-		{
-			const NodeComponentNew::Private* parent = nullptr;
-			if (m_pParent)
-			{
-				NodeComponent::Private* ParentPrivate = dynamic_cast<NodeComponent::Private*>(m_pParent->GetBasePrivate());
-				parent = ParentPrivate->GetParentComponentNew();
-			}
-
-			check(parent);
-			return parent;
-		}
+		const NodeComponentNew::Private* GetParentComponentNew() const override;
 	};
 
 }
diff --git a/Engine/Plugins/Experimental/Mutable/Source/MutableTools/Private/MuT/NodeImageBinarisePrivate.h b/Engine/Plugins/Experimental/Mutable/Source/MutableTools/Private/MuT/NodeImageBinarisePrivate.h
--- a/Engine/Plugins/Experimental/Mutable/Source/MutableTools/Private/MuT/NodeImageBinarisePrivate.h
+++ b/Engine/Plugins/Experimental/Mutable/Source/MutableTools/Private/MuT/NodeImageBinarisePrivate.h
@@ -15,10 +15,6 @@ namespace mu
 
 	class NodeImageBinarise::Private : public NodeImage::Private
 	{
-	public:
-
-		MUTABLE_DEFINE_CONST_VISITABLE();
-
 	public:
 
 		static NODE_TYPE s_type;
diff --git a/Engine/Plugins/Experimental/Mutable/Source/MutableTools/Private/MuT/NodeImageConstantPrivate.h b/Engine/Plugins/Experimental/Mutable/Source/MutableTools/Private/MuT/NodeImageConstantPrivate.h
--- a/Engine/Plugins/Experimental/Mutable/Source/MutableTools/Private/MuT/NodeImageConstantPrivate.h
+++ b/Engine/Plugins/Experimental/Mutable/Source/MutableTools/Private/MuT/NodeImageConstantPrivate.h
@@ -12,10 +12,6 @@ namespace mu
 {
 	class NodeImageConstant::Private : public NodeImage::Private
 	{
-	public:
-
-		MUTABLE_DEFINE_CONST_VISITABLE();
-
 	public:
 
 		static NODE_TYPE s_type;
diff --git a/Engine/Plugins/Experimental/Mutable/Source/MutableTools/Private/MuT/NodeImageGradientPrivate.h b/Engine/Plugins/Experimental/Mutable/Source/MutableTools/Private/MuT/NodeImageGradientPrivate.h
--- a/Engine/Plugins/Experimental/Mutable/Source/MutableTools/Private/MuT/NodeImageGradientPrivate.h
+++ b/Engine/Plugins/Experimental/Mutable/Source/MutableTools/Private/MuT/NodeImageGradientPrivate.h
@@ -1,14 +1,11 @@
 // Copyright Epic Games, Inc. All Rights Reserved.
 
-
 #pragma once
 
-
 #include "MuT/NodeColour.h"
 #include "MuT/NodeImagePrivate.h"
 #include "MuT/NodeImageGradient.h"
 #include "MuT/AST.h"
-
 #include "MuR/MutableMath.h"
 
 
@@ -18,10 +15,6 @@ namespace mu
 
 	class NodeImageGradient::Private : public NodeImage::Private
 	{
-	public:
-
-		MUTABLE_DEFINE_CONST_VISITABLE()
-
 	public:
 
 		static NODE_TYPE s_type;
diff --git a/Engine/Plugins/Experimental/Mutable/Source/MutableTools/Private/MuT/NodeImageLuminancePrivate.h b/Engine/Plugins/Experimental/Mutable/Source/MutableTools/Private/MuT/NodeImageLuminancePrivate.h
--- a/Engine/Plugins/Experimental/Mutable/Source/MutableTools/Private/MuT/NodeImageLuminancePrivate.h
+++ b/Engine/Plugins/Experimental/Mutable/Source/MutableTools/Private/MuT/NodeImageLuminancePrivate.h
@@ -9,13 +9,8 @@
 namespace mu
 {
 
-
 	class NodeImageLuminance::Private : public NodeImage::Private
 	{
-	public:
-
-		MUTABLE_DEFINE_CONST_VISITABLE()
-
 	public:
 
 		static NODE_TYPE s_type;
diff --git a/Engine/Plugins/Experimental/Mutable/Source/MutableTools/Private/MuT/NodeImageSaturatePrivate.h b/Engine/Plugins/Experimental/Mutable/Source/MutableTools/Private/MuT/NodeImageSaturatePrivate.h
--- a/Engine/Plugins/Experimental/Mutable/Source/MutableTools/Private/MuT/NodeImageSaturatePrivate.h
+++ b/Engine/Plugins/Experimental/Mutable/Source/MutableTools/Private/MuT/NodeImageSaturatePrivate.h
@@ -11,13 +11,8 @@
 namespace mu
 {
 
-
 	class NodeImageSaturate::Private : public NodeImage::Private
 	{
-	public:
-
-		MUTABLE_DEFINE_CONST_VISITABLE()
-
 	public:
 
 		static NODE_TYPE s_type;
diff --git a/Engine/Plugins/Experimental/Mutable/Source/MutableTools/Private/MuT/NodeImageTransform.cpp b/Engine/Plugins/Experimental/Mutable/Source/MutableTools/Private/MuT/NodeImageTransform.cpp
--- a/Engine/Plugins/Experimental/Mutable/Source/MutableTools/Private/MuT/NodeImageTransform.cpp
+++ b/Engine/Plugins/Experimental/Mutable/Source/MutableTools/Private/MuT/NodeImageTransform.cpp
@@ -136,7 +136,7 @@ namespace mu
 	//---------------------------------------------------------------------------------------------
 	bool NodeImageTransform::GetKeepAspectRatio() const
 	{
-		return m_pD->bKeepAspectRatio;;
+		return m_pD->bKeepAspectRatio;
 	}
 
 	//---------------------------------------------------------------------------------------------
diff --git a/Engine/Plugins/Experimental/Mutable/Source/MutableTools/Private/MuT/NodeImageVariationPrivate.h b/Engine/Plugins/Experimental/Mutable/Source/MutableTools/Private/MuT/NodeImageVariationPrivate.h
--- a/Engine/Plugins/Experimental/Mutable/Source/MutableTools/Private/MuT/NodeImageVariationPrivate.h
+++ b/Engine/Plugins/Experimental/Mutable/Source/MutableTools/Private/MuT/NodeImageVariationPrivate.h
@@ -14,9 +14,6 @@ namespace mu
     //---------------------------------------------------------------------------------------------
     class NodeImageVariation::Private : public Node::Private
     {
-    public:
-        MUTABLE_DEFINE_CONST_VISITABLE()
-
     public:
         Private() {}
 
diff --git a/Engine/Plugins/Experimental/Mutable/Source/MutableTools/Private/MuT/NodeMeshFormatPrivate.h b/Engine/Plugins/Experimental/Mutable/Source/MutableTools/Private/MuT/NodeMeshFormatPrivate.h
--- a/Engine/Plugins/Experimental/Mutable/Source/MutableTools/Private/MuT/NodeMeshFormatPrivate.h
+++ b/Engine/Plugins/Experimental/Mutable/Source/MutableTools/Private/MuT/NodeMeshFormatPrivate.h
@@ -15,11 +15,6 @@ namespace mu
 	{
 	public:
 
-		MUTABLE_DEFINE_CONST_VISITABLE()
-
-	public:
-
-
 		static NODE_TYPE s_type;
 
 		//! Source mesh to transform
diff --git a/Engine/Plugins/Experimental/Mutable/Source/MutableTools/Private/MuT/NodeObjectNewPrivate.h b/Engine/Plugins/Experimental/Mutable/Source/MutableTools/Private/MuT/NodeObjectNewPrivate.h
--- a/Engine/Plugins/Experimental/Mutable/Source/MutableTools/Private/MuT/NodeObjectNewPrivate.h
+++ b/Engine/Plugins/Experimental/Mutable/Source/MutableTools/Private/MuT/NodeObjectNewPrivate.h
@@ -17,10 +17,6 @@ namespace mu
 
 	class NodeObjectNew::Private : public NodeObject::Private
 	{
-	public:
-
-		MUTABLE_DEFINE_CONST_VISITABLE()
-
 	public:
 
 		static NODE_TYPE s_type;
@@ -109,7 +105,7 @@ namespace mu
 				{
 					std::string Temp;
 					arch >> Temp;
-					m_extensionDataNodes[i].Name = Temp.c_str();;
+					m_extensionDataNodes[i].Name = Temp.c_str();
 					arch >> m_extensionDataNodes[i].Node;
 				}
 			}
diff --git a/Engine/Plugins/Experimental/Mutable/Source/MutableTools/Private/MuT/NodePrivate.h b/Engine/Plugins/Experimental/Mutable/Source/MutableTools/Private/MuT/NodePrivate.h
--- a/Engine/Plugins/Experimental/Mutable/Source/MutableTools/Private/MuT/NodePrivate.h
+++ b/Engine/Plugins/Experimental/Mutable/Source/MutableTools/Private/MuT/NodePrivate.h
@@ -1,10 +1,8 @@
 // Copyright Epic Games, Inc. All Rights Reserved.
 
-
 #pragma once
 
 #include "MuT/Node.h"
-#include "MuT/Visitor.h"
 #include "MuT/AST.h"
 #include "MuR/Operations.h"
 
@@ -12,18 +10,17 @@
 namespace mu
 {
 
-
-    //!
-    class Node::Private :
-        public Base,
-        public BaseVisitable<Ptr<class ASTOp>,DefaultCatchAll,true>
+    class Node::Private
     {
     public:
 
-        //! This is an opaque context used to attach to reported error messages.
+		/** Force a virtual destructor. */
+		virtual ~Private() = default; 
+
+        /** This is an opaque context used to attach to reported error messages. */
 		const void* m_errorContext = nullptr;
 
-		//! Generic pointer to the node owning this private.
+		/** Generic pointer to the node owning this private. */
 		const Node* m_pNode = nullptr;
     };
 
diff --git a/Engine/Plugins/Experimental/Mutable/Source/MutableTools/Private/MuT/NodeProjectorPrivate.h b/Engine/Plugins/Experimental/Mutable/Source/MutableTools/Private/MuT/NodeProjectorPrivate.h
--- a/Engine/Plugins/Experimental/Mutable/Source/MutableTools/Private/MuT/NodeProjectorPrivate.h
+++ b/Engine/Plugins/Experimental/Mutable/Source/MutableTools/Private/MuT/NodeProjectorPrivate.h
@@ -20,10 +20,6 @@ namespace mu
 
 	class NodeProjectorConstant::Private : public NodeProjector::Private
 	{
-	public:
-
-		MUTABLE_DEFINE_CONST_VISITABLE()
-
 	public:
 
 		static NODE_TYPE s_type;
@@ -69,10 +65,6 @@ namespace mu
 
 	class NodeProjectorParameter::Private : public NodeProjectorConstant::Private
 	{
-	public:
-
-		MUTABLE_DEFINE_CONST_VISITABLE()
-
 	public:
 
 		static NODE_TYPE s_type;
diff --git a/Engine/Plugins/Experimental/Mutable/Source/MutableTools/Private/MuT/NodeSurfaceVariationPrivate.h b/Engine/Plugins/Experimental/Mutable/Source/MutableTools/Private/MuT/NodeSurfaceVariationPrivate.h
--- a/Engine/Plugins/Experimental/Mutable/Source/MutableTools/Private/MuT/NodeSurfaceVariationPrivate.h
+++ b/Engine/Plugins/Experimental/Mutable/Source/MutableTools/Private/MuT/NodeSurfaceVariationPrivate.h
@@ -18,11 +18,6 @@ namespace mu
 	{
 	public:
 
-		MUTABLE_DEFINE_CONST_VISITABLE()
-
-	public:
-
-
 		static NODE_TYPE s_type;
 
         TArray<NodeSurfacePtr> m_defaultSurfaces;
diff --git a/Engine/Plugins/Experimental/Mutable/Source/MutableTools/Private/MuT/Streams.cpp b/Engine/Plugins/Experimental/Mutable/Source/MutableTools/Private/MuT/Streams.cpp
--- a/Engine/Plugins/Experimental/Mutable/Source/MutableTools/Private/MuT/Streams.cpp
+++ b/Engine/Plugins/Experimental/Mutable/Source/MutableTools/Private/MuT/Streams.cpp
@@ -269,7 +269,7 @@ namespace mu
     //---------------------------------------------------------------------------------------------
     //---------------------------------------------------------------------------------------------
     //---------------------------------------------------------------------------------------------
-    class ProxyFactoryMutableSourceFile::Private : public Base
+    class ProxyFactoryMutableSourceFile::Private
     {
     public:
         InputFileStream* Stream = nullptr;
diff --git a/Engine/Plugins/Experimental/Mutable/Source/MutableTools/Private/MuT/StreamsPrivate.h b/Engine/Plugins/Experimental/Mutable/Source/MutableTools/Private/MuT/StreamsPrivate.h
--- a/Engine/Plugins/Experimental/Mutable/Source/MutableTools/Private/MuT/StreamsPrivate.h
+++ b/Engine/Plugins/Experimental/Mutable/Source/MutableTools/Private/MuT/StreamsPrivate.h
@@ -17,7 +17,7 @@ namespace mu
 
 
 	//---------------------------------------------------------------------------------------------
-	class MUTABLETOOLS_API InputFileStream::Private : public Base
+	class MUTABLETOOLS_API InputFileStream::Private
 	{
 	public:
 
@@ -36,7 +36,7 @@ namespace mu
 
 
 	//---------------------------------------------------------------------------------------------
-	class MUTABLETOOLS_API OutputFileStream::Private : public Base
+	class MUTABLETOOLS_API OutputFileStream::Private
 	{
 	public:
 
diff --git a/Engine/Plugins/Experimental/Mutable/Source/MutableTools/Private/MuT/TablePrivate.h b/Engine/Plugins/Experimental/Mutable/Source/MutableTools/Private/MuT/TablePrivate.h
--- a/Engine/Plugins/Experimental/Mutable/Source/MutableTools/Private/MuT/TablePrivate.h
+++ b/Engine/Plugins/Experimental/Mutable/Source/MutableTools/Private/MuT/TablePrivate.h
@@ -4,12 +4,11 @@
 
 #include "MuT/Table.h"
 
-#include "MuT/Visitor.h"
 #include "MuT/ErrorLogPrivate.h"
-
 #include "MuR/MutableMath.h"
 #include "MuR/SerialisationPrivate.h"
 
+#include <string>
 
 namespace mu
 {
@@ -45,7 +44,7 @@ namespace mu
 
 
 	//!
-	class Table::Private : public Base
+	class Table::Private
 	{
 	public:
 
@@ -114,7 +113,7 @@ namespace mu
 			{
 				if (ver <= 2)
 				{
-					string LegacyString;
+					std::string LegacyString;
 					arch >> LegacyString;
 					Columns[c].Name = LegacyString.c_str();
 				}
@@ -181,7 +180,7 @@ namespace mu
 					{
 						if (ver <= 2)
 						{
-							string LegacyString;
+							std::string LegacyString;
 							arch >> LegacyString;
 							v.String = LegacyString.c_str();
 						}
diff --git a/Engine/Plugins/Experimental/Mutable/Source/MutableTools/Public/MuT/NodeBool.h b/Engine/Plugins/Experimental/Mutable/Source/MutableTools/Public/MuT/NodeBool.h
--- a/Engine/Plugins/Experimental/Mutable/Source/MutableTools/Public/MuT/NodeBool.h
+++ b/Engine/Plugins/Experimental/Mutable/Source/MutableTools/Public/MuT/NodeBool.h
@@ -25,7 +25,7 @@ namespace mu
 		{
 			Constant = 0,
 			Parameter = 1,
-			IsNull = 2,
+			DEPRECATED_IsNull = 2,
 			Not = 3,
 			And = 4,
 
@@ -179,56 +179,6 @@ namespace mu
 	};
 
 
-	//---------------------------------------------------------------------------------------------
-	//! Node that returns true if there is an input or false if there is nothing connected.
-	//! This node is mostly useful when used in the model Transform's conditions.
-	//! \ingroup model
-	//---------------------------------------------------------------------------------------------
-	class MUTABLETOOLS_API NodeBoolIsNull : public NodeBool
-	{
-	public:
-
-		//-----------------------------------------------------------------------------------------
-		// Life cycle
-		//-----------------------------------------------------------------------------------------
-
-		NodeBoolIsNull();
-
-		void SerialiseWrapper(OutputArchive& arch) const override;
-		static void Serialise( const NodeBoolIsNull* pNode, OutputArchive& arch );
-		static Ptr<NodeBoolIsNull> StaticUnserialise( InputArchive& arch );
-
-
-		//-----------------------------------------------------------------------------------------
-		// Node Interface
-		//-----------------------------------------------------------------------------------------
-
-        const NODE_TYPE* GetType() const override;
-		static const NODE_TYPE* GetStaticType();
-
-		//-----------------------------------------------------------------------------------------
-		// Own Interface
-		//-----------------------------------------------------------------------------------------
-
-		//-----------------------------------------------------------------------------------------
-		// Interface pattern
-		//-----------------------------------------------------------------------------------------
-		class Private;
-		Private* GetPrivate() const;
-        Node::Private* GetBasePrivate() const override;
-
-	protected:
-
-		//! Forbidden. 
-		~NodeBoolIsNull();
-
-	private:
-
-		Private* m_pD;
-
-	};
-
-
 	//---------------------------------------------------------------------------------------------
 	//! Node that returns the oposite of the input value.
 	//! \ingroup model
diff --git a/Engine/Plugins/Experimental/Mutable/Source/MutableTools/Public/MuT/NodeImageNormalComposite.h b/Engine/Plugins/Experimental/Mutable/Source/MutableTools/Public/MuT/NodeImageNormalComposite.h
--- a/Engine/Plugins/Experimental/Mutable/Source/MutableTools/Public/MuT/NodeImageNormalComposite.h
+++ b/Engine/Plugins/Experimental/Mutable/Source/MutableTools/Public/MuT/NodeImageNormalComposite.h
@@ -23,12 +23,6 @@ namespace mu
 	{
 	public:
 
-		//MUTABLE_DEFINE_CONST_VISITABLE()
-		
-		//-----------------------------------------------------------------------------------------
-		// Life cycle
-		//-----------------------------------------------------------------------------------------
-
 		NodeImageNormalComposite();
 
 		void SerialiseWrapper(OutputArchive& arch) const override;
