[Mutable] Commandlet wait for MIPs to be streamed.

#rnx

diff --git a/Engine/Plugins/Experimental/Mutable/Source/MutableValidation/Private/MuV/CustomizableObjectValidationCommandlet.cpp b/Engine/Plugins/Experimental/Mutable/Source/MutableValidation/Private/MuV/CustomizableObjectValidationCommandlet.cpp
--- a/Engine/Plugins/Experimental/Mutable/Source/MutableValidation/Private/MuV/CustomizableObjectValidationCommandlet.cpp
+++ b/Engine/Plugins/Experimental/Mutable/Source/MutableValidation/Private/MuV/CustomizableObjectValidationCommandlet.cpp
@@ -2,6 +2,7 @@
 
 #include "MuV/CustomizableObjectValidationCommandlet.h"
 
+#include "Components/SkeletalMeshComponent.h"
 #include "AssetRegistry/AssetRegistryModule.h"
 #include "Containers/Array.h"
 #include "MuCO/CustomizableObject.h"
@@ -162,6 +163,32 @@ int32 UCustomizableObjectValidationCommandlet::Main(const FString& Params)
             	// Wait until current instance turns invalid
             	if (InstanceBeingUpdated)
             	{
+            		// Wait until all MIPs gets streamed
+            		if (!ComponentsBeingUpdated.IsEmpty())
+            		{
+            			bool bFullyStreamed = true;
+            			for (auto It = ComponentsBeingUpdated.CreateIterator(); It && bFullyStreamed; ++It)
+            			{
+            				TObjectPtr<USkeletalMeshComponent>& ComponentBeingUpdated = *It;
+            		
+            				FStreamingTextureLevelContext LevelContext(EMaterialQualityLevel::Num, ComponentBeingUpdated);
+            				TArray<FStreamingRenderAssetPrimitiveInfo> RenderAssetInfoArray;
+            				ComponentBeingUpdated->GetStreamingRenderAssetInfo(LevelContext, RenderAssetInfoArray);
+
+            				for (auto ItAsset = RenderAssetInfoArray.CreateIterator(); ItAsset && bFullyStreamed; ++ItAsset)
+            				{
+            					bFullyStreamed = ItAsset->RenderAsset->IsFullyStreamedIn();
+            				}
+            			}
+
+            			if (bFullyStreamed)
+            			{
+            				UE_LOG(LogMutable,Display,TEXT("Instance %s finished streaming all MIPs."), *InstanceBeingUpdated->GetName());
+            				ComponentsBeingUpdated.Reset();
+            				InstanceBeingUpdated = nullptr;
+            			}
+            		}
+            		
             		continue;
             	}
            
@@ -176,6 +203,8 @@ int32 UCustomizableObjectValidationCommandlet::Main(const FString& Params)
             	if (InstanceBeingUpdated)
             	{
             		UE_LOG(LogMutable,Display,TEXT("Invoking update for %s instance."),*InstanceBeingUpdated->GetName());
+            		// Instance update delegate
+            		FInstanceUpdateDelegate InstanceUpdateDelegate;
             		InstanceUpdateDelegate.BindDynamic(this, &UCustomizableObjectValidationCommandlet::OnInstanceUpdate);
             		InstanceBeingUpdated->UpdateSkeletalMeshAsyncResult(InstanceUpdateDelegate);
             	}
@@ -223,17 +252,38 @@ void UCustomizableObjectValidationCommandlet::OnInstanceUpdate(const FUpdateCont
 	if (InstanceUpdateResult == EUpdateResult::Success)
 	{
 		UE_LOG(LogMutable,Display,TEXT("Instance %s finished update succesfully."),*InstanceName);
+
+		// Request load all MIPs
+		UE_LOG(LogMutable,Display,TEXT("Instance %s rquesting streaming all MIPs."), *InstanceBeingUpdated->GetName());
+
+		check(ComponentsBeingUpdated.IsEmpty());
+		for (int32 Index = 0; Index < InstanceBeingUpdated->GetNumComponents(); ++Index)
+		{
+			USkeletalMeshComponent* SkeletalComponent = NewObject<USkeletalMeshComponent>();
+			SkeletalComponent->SetSkeletalMesh(InstanceBeingUpdated->GetSkeletalMesh(Index));
+            
+			ComponentsBeingUpdated.Add(SkeletalComponent);            			
+		}
+		
+		for (TObjectPtr<USkeletalMeshComponent>& ComponentBeingUpdated : ComponentsBeingUpdated)
+		{
+			FStreamingTextureLevelContext LevelContext(EMaterialQualityLevel::Num, ComponentBeingUpdated);
+			TArray<FStreamingRenderAssetPrimitiveInfo> RenderAssetInfoArray;
+			ComponentBeingUpdated->GetStreamingRenderAssetInfo(LevelContext, RenderAssetInfoArray);
+
+			for (const FStreamingRenderAssetPrimitiveInfo& Info : RenderAssetInfoArray)
+			{
+				Info.RenderAsset->StreamIn(MAX_int32, true);
+			}
+		}
 	}
 	else
 	{
 		const FString OutputStatus = UEnum::GetValueAsString(Result.UpdateResult);
-		UE_LOG(LogMutable,Error,TEXT("Instance %s finished update with anomalous state : %s."),*InstanceName, *OutputStatus);
+		UE_LOG(LogMutable,Error,TEXT("Instance %s finished update with anomalous state : %s."), *InstanceName, *OutputStatus);
 		bInstanceFailedUpdate = true;
-	}
-	
-	// TODO: Wait for the mips before starting the update of another instance. 
-	
-	InstanceUpdateDelegate.Unbind();
-	InstanceBeingUpdated = nullptr;
+
+		InstanceBeingUpdated = nullptr;
+	}	
 }
 
diff --git a/Engine/Plugins/Experimental/Mutable/Source/MutableValidation/Public/MuV/CustomizableObjectValidationCommandlet.h b/Engine/Plugins/Experimental/Mutable/Source/MutableValidation/Public/MuV/CustomizableObjectValidationCommandlet.h
--- a/Engine/Plugins/Experimental/Mutable/Source/MutableValidation/Public/MuV/CustomizableObjectValidationCommandlet.h
+++ b/Engine/Plugins/Experimental/Mutable/Source/MutableValidation/Public/MuV/CustomizableObjectValidationCommandlet.h
@@ -19,28 +19,28 @@ public:
 	virtual int32 Main(const FString& Params) override;
 	
 private:
-
 	/** Customizable Object to be tested */
 	UPROPERTY()
-	UCustomizableObject* ToTestCustomizableObject = nullptr;
+	TObjectPtr<UCustomizableObject> ToTestCustomizableObject = nullptr;
 	
-	/** Customizable Object Instance currently being updated*/
+	/** Customizable Object Instance currently being updated */
+	UPROPERTY()
+	TObjectPtr<UCustomizableObjectInstance> InstanceBeingUpdated = nullptr;
+
+	/** Customizable Skeletal Component currently being updated */
 	UPROPERTY()
-	UCustomizableObjectInstance* InstanceBeingUpdated = nullptr;
+	TArray<TObjectPtr<USkeletalMeshComponent>> ComponentsBeingUpdated;
 
-	/** Array of COI to be generated with randomized parameter values*/
+	/** Array of COI to be generated with randomized parameter values */
 	UPROPERTY()
-	TArray<UCustomizableObjectInstance*> InstancesToProcess;
+	TArray<TObjectPtr<UCustomizableObjectInstance>> InstancesToProcess;
 	
-	/** Handle to be able to unbind OnInstanceUpdated(...) from instance end of update delegate.*/
+	/** Handle to be able to unbind OnInstanceUpdated(...) from instance end of update delegate. */
 	FDelegateHandle OnInstanceUpdateHandle;
 
-	/** Did any of the instances fail the UpdateSkeletalMesh process?*/
+	/** Did any of the instances fail the UpdateSkeletalMesh process? */
 	bool bInstanceFailedUpdate = false;
 
-	// Instance update delegate
-	FInstanceUpdateDelegate InstanceUpdateDelegate;
-
 	/** Callback invoked once the currently updating instance has done the update.
 	 * @paragm Result is a container that provides us with data related with the instance updating process.
 	 */
