[Mutable] Fix crash when closing the editor with Mutable Benchmark enabled.

#rnx

diff --git a/Engine/Plugins/Experimental/Mutable/Source/CustomizableObject/Private/MuCO/CustomizableObjectMipDataProvider.cpp b/Engine/Plugins/Experimental/Mutable/Source/CustomizableObject/Private/MuCO/CustomizableObjectMipDataProvider.cpp
--- a/Engine/Plugins/Experimental/Mutable/Source/CustomizableObject/Private/MuCO/CustomizableObjectMipDataProvider.cpp
+++ b/Engine/Plugins/Experimental/Mutable/Source/CustomizableObject/Private/MuCO/CustomizableObjectMipDataProvider.cpp
@@ -248,6 +248,11 @@ namespace impl
 			FFunctionGraphTask::CreateAndDispatchWhenReady(
 			[CustomizableObjectPathName, InstancePathName, Time, PeakMemory, RealMemoryPeak]()
 			{
+				if (!UCustomizableObjectSystem::IsCreated()) // We are shutting down
+				{
+					return;	
+				}
+				
 				UCustomizableObjectSystem* System = UCustomizableObjectSystem::GetInstance();
 				if (!System)
 				{
