[mutable] Fix comparison bug detected by static analysis.
#rnx

diff --git a/Engine/Plugins/Experimental/Mutable/Source/MutableTools/Private/MuT/CodeGenerator.h b/Engine/Plugins/Experimental/Mutable/Source/MutableTools/Private/MuT/CodeGenerator.h
--- a/Engine/Plugins/Experimental/Mutable/Source/MutableTools/Private/MuT/CodeGenerator.h
+++ b/Engine/Plugins/Experimental/Mutable/Source/MutableTools/Private/MuT/CodeGenerator.h
@@ -251,7 +251,7 @@ namespace mu
 			{
 				return m_pObject == Other.m_pObject
 					&&
-					m_lod == m_lod;
+					m_lod == Other.m_lod;
 			}
 
 			friend FORCEINLINE uint32 GetTypeHash(const FAdditionalComponentKey& InKey)
