[mutable] Remove unnecessary localization macro undef.
#rnx

diff --git a/Engine/Plugins/Experimental/Mutable/Source/CustomizableObjectEditor/Private/MuCOE/SMutableCodeViewer.cpp b/Engine/Plugins/Experimental/Mutable/Source/CustomizableObjectEditor/Private/MuCOE/SMutableCodeViewer.cpp
--- a/Engine/Plugins/Experimental/Mutable/Source/CustomizableObjectEditor/Private/MuCOE/SMutableCodeViewer.cpp
+++ b/Engine/Plugins/Experimental/Mutable/Source/CustomizableObjectEditor/Private/MuCOE/SMutableCodeViewer.cpp
@@ -3182,7 +3182,4 @@ int32 FMutableCodeTreeElement::GetStateIndex() const
 }
 
 
-#undef LOCTEXT_NAMESPACE
-
-
 #undef LOCTEXT_NAMESPACE
