[Mutable] Fix UCustomizableObjectInstanceUsage GC  crash.

#rnx

diff --git a/Engine/Plugins/Experimental/Mutable/Source/CustomizableObject/Classes/MuCO/CustomizableObjectInstanceUsage.h b/Engine/Plugins/Experimental/Mutable/Source/CustomizableObject/Classes/MuCO/CustomizableObjectInstanceUsage.h
--- a/Engine/Plugins/Experimental/Mutable/Source/CustomizableObject/Classes/MuCO/CustomizableObjectInstanceUsage.h
+++ b/Engine/Plugins/Experimental/Mutable/Source/CustomizableObject/Classes/MuCO/CustomizableObjectInstanceUsage.h
@@ -111,6 +111,7 @@ private:
 	virtual TStatId GetStatId() const override;
 	virtual bool IsTickableWhenPaused() const override { return true; }
 	virtual bool IsTickableInEditor() const override { return true; }
+	virtual bool IsTickable() const override;
 	// End FTickableGameObject
 
 	friend UCustomizableSkeletalComponent;
diff --git a/Engine/Plugins/Experimental/Mutable/Source/CustomizableObject/Private/MuCO/CustomizableObjectInstanceUsage.cpp b/Engine/Plugins/Experimental/Mutable/Source/CustomizableObject/Private/MuCO/CustomizableObjectInstanceUsage.cpp
--- a/Engine/Plugins/Experimental/Mutable/Source/CustomizableObject/Private/MuCO/CustomizableObjectInstanceUsage.cpp
+++ b/Engine/Plugins/Experimental/Mutable/Source/CustomizableObject/Private/MuCO/CustomizableObjectInstanceUsage.cpp
@@ -584,6 +584,12 @@ TStatId UCustomizableObjectInstanceUsage::GetStatId() const
 }
 
 
+bool UCustomizableObjectInstanceUsage::IsTickable() const
+{
+	return !HasAnyFlags(RF_BeginDestroyed);
+}
+
+
 bool UCustomizableObjectInstanceUsage::IsNetMode(ENetMode InNetMode) const
 {
 	if (CustomizableSkeletalComponent)
