[Mutable] Fix assert if a passthrough texture is missing from the array of referenced passthrough textures.

#rnx

diff --git a/Engine/Plugins/Experimental/Mutable/Source/CustomizableObject/Private/MuCO/CustomizableObjectInstance.cpp b/Engine/Plugins/Experimental/Mutable/Source/CustomizableObject/Private/MuCO/CustomizableObjectInstance.cpp
--- a/Engine/Plugins/Experimental/Mutable/Source/CustomizableObject/Private/MuCO/CustomizableObjectInstance.cpp
+++ b/Engine/Plugins/Experimental/Mutable/Source/CustomizableObject/Private/MuCO/CustomizableObjectInstance.cpp
@@ -5623,10 +5623,12 @@ void UCustomizableInstancePrivateData::BuildMaterials(const TSharedRef<FUpdateCo
 									PassThroughTexture = Ref.LoadSynchronous();
 								}
 							}
-							else
+
+							if (!PassThroughTexture)
 							{
-								// internal error.
-								UE_LOG(LogMutable, Error, TEXT("Referenced image [%d] was not stored in the resource array."), ReferenceID);
+								// Internal error.
+								UE_LOG(LogMutable, Error, TEXT("Missing referenced image [%d]."), ReferenceID);
+								continue;
 							}
 						}
 
