[Mutable] Use the default material in sections without one.
- Fix unknown RenderSection.MaterialIndex in the UCustomizableInstancePrivateData::RegenerateImportedModel method because of an empty LODMaterialMap.
- RegenerateImportedModel and InitResources once per SkeletalMesh instead of with every update.

#rnx

diff --git a/Engine/Plugins/Experimental/Mutable/Source/CustomizableObject/Private/MuCO/CustomizableObjectInstance.cpp b/Engine/Plugins/Experimental/Mutable/Source/CustomizableObject/Private/MuCO/CustomizableObjectInstance.cpp
--- a/Engine/Plugins/Experimental/Mutable/Source/CustomizableObject/Private/MuCO/CustomizableObjectInstance.cpp
+++ b/Engine/Plugins/Experimental/Mutable/Source/CustomizableObject/Private/MuCO/CustomizableObjectInstance.cpp
@@ -3575,6 +3575,10 @@ void UCustomizableInstancePrivateData::BuildMeshSockets(const TSharedRef<FUpdate
 			}
 		}
 	}
+
+#if !WITH_EDITOR
+	SkeletalMesh->RebuildSocketMap();
+#endif // !WITH_EDITOR
 }
 
 
@@ -5473,8 +5477,10 @@ void UCustomizableInstancePrivateData::BuildMaterials(const TSharedRef<FUpdateCo
 				UMaterialInterface* MaterialTemplate = ReferencedMaterials[ReferencedMaterialIndex];
 				if (!MaterialTemplate)
 				{
+					// Missing MaterialTemplate. Use DefaultMaterial instead. 
+					MaterialTemplate = UMaterial::GetDefaultMaterial(MD_Surface);
+					check(MaterialTemplate);
 					UE_LOG(LogMutable, Error, TEXT("Build Materials: Missing referenced template to use as parent material on CustomizableObject [%s]."), *CustomizableObject->GetName());
-					continue;
 				}
 
 				// This section will require a new slot
@@ -5484,7 +5490,7 @@ void UCustomizableInstancePrivateData::BuildMaterials(const TSharedRef<FUpdateCo
 				const int32 MaterialSlotIndex = Materials.Num();
 				FSkeletalMaterial& MaterialSlot = Materials.AddDefaulted_GetRef();
 				MaterialSlot.MaterialInterface = MaterialTemplate;
-				MaterialSlot.MaterialSlotName = CustomizableObject->ReferencedMaterialSlotNames[Surface.MaterialIndex];
+				MaterialSlot.MaterialSlotName = CustomizableObject->ReferencedMaterialSlotNames.IsValidIndex(Surface.MaterialIndex) ? CustomizableObject->ReferencedMaterialSlotNames[Surface.MaterialIndex] : NAME_None;
 				SetMeshUVChannelDensity(MaterialSlot.UVChannelData, RefSkeletalMeshData->Settings.DefaultUVChannelDensity);
 
 				const int32 LODMaterialIndex = SkeletalMesh->GetLODInfoArray()[LODIndex].LODMaterialMap.Add(MaterialSlotIndex);
diff --git a/Engine/Plugins/Experimental/Mutable/Source/CustomizableObject/Private/MuCO/CustomizableObjectSystem.cpp b/Engine/Plugins/Experimental/Mutable/Source/CustomizableObject/Private/MuCO/CustomizableObjectSystem.cpp
--- a/Engine/Plugins/Experimental/Mutable/Source/CustomizableObject/Private/MuCO/CustomizableObjectSystem.cpp
+++ b/Engine/Plugins/Experimental/Mutable/Source/CustomizableObject/Private/MuCO/CustomizableObjectSystem.cpp
@@ -884,18 +884,6 @@ void UpdateSkeletalMesh(const TSharedRef<FUpdateContextPrivate>& Context)
 
 	UCustomizableObjectInstance* CustomizableObjectInstance = Context->Instance.Get();
 	check(CustomizableObjectInstance);
-	
-	for (int32 ComponentIndex = 0; ComponentIndex < CustomizableObjectInstance->SkeletalMeshes.Num(); ++ComponentIndex)
-	{
-		if (TObjectPtr<USkeletalMesh> SkeletalMesh = CustomizableObjectInstance->SkeletalMeshes[ComponentIndex])
-		{
-#if WITH_EDITOR
-			UCustomizableInstancePrivateData::RegenerateImportedModel(SkeletalMesh);
-#else
-			SkeletalMesh->RebuildSocketMap();
-#endif
-		}
-	}
 
 	UCustomizableInstancePrivateData* CustomizableObjectInstancePrivateData = CustomizableObjectInstance->GetPrivate();
 	check(CustomizableObjectInstancePrivateData != nullptr);
@@ -2347,13 +2335,14 @@ namespace impl
 				{
 					MUTABLE_CPUPROFILER_SCOPE(UpdateSkeletalMesh_PostBeginUpdate2);
 
-					for (int32 Component = 0; Component < CustomizableObjectInstance->SkeletalMeshes.Num(); ++Component)
+					for (TObjectPtr<USkeletalMesh>& SkeletalMesh : CustomizableObjectInstance->SkeletalMeshes)
 					{
-						if (CustomizableObjectInstance->SkeletalMeshes[Component] && CustomizableObjectInstance->SkeletalMeshes[Component]->GetLODInfoArray().Num())
+						if (SkeletalMesh && SkeletalMesh->GetResourceForRendering() && !SkeletalMesh->GetResourceForRendering()->IsInitialized())
 						{
-							MUTABLE_CPUPROFILER_SCOPE(UpdateSkeletalMesh_PostEditChangeProperty);
-
-							CustomizableInstancePrivateData->PostEditChangePropertyWithoutEditor(CustomizableObjectInstance->SkeletalMeshes[Component]);
+#if WITH_EDITOR
+							UCustomizableInstancePrivateData::RegenerateImportedModel(SkeletalMesh);
+#endif
+							CustomizableInstancePrivateData->PostEditChangePropertyWithoutEditor(SkeletalMesh);
 						}
 					}
 				}
