[mutable] Cleaner handling of flags derived from compile options.
Removed deprecated texture compression property.
Added in-game warning when using unoptimized COs.
#rnx

diff --git a/Engine/Plugins/Experimental/Mutable/Source/CustomizableObject/Classes/MuCO/CustomizableObject.h b/Engine/Plugins/Experimental/Mutable/Source/CustomizableObject/Classes/MuCO/CustomizableObject.h
--- a/Engine/Plugins/Experimental/Mutable/Source/CustomizableObject/Classes/MuCO/CustomizableObject.h
+++ b/Engine/Plugins/Experimental/Mutable/Source/CustomizableObject/Classes/MuCO/CustomizableObject.h
@@ -11,14 +11,12 @@
 #include "MuCO/CustomizableObjectParameterTypeDefinitions.h"
 #include "MuCO/CustomizableObjectUIData.h"
 
-class FReply;
-
-#if WITH_EDITORONLY_DATA
-#endif
-
-
 #include "CustomizableObject.generated.h"
 
+/** Index of the maximum optimization level when compiling CustomizableObjects */
+#define UE_MUTABLE_MAX_OPTIMIZATION			2
+
+class FReply;
 class FMemoryReaderView;
 class FMemoryWriter64;
 class FObjectPreSaveContext;
@@ -168,18 +166,14 @@ USTRUCT()
 struct FCompilationOptions
 {
 	GENERATED_USTRUCT_BODY()
-
-	// Flag to know if texture compression should be enabled
-	UPROPERTY()
-	bool bTextureCompression_DEPRECATED = true;
 	
 	// Enum to know what texture compression should be used
 	UPROPERTY()
 	ECustomizableObjectTextureCompression TextureCompression = ECustomizableObjectTextureCompression::Fast;
 
-	// From 0 to 2
+	// From 0 to UE_MUTABLE_MAX_OPTIMIZATION
 	UPROPERTY()
-	int32 OptimizationLevel = 1;
+	int32 OptimizationLevel = UE_MUTABLE_MAX_OPTIMIZATION;
 
 	// Use the disk to store intermediate compilation data. This slows down the object compilation
 	// but it may be necessary for huge objects.
@@ -1121,6 +1115,12 @@ public:
 	UPROPERTY(VisibleAnywhere, Category = CustomizableObject)
 	TArray<FMutableModelImageProperties> ImageProperties;
 
+	/** This is a non-user-controlled flag to disable streaming (set at object compilation time, depending on optimization). */
+	bool bDisableTextureStreaming = false;
+
+	/** If the object is compiled, this flag is true if it was compiled with maximum optimizations. If the object is not compiled, its value is meaningless. */
+	bool bIsCompiledWithOptimization = true;
+
 	UPROPERTY(Transient)
 	TArray<FMorphTargetInfo> ContributingMorphTargetsInfo;
 	
@@ -1397,7 +1397,7 @@ private:
 	// This is a manual version number for the binary blobs in this asset.
 	// Increasing it invalidates all the previously compiled models.
 	// Warning: If while merging code both versions have changed, take the highest+1.
-	static const int32 CurrentSupportedVersion = 414;
+	static const int32 CurrentSupportedVersion = 415;
 
 	// Compile the object for a specific platform - Compile for Cook Customizable Object
 	void CompileForTargetPlatform(const ITargetPlatform* TargetPlatform);
diff --git a/Engine/Plugins/Experimental/Mutable/Source/CustomizableObject/Private/MuCO/CustomizableObject.cpp b/Engine/Plugins/Experimental/Mutable/Source/CustomizableObject/Private/MuCO/CustomizableObject.cpp
--- a/Engine/Plugins/Experimental/Mutable/Source/CustomizableObject/Private/MuCO/CustomizableObject.cpp
+++ b/Engine/Plugins/Experimental/Mutable/Source/CustomizableObject/Private/MuCO/CustomizableObject.cpp
@@ -223,15 +223,6 @@ void UCustomizableObject::PostLoad()
 
 	const int32 CustomizableObjectCustomVersion = GetLinkerCustomVersion(FCustomizableObjectCustomVersion::GUID);
 
-	// Convert texture compression option
-	if (CustomizableObjectCustomVersion < FCustomizableObjectCustomVersion::TextureCompressionEnum)
-	{
-		CompileOptions.TextureCompression = CompileOptions.bTextureCompression_DEPRECATED
-			? ECustomizableObjectTextureCompression::Fast
-			: ECustomizableObjectTextureCompression::None
-			;
-	}
-
 	// Update state never-stream flag from deprecated enum
 	if (CustomizableObjectCustomVersion < FCustomizableObjectCustomVersion::CustomizableObjectStateHasSeparateNeverStreamFlag)
 	{
@@ -825,7 +816,7 @@ void UCustomizableObject::CompileForTargetPlatform(const ITargetPlatform* Target
 	if (bIsRootObject && bIsRelevantForThisTarget)
 	{
 		FCompilationOptions Options;
-		Options.OptimizationLevel = 2;	// max optimization when packaging.
+		Options.OptimizationLevel = UE_MUTABLE_MAX_OPTIMIZATION;	// max optimization when packaging.
 		Options.TextureCompression = ECustomizableObjectTextureCompression::HighQuality;
 		Options.bIsCooking = true;
 		Options.TargetPlatform = TargetPlatform;
@@ -851,6 +842,13 @@ bool UCustomizableObject::ConditionalAutoCompile()
 	// Don't compile compiled objects
 	if (IsCompiled())
 	{
+		// Show a warning if the compilation was not done with optimizations.
+		if (bIsCompiledWithOptimization)
+		{
+			FString Msg = FString::Printf(TEXT("Warning: Customizable Object [%s] was compiled without optimization."), *GetName());
+			GEngine->AddOnScreenDebugMessage((uint64)((PTRINT)this), 10.0f, FColor::Red, Msg);
+		}
+
 		return true;
 	}
 
@@ -982,6 +980,10 @@ void UCustomizableObject::SaveEmbeddedData(FArchive& Ar)
 
 	if (Private->GetModel())
 	{
+		// General derived flags
+		Ar << bDisableTextureStreaming;
+		Ar << bIsCompiledWithOptimization;
+
 		// Serialize morph data
 		{
 			Ar << ContributingMorphTargetsInfo;
@@ -1022,6 +1024,10 @@ void UCustomizableObject::LoadEmbeddedData(FArchive& Ar)
 
 	if(CurrentSupportedVersion == InternalVersion)
 	{
+		// General derived flags
+		Ar << bDisableTextureStreaming;
+		Ar << bIsCompiledWithOptimization;
+
 		// Load morph data
 		{
 			Ar << ContributingMorphTargetsInfo;
diff --git a/Engine/Plugins/Experimental/Mutable/Source/CustomizableObject/Private/MuCO/CustomizableObjectSystem.cpp b/Engine/Plugins/Experimental/Mutable/Source/CustomizableObject/Private/MuCO/CustomizableObjectSystem.cpp
--- a/Engine/Plugins/Experimental/Mutable/Source/CustomizableObject/Private/MuCO/CustomizableObjectSystem.cpp
+++ b/Engine/Plugins/Experimental/Mutable/Source/CustomizableObject/Private/MuCO/CustomizableObjectSystem.cpp
@@ -935,7 +935,7 @@ void FCustomizableObjectSystemPrivate::GetMipStreamingConfig(const UCustomizable
 		bOutNeverStream = State->bDisableTextureStreaming;
 
 		// Was streaming disabled at object-compilation time? 
-		if (State->bDisableTextureStreamingOverride)
+		if (Instance.GetCustomizableObject()->bDisableTextureStreaming)
 		{
 			bOutNeverStream = true;
 		}
diff --git a/Engine/Plugins/Experimental/Mutable/Source/CustomizableObject/Public/MuCO/CustomizableObjectUIData.h b/Engine/Plugins/Experimental/Mutable/Source/CustomizableObject/Public/MuCO/CustomizableObjectUIData.h
--- a/Engine/Plugins/Experimental/Mutable/Source/CustomizableObject/Public/MuCO/CustomizableObjectUIData.h
+++ b/Engine/Plugins/Experimental/Mutable/Source/CustomizableObject/Public/MuCO/CustomizableObjectUIData.h
@@ -262,9 +262,6 @@ struct FParameterUIData
 	UPROPERTY()
 	bool bDisableTextureStreaming = false;
 
-	/** This is a non-user-controlled flag to disable streaming at object compilation time. */
-	bool bDisableTextureStreamingOverride = false;
-
 	/** In this mode instances and their temp data will be reused between updates. It will be much faster but spend as much as ten times the memory.
 	    Useful for customization lockers with few characters that are going to have their parameters changed many times, not for in-game */
 	UPROPERTY(BlueprintReadWrite, Category = CustomizableObject)
@@ -282,7 +279,6 @@ struct FParameterUIData
 			|| ArrayIntegerParameterOption != Other.ArrayIntegerParameterOption || IntegerParameterGroupType != Other.IntegerParameterGroupType 
 			|| bLiveUpdateMode != Other.bLiveUpdateMode || bReuseInstanceTextures != Other.bReuseInstanceTextures
 			|| bDisableTextureStreaming != Other.bDisableTextureStreaming
-			|| bDisableTextureStreamingOverride != Other.bDisableTextureStreamingOverride
 			|| !ForcedParameterValues.OrderIndependentCompareEqual(Other.ForcedParameterValues)
 			)
 		{
@@ -301,7 +297,6 @@ struct FParameterUIData
 		Ar << UIData.IntegerParameterGroupType;
 		Ar << UIData.TextureCompressionStrategy;
 		Ar << UIData.bDisableTextureStreaming;
-		Ar << UIData.bDisableTextureStreamingOverride;
 		Ar << UIData.bLiveUpdateMode;
 		Ar << UIData.bReuseInstanceTextures;
 		Ar << UIData.ForcedParameterValues;
diff --git a/Engine/Plugins/Experimental/Mutable/Source/CustomizableObjectEditor/Private/MuCOE/CustomizableObjectCompiler.cpp b/Engine/Plugins/Experimental/Mutable/Source/CustomizableObjectEditor/Private/MuCOE/CustomizableObjectCompiler.cpp
--- a/Engine/Plugins/Experimental/Mutable/Source/CustomizableObjectEditor/Private/MuCOE/CustomizableObjectCompiler.cpp
+++ b/Engine/Plugins/Experimental/Mutable/Source/CustomizableObjectEditor/Private/MuCOE/CustomizableObjectCompiler.cpp
@@ -1057,22 +1057,14 @@ void FCustomizableObjectCompiler::CompileInternal(UCustomizableObject* Object, c
 		Object->ParameterUIDataMap = GenerationContext.ParameterUIDataMap;
 
 		Object->StateUIDataMap = GenerationContext.StateUIDataMap;
-		if (GenerationContext.Options.OptimizationLevel == 0 )
+		if (GenerationContext.Options.OptimizationLevel == 0)
 		{
 			// If the optimization level is "none" disable texture streaming, because textures are all referenced
 			// unreal assets and progressive generation is not supported.
-			for (TPair<FString, FParameterUIData>& StateData : Object->StateUIDataMap)
-			{
-				StateData.Value.bDisableTextureStreamingOverride = true;
-			}
-
-			// If there is no state data, we need to add the default one to be able to disable streaming.
-			if (Object->StateUIDataMap.IsEmpty())
-			{
-				FParameterUIData& Data = Object->StateUIDataMap.Add(TEXT("Default"));
-				Data.bDisableTextureStreamingOverride = true;
-			}
+			Object->bDisableTextureStreaming = true;
 		}
+		
+		Object->bIsCompiledWithOptimization = GenerationContext.Options.OptimizationLevel < UE_MUTABLE_MAX_OPTIMIZATION;
 
 		Object->AlwaysLoadedExtensionData = MoveTemp(GenerationContext.AlwaysLoadedExtensionData);
 
diff --git a/Engine/Plugins/Experimental/Mutable/Source/CustomizableObjectEditor/Private/MuCOE/CustomizableObjectEditor.cpp b/Engine/Plugins/Experimental/Mutable/Source/CustomizableObjectEditor/Private/MuCOE/CustomizableObjectEditor.cpp
--- a/Engine/Plugins/Experimental/Mutable/Source/CustomizableObjectEditor/Private/MuCOE/CustomizableObjectEditor.cpp
+++ b/Engine/Plugins/Experimental/Mutable/Source/CustomizableObjectEditor/Private/MuCOE/CustomizableObjectEditor.cpp
@@ -1311,9 +1311,10 @@ TSharedRef<SWidget> FCustomizableObjectEditor::GenerateCompileOptionsMenuContent
 	{
 		// Level
 		CompileOptimizationStrings.Empty();
-		CompileOptimizationStrings.Add(MakeShareable(new FString(LOCTEXT("OptimizationNone", "None").ToString())));
+		CompileOptimizationStrings.Add(MakeShareable(new FString(LOCTEXT("OptimizationNone", "None (Disable texture streaming)").ToString())));
 		CompileOptimizationStrings.Add(MakeShareable(new FString(LOCTEXT("OptimizationMin", "Minimal").ToString())));
 		CompileOptimizationStrings.Add(MakeShareable(new FString(LOCTEXT("OptimizationMax", "Maximum").ToString())));
+		check(CompileOptimizationStrings.Num() == UE_MUTABLE_MAX_OPTIMIZATION + 1);
 
 		if (CustomizableObject)
 		{
@@ -1371,7 +1372,6 @@ TSharedRef<SWidget> FCustomizableObjectEditor::GenerateCompileOptionsMenuContent
 		MenuBuilder.AddWidget(CompileTilingCombo.ToSharedRef(), LOCTEXT("MutableCompileImageTiling", "Image Tiling"));
 
 		MenuBuilder.AddMenuEntry(FCustomizableObjectEditorCommands::Get().CompileOptions_UseDiskCompilation);
-		MenuBuilder.AddMenuEntry(FCustomizableObjectEditorCommands::Get().CompileOptions_EnableTextureCompression);
 	}
 	MenuBuilder.EndSection();
 
diff --git a/Engine/Plugins/Experimental/Mutable/Source/CustomizableObjectEditor/Private/MuCOE/CustomizableObjectEditorActions.cpp b/Engine/Plugins/Experimental/Mutable/Source/CustomizableObjectEditor/Private/MuCOE/CustomizableObjectEditorActions.cpp
--- a/Engine/Plugins/Experimental/Mutable/Source/CustomizableObjectEditor/Private/MuCOE/CustomizableObjectEditorActions.cpp
+++ b/Engine/Plugins/Experimental/Mutable/Source/CustomizableObjectEditor/Private/MuCOE/CustomizableObjectEditorActions.cpp
@@ -24,7 +24,6 @@ void FCustomizableObjectEditorCommands::RegisterCommands()
 	UI_COMMAND(Compile, "Compile", "Compile the source graph of the customizable object and update the previews.", EUserInterfaceActionType::Button, FInputChord());
 	UI_COMMAND(CompileOnlySelected, "Compile Only Selected", "Compile the source graph of the customizable object and update the previews, only for the selected options in the preview. The rest of options are discarded. If they are selected, press again this button to see their effect in the preview.", EUserInterfaceActionType::Button, FInputChord());
 	UI_COMMAND(ResetCompileOptions, "Reset Compilation Options", "Set reasonable defaults for the compilation options.", EUserInterfaceActionType::Button, FInputChord());
-	UI_COMMAND(CompileOptions_EnableTextureCompression, "Enable texture compression.", "Only for debug. Do not use.", EUserInterfaceActionType::ToggleButton, FInputChord());
 	UI_COMMAND(CompileOptions_UseDiskCompilation, "Enable compiling using the disk as memory.", "This is very slow but supports compiling huge objects. It requires a lot of free space in the OS disk.", EUserInterfaceActionType::ToggleButton, FInputChord());
 	UI_COMMAND(Debug, "Debug", "Open the CustomizableObject debugger tab for this object.", EUserInterfaceActionType::Button, FInputChord());
 
diff --git a/Engine/Plugins/Experimental/Mutable/Source/CustomizableObjectEditor/Public/MuCOE/CustomizableObjectEditorActions.h b/Engine/Plugins/Experimental/Mutable/Source/CustomizableObjectEditor/Public/MuCOE/CustomizableObjectEditorActions.h
--- a/Engine/Plugins/Experimental/Mutable/Source/CustomizableObjectEditor/Public/MuCOE/CustomizableObjectEditorActions.h
+++ b/Engine/Plugins/Experimental/Mutable/Source/CustomizableObjectEditor/Public/MuCOE/CustomizableObjectEditorActions.h
@@ -18,7 +18,6 @@ public:
 	TSharedPtr< FUICommandInfo > Compile;
 	TSharedPtr< FUICommandInfo > CompileOnlySelected;
 	TSharedPtr< FUICommandInfo > ResetCompileOptions;
-	TSharedPtr< FUICommandInfo > CompileOptions_EnableTextureCompression;
 	TSharedPtr< FUICommandInfo > CompileOptions_UseDiskCompilation;
 	TSharedPtr< FUICommandInfo > Debug;
 
