[Mutable] Make CompilationGuid transient and editor only to fix cook indeterminism.

#rnx

diff --git a/Engine/Plugins/Experimental/Mutable/Source/CustomizableObject/Classes/MuCO/CustomizableObject.h b/Engine/Plugins/Experimental/Mutable/Source/CustomizableObject/Classes/MuCO/CustomizableObject.h
--- a/Engine/Plugins/Experimental/Mutable/Source/CustomizableObject/Classes/MuCO/CustomizableObject.h
+++ b/Engine/Plugins/Experimental/Mutable/Source/CustomizableObject/Classes/MuCO/CustomizableObject.h
@@ -1487,8 +1487,6 @@ public:
 
 	void Serialize(FArchive& Ar) override;
 
-	FGuid GetCompilationGuid() const;
-
 	int32 FindState( const FString& Name ) const;
 
 	/** Return the number of object states that are defined in the CustomizableObject. */
@@ -1588,15 +1586,22 @@ public:
 	UPROPERTY()
 	bool bIsChildObject = false;
 
-	/** Unique Identifier - used to locate Model and Streamable data on disk. Should not be modified. */
+	/** Unique Identifier - Deterministic. Used to locate Model and Streamable data on disk. Should not be modified. */
 	UPROPERTY(Transient)
 	FGuid Identifier;
 
+	/** Unique Identifier - Regenerated each time the object is compiled. */
+	UPROPERTY(Transient)
+	FGuid CompilationGuid;
+
 	ECustomizableObjectCompilationState CompilationState = ECustomizableObjectCompilationState::None;
 
 	FPostCompileDelegate PostCompileDelegate;
 
 	void PostCompile();
+
+	FGuid GetCompilationGuid() const;
+
 #endif
 
 	/** Create a new instance of this object. The instance parameters will be initialized with the object default values. */
@@ -1651,10 +1656,6 @@ private:
 	/** Cache of merged skeletons */
 	TArray<FMergedSkeleton> MergedSkeletons;
 
-	/** Unique identifier. Regenerated each time the object is compiled. */
-	UPROPERTY()
-	FGuid CompilationGuid;
-
 	/** BulkData that stores all in-game resources used by Mutable when generating instances.
 	  * Only valid in packaged builds */
 	UPROPERTY()
diff --git a/Engine/Plugins/Experimental/Mutable/Source/CustomizableObject/Private/MuCO/CustomizableObjectInstanceDescriptor.cpp b/Engine/Plugins/Experimental/Mutable/Source/CustomizableObject/Private/MuCO/CustomizableObjectInstanceDescriptor.cpp
--- a/Engine/Plugins/Experimental/Mutable/Source/CustomizableObject/Private/MuCO/CustomizableObjectInstanceDescriptor.cpp
+++ b/Engine/Plugins/Experimental/Mutable/Source/CustomizableObject/Private/MuCO/CustomizableObjectInstanceDescriptor.cpp
@@ -2574,11 +2574,13 @@ void FCustomizableObjectInstanceDescriptor::CreateParametersLookupTable()
 
 FDescriptorHash::FDescriptorHash(const FCustomizableObjectInstanceDescriptor& Descriptor)
 {
+#if WITH_EDITORONLY_DATA
 	if (Descriptor.CustomizableObject)
 	{
 		Hash = HashCombine(Hash, GetTypeHash(Descriptor.CustomizableObject->GetPathName()));
 		Hash = HashCombine(Hash, GetTypeHash(Descriptor.CustomizableObject->GetCompilationGuid()));
 	}
+#endif
 
 	for (const FCustomizableObjectBoolParameterValue& Value : Descriptor.BoolParameters)
 	{
