[mutable] Fix crash in degenerated case of object with no meshes.
#rnx

diff --git a/Engine/Plugins/Experimental/Mutable/Source/CustomizableObject/Private/MuCO/CustomizableObjectInstance.cpp b/Engine/Plugins/Experimental/Mutable/Source/CustomizableObject/Private/MuCO/CustomizableObjectInstance.cpp
--- a/Engine/Plugins/Experimental/Mutable/Source/CustomizableObject/Private/MuCO/CustomizableObjectInstance.cpp
+++ b/Engine/Plugins/Experimental/Mutable/Source/CustomizableObject/Private/MuCO/CustomizableObjectInstance.cpp
@@ -5396,6 +5396,12 @@ void UCustomizableInstancePrivateData::BuildMaterials(const TSharedRef<FUpdateCo
 
 	for (int32 ComponentIndex = 0; ComponentIndex < NumComponents; ++ComponentIndex)
 	{
+		if (!Public->SkeletalMeshes.IsValidIndex(ComponentIndex))
+		{
+			// This could happen if there are no meshes at all.
+			continue;
+		}
+
 		USkeletalMesh* SkeletalMesh = Public->SkeletalMeshes[ComponentIndex];
 
 		if (!SkeletalMesh)
