[mutable] Updated mutable documentation page URL

diff --git a/Engine/Plugins/Experimental/Mutable/Source/CustomizableObjectEditor/Private/MuCOE/CustomizableObjectEditor.h b/Engine/Plugins/Experimental/Mutable/Source/CustomizableObjectEditor/Private/MuCOE/CustomizableObjectEditor.h
--- a/Engine/Plugins/Experimental/Mutable/Source/CustomizableObjectEditor/Private/MuCOE/CustomizableObjectEditor.h
+++ b/Engine/Plugins/Experimental/Mutable/Source/CustomizableObjectEditor/Private/MuCOE/CustomizableObjectEditor.h
@@ -448,7 +448,7 @@ private:
 	void ExtendToolbar();
 
 	/** URL to open when pressing the documentation button generated by UE */
-	const FString DocumentationURL{ TEXT("https://work.anticto.com/w/mutable/unreal-engine-4/") };
+	const FString DocumentationURL{ TEXT("[REDACTED_PROFILE_URL]/Mutable-Documentation/wiki") };
 
 	/** Postponed work to do when OnUpdatePreviewInstance is called. Emptied at the end on each OnUpdatePreviewInstance call. */
 	TArray<TFunction<void()>> OnUpdatePreviewInstanceWork;
