[Mutable] Unload roms after compiling a CO.

#rnx

diff --git a/Engine/Plugins/Experimental/Mutable/Source/CustomizableObjectEditor/Private/MuCOE/CustomizableObjectCompiler.cpp b/Engine/Plugins/Experimental/Mutable/Source/CustomizableObjectEditor/Private/MuCOE/CustomizableObjectCompiler.cpp
--- a/Engine/Plugins/Experimental/Mutable/Source/CustomizableObjectEditor/Private/MuCOE/CustomizableObjectCompiler.cpp
+++ b/Engine/Plugins/Experimental/Mutable/Source/CustomizableObjectEditor/Private/MuCOE/CustomizableObjectCompiler.cpp
@@ -24,6 +24,8 @@
 #include "Interfaces/ITargetPlatform.h"
 #include "Misc/App.h"
 #include "MuCO/CustomizableObjectPrivate.h"
+#include "MuR/Model.h"
+#include "MuR/ModelPrivate.h"
 
 class UTexture2D;
 
@@ -1360,6 +1362,8 @@ void FCustomizableObjectCompiler::FinishCompilation()
 	}
 	else
 	{
+		CurrentObject->GetModel()->GetPrivate()->UnloadRoms();
+
 		// when skipping the SaveDerivedData task unlock the object so that instances can be updated
 		UCustomizableObjectSystem::GetInstance()->UnlockObject(CurrentObject);
 	}
@@ -1380,6 +1384,8 @@ void FCustomizableObjectCompiler::FinishSavingDerivedData()
 	SaveDDThread.Reset();
 	SaveDDTask.Reset();
 
+	CurrentObject->GetModel()->GetPrivate()->UnloadRoms();
+	
 	// Unlock the object so that instances can be updated
 	UCustomizableObjectSystem::GetInstance()->UnlockObject(CurrentObject);
 
diff --git a/Engine/Plugins/Experimental/Mutable/Source/MutableRuntime/Internal/MuR/ModelPrivate.h b/Engine/Plugins/Experimental/Mutable/Source/MutableRuntime/Internal/MuR/ModelPrivate.h
--- a/Engine/Plugins/Experimental/Mutable/Source/MutableRuntime/Internal/MuR/ModelPrivate.h
+++ b/Engine/Plugins/Experimental/Mutable/Source/MutableRuntime/Internal/MuR/ModelPrivate.h
@@ -675,6 +675,10 @@ namespace mu
         //!
         FProgram m_program;
 
+
+    	void UnloadRoms();
+
+
         //!
         void Serialise( OutputArchive& arch ) const
         {
diff --git a/Engine/Plugins/Experimental/Mutable/Source/MutableRuntime/Private/MuR/Model.cpp b/Engine/Plugins/Experimental/Mutable/Source/MutableRuntime/Private/MuR/Model.cpp
--- a/Engine/Plugins/Experimental/Mutable/Source/MutableRuntime/Private/MuR/Model.cpp
+++ b/Engine/Plugins/Experimental/Mutable/Source/MutableRuntime/Private/MuR/Model.cpp
@@ -69,6 +69,16 @@ namespace mu
 #endif
     }
 
+	
+	//---------------------------------------------------------------------------------------------
+    void Model::Private::UnloadRoms()
+    {
+	    for (int32 RomIndex = 0; RomIndex < m_program.m_roms.Num(); ++RomIndex)
+	    {
+		    m_program.UnloadRom(RomIndex);
+	    }
+    }
+	
 
     //---------------------------------------------------------------------------------------------
     //---------------------------------------------------------------------------------------------
