[Mutable] Removed Transient flag from UCustomizableObjectInstancePrivateData.

- Renamed to UCustomizableObjectInstancePrivate.

#rnx

diff --git a/Engine/Plugins/Experimental/Mutable/Source/CustomizableObject/Classes/MuCO/CustomizableObjectInstance.h b/Engine/Plugins/Experimental/Mutable/Source/CustomizableObject/Classes/MuCO/CustomizableObjectInstance.h
--- a/Engine/Plugins/Experimental/Mutable/Source/CustomizableObject/Classes/MuCO/CustomizableObjectInstance.h
+++ b/Engine/Plugins/Experimental/Mutable/Source/CustomizableObject/Classes/MuCO/CustomizableObjectInstance.h
@@ -111,7 +111,7 @@ class CUSTOMIZABLEOBJECT_API UCustomizableObjectInstance : public UObject
 {
 	GENERATED_BODY()
 
-	friend UCustomizableInstancePrivateData;
+	friend UCustomizableInstancePrivate;
 	friend FMutableUpdateCandidate;
 
 public:
@@ -677,7 +677,7 @@ public:
 	/** Flag to know if a property of this instance changed in the editor */
 	bool bEditorPropertyChanged = false;
 
-	UCustomizableInstancePrivateData* GetPrivate() const;
+	UCustomizableInstancePrivate* GetPrivate() const;
 
 	// TEMP VARIABLE to check the Min desired LODs for this instance
 	TWeakObjectPtr<UCustomizableObjectInstanceUsage> NearestToActor;
@@ -699,8 +699,8 @@ private:
 	UPROPERTY()
 	FCustomizableObjectInstanceDescriptor Descriptor;
 
-	UPROPERTY( Transient )
-	TObjectPtr<UCustomizableInstancePrivateData> PrivateData;
+	UPROPERTY()
+	TObjectPtr<UCustomizableInstancePrivate> PrivateData;
 
 	/** LODs applied on the beginning of the last update. Represent the actual LODs the Instance is using (not strictly true since an update can fail). */
 	int32 CurrentMinLOD = -1;
diff --git a/Engine/Plugins/Experimental/Mutable/Source/CustomizableObject/Internal/MuCO/CustomizableInstancePrivateData.h b/Engine/Plugins/Experimental/Mutable/Source/CustomizableObject/Internal/MuCO/CustomizableInstancePrivateData.h
--- a/Engine/Plugins/Experimental/Mutable/Source/CustomizableObject/Internal/MuCO/CustomizableInstancePrivateData.h
+++ b/Engine/Plugins/Experimental/Mutable/Source/CustomizableObject/Internal/MuCO/CustomizableInstancePrivateData.h
@@ -154,12 +154,12 @@ struct FAnimBpGeneratedPhysicsAssets
 };
 
 UCLASS()
-class UCustomizableInstancePrivateData : public UObject
+class UCustomizableInstancePrivate : public UObject
 {
 public:
 	GENERATED_BODY()
 
-	UCustomizableInstancePrivateData();
+	UCustomizableInstancePrivate();
 
 	UPROPERTY( Transient )
 	TArray<FGeneratedMaterial> GeneratedMaterials;
diff --git a/Engine/Plugins/Experimental/Mutable/Source/CustomizableObject/Private/MuCO/CustomizableObjectInstance.cpp b/Engine/Plugins/Experimental/Mutable/Source/CustomizableObject/Private/MuCO/CustomizableObjectInstance.cpp
--- a/Engine/Plugins/Experimental/Mutable/Source/CustomizableObject/Private/MuCO/CustomizableObjectInstance.cpp
+++ b/Engine/Plugins/Experimental/Mutable/Source/CustomizableObject/Private/MuCO/CustomizableObjectInstance.cpp
@@ -142,7 +142,7 @@ uint32 GetTypeHash(const FMutableMaterialPlaceholder& PlaceHolder)
 }
 
 
-UTexture2D* UCustomizableInstancePrivateData::CreateTexture()
+UTexture2D* UCustomizableInstancePrivate::CreateTexture()
 {
 	UTexture2D* NewTexture = NewObject<UTexture2D>(
 		GetTransientPackage(),
@@ -156,7 +156,7 @@ UTexture2D* UCustomizableInstancePrivateData::CreateTexture()
 }
 
 
-mu::FResourceID UCustomizableInstancePrivateData::GetLastMeshId(int32 ComponentIndex, int32 LODIndex) const
+mu::FResourceID UCustomizableInstancePrivate::GetLastMeshId(int32 ComponentIndex, int32 LODIndex) const
 {
 	const FCustomizableInstanceComponentData* ComponentData = GetComponentData(ComponentIndex);
 
@@ -172,7 +172,7 @@ mu::FResourceID UCustomizableInstancePrivateData::GetLastMeshId(int32 ComponentI
 }
 
 
-void UCustomizableInstancePrivateData::SetLastMeshId(int32 ComponentIndex, int32 LODIndex, mu::FResourceID MeshId)
+void UCustomizableInstancePrivate::SetLastMeshId(int32 ComponentIndex, int32 LODIndex, mu::FResourceID MeshId)
 {
 	FCustomizableInstanceComponentData* ComponentData = GetComponentData(ComponentIndex);
 	if (ComponentData && ComponentData->LastMeshIdPerLOD.IsValidIndex(LODIndex))
@@ -186,7 +186,7 @@ void UCustomizableInstancePrivateData::SetLastMeshId(int32 ComponentIndex, int32
 }
 
 
-void UCustomizableInstancePrivateData::InvalidateGeneratedData()
+void UCustomizableInstancePrivate::InvalidateGeneratedData()
 {
 	SkeletalMeshStatus = ESkeletalMeshStatus::NotGenerated;
 
@@ -199,7 +199,7 @@ void UCustomizableInstancePrivateData::InvalidateGeneratedData()
 }
 
 
-void UCustomizableInstancePrivateData::InitCustomizableObjectData(const UCustomizableObject* InCustomizableObject)
+void UCustomizableInstancePrivate::InitCustomizableObjectData(const UCustomizableObject* InCustomizableObject)
 {
 	InvalidateGeneratedData();
 
@@ -220,13 +220,13 @@ void UCustomizableInstancePrivateData::InitCustomizableObjectData(const UCustomi
 }
 
 
-FCustomizableInstanceComponentData* UCustomizableInstancePrivateData::GetComponentData(int32 ComponentIndex)
+FCustomizableInstanceComponentData* UCustomizableInstancePrivate::GetComponentData(int32 ComponentIndex)
 {
 	return ComponentsData.IsValidIndex(ComponentIndex) ? &ComponentsData[ComponentIndex] : nullptr;
 }
 
 
-const FCustomizableInstanceComponentData* UCustomizableInstancePrivateData::GetComponentData(int32 ComponentIndex) const
+const FCustomizableInstanceComponentData* UCustomizableInstancePrivate::GetComponentData(int32 ComponentIndex) const
 {
 	return ComponentsData.IsValidIndex(ComponentIndex) ? &ComponentsData[ComponentIndex] : nullptr;
 }
@@ -236,7 +236,7 @@ UCustomizableObjectInstance::UCustomizableObjectInstance()
 {
 	SetFlags(RF_Transactional);
 	
-	PrivateData = CreateDefaultSubobject<UCustomizableInstancePrivateData>(FName("PrivateData"), true);
+	PrivateData = CreateDefaultSubobject<UCustomizableInstancePrivate>(FName("Private"));
 }
 
 
@@ -271,7 +271,7 @@ void UCustomizableObjectInstance::SetDescriptor(const FCustomizableObjectInstanc
 }
 
 
-void UCustomizableInstancePrivateData::PrepareForUpdate(const TSharedRef<FUpdateContextPrivate>& OperationData)
+void UCustomizableInstancePrivate::PrepareForUpdate(const TSharedRef<FUpdateContextPrivate>& OperationData)
 {
 	// Clear the ComponentData from previous updates
 	for (FCustomizableInstanceComponentData& ComponentData : ComponentsData)
@@ -296,7 +296,7 @@ void UCustomizableInstancePrivateData::PrepareForUpdate(const TSharedRef<FUpdate
 #if WITH_EDITOR
 
 
-void UCustomizableInstancePrivateData::PostDuplicate(bool bDuplicateForPIE)
+void UCustomizableInstancePrivate::PostDuplicate(bool bDuplicateForPIE)
 {
 	Super::PostDuplicate(bDuplicateForPIE);
 
@@ -308,7 +308,7 @@ void UCustomizableInstancePrivateData::PostDuplicate(bool bDuplicateForPIE)
 }
 
 
-void UCustomizableInstancePrivateData::OnPostCompile()
+void UCustomizableInstancePrivate::OnPostCompile()
 {
 	UCustomizableObjectInstance* Instance = Cast<UCustomizableObjectInstance>(GetOuter());
 	if (!Instance)
@@ -333,7 +333,7 @@ void UCustomizableObjectInstance::BindPostCompileDelegate(UCustomizableObject* I
 	if (InCustomizableObject &&
 		!InCustomizableObject->PostCompileDelegate.IsBoundToObject(GetPrivate()))
 	{
-		InCustomizableObject->PostCompileDelegate.AddUObject(GetPrivate(), &UCustomizableInstancePrivateData::OnPostCompile);
+		InCustomizableObject->PostCompileDelegate.AddUObject(GetPrivate(), &UCustomizableInstancePrivate::OnPostCompile);
 	}
 }
 
@@ -430,7 +430,7 @@ void UCustomizableObjectInstance::DestroyLiveUpdateInstance()
 }
 
 
-void UCustomizableInstancePrivateData::ReleaseMutableResources(bool bCalledFromBeginDestroy, const UCustomizableObjectInstance& Instance)
+void UCustomizableInstancePrivate::ReleaseMutableResources(bool bCalledFromBeginDestroy, const UCustomizableObjectInstance& Instance)
 {
 	GeneratedMaterials.Empty();
 
@@ -660,7 +660,7 @@ bool UCustomizableObjectInstance::IsParamMultidimensional(const FString& ParamNa
 
 
 // Only safe to call if the Mutable texture ref count system returns 0 and absolutely sure nobody holds a reference to the texture
-void UCustomizableInstancePrivateData::ReleaseMutableTexture(const FMutableImageCacheKey& MutableTextureKey, UTexture2D* Texture, FMutableResourceCache& Cache)
+void UCustomizableInstancePrivate::ReleaseMutableTexture(const FMutableImageCacheKey& MutableTextureKey, UTexture2D* Texture, FMutableResourceCache& Cache)
 {
 	if (ensure(Texture) && Texture->IsValidLowLevel())
 	{
@@ -778,7 +778,7 @@ bool UCustomizableObjectInstance::IsParameterRelevant(const FString& ParamName)
 }
 
 
-void UCustomizableInstancePrivateData::PostEditChangePropertyWithoutEditor(USkeletalMesh* InSkeletalMesh)
+void UCustomizableInstancePrivate::PostEditChangePropertyWithoutEditor(USkeletalMesh* InSkeletalMesh)
 {
 	MUTABLE_CPUPROFILER_SCOPE(UCustomizableInstancePrivateData::PostEditChangePropertyWithoutEditor);
 
@@ -836,7 +836,7 @@ void UCustomizableObjectInstance::UpdateSkeletalMeshAsyncResult(FInstanceUpdateD
 }
 
 
-void UCustomizableInstancePrivateData::TickUpdateCloseCustomizableObjects(UCustomizableObjectInstance& Public, FMutableInstanceUpdateMap& InOutRequestedUpdates)
+void UCustomizableInstancePrivate::TickUpdateCloseCustomizableObjects(UCustomizableObjectInstance& Public, FMutableInstanceUpdateMap& InOutRequestedUpdates)
 {
 	if (!Public.CanUpdateInstance())
 	{
@@ -888,7 +888,7 @@ void UCustomizableInstancePrivateData::TickUpdateCloseCustomizableObjects(UCusto
 }
 
 
-void UCustomizableInstancePrivateData::UpdateInstanceIfNotGenerated(UCustomizableObjectInstance& Public, FMutableInstanceUpdateMap& InOutRequestedUpdates)
+void UCustomizableInstancePrivate::UpdateInstanceIfNotGenerated(UCustomizableObjectInstance& Public, FMutableInstanceUpdateMap& InOutRequestedUpdates)
 {
 	if (SkeletalMeshStatus != ESkeletalMeshStatus::NotGenerated)
 	{
@@ -1026,7 +1026,7 @@ bool AreSkeletonsCompatible(const TArray<TObjectPtr<USkeleton>>& InSkeletons)
 #endif
 
 
-USkeleton* UCustomizableInstancePrivateData::MergeSkeletons(UCustomizableObject& CustomizableObject, const FMutableRefSkeletalMeshData& RefSkeletalMeshData, int32 ComponentIndex)
+USkeleton* UCustomizableInstancePrivate::MergeSkeletons(UCustomizableObject& CustomizableObject, const FMutableRefSkeletalMeshData& RefSkeletalMeshData, int32 ComponentIndex)
 {
 	MUTABLE_CPUPROFILER_SCOPE(BuildSkeletonData_MergeSkeletons);
 
@@ -1373,7 +1373,7 @@ namespace
 	}
 }
 
-UPhysicsAsset* UCustomizableInstancePrivateData::GetOrBuildMainPhysicsAsset(
+UPhysicsAsset* UCustomizableInstancePrivate::GetOrBuildMainPhysicsAsset(
 	TObjectPtr<UPhysicsAsset> TemplateAsset,
 	const mu::PhysicsBody* MutablePhysics,
 	const UCustomizableObject& CustomizableObject,
@@ -1688,7 +1688,7 @@ void SetMeshUVChannelDensity(FMeshUVChannelInfo& UVChannelInfo, float Density =
 }
 
 
-bool UCustomizableInstancePrivateData::DoComponentsNeedUpdate(UCustomizableObjectInstance* Public, const TSharedRef<FUpdateContextPrivate>& OperationData, TArray<bool>& OutComponentNeedsUpdate, bool& bHasInvalidMesh)
+bool UCustomizableInstancePrivate::DoComponentsNeedUpdate(UCustomizableObjectInstance* Public, const TSharedRef<FUpdateContextPrivate>& OperationData, TArray<bool>& OutComponentNeedsUpdate, bool& bHasInvalidMesh)
 {
 	UCustomizableObject* CustomizableObject = Public->GetCustomizableObject();
 	check(CustomizableObject);
@@ -1785,7 +1785,7 @@ bool UCustomizableInstancePrivateData::DoComponentsNeedUpdate(UCustomizableObjec
 }
 
 
-bool UCustomizableInstancePrivateData::UpdateSkeletalMesh_PostBeginUpdate0(UCustomizableObjectInstance* Public, const TSharedRef<FUpdateContextPrivate>& OperationData)
+bool UCustomizableInstancePrivate::UpdateSkeletalMesh_PostBeginUpdate0(UCustomizableObjectInstance* Public, const TSharedRef<FUpdateContextPrivate>& OperationData)
 {
 	MUTABLE_CPUPROFILER_SCOPE(UCustomizableInstancePrivateData::UpdateSkeletalMesh_PostBeginUpdate0)
 
@@ -2118,7 +2118,7 @@ bool UCustomizableInstancePrivateData::UpdateSkeletalMesh_PostBeginUpdate0(UCust
 }
 
 
-UCustomizableInstancePrivateData::UCustomizableInstancePrivateData()
+UCustomizableInstancePrivate::UCustomizableInstancePrivate()
 {
 	MinSquareDistFromComponentToPlayer = FLT_MAX;
 	LastMinSquareDistFromComponentToPlayer = FLT_MAX;
@@ -2754,7 +2754,7 @@ bool UCustomizableObjectInstance::IsSelectedParameterProfileDirty() const
 //}
 
 
-void UCustomizableInstancePrivateData::DiscardResourcesAndSetReferenceSkeletalMesh(UCustomizableObjectInstance* Instance)
+void UCustomizableInstancePrivate::DiscardResourcesAndSetReferenceSkeletalMesh(UCustomizableObjectInstance* Instance)
 {
 	if (SkeletalMeshStatus == ESkeletalMeshStatus::Success)
 	{
@@ -2862,7 +2862,7 @@ const TArray<TObjectPtr<UTexture2D>>& UCustomizableObjectInstance::GetTexturePar
 #endif
 
 
-UCustomizableInstancePrivateData* UCustomizableObjectInstance::GetPrivate() const
+UCustomizableInstancePrivate* UCustomizableObjectInstance::GetPrivate() const
 { 
 	check(PrivateData); // Currently this is initialized in the constructor so we expect it always to exist.
 	return PrivateData; 
@@ -2917,7 +2917,7 @@ FTexturePlatformData* MutableCreateImagePlatformData(mu::Ptr<const mu::Image> Mu
 
 	if (SizeX <= 0 || SizeY <= 0)
 	{
-		UE_LOG(LogMutable, Warning, TEXT("Invalid parameters specified for UCustomizableInstancePrivateData::MutableCreateImagePlatformData()"));
+		UE_LOG(LogMutable, Warning, TEXT("Invalid parameters specified for UCustomizableInstancePrivate::MutableCreateImagePlatformData()"));
 		return nullptr;
 	}
 
@@ -3264,7 +3264,7 @@ void ConvertImage(UTexture2D* Texture, mu::Ptr<const mu::Image> MutableImage, co
 }
 
 
-void UCustomizableInstancePrivateData::InitSkeletalMeshData(const TSharedRef<FUpdateContextPrivate>& OperationData, USkeletalMesh* SkeletalMesh, const FMutableRefSkeletalMeshData* RefSkeletalMeshData, const UCustomizableObject& CustomizableObject, int32 ComponentIndex)
+void UCustomizableInstancePrivate::InitSkeletalMeshData(const TSharedRef<FUpdateContextPrivate>& OperationData, USkeletalMesh* SkeletalMesh, const FMutableRefSkeletalMeshData* RefSkeletalMeshData, const UCustomizableObject& CustomizableObject, int32 ComponentIndex)
 {
 	MUTABLE_CPUPROFILER_SCOPE(UCustomizableInstancePrivateData::InitSkeletalMesh);
 
@@ -3349,7 +3349,7 @@ void UCustomizableInstancePrivateData::InitSkeletalMeshData(const TSharedRef<FUp
 }
 
 
-bool UCustomizableInstancePrivateData::BuildSkeletonData(const TSharedRef<FUpdateContextPrivate>& OperationData, USkeletalMesh& SkeletalMesh, const FMutableRefSkeletalMeshData& RefSkeletalMeshData, UCustomizableObject& CustomizableObject, int32 ComponentIndex)
+bool UCustomizableInstancePrivate::BuildSkeletonData(const TSharedRef<FUpdateContextPrivate>& OperationData, USkeletalMesh& SkeletalMesh, const FMutableRefSkeletalMeshData& RefSkeletalMeshData, UCustomizableObject& CustomizableObject, int32 ComponentIndex)
 {
 	MUTABLE_CPUPROFILER_SCOPE(UCustomizableInstancePrivateData::BuildSkeletonData);
 
@@ -3478,7 +3478,7 @@ bool UCustomizableInstancePrivateData::BuildSkeletonData(const TSharedRef<FUpdat
 }
 
 
-void UCustomizableInstancePrivateData::BuildMeshSockets(const TSharedRef<FUpdateContextPrivate>& OperationData, USkeletalMesh* SkeletalMesh, const FMutableRefSkeletalMeshData* RefSkeletalMeshData, UCustomizableObjectInstance* Public, mu::MeshPtrConst MutableMesh)
+void UCustomizableInstancePrivate::BuildMeshSockets(const TSharedRef<FUpdateContextPrivate>& OperationData, USkeletalMesh* SkeletalMesh, const FMutableRefSkeletalMeshData* RefSkeletalMeshData, UCustomizableObjectInstance* Public, mu::MeshPtrConst MutableMesh)
 {
 	// Build mesh sockets.
 	MUTABLE_CPUPROFILER_SCOPE(UCustomizableInstancePrivateData::BuildMeshSockets);
@@ -3582,7 +3582,7 @@ void UCustomizableInstancePrivateData::BuildMeshSockets(const TSharedRef<FUpdate
 }
 
 
-void UCustomizableInstancePrivateData::BuildOrCopyElementData(const TSharedRef<FUpdateContextPrivate>& OperationData, USkeletalMesh* SkeletalMesh, UCustomizableObjectInstance* CustomizableObjectInstance, int32 ComponentIndex)
+void UCustomizableInstancePrivate::BuildOrCopyElementData(const TSharedRef<FUpdateContextPrivate>& OperationData, USkeletalMesh* SkeletalMesh, UCustomizableObjectInstance* CustomizableObjectInstance, int32 ComponentIndex)
 {
 	MUTABLE_CPUPROFILER_SCOPE(UCustomizableInstancePrivateData::BuildOrCopyElementData);
 
@@ -3610,7 +3610,7 @@ void UCustomizableInstancePrivateData::BuildOrCopyElementData(const TSharedRef<F
 }
 
 
-void UCustomizableInstancePrivateData::BuildOrCopyMorphTargetsData(const TSharedRef<FUpdateContextPrivate>& OperationData, USkeletalMesh* SkeletalMesh, const USkeletalMesh* LastUpdateSkeletalMesh, UCustomizableObjectInstance* CustomizableObjectInstance, int32 ComponentIndex)
+void UCustomizableInstancePrivate::BuildOrCopyMorphTargetsData(const TSharedRef<FUpdateContextPrivate>& OperationData, USkeletalMesh* SkeletalMesh, const USkeletalMesh* LastUpdateSkeletalMesh, UCustomizableObjectInstance* CustomizableObjectInstance, int32 ComponentIndex)
 {
 	MUTABLE_CPUPROFILER_SCOPE(UCustomizableInstancePrivateData::BuildOrCopyMorphTargetsData);
 
@@ -3762,7 +3762,7 @@ namespace
 }
 
 
-void UCustomizableInstancePrivateData::BuildOrCopyClothingData(const TSharedRef<FUpdateContextPrivate>&OperationData, USkeletalMesh * SkeletalMesh, const USkeletalMesh* LastUpdateSkeletalMesh, UCustomizableObjectInstance * CustomizableObjectInstance, int32 ComponentIndex)
+void UCustomizableInstancePrivate::BuildOrCopyClothingData(const TSharedRef<FUpdateContextPrivate>&OperationData, USkeletalMesh * SkeletalMesh, const USkeletalMesh* LastUpdateSkeletalMesh, UCustomizableObjectInstance * CustomizableObjectInstance, int32 ComponentIndex)
 {
 	MUTABLE_CPUPROFILER_SCOPE(UCustomizableInstancePrivateData::BuildOrCopyClothingData);
 
@@ -4628,7 +4628,7 @@ void UCustomizableInstancePrivateData::BuildOrCopyClothingData(const TSharedRef<
 }
 
 
-bool UCustomizableInstancePrivateData::BuildOrCopyRenderData(const TSharedRef<FUpdateContextPrivate>& OperationData, USkeletalMesh* SkeletalMesh, const USkeletalMesh* LastUpdateSkeletalMesh, UCustomizableObjectInstance* Public, int32 ComponentIndex)
+bool UCustomizableInstancePrivate::BuildOrCopyRenderData(const TSharedRef<FUpdateContextPrivate>& OperationData, USkeletalMesh* SkeletalMesh, const USkeletalMesh* LastUpdateSkeletalMesh, UCustomizableObjectInstance* Public, int32 ComponentIndex)
 {
 	MUTABLE_CPUPROFILER_SCOPE(UCustomizableInstancePrivateData::BuildOrCopyRenderData);
 
@@ -4798,7 +4798,7 @@ FAutoConsoleVariableRef CVarMutableHighPriorityLoading(
 	TEXT("If enabled, the request to load additional assets will have high priority."));
 
 
-FGraphEventRef UCustomizableInstancePrivateData::LoadAdditionalAssetsAsync(const TSharedRef<FUpdateContextPrivate>& OperationData, UCustomizableObjectInstance* Public, FStreamableManager& StreamableManager)
+FGraphEventRef UCustomizableInstancePrivate::LoadAdditionalAssetsAsync(const TSharedRef<FUpdateContextPrivate>& OperationData, UCustomizableObjectInstance* Public, FStreamableManager& StreamableManager)
 {
 	MUTABLE_CPUPROFILER_SCOPE(UCustomizableInstancePrivateData::LoadAdditionalAssetsAsync);
 
@@ -5084,7 +5084,7 @@ const TArray<TObjectPtr<UMaterialInterface>>* UCustomizableObjectInstance::GetOv
 }
 
 
-void UCustomizableInstancePrivateData::AdditionalAssetsAsyncLoaded(UCustomizableObjectInstance* Public)
+void UCustomizableInstancePrivate::AdditionalAssetsAsyncLoaded(UCustomizableObjectInstance* Public)
 {
 	MUTABLE_CPUPROFILER_SCOPE(UCustomizableInstancePrivateData::AdditionalAssetsAsyncLoaded);
 
@@ -5325,7 +5325,7 @@ void UpdateTextureRegionsMutable(UTexture2D* Texture, int32 MipIndex, uint32 Num
 }
 
 
-void UCustomizableInstancePrivateData::ReuseTexture(UTexture2D* Texture, TSharedRef<FTexturePlatformData, ESPMode::ThreadSafe>& PlatformData)
+void UCustomizableInstancePrivate::ReuseTexture(UTexture2D* Texture, TSharedRef<FTexturePlatformData, ESPMode::ThreadSafe>& PlatformData)
 {
 	uint32 NumMips = PlatformData->Mips.Num();
 
@@ -5359,7 +5359,7 @@ FAutoConsoleVariableRef CVarMutableReuseMaterialInstances(
 	bReuseMaterialInstances,
 	TEXT("If true, allow reuse of MaterialInstances between updates."));
 
-void UCustomizableInstancePrivateData::BuildMaterials(const TSharedRef<FUpdateContextPrivate>& OperationData, UCustomizableObjectInstance* Public)
+void UCustomizableInstancePrivate::BuildMaterials(const TSharedRef<FUpdateContextPrivate>& OperationData, UCustomizableObjectInstance* Public)
 {
 	MUTABLE_CPUPROFILER_SCOPE(UCustomizableInstancePrivateData::BuildMaterials)
 
@@ -6508,7 +6508,7 @@ void UCustomizableObjectInstance::AnimInstanceFixup(UAnimInstance* InAnimInstanc
 	}
 }
 
-const TArray<FAnimInstanceOverridePhysicsAsset>* UCustomizableInstancePrivateData::GetGeneratedPhysicsAssetsForAnimInstance(TSubclassOf<UAnimInstance> AnimInstanceClass) const
+const TArray<FAnimInstanceOverridePhysicsAsset>* UCustomizableInstancePrivate::GetGeneratedPhysicsAssetsForAnimInstance(TSubclassOf<UAnimInstance> AnimInstanceClass) const
 {
 	const FAnimBpGeneratedPhysicsAssets* Found = AnimBpPhysicsAssets.Find(AnimInstanceClass);
 
@@ -6582,7 +6582,7 @@ void UCustomizableObjectInstance::ForEachAnimInstance(int32 ComponentIndex, FEac
 
 TSet<UAssetUserData*> UCustomizableObjectInstance::GetMergedAssetUserData(int32 ComponentIndex) const
 {
-	UCustomizableInstancePrivateData* PrivateInstanceData = GetPrivate();
+	UCustomizableInstancePrivate* PrivateInstanceData = GetPrivate();
 
 	if (PrivateInstanceData && PrivateInstanceData->ComponentsData.IsValidIndex(ComponentIndex))
 	{
@@ -6605,7 +6605,7 @@ TSet<UAssetUserData*> UCustomizableObjectInstance::GetMergedAssetUserData(int32
 
 #if WITH_EDITORONLY_DATA
 
-void UCustomizableInstancePrivateData::RegenerateImportedModel(USkeletalMesh* SkeletalMesh)
+void UCustomizableInstancePrivate::RegenerateImportedModel(USkeletalMesh* SkeletalMesh)
 {
 	FSkeletalMeshRenderData* SkelResource = SkeletalMesh->GetResourceForRendering();
 	if (!SkelResource)
@@ -6761,7 +6761,7 @@ void UCustomizableInstancePrivateData::RegenerateImportedModel(USkeletalMesh* Sk
 			else
 			{
 				// The material should have been in the LODMaterialMap
-				ensureMsgf(false, TEXT("Unexpected material index in UCustomizableInstancePrivateData::RegenerateImportedModel"));
+				ensureMsgf(false, TEXT("Unexpected material index in UCustomizableInstancePrivate::RegenerateImportedModel"));
 
 				// Fallback index, may shift materials around sections
 				if (SkeletalMesh->GetMaterials().IsValidIndex(RenderSection.MaterialIndex))
diff --git a/Engine/Plugins/Experimental/Mutable/Source/CustomizableObject/Private/MuCO/CustomizableObjectSystem.cpp b/Engine/Plugins/Experimental/Mutable/Source/CustomizableObject/Private/MuCO/CustomizableObjectSystem.cpp
--- a/Engine/Plugins/Experimental/Mutable/Source/CustomizableObject/Private/MuCO/CustomizableObjectSystem.cpp
+++ b/Engine/Plugins/Experimental/Mutable/Source/CustomizableObject/Private/MuCO/CustomizableObjectSystem.cpp
@@ -772,7 +772,7 @@ void FinishUpdateGlobal(const TSharedRef<FUpdateContextPrivate>& Context)
 
 	if (Instance)
 	{
-		UCustomizableInstancePrivateData* PrivateInstance = Instance->GetPrivate();
+		UCustomizableInstancePrivate* PrivateInstance = Instance->GetPrivate();
 		
 		switch (Context->UpdateResult)
 		{
@@ -885,7 +885,7 @@ void UpdateSkeletalMesh(const TSharedRef<FUpdateContextPrivate>& Context)
 	UCustomizableObjectInstance* CustomizableObjectInstance = Context->Instance.Get();
 	check(CustomizableObjectInstance);
 
-	UCustomizableInstancePrivateData* CustomizableObjectInstancePrivateData = CustomizableObjectInstance->GetPrivate();
+	UCustomizableInstancePrivate* CustomizableObjectInstancePrivateData = CustomizableObjectInstance->GetPrivate();
 	check(CustomizableObjectInstancePrivateData != nullptr);
 	for (TObjectIterator<UCustomizableObjectInstanceUsage> It; It; ++It)
 	{
@@ -1029,7 +1029,7 @@ bool FCustomizableObjectSystemPrivate::TextureHasReferences(const FMutableImageC
 EUpdateRequired FCustomizableObjectSystemPrivate::IsUpdateRequired(const UCustomizableObjectInstance& Instance, bool bOnlyUpdateIfNotGenerated, bool bOnlyUpdateIfLODs, bool bIgnoreCloseDist) const
 {
 	UCustomizableObjectSystem* System = UCustomizableObjectSystem::GetInstance();
-	const UCustomizableInstancePrivateData* const Private = Instance.GetPrivate();
+	const UCustomizableInstancePrivate* const Private = Instance.GetPrivate();
 	
 	if (!Instance.CanUpdateInstance())
 	{
@@ -1083,7 +1083,7 @@ EUpdateRequired FCustomizableObjectSystemPrivate::IsUpdateRequired(const UCustom
 
 EQueuePriorityType FCustomizableObjectSystemPrivate::GetUpdatePriority(const UCustomizableObjectInstance& Instance,	bool bForceHighPriority) const
 {
-	const UCustomizableInstancePrivateData* InstancePrivate = Instance.GetPrivate();
+	const UCustomizableInstancePrivate* InstancePrivate = Instance.GetPrivate();
 		
 	const bool bNotGenerated = InstancePrivate->SkeletalMeshStatus == ESkeletalMeshStatus::NotGenerated;
 	const bool bShouldUpdateLODs = InstancePrivate->HasCOInstanceFlags(PendingLODsUpdate);
@@ -1128,7 +1128,7 @@ void FCustomizableObjectSystemPrivate::EnqueueUpdateSkeletalMesh(const TSharedRe
 	UCustomizableObjectInstance* Instance = Context->Instance.Get();
 	check(Instance);
 	
-	UCustomizableInstancePrivateData* InstancePrivate = Instance->GetPrivate();
+	UCustomizableInstancePrivate* InstancePrivate = Instance->GetPrivate();
 
 	const EQueuePriorityType Priority = GetUpdatePriority(*Instance, Context->bForceHighPriority);
 	const uint32 InstanceId = Instance->GetUniqueID();
@@ -1683,7 +1683,7 @@ namespace impl
 		check(System != nullptr);
 
 		const UCustomizableObject* CustomizableObject = OperationData->Instance->GetCustomizableObject();
-		UCustomizableInstancePrivateData* CustomizableObjectInstancePrivateData = OperationData->Instance->GetPrivate();
+		UCustomizableInstancePrivate* CustomizableObjectInstancePrivateData = OperationData->Instance->GetPrivate();
 
 		CustomizableObjectInstancePrivateData->PassThroughTexturesToLoad.Empty();
 
@@ -2273,10 +2273,10 @@ namespace impl
 		{
 			FMutableResourceCache& Cache = CustomizableObjectSystemPrivateData->GetObjectCache(CustomizableObjectInstance->GetCustomizableObject());
 
-			UCustomizableInstancePrivateData* CustomizableObjectInstancePrivateData = CustomizableObjectInstance->GetPrivate();
+			UCustomizableInstancePrivate* CustomizableObjectInstancePrivateData = CustomizableObjectInstance->GetPrivate();
 			for (FGeneratedTexture& GeneratedTexture : CustomizableObjectInstancePrivateData->TexturesToRelease)
 			{
-				UCustomizableInstancePrivateData::ReleaseMutableTexture(GeneratedTexture.Key, Cast<UTexture2D>(GeneratedTexture.Texture), Cache);
+				UCustomizableInstancePrivate::ReleaseMutableTexture(GeneratedTexture.Key, Cast<UTexture2D>(GeneratedTexture.Texture), Cache);
 			}
 
 			CustomizableObjectInstancePrivateData->TexturesToRelease.Empty();
@@ -2311,7 +2311,7 @@ namespace impl
 		const bool bInstanceInvalid = !CustomizableObjectInstance || !CustomizableObjectInstance->IsValidLowLevel();
 		if (!bInstanceInvalid)
 		{
-			UCustomizableInstancePrivateData* CustomizableInstancePrivateData = CustomizableObjectInstance->GetPrivate();
+			UCustomizableInstancePrivate* CustomizableInstancePrivateData = CustomizableObjectInstance->GetPrivate();
 
 #if WITH_EDITOR
 			CustomizableObjectInstance->LastUpdateMutableRuntimeCycles = OperationData->MutableRuntimeCycles;
@@ -2340,7 +2340,7 @@ namespace impl
 						if (SkeletalMesh && SkeletalMesh->GetResourceForRendering() && !SkeletalMesh->GetResourceForRendering()->IsInitialized())
 						{
 #if WITH_EDITOR
-							UCustomizableInstancePrivateData::RegenerateImportedModel(SkeletalMesh);
+							UCustomizableInstancePrivate::RegenerateImportedModel(SkeletalMesh);
 #endif
 							CustomizableInstancePrivateData->PostEditChangePropertyWithoutEditor(SkeletalMesh);
 						}
@@ -2430,7 +2430,7 @@ namespace impl
 			return;
 		}
 
-		UCustomizableInstancePrivateData* ObjectInstancePrivateData = ObjectInstance->GetPrivate();
+		UCustomizableInstancePrivate* ObjectInstancePrivateData = ObjectInstance->GetPrivate();
 		check(ObjectInstancePrivateData != nullptr);
 
 		if (OperationData->bLiveUpdateMode)
@@ -2667,7 +2667,7 @@ namespace impl
 
 		UCustomizableObjectInstance* CandidateInstance = Operation->Instance.Get();
 		
-		UCustomizableInstancePrivateData* CandidateInstancePrivateData = CandidateInstance->GetPrivate();
+		UCustomizableInstancePrivate* CandidateInstancePrivateData = CandidateInstance->GetPrivate();
 		if (!CandidateInstancePrivateData)
 		{
 			System->ClearCurrentMutableOperation();
@@ -3010,7 +3010,7 @@ bool UCustomizableObjectSystem::Tick(float DeltaTime)
 		{
 			if (IsValid(*CustomizableObjectInstance) && CustomizableObjectInstance->GetPrivate())
 			{
-				UCustomizableInstancePrivateData* ObjectInstancePrivateData = CustomizableObjectInstance->GetPrivate();
+				UCustomizableInstancePrivate* ObjectInstancePrivateData = CustomizableObjectInstance->GetPrivate();
 
 				if (ObjectInstancePrivateData->HasCOInstanceFlags(UsedByComponentInPlay))
 				{
@@ -3199,7 +3199,7 @@ void UCustomizableObjectSystem::DiscardInstances()
 		const bool bUpdating = Private->CurrentMutableOperation && Private->CurrentMutableOperation->Instance != Iterator->CustomizableObjectInstance;
 		if (COI && !bUpdating)
 		{
-			UCustomizableInstancePrivateData* COIPrivateData = COI ? COI->GetPrivate() : nullptr;
+			UCustomizableInstancePrivate* COIPrivateData = COI ? COI->GetPrivate() : nullptr;
 
 			// Only discard resources if the instance is still out range (it could have got closer to the player since the task was queued)
 			if (!CurrentInstanceLODManagement->IsOnlyUpdateCloseCustomizableObjectsEnabled() ||
diff --git a/Engine/Plugins/Experimental/Mutable/Source/CustomizableObject/Public/MuCO/CustomizableObjectInstanceDescriptor.h b/Engine/Plugins/Experimental/Mutable/Source/CustomizableObject/Public/MuCO/CustomizableObjectInstanceDescriptor.h
--- a/Engine/Plugins/Experimental/Mutable/Source/CustomizableObject/Public/MuCO/CustomizableObjectInstanceDescriptor.h
+++ b/Engine/Plugins/Experimental/Mutable/Source/CustomizableObject/Public/MuCO/CustomizableObjectInstanceDescriptor.h
@@ -12,7 +12,7 @@ class UTexture2D;
 enum class ECustomizableObjectProjectorType : uint8;
 
 class FArchive;
-class UCustomizableInstancePrivateData;
+class UCustomizableInstancePrivate;
 class UCustomizableObject;
 class UCustomizableObjectInstance;
 class FDescriptorHash;
@@ -403,7 +403,7 @@ private:
 	friend FDescriptorHash;
 	friend FDescriptorRuntimeHash;
 	friend UCustomizableObjectInstance;
-	friend UCustomizableInstancePrivateData;
+	friend UCustomizableInstancePrivate;
 	friend FMultilayerProjector;
 	friend FMutableUpdateCandidate;
 };
diff --git a/Engine/Plugins/Experimental/Mutable/Source/CustomizableObjectEditor/Private/MuCOE/SCustomizableObjectEditorViewportToolBar.cpp b/Engine/Plugins/Experimental/Mutable/Source/CustomizableObjectEditor/Private/MuCOE/SCustomizableObjectEditorViewportToolBar.cpp
--- a/Engine/Plugins/Experimental/Mutable/Source/CustomizableObjectEditor/Private/MuCOE/SCustomizableObjectEditorViewportToolBar.cpp
+++ b/Engine/Plugins/Experimental/Mutable/Source/CustomizableObjectEditor/Private/MuCOE/SCustomizableObjectEditorViewportToolBar.cpp
@@ -208,7 +208,7 @@ FText SCustomizableObjectEditorViewportToolBar::GetCompileErrorOverlayText() con
 		return LOCTEXT("UpdatingSkeletalMesh", "Updating Skeletal Mesh");
 	}
 
-	const UCustomizableInstancePrivateData* PrivateInstance = Instance->GetPrivate();
+	const UCustomizableInstancePrivate* PrivateInstance = Instance->GetPrivate();
 	
 	switch (PrivateInstance->SkeletalMeshStatus)
 	{
