[mutable] Ensure optimization is applied to state-specific code when a pixel format is changed because of state options.
#rnx

diff --git a/Engine/Plugins/Experimental/Mutable/Source/MutableTools/Private/MuT/CodeOptimiser_Parameters.cpp b/Engine/Plugins/Experimental/Mutable/Source/MutableTools/Private/MuT/CodeOptimiser_Parameters.cpp
--- a/Engine/Plugins/Experimental/Mutable/Source/MutableTools/Private/MuT/CodeOptimiser_Parameters.cpp
+++ b/Engine/Plugins/Experimental/Mutable/Source/MutableTools/Private/MuT/CodeOptimiser_Parameters.cpp
@@ -2015,12 +2015,14 @@ namespace mu
             }
 
 			// Apply texture compression strategy
+			bool bModified = false;
 			switch (m_states[s].nodeState.m_optimisation.TextureCompressionStrategy)
 			{
 			case ETextureCompressionStrategy::DontCompressRuntime:
 			{
 				MUTABLE_CPUPROFILER_SCOPE(RuntimeTextureCompressionRemover);
 				RuntimeTextureCompressionRemoverAST r(&m_states[s], false);
+				bModified = true;
 				break;
 			}
 
@@ -2028,6 +2030,7 @@ namespace mu
 			{
 				MUTABLE_CPUPROFILER_SCOPE(RuntimeTextureCompressionRemover);
 				RuntimeTextureCompressionRemoverAST r(&m_states[s], true);
+				bModified = true;
 				break;
 			}
 
@@ -2036,7 +2039,7 @@ namespace mu
 			}
 
             // If a state has no runtime parameters, skip its optimisation alltogether
-            if (m_states[s].nodeState.m_runtimeParams.Num())
+            if (bModified || m_states[s].nodeState.m_runtimeParams.Num())
             {
                 // Promote the intructions that depend on runtime parameters, and sink new
                 // format instructions.
