[mutable] Fixed bug where two table nodes from the same data table could not have different values in the parameter "Add None option"

#rnx

diff --git a/Engine/Plugins/Experimental/Mutable/Source/CustomizableObjectEditor/Private/MuCOE/GenerateMutableSource/GenerateMutableSourceColor.cpp b/Engine/Plugins/Experimental/Mutable/Source/CustomizableObjectEditor/Private/MuCOE/GenerateMutableSource/GenerateMutableSourceColor.cpp
--- a/Engine/Plugins/Experimental/Mutable/Source/CustomizableObjectEditor/Private/MuCOE/GenerateMutableSource/GenerateMutableSourceColor.cpp
+++ b/Engine/Plugins/Experimental/Mutable/Source/CustomizableObjectEditor/Private/MuCOE/GenerateMutableSource/GenerateMutableSourceColor.cpp
@@ -319,6 +319,7 @@ mu::NodeColourPtr GenerateMutableSourceColor(const UEdGraphPin* Pin, FMutableGra
 						ColorTableNode->SetTable(Table);
 						ColorTableNode->SetColumn(ColumnName);
 						ColorTableNode->SetParameterName(TypedNodeTable->ParameterName);
+						ColorTableNode->SetNoneOption(TypedNodeTable->bAddNoneOption);
 
 						GenerationContext.AddParameterNameUnique(Node, TypedNodeTable->ParameterName);
 					}
diff --git a/Engine/Plugins/Experimental/Mutable/Source/CustomizableObjectEditor/Private/MuCOE/GenerateMutableSource/GenerateMutableSourceFloat.cpp b/Engine/Plugins/Experimental/Mutable/Source/CustomizableObjectEditor/Private/MuCOE/GenerateMutableSource/GenerateMutableSourceFloat.cpp
--- a/Engine/Plugins/Experimental/Mutable/Source/CustomizableObjectEditor/Private/MuCOE/GenerateMutableSource/GenerateMutableSourceFloat.cpp
+++ b/Engine/Plugins/Experimental/Mutable/Source/CustomizableObjectEditor/Private/MuCOE/GenerateMutableSource/GenerateMutableSourceFloat.cpp
@@ -359,6 +359,7 @@ mu::NodeScalarPtr GenerateMutableSourceFloat(const UEdGraphPin* Pin, FMutableGra
 						ScalarTableNode->SetTable(Table);
 						ScalarTableNode->SetColumn(ColumnName);
 						ScalarTableNode->SetParameterName(TypedNodeTable->ParameterName);
+						ScalarTableNode->SetNoneOption(TypedNodeTable->bAddNoneOption);
 
 						GenerationContext.AddParameterNameUnique(Node, TypedNodeTable->ParameterName);
 					}
diff --git a/Engine/Plugins/Experimental/Mutable/Source/CustomizableObjectEditor/Private/MuCOE/GenerateMutableSource/GenerateMutableSourceImage.cpp b/Engine/Plugins/Experimental/Mutable/Source/CustomizableObjectEditor/Private/MuCOE/GenerateMutableSource/GenerateMutableSourceImage.cpp
--- a/Engine/Plugins/Experimental/Mutable/Source/CustomizableObjectEditor/Private/MuCOE/GenerateMutableSource/GenerateMutableSourceImage.cpp
+++ b/Engine/Plugins/Experimental/Mutable/Source/CustomizableObjectEditor/Private/MuCOE/GenerateMutableSource/GenerateMutableSourceImage.cpp
@@ -1071,6 +1071,7 @@ mu::NodeImagePtr GenerateMutableSourceImage(const UEdGraphPin* Pin, FMutableGrap
 							ImageTableNode->SetTable(Table);
 							ImageTableNode->SetColumn(ColumnName);
 							ImageTableNode->SetParameterName(TypedNodeTable->ParameterName);
+							ImageTableNode->SetNoneOption(TypedNodeTable->bAddNoneOption);
 
 							// TextureArrays are passthrough textures and do not need this step
 							if (DefaultTexture2D)
diff --git a/Engine/Plugins/Experimental/Mutable/Source/CustomizableObjectEditor/Private/MuCOE/GenerateMutableSource/GenerateMutableSourceMesh.cpp b/Engine/Plugins/Experimental/Mutable/Source/CustomizableObjectEditor/Private/MuCOE/GenerateMutableSource/GenerateMutableSourceMesh.cpp
--- a/Engine/Plugins/Experimental/Mutable/Source/CustomizableObjectEditor/Private/MuCOE/GenerateMutableSource/GenerateMutableSourceMesh.cpp
+++ b/Engine/Plugins/Experimental/Mutable/Source/CustomizableObjectEditor/Private/MuCOE/GenerateMutableSource/GenerateMutableSourceMesh.cpp
@@ -3689,6 +3689,7 @@ mu::NodeMeshPtr GenerateMutableSourceMesh(const UEdGraphPin* Pin,
 						MeshTableNode->SetTable(Table);
 						MeshTableNode->SetColumn(MutableColumnName);
 						MeshTableNode->SetParameterName(TypedNodeTable->ParameterName);
+						MeshTableNode->SetNoneOption(TypedNodeTable->bAddNoneOption);
 
 						GenerationContext.AddParameterNameUnique(Node, TypedNodeTable->ParameterName);
 
diff --git a/Engine/Plugins/Experimental/Mutable/Source/CustomizableObjectEditor/Private/MuCOE/GenerateMutableSource/GenerateMutableSourceTable.cpp b/Engine/Plugins/Experimental/Mutable/Source/CustomizableObjectEditor/Private/MuCOE/GenerateMutableSource/GenerateMutableSourceTable.cpp
--- a/Engine/Plugins/Experimental/Mutable/Source/CustomizableObjectEditor/Private/MuCOE/GenerateMutableSource/GenerateMutableSourceTable.cpp
+++ b/Engine/Plugins/Experimental/Mutable/Source/CustomizableObjectEditor/Private/MuCOE/GenerateMutableSource/GenerateMutableSourceTable.cpp
@@ -724,9 +724,6 @@ mu::TablePtr GenerateMutableSourceTable(const FString& TableName, const UEdGraph
 			
 			ParameterUIData.IntegerParameterGroupType = TypedTable->bAddNoneOption ? ECustomizableObjectGroupType::COGT_ONE_OR_NONE : ECustomizableObjectGroupType::COGT_ONE;
 
-			// Adding a None Option
-			MutableTable->SetNoneOption(TypedTable->bAddNoneOption);
-
 			for (int32 i = 0; i < RowNames.Num(); ++i)
 			{
 				MutableTable->AddRow(i);
diff --git a/Engine/Plugins/Experimental/Mutable/Source/MutableTools/Private/MuT/CodeGenerator.cpp b/Engine/Plugins/Experimental/Mutable/Source/MutableTools/Private/MuT/CodeGenerator.cpp
--- a/Engine/Plugins/Experimental/Mutable/Source/MutableTools/Private/MuT/CodeGenerator.cpp
+++ b/Engine/Plugins/Experimental/Mutable/Source/MutableTools/Private/MuT/CodeGenerator.cpp
@@ -318,7 +318,7 @@ namespace mu
 	}
 
 
-	Ptr<ASTOp> CodeGenerator::GenerateTableVariable(TablePtr pTable, const FString& strName)
+	Ptr<ASTOp> CodeGenerator::GenerateTableVariable(TablePtr pTable, const FString& strName, bool bAddNoneOption)
 	{
 		Ptr<ASTOp> result;
 
@@ -345,7 +345,7 @@ namespace mu
 				}
 			}
 
-			if (pTable->GetPrivate()->bNoneOption)
+			if (bAddNoneOption)
 			{
 				FParameterDesc::FIntValueDesc nullValue;
 				nullValue.m_value = -1;
diff --git a/Engine/Plugins/Experimental/Mutable/Source/MutableTools/Private/MuT/CodeGenerator.h b/Engine/Plugins/Experimental/Mutable/Source/MutableTools/Private/MuT/CodeGenerator.h
--- a/Engine/Plugins/Experimental/Mutable/Source/MutableTools/Private/MuT/CodeGenerator.h
+++ b/Engine/Plugins/Experimental/Mutable/Source/MutableTools/Private/MuT/CodeGenerator.h
@@ -298,7 +298,7 @@ namespace mu
 
         //-----------------------------------------------------------------------------------------
         //!
-        Ptr<ASTOp> GenerateTableVariable( TablePtr pTable, const FString& strName );
+        Ptr<ASTOp> GenerateTableVariable( TablePtr pTable, const FString& strName, bool bAddNoneOption );
 
         //!
         Ptr<ASTOp> GenerateMissingBoolCode(const TCHAR* strWhere, bool value, const void* errorContext );
@@ -713,34 +713,34 @@ namespace mu
     template<class NODE_TABLE_PRIVATE, ETableColumnType TYPE, OP_TYPE OPTYPE, typename F>
     Ptr<ASTOp> CodeGenerator::GenerateTableSwitch( const NODE_TABLE_PRIVATE& node, F&& GenerateOption )
     {
-        Ptr<Table> pTable;
-        Ptr<ASTOp> variable;
+        Ptr<Table> Table;
+        Ptr<ASTOp> Variable;
 
-        TPair<TablePtr,Ptr<ASTOp>>* it = GeneratedTables.Find( TPair<TablePtr,FString>(node.m_pTable,node.m_parameterName) );
+        TPair<TablePtr,Ptr<ASTOp>>* it = GeneratedTables.Find( TPair<TablePtr,FString>(node.Table,node.ParameterName) );
         if ( it )
         {
-            pTable = it->Key;
-            variable = it->Value;
+            Table = it->Key;
+            Variable = it->Value;
         }
 
-        if ( !pTable )
+        if ( !Table )
         {
             // Create the table variable expression
-            pTable = node.m_pTable;
-            variable = GenerateTableVariable( pTable, node.m_parameterName );
+            Table = node.Table;
+            Variable = GenerateTableVariable( Table, node.ParameterName, node.bNoneOption);
 
-            GeneratedTables.Add( TPair<TablePtr, FString>(node.m_pTable,node.m_parameterName), TPair<TablePtr,Ptr<ASTOp>>( pTable, variable ) );
+            GeneratedTables.Add( TPair<TablePtr, FString>(node.Table,node.ParameterName), TPair<TablePtr,Ptr<ASTOp>>( Table, Variable ) );
         }
 
         // Verify that the table column is the right type
-        int32 colIndex = pTable->FindColumn( node.m_columnName );
-        if ( colIndex<0 )
+        int32 ColIndex = Table->FindColumn( node.ColumnName );
+        if ( ColIndex<0 )
         {
             m_pErrorLog->GetPrivate()->Add("Table column not found.", ELMT_ERROR, node.m_errorContext);
             return nullptr;
         }
 
-        if ( pTable->GetPrivate()->Columns[ colIndex ].Type != TYPE )
+        if ( Table->GetPrivate()->Columns[ ColIndex ].Type != TYPE )
         {
             m_pErrorLog->GetPrivate()->Add("Table column type is not the right type.",
                                            ELMT_ERROR, node.m_errorContext);
@@ -749,22 +749,22 @@ namespace mu
 
         // Create the switch to cover all the options
         Ptr<ASTOp> lastSwitch;
-        int32 rows = pTable->GetPrivate()->Rows.Num();
+        int32 Rows = Table->GetPrivate()->Rows.Num();
 
         Ptr<ASTOpSwitch> SwitchOp = new ASTOpSwitch();
 		SwitchOp->type = OPTYPE;
-		SwitchOp->variable = variable;
+		SwitchOp->variable = Variable;
 		SwitchOp->def = GenerateDefaultTableValue(TYPE);
 
-		for (int32 i = 0; i < rows; ++i)
+		for (int32 i = 0; i < Rows; ++i)
         {
-            check( pTable->GetPrivate()->Rows[i].Id <= 0xFFFF);
-            auto condition = (uint16)pTable->GetPrivate()->Rows[i].Id;
-            Ptr<ASTOp> Branch = GenerateOption( node, colIndex, (int)i, m_pErrorLog.get() );
+            check( Table->GetPrivate()->Rows[i].Id <= 0xFFFF);
+            auto Condition = (uint16)Table->GetPrivate()->Rows[i].Id;
+            Ptr<ASTOp> Branch = GenerateOption( node, ColIndex, (int)i, m_pErrorLog.get() );
 
 			if (Branch || TYPE != ETableColumnType::Mesh)
 			{
-				SwitchOp->cases.Add(ASTOpSwitch::FCase(condition, SwitchOp, Branch));
+				SwitchOp->cases.Add(ASTOpSwitch::FCase(Condition, SwitchOp, Branch));
 			}
         }
 
diff --git a/Engine/Plugins/Experimental/Mutable/Source/MutableTools/Private/MuT/CodeGenerator_Color.cpp b/Engine/Plugins/Experimental/Mutable/Source/MutableTools/Private/MuT/CodeGenerator_Color.cpp
--- a/Engine/Plugins/Experimental/Mutable/Source/MutableTools/Private/MuT/CodeGenerator_Color.cpp
+++ b/Engine/Plugins/Experimental/Mutable/Source/MutableTools/Private/MuT/CodeGenerator_Color.cpp
@@ -491,7 +491,7 @@ namespace mu
 			[this, &Options](const NodeColourTable::Private& node, int colIndex, int row, ErrorLog* pErrorLog)
 			{
 				NodeColourConstantPtr CellData = new NodeColourConstant();
-				FVector4f Colour = node.m_pTable->GetPrivate()->Rows[row].Values[colIndex].Color;
+				FVector4f Colour = node.Table->GetPrivate()->Rows[row].Values[colIndex].Color;
 				CellData->SetValue(Colour);
 				return Generate(CellData, Options);
 			});
diff --git a/Engine/Plugins/Experimental/Mutable/Source/MutableTools/Private/MuT/CodeGenerator_Image.cpp b/Engine/Plugins/Experimental/Mutable/Source/MutableTools/Private/MuT/CodeGenerator_Image.cpp
--- a/Engine/Plugins/Experimental/Mutable/Source/MutableTools/Private/MuT/CodeGenerator_Image.cpp
+++ b/Engine/Plugins/Experimental/Mutable/Source/MutableTools/Private/MuT/CodeGenerator_Image.cpp
@@ -1867,7 +1867,7 @@ namespace mu
 		Result.op = GenerateTableSwitch<NodeImageTable::Private, ETableColumnType::Image, OP_TYPE::IM_SWITCH>(node,
 			[this,Options](const NodeImageTable::Private& node, int colIndex, int row, ErrorLog* pErrorLog)
 			{
-				const FTableValue& CellData = node.m_pTable->GetPrivate()->Rows[row].Values[colIndex];
+				const FTableValue& CellData = node.Table->GetPrivate()->Rows[row].Values[colIndex];
 				ImagePtrConst pImage = nullptr;
 
 				if (Ptr<ResourceProxy<Image>> pProxyImage = CellData.ProxyImage)
diff --git a/Engine/Plugins/Experimental/Mutable/Source/MutableTools/Private/MuT/CodeGenerator_Mesh.cpp b/Engine/Plugins/Experimental/Mutable/Source/MutableTools/Private/MuT/CodeGenerator_Mesh.cpp
--- a/Engine/Plugins/Experimental/Mutable/Source/MutableTools/Private/MuT/CodeGenerator_Mesh.cpp
+++ b/Engine/Plugins/Experimental/Mutable/Source/MutableTools/Private/MuT/CodeGenerator_Mesh.cpp
@@ -1120,7 +1120,7 @@ namespace mu
 		Ptr<ASTOp> Op = GenerateTableSwitch<NodeMeshTable::Private, ETableColumnType::Mesh, OP_TYPE::ME_SWITCH>(*TableNode->GetPrivate(),
 			[this, &NewResult, &bFirstRowGenerated, &InOptions] (const NodeMeshTable::Private& node, int colIndex, int row, ErrorLog* pErrorLog)
 			{
-				mu::Ptr<mu::Mesh> pMesh = node.m_pTable->GetPrivate()->Rows[row].Values[colIndex].Mesh;
+				mu::Ptr<mu::Mesh> pMesh = node.Table->GetPrivate()->Rows[row].Values[colIndex].Mesh;
 				FMeshGenerationResult BranchResults;
 
 				if (pMesh)
@@ -1129,11 +1129,11 @@ namespace mu
 					pCell->SetValue(pMesh);
 
 					// TODO Take into account layout strategy
-					int numLayouts = node.m_layouts.Num();
+					int numLayouts = node.Layouts.Num();
 					pCell->SetLayoutCount(numLayouts);
 					for (int i = 0; i < numLayouts; ++i)
 					{
-						pCell->SetLayout(i, node.m_layouts[i]);
+						pCell->SetLayout(i, node.Layouts[i]);
 					}
 
 					FMeshGenerationOptions TargetOptions = InOptions;
@@ -1143,7 +1143,7 @@ namespace mu
 						TargetOptions.OverrideLayouts = NewResult.GeneratedLayouts;
 					}
 
-					TargetOptions.OverrideContext = node.m_pTable->GetPrivate()->Rows[row].Values[colIndex].ErrorContext;
+					TargetOptions.OverrideContext = node.Table->GetPrivate()->Rows[row].Values[colIndex].ErrorContext;
 
 					GenerateMesh(TargetOptions, BranchResults, pCell);
 
diff --git a/Engine/Plugins/Experimental/Mutable/Source/MutableTools/Private/MuT/CodeGenerator_Scalar.cpp b/Engine/Plugins/Experimental/Mutable/Source/MutableTools/Private/MuT/CodeGenerator_Scalar.cpp
--- a/Engine/Plugins/Experimental/Mutable/Source/MutableTools/Private/MuT/CodeGenerator_Scalar.cpp
+++ b/Engine/Plugins/Experimental/Mutable/Source/MutableTools/Private/MuT/CodeGenerator_Scalar.cpp
@@ -408,7 +408,7 @@ namespace mu
 			[this,&Options](const NodeScalarTable::Private& node, int colIndex, int row, ErrorLog* pErrorLog)
 			{
 				NodeScalarConstantPtr pCell = new NodeScalarConstant();
-				float scalar = node.m_pTable->GetPrivate()->Rows[row].Values[colIndex].Scalar;
+				float scalar = node.Table->GetPrivate()->Rows[row].Values[colIndex].Scalar;
 				pCell->SetValue(scalar);
 				return Generate(pCell, Options);
 			});
diff --git a/Engine/Plugins/Experimental/Mutable/Source/MutableTools/Private/MuT/NodeColourTable.cpp b/Engine/Plugins/Experimental/Mutable/Source/MutableTools/Private/MuT/NodeColourTable.cpp
--- a/Engine/Plugins/Experimental/Mutable/Source/MutableTools/Private/MuT/NodeColourTable.cpp
+++ b/Engine/Plugins/Experimental/Mutable/Source/MutableTools/Private/MuT/NodeColourTable.cpp
@@ -33,28 +33,35 @@ namespace mu
 	//---------------------------------------------------------------------------------------------
 	void NodeColourTable::SetColumn( const FString& strName )
 	{
-		m_pD->m_columnName = strName;
+		m_pD->ColumnName = strName;
 	}
 
 
 	//---------------------------------------------------------------------------------------------
 	void NodeColourTable::SetParameterName( const FString& strName )
 	{
-		m_pD->m_parameterName = strName;
+		m_pD->ParameterName = strName;
 	}
 
 
 	//---------------------------------------------------------------------------------------------
 	void NodeColourTable::SetTable( TablePtr pTable )
 	{
-		m_pD->m_pTable = pTable;
+		m_pD->Table = pTable;
 	}
 
 
 	//---------------------------------------------------------------------------------------------
 	TablePtr NodeColourTable::GetTable() const
 	{
-		return m_pD->m_pTable;
+		return m_pD->Table;
+	}
+
+
+	//---------------------------------------------------------------------------------------------
+	void NodeColourTable::SetNoneOption(bool bAddNoneOption)
+	{
+		m_pD->bNoneOption = bAddNoneOption;
 	}
 
 
diff --git a/Engine/Plugins/Experimental/Mutable/Source/MutableTools/Private/MuT/NodeColourTablePrivate.h b/Engine/Plugins/Experimental/Mutable/Source/MutableTools/Private/MuT/NodeColourTablePrivate.h
--- a/Engine/Plugins/Experimental/Mutable/Source/MutableTools/Private/MuT/NodeColourTablePrivate.h
+++ b/Engine/Plugins/Experimental/Mutable/Source/MutableTools/Private/MuT/NodeColourTablePrivate.h
@@ -5,7 +5,7 @@
 
 #include "MuT/NodePrivate.h"
 #include "MuT/NodeColourTable.h"
-#include "MuT/Table.h"
+#include "MuT/TablePrivate.h"
 #include "MuT/AST.h"
 
 
@@ -19,19 +19,21 @@ namespace mu
 
 		static NODE_TYPE s_type;
 
-		FString m_parameterName;
-		TablePtr m_pTable;
-		FString m_columnName;
+		FString ParameterName;
+		TablePtr Table;
+		FString ColumnName;
+		bool bNoneOption = false;
 
 		//!
 		void Serialise( OutputArchive& arch ) const
 		{
-            uint32 ver = 1;
+            uint32_t ver = 2;
 			arch << ver;
 
-			arch << m_parameterName;
-			arch << m_pTable;
-			arch << m_columnName;
+			arch << ParameterName;
+			arch << Table;
+			arch << ColumnName;
+			arch << bNoneOption;
 		}
 
 		//!
@@ -39,30 +41,39 @@ namespace mu
 		{
             uint32 ver;
 			arch >> ver;
-			check(ver <= 1);
+			check(ver <= 2);
 
 			if (ver == 0)
 			{
 				std::string Temp;
 				arch >> Temp;
-				m_parameterName = Temp.c_str();
+				ParameterName = Temp.c_str();
 			}
 			else
 			{
-				arch >> m_parameterName;
+				arch >> ParameterName;
 			}
 
-			arch >> m_pTable;
+			arch >> Table;
 
 			if (ver == 0)
 			{
 				std::string Temp;
 				arch >> Temp;
-				m_columnName = Temp.c_str();
+				ColumnName = Temp.c_str();
 			}
 			else
 			{
-				arch >> m_columnName;
+				arch >> ColumnName;
+			}
+
+			if (ver >= 2)
+			{
+				arch >> bNoneOption;
+			}
+			else
+			{
+				bNoneOption = Table->GetPrivate()->bNoneOption_DEPRECATED;
 			}
 		}
 
diff --git a/Engine/Plugins/Experimental/Mutable/Source/MutableTools/Private/MuT/NodeImageTable.cpp b/Engine/Plugins/Experimental/Mutable/Source/MutableTools/Private/MuT/NodeImageTable.cpp
--- a/Engine/Plugins/Experimental/Mutable/Source/MutableTools/Private/MuT/NodeImageTable.cpp
+++ b/Engine/Plugins/Experimental/Mutable/Source/MutableTools/Private/MuT/NodeImageTable.cpp
@@ -32,28 +32,28 @@ namespace mu
 	//---------------------------------------------------------------------------------------------
 	void NodeImageTable::SetColumn( const FString& strName )
 	{
-		m_pD->m_columnName = strName;
+		m_pD->ColumnName = strName;
 	}
 
 
 	//---------------------------------------------------------------------------------------------
 	void NodeImageTable::SetParameterName( const FString& strName )
 	{
-		m_pD->m_parameterName = strName;
+		m_pD->ParameterName = strName;
 	}
 
 
 	//---------------------------------------------------------------------------------------------
 	void NodeImageTable::SetTable( TablePtr pTable )
 	{
-		m_pD->m_pTable = pTable;
+		m_pD->Table = pTable;
 	}
 
 
 	//---------------------------------------------------------------------------------------------
 	TablePtr NodeImageTable::GetTable() const
 	{
-		return m_pD->m_pTable;
+		return m_pD->Table;
 	}
 
 
@@ -70,6 +70,13 @@ namespace mu
 		return m_pD->MaxTextureSize;
 	}
 
+	
+	//---------------------------------------------------------------------------------------------
+	void NodeImageTable::SetNoneOption(bool bAddNoneOption)
+	{
+		m_pD->bNoneOption = bAddNoneOption;
+	}
+
 
 	//---------------------------------------------------------------------------------------------
 	void NodeImageTable::SetReferenceImageDescriptor(const FImageDesc& ImageDesc)
diff --git a/Engine/Plugins/Experimental/Mutable/Source/MutableTools/Private/MuT/NodeImageTablePrivate.h b/Engine/Plugins/Experimental/Mutable/Source/MutableTools/Private/MuT/NodeImageTablePrivate.h
--- a/Engine/Plugins/Experimental/Mutable/Source/MutableTools/Private/MuT/NodeImageTablePrivate.h
+++ b/Engine/Plugins/Experimental/Mutable/Source/MutableTools/Private/MuT/NodeImageTablePrivate.h
@@ -4,7 +4,7 @@
 
 #include "MuT/NodePrivate.h"
 #include "MuT/NodeImageTable.h"
-#include "MuT/Table.h"
+#include "MuT/TablePrivate.h"
 #include "MuT/AST.h"
 
 
@@ -18,29 +18,32 @@ namespace mu
 
 		static NODE_TYPE s_type;
 
-		FString m_parameterName;
-		TablePtr m_pTable;
-		FString m_columnName;
+		FString ParameterName;
+		TablePtr Table;
+		FString ColumnName;
 		uint16 MaxTextureSize = 0;
 		FImageDesc ReferenceImageDesc;
+		bool bNoneOption = false;
 
 		//!
 		void Serialise( OutputArchive& arch ) const
 		{
-            uint32_t ver = 4;
+            uint32_t ver = 5;
 			arch << ver;
 
-			arch << m_parameterName;
-			arch << m_pTable;
-			arch << m_columnName;
+			arch << ParameterName;
+			arch << Table;
+			arch << ColumnName;
 			arch << MaxTextureSize;
-			
+
 			// FImageDesc
 			{
 				arch << ReferenceImageDesc.m_size;
 				arch << uint8(ReferenceImageDesc.m_format);
 				arch << ReferenceImageDesc.m_lods;
 			}
+
+			arch << bNoneOption;
 		}
 
 		//!
@@ -48,36 +51,40 @@ namespace mu
 		{
             uint32_t ver;
 			arch >> ver;
-			check(ver>=1 && ver<= 4);
+			check(ver>=1 && ver<= 5);
 
 			if (ver == 1)
 			{
 				std::string Temp;
 				arch >> Temp;
-				m_parameterName = Temp.c_str();
+				ParameterName = Temp.c_str();
 			}
 			else
 			{
-				arch >> m_parameterName;
+				arch >> ParameterName;
 			}
 
-			arch >> m_pTable;
+			arch >> Table;
 
 			if (ver == 1)
 			{
 				std::string Temp;
 				arch >> Temp;
-				m_columnName = Temp.c_str();
+				ColumnName = Temp.c_str();
 			}
 			else
 			{
-				arch >> m_columnName;
+				arch >> ColumnName;
 			}
 
 			if (ver >= 3)
 			{
 				arch >> MaxTextureSize;
 			}
+			else
+			{
+				bNoneOption = Table->GetPrivate()->bNoneOption_DEPRECATED;
+			}
 
 			if (ver >= 4)
 			{
@@ -89,6 +96,11 @@ namespace mu
 				
 				arch >> ReferenceImageDesc.m_lods;
 			}
+			
+			if (ver >= 5)
+			{
+				arch >> bNoneOption;
+			}
 		}
 	};
 
diff --git a/Engine/Plugins/Experimental/Mutable/Source/MutableTools/Private/MuT/NodeMeshTable.cpp b/Engine/Plugins/Experimental/Mutable/Source/MutableTools/Private/MuT/NodeMeshTable.cpp
--- a/Engine/Plugins/Experimental/Mutable/Source/MutableTools/Private/MuT/NodeMeshTable.cpp
+++ b/Engine/Plugins/Experimental/Mutable/Source/MutableTools/Private/MuT/NodeMeshTable.cpp
@@ -33,42 +33,42 @@ namespace mu
 	//---------------------------------------------------------------------------------------------
 	void NodeMeshTable::SetColumn( const FString& strName )
 	{
-		m_pD->m_columnName = strName;
+		m_pD->ColumnName = strName;
 	}
 
 
 	//---------------------------------------------------------------------------------------------
 	void NodeMeshTable::SetParameterName( const FString& strName )
 	{
-		m_pD->m_parameterName = strName;
+		m_pD->ParameterName = strName;
 	}
 
 
 	//---------------------------------------------------------------------------------------------
 	void NodeMeshTable::SetTable( TablePtr pTable )
 	{
-		m_pD->m_pTable = pTable;
+		m_pD->Table = pTable;
 	}
 
 
 	//---------------------------------------------------------------------------------------------
 	TablePtr NodeMeshTable::GetTable() const
 	{
-		return m_pD->m_pTable;
+		return m_pD->Table;
 	}
 
 
 	//---------------------------------------------------------------------------------------------
 	int NodeMeshTable::GetLayoutCount() const
 	{
-		return m_pD->m_layouts.Num();
+		return m_pD->Layouts.Num();
 	}
 
 
 	//---------------------------------------------------------------------------------------------
 	void NodeMeshTable::SetLayoutCount( int i )
 	{
-		m_pD->m_layouts.SetNum( i );
+		m_pD->Layouts.SetNum( i );
 	}
 
 
@@ -84,7 +84,7 @@ namespace mu
 	void NodeMeshTable::SetLayout( int i, NodeLayoutPtr pLayout )
 	{
 		check( i>=0 && i<GetLayoutCount() );
-		m_pD->m_layouts[i] = pLayout;
+		m_pD->Layouts[i] = pLayout;
 	}
 
 
@@ -93,14 +93,21 @@ namespace mu
 	{
 		NodeLayoutPtr pResult;
 
-		if ( i>=0 && i<m_layouts.Num() )
+		if ( i>=0 && i< Layouts.Num() )
 		{
-			pResult = m_layouts[i];
+			pResult = Layouts[i];
 		}
 
 		return pResult;
 	}
 
+
+	//---------------------------------------------------------------------------------------------
+	void NodeMeshTable::SetNoneOption(bool bAddNoneOption)
+	{
+		m_pD->bNoneOption = bAddNoneOption;
+	}
+
 }
 
 
diff --git a/Engine/Plugins/Experimental/Mutable/Source/MutableTools/Private/MuT/NodeMeshTablePrivate.h b/Engine/Plugins/Experimental/Mutable/Source/MutableTools/Private/MuT/NodeMeshTablePrivate.h
--- a/Engine/Plugins/Experimental/Mutable/Source/MutableTools/Private/MuT/NodeMeshTablePrivate.h
+++ b/Engine/Plugins/Experimental/Mutable/Source/MutableTools/Private/MuT/NodeMeshTablePrivate.h
@@ -5,7 +5,7 @@
 #include "MuT/NodeMeshPrivate.h"
 #include "MuT/NodeMeshTable.h"
 #include "MuT/NodeLayout.h"
-#include "MuT/Table.h"
+#include "MuT/TablePrivate.h"
 #include "MuT/AST.h"
 
 
@@ -19,22 +19,24 @@ namespace mu
 
 		static NODE_TYPE s_type;
 
-		FString m_parameterName;
-		TablePtr m_pTable;
-		FString m_columnName;
+		FString ParameterName;
+		TablePtr Table;
+		FString ColumnName;
+		bool bNoneOption = false;
 
-		TArray<NodeLayoutPtr> m_layouts;
+		TArray<NodeLayoutPtr> Layouts;
 
 		//!
 		void Serialise( OutputArchive& arch ) const
 		{
-            uint32 ver = 2;
+            uint32 ver = 3;
 			arch << ver;
 
-			arch << m_parameterName;
-			arch << m_pTable;
-			arch << m_columnName;
-			arch << m_layouts;
+			arch << ParameterName;
+			arch << Table;
+			arch << ColumnName;
+			arch << Layouts;
+			arch << bNoneOption;
 		}
 
 		//!
@@ -42,32 +44,42 @@ namespace mu
 		{
             uint32 ver;
 			arch >> ver;
-            check(ver>=1 && ver<=2);
+            check(ver>=1 && ver<=3);
 
 			if (ver == 1)
 			{
 				std::string Temp;
 				arch >> Temp;
-				m_parameterName = Temp.c_str();
+				ParameterName = Temp.c_str();
 			}
 			else
 			{
-				arch >> m_parameterName;
+				arch >> ParameterName;
 			}
 
-			arch >> m_pTable;
+			arch >> Table;
 
 			if (ver == 1)
 			{
 				std::string Temp;
 				arch >> Temp;
-				m_columnName = Temp.c_str();
+				ColumnName = Temp.c_str();
 			}
 			else
 			{
-				arch >> m_columnName;
+				arch >> ColumnName;
+			}
+			
+			arch >> Layouts;
+
+			if(ver >= 3)
+			{
+				arch >> bNoneOption;
+			}
+			else
+			{
+				bNoneOption = Table->GetPrivate()->bNoneOption_DEPRECATED;
 			}
-			arch >> m_layouts;
 		}
 
 		// NodeMesh::Private interface
diff --git a/Engine/Plugins/Experimental/Mutable/Source/MutableTools/Private/MuT/NodeScalarTable.cpp b/Engine/Plugins/Experimental/Mutable/Source/MutableTools/Private/MuT/NodeScalarTable.cpp
--- a/Engine/Plugins/Experimental/Mutable/Source/MutableTools/Private/MuT/NodeScalarTable.cpp
+++ b/Engine/Plugins/Experimental/Mutable/Source/MutableTools/Private/MuT/NodeScalarTable.cpp
@@ -32,31 +32,37 @@ namespace mu
 	//---------------------------------------------------------------------------------------------
 	void NodeScalarTable::SetColumn( const FString& strName )
 	{
-		m_pD->m_columnName = strName;
+		m_pD->ColumnName = strName;
 	}
 
 
 	//---------------------------------------------------------------------------------------------
 	void NodeScalarTable::SetParameterName( const FString& strName )
 	{
-		m_pD->m_parameterName = strName;
+		m_pD->ParameterName = strName;
 	}
 
 
 	//---------------------------------------------------------------------------------------------
 	void NodeScalarTable::SetTable( TablePtr pTable )
 	{
-		m_pD->m_pTable = pTable;
+		m_pD->Table = pTable;
 	}
 
 
 	//---------------------------------------------------------------------------------------------
 	TablePtr NodeScalarTable::GetTable() const
 	{
-		return m_pD->m_pTable;
+		return m_pD->Table;
 	}
 
 
+	//---------------------------------------------------------------------------------------------
+	void NodeScalarTable::SetNoneOption(bool bAddNoneOption)
+	{
+		m_pD->bNoneOption = bAddNoneOption;
+	}
+
 }
 
 
diff --git a/Engine/Plugins/Experimental/Mutable/Source/MutableTools/Private/MuT/NodeScalarTablePrivate.h b/Engine/Plugins/Experimental/Mutable/Source/MutableTools/Private/MuT/NodeScalarTablePrivate.h
--- a/Engine/Plugins/Experimental/Mutable/Source/MutableTools/Private/MuT/NodeScalarTablePrivate.h
+++ b/Engine/Plugins/Experimental/Mutable/Source/MutableTools/Private/MuT/NodeScalarTablePrivate.h
@@ -4,7 +4,7 @@
 
 #include "MuT/NodePrivate.h"
 #include "MuT/NodeScalarTable.h"
-#include "MuT/Table.h"
+#include "MuT/TablePrivate.h"
 #include "MuT/AST.h"
 
 
@@ -18,19 +18,21 @@ namespace mu
 
 		static NODE_TYPE s_type;
 
-		FString m_parameterName;
-		Ptr<Table> m_pTable;
-		FString m_columnName;
+		FString ParameterName;
+		TablePtr Table;
+		FString ColumnName;
+		bool bNoneOption = false;
 
 		//!
 		void Serialise( OutputArchive& arch ) const
 		{
-            uint32 ver = 1;
+            uint32 ver = 2;
 			arch << ver;
 
-			arch << m_parameterName;
-			arch << m_pTable;
-			arch << m_columnName;
+			arch << ParameterName;
+			arch << Table;
+			arch << ColumnName;
+			arch << bNoneOption;
 		}
 
 		//!
@@ -38,30 +40,39 @@ namespace mu
 		{
             uint32 ver;
 			arch >> ver;
-            check(ver<=1);
+            check(ver <= 2);
 
 			if (ver==0)
 			{
 				std::string Temp;
 				arch >> Temp;
-				m_parameterName = Temp.c_str();
+				ParameterName = Temp.c_str();
 			}
 			else
 			{
-				arch >> m_parameterName;
+				arch >> ParameterName;
 			}
 
-			arch >> m_pTable;
+			arch >> Table;
 
 			if (ver == 0)
 			{
 				std::string Temp;
 				arch >> Temp;
-				m_columnName = Temp.c_str();
+				ColumnName = Temp.c_str();
 			}
 			else
 			{
-				arch >> m_columnName;
+				arch >> ColumnName;
+			}
+
+			if(ver >= 2)
+			{
+				arch >> bNoneOption;
+			}
+			else
+			{
+				bNoneOption = Table->GetPrivate()->bNoneOption_DEPRECATED;
 			}
 		}
 
diff --git a/Engine/Plugins/Experimental/Mutable/Source/MutableTools/Private/MuT/Table.cpp b/Engine/Plugins/Experimental/Mutable/Source/MutableTools/Private/MuT/Table.cpp
--- a/Engine/Plugins/Experimental/Mutable/Source/MutableTools/Private/MuT/Table.cpp
+++ b/Engine/Plugins/Experimental/Mutable/Source/MutableTools/Private/MuT/Table.cpp
@@ -120,12 +120,6 @@ namespace mu
 	}
 
 
-	void Table::SetNoneOption(bool bAddOption)
-	{
-		m_pD->bNoneOption = bAddOption;
-	}
-
-
     void Table::SetCell( int32 Column, uint32 RowId, float Value, const void* ErrorContext)
 	{
 		int32 Row = m_pD->FindRow(RowId);
diff --git a/Engine/Plugins/Experimental/Mutable/Source/MutableTools/Private/MuT/TablePrivate.h b/Engine/Plugins/Experimental/Mutable/Source/MutableTools/Private/MuT/TablePrivate.h
--- a/Engine/Plugins/Experimental/Mutable/Source/MutableTools/Private/MuT/TablePrivate.h
+++ b/Engine/Plugins/Experimental/Mutable/Source/MutableTools/Private/MuT/TablePrivate.h
@@ -51,12 +51,15 @@ namespace mu
 		FString Name;
 		TArray<FTableColumn> Columns;
 		TArray<FTableRow> Rows;
-		bool bNoneOption = false;
+
+		// Transient value for serialization compatibility
+		bool bNoneOption_DEPRECATED = false;
+
 
 		//!
 		void Serialise( OutputArchive& arch ) const
 		{
-            uint32 ver = 3;
+            uint32 ver = 4;
 			arch << ver;
 
             arch << Columns.Num();
@@ -95,8 +98,6 @@ namespace mu
 					}
 				}
 			}
-
-			arch << bNoneOption;
 		}
 
 		//!
@@ -104,7 +105,7 @@ namespace mu
 		{
             uint32 ver;
 			arch >> ver;
-			check(ver<=3);
+			check(ver<=4);
 
             uint32 ColumnCount;
 			arch >> ColumnCount;
@@ -195,9 +196,9 @@ namespace mu
 				}
 			}
 
-			if (ver >= 1)
+			if (ver >= 1 && ver <= 3)
 			{
-				arch >> bNoneOption;
+				arch >> bNoneOption_DEPRECATED;
 			}
 		}
 
diff --git a/Engine/Plugins/Experimental/Mutable/Source/MutableTools/Public/MuT/NodeColourTable.h b/Engine/Plugins/Experimental/Mutable/Source/MutableTools/Public/MuT/NodeColourTable.h
--- a/Engine/Plugins/Experimental/Mutable/Source/MutableTools/Public/MuT/NodeColourTable.h
+++ b/Engine/Plugins/Experimental/Mutable/Source/MutableTools/Public/MuT/NodeColourTable.h
@@ -58,6 +58,9 @@ namespace mu
 		//!
 		void SetColumn( const FString& strName );
 
+		//!
+		void SetNoneOption(bool bAddNoneOption);
+
 		//-----------------------------------------------------------------------------------------
 		// Interface pattern
 		//-----------------------------------------------------------------------------------------
diff --git a/Engine/Plugins/Experimental/Mutable/Source/MutableTools/Public/MuT/NodeImageTable.h b/Engine/Plugins/Experimental/Mutable/Source/MutableTools/Public/MuT/NodeImageTable.h
--- a/Engine/Plugins/Experimental/Mutable/Source/MutableTools/Public/MuT/NodeImageTable.h
+++ b/Engine/Plugins/Experimental/Mutable/Source/MutableTools/Public/MuT/NodeImageTable.h
@@ -67,6 +67,9 @@ namespace mu
 		void SetMaxTextureSize(uint16 Size);
 		uint16 GetMaxTextureSize();
 
+		//! Adds the "None" option to the parameter that represents this table column
+		void SetNoneOption(bool bAddOption);
+		
 		//!
 		void SetReferenceImageDescriptor(const FImageDesc& Descriptor);
 
diff --git a/Engine/Plugins/Experimental/Mutable/Source/MutableTools/Public/MuT/NodeMeshTable.h b/Engine/Plugins/Experimental/Mutable/Source/MutableTools/Public/MuT/NodeMeshTable.h
--- a/Engine/Plugins/Experimental/Mutable/Source/MutableTools/Public/MuT/NodeMeshTable.h
+++ b/Engine/Plugins/Experimental/Mutable/Source/MutableTools/Public/MuT/NodeMeshTable.h
@@ -71,6 +71,9 @@ namespace mu
 		NodeLayoutPtr GetLayout( int index ) const;
 		void SetLayout( int index, NodeLayoutPtr );
 
+		//! Adds the "None" option to the parameter that represents this table column
+		void SetNoneOption(bool bAddOption);
+
 		//-----------------------------------------------------------------------------------------
 		// Interface pattern
 		//-----------------------------------------------------------------------------------------
diff --git a/Engine/Plugins/Experimental/Mutable/Source/MutableTools/Public/MuT/NodeScalarTable.h b/Engine/Plugins/Experimental/Mutable/Source/MutableTools/Public/MuT/NodeScalarTable.h
--- a/Engine/Plugins/Experimental/Mutable/Source/MutableTools/Public/MuT/NodeScalarTable.h
+++ b/Engine/Plugins/Experimental/Mutable/Source/MutableTools/Public/MuT/NodeScalarTable.h
@@ -59,6 +59,9 @@ namespace mu
 		//!
 		void SetColumn( const FString& strName );
 
+		//! Adds the "None" option to the parameter that represents this table column
+		void SetNoneOption(bool bAddOption);
+
 		//-----------------------------------------------------------------------------------------
 		// Interface pattern
 		//-----------------------------------------------------------------------------------------
diff --git a/Engine/Plugins/Experimental/Mutable/Source/MutableTools/Public/MuT/Table.h b/Engine/Plugins/Experimental/Mutable/Source/MutableTools/Public/MuT/Table.h
--- a/Engine/Plugins/Experimental/Mutable/Source/MutableTools/Public/MuT/Table.h
+++ b/Engine/Plugins/Experimental/Mutable/Source/MutableTools/Public/MuT/Table.h
@@ -64,9 +64,6 @@ namespace mu
 		//!
         void AddRow( uint32 id );
 
-		//! Adds a "NONE" option
-		void SetNoneOption(bool bAddOption);
-
 		//!
         void SetCell( int32 Column, uint32 RowId, float Value, const void* ErrorContext = nullptr);
         void SetCell( int32 Column, uint32 RowId, const FVector4f& Value, const void* ErrorContext = nullptr);
