[Mutable] Converted FCustomizableObjectSystemPrivate to a UObject.

- Fixed cyclical pointer reference.

#rnx

diff --git a/Engine/Plugins/Experimental/Mutable/Source/CustomizableObject/Classes/MuCO/CustomizableObjectInstance.h b/Engine/Plugins/Experimental/Mutable/Source/CustomizableObject/Classes/MuCO/CustomizableObjectInstance.h
--- a/Engine/Plugins/Experimental/Mutable/Source/CustomizableObject/Classes/MuCO/CustomizableObjectInstance.h
+++ b/Engine/Plugins/Experimental/Mutable/Source/CustomizableObject/Classes/MuCO/CustomizableObjectInstance.h
@@ -17,7 +17,7 @@ namespace mu
 	class Image;
 }
 
-class FCustomizableObjectSystemPrivate;
+class UCustomizableObjectSystemPrivate;
 class USkeletalMesh;
 class AActor;
 class FProperty;
diff --git a/Engine/Plugins/Experimental/Mutable/Source/CustomizableObject/Classes/MuCO/CustomizableObjectSystem.h b/Engine/Plugins/Experimental/Mutable/Source/CustomizableObject/Classes/MuCO/CustomizableObjectSystem.h
--- a/Engine/Plugins/Experimental/Mutable/Source/CustomizableObject/Classes/MuCO/CustomizableObjectSystem.h
+++ b/Engine/Plugins/Experimental/Mutable/Source/CustomizableObject/Classes/MuCO/CustomizableObjectSystem.h
@@ -27,7 +27,7 @@ class UDefaultImageProvider;
 class USkeletalMesh;
 class UMaterialInterface;
 class UTexture2D;
-class FCustomizableObjectSystemPrivate; // This is used to hide Mutable SDK members in the public headers.
+class UCustomizableObjectSystemPrivate; // This is used to hide Mutable SDK members in the public headers.
 class FUpdateContextPrivate;
 struct FFrame;
 struct FGuid;
@@ -309,11 +309,11 @@ public:
 	void LogShowData(bool bFullInfo, bool ShowMaterialInfo) const;
 
 	// Give access to the internal object data.
-	FCustomizableObjectSystemPrivate* GetPrivate();
-	const FCustomizableObjectSystemPrivate* GetPrivate() const;
+	UCustomizableObjectSystemPrivate* GetPrivate();
+	const UCustomizableObjectSystemPrivate* GetPrivate() const;
 
-	FCustomizableObjectSystemPrivate* GetPrivateChecked();
-	const FCustomizableObjectSystemPrivate* GetPrivateChecked() const;
+	UCustomizableObjectSystemPrivate* GetPrivateChecked();
+	const UCustomizableObjectSystemPrivate* GetPrivateChecked() const;
 	
 	FStreamableManager& GetStreamableManager();
 
@@ -377,7 +377,8 @@ public:
 	TArray<TObjectPtr<UTexture2D>> ProtectedCachedTextures;
 
 private:
-	TSharedPtr<FCustomizableObjectSystemPrivate> Private = nullptr;
+	UPROPERTY(Transient)
+	TObjectPtr<UCustomizableObjectSystemPrivate> Private = nullptr;
 
 	// For async material loading
 	FStreamableManager StreamableManager;
@@ -425,7 +426,7 @@ private:
 #endif
 	
 	// Friends
-	friend class FCustomizableObjectSystemPrivate;
+	friend class UCustomizableObjectSystemPrivate;
 };
 
 
diff --git a/Engine/Plugins/Experimental/Mutable/Source/CustomizableObject/Private/MuCO/CustomizableObjectInstance.cpp b/Engine/Plugins/Experimental/Mutable/Source/CustomizableObject/Private/MuCO/CustomizableObjectInstance.cpp
--- a/Engine/Plugins/Experimental/Mutable/Source/CustomizableObject/Private/MuCO/CustomizableObjectInstance.cpp
+++ b/Engine/Plugins/Experimental/Mutable/Source/CustomizableObject/Private/MuCO/CustomizableObjectInstance.cpp
@@ -418,10 +418,10 @@ void UCustomizableObjectInstance::DestroyLiveUpdateInstance()
 {
 	if (PrivateData && PrivateData->LiveUpdateModeInstanceID)
 	{
-		// If FCustomizableObjectSystemPrivate::SSystem is nullptr it means it has already been destroyed, no point in registering an instanceID release
+		// If UCustomizableObjectSystemPrivate::SSystem is nullptr it means it has already been destroyed, no point in registering an instanceID release
 		// since the Mutable system has already been destroyed. Just checking UCustomizableObjectSystem::GetInstance() will try to recreate the system when
-		// everything is shutting down, so it's better to check FCustomizableObjectSystemPrivate::SSystem first here
-		if (FCustomizableObjectSystemPrivate::SSystem && UCustomizableObjectSystem::GetInstance() && UCustomizableObjectSystem::GetInstance()->GetPrivate())
+		// everything is shutting down, so it's better to check UCustomizableObjectSystemPrivate::SSystem first here
+		if (UCustomizableObjectSystemPrivate::SSystem && UCustomizableObjectSystem::GetInstance() && UCustomizableObjectSystem::GetInstance()->GetPrivate())
 		{
 			UCustomizableObjectSystem::GetInstance()->GetPrivate()->InitInstanceIDRelease(PrivateData->LiveUpdateModeInstanceID);
 			PrivateData->LiveUpdateModeInstanceID = 0;
@@ -436,7 +436,7 @@ void UCustomizableInstancePrivate::ReleaseMutableResources(bool bCalledFromBegin
 
 	if (UCustomizableObjectSystem::IsCreated()) // Need to check this because the object might be destroyed after the CustomizableObjectSystem at shutdown
 	{
-		FCustomizableObjectSystemPrivate* CustomizableObjectSystem = UCustomizableObjectSystem::GetInstance()->GetPrivate();
+		UCustomizableObjectSystemPrivate* CustomizableObjectSystem = UCustomizableObjectSystem::GetInstance()->GetPrivate();
 		// Get the cache of resources of all live instances of this object
 		FMutableResourceCache& Cache = CustomizableObjectSystem->GetObjectCache(Instance.GetCustomizableObject());
 
@@ -813,7 +813,7 @@ bool UCustomizableObjectInstance::CanUpdateInstance() const
 
 void UCustomizableObjectInstance::UpdateSkeletalMeshAsync(bool bIgnoreCloseDist, bool bForceHighPriority)
 {
-	FCustomizableObjectSystemPrivate* SystemPrivate = UCustomizableObjectSystem::GetInstance()->GetPrivate();
+	UCustomizableObjectSystemPrivate* SystemPrivate = UCustomizableObjectSystem::GetInstance()->GetPrivate();
 
 	const TSharedRef<FUpdateContextPrivate> Context = MakeShared<FUpdateContextPrivate>(*this);
 	Context->bIgnoreCloseDist = bIgnoreCloseDist;
@@ -825,7 +825,7 @@ void UCustomizableObjectInstance::UpdateSkeletalMeshAsync(bool bIgnoreCloseDist,
 
 void UCustomizableObjectInstance::UpdateSkeletalMeshAsyncResult(FInstanceUpdateDelegate Callback, bool bIgnoreCloseDist, bool bForceHighPriority)
 {
-	FCustomizableObjectSystemPrivate* SystemPrivate = UCustomizableObjectSystem::GetInstance()->GetPrivate();
+	UCustomizableObjectSystemPrivate* SystemPrivate = UCustomizableObjectSystem::GetInstance()->GetPrivate();
 
 	const TSharedRef<FUpdateContextPrivate> Context = MakeShared<FUpdateContextPrivate>(*this);
 	Context->bIgnoreCloseDist = bIgnoreCloseDist;
@@ -843,7 +843,7 @@ void UCustomizableInstancePrivate::TickUpdateCloseCustomizableObjects(UCustomiza
 		return;
 	}
 
-	const FCustomizableObjectSystemPrivate* SystemPrivate = UCustomizableObjectSystem::GetInstance()->GetPrivate();	
+	const UCustomizableObjectSystemPrivate* SystemPrivate = UCustomizableObjectSystem::GetInstance()->GetPrivate();	
 
 	const EUpdateRequired UpdateRequired = SystemPrivate->IsUpdateRequired(Public, true, true, false);
 	if (UpdateRequired != EUpdateRequired::NoUpdate) // Since this is done in the tick, avoid starting an update that we know for sure that would not be performed. Once started it has some performance implications that we want to avoid.
@@ -900,7 +900,7 @@ void UCustomizableInstancePrivate::UpdateInstanceIfNotGenerated(UCustomizableObj
 		return;
 	}
 
-	FCustomizableObjectSystemPrivate* SystemPrivate = UCustomizableObjectSystem::GetInstance()->GetPrivate();
+	UCustomizableObjectSystemPrivate* SystemPrivate = UCustomizableObjectSystem::GetInstance()->GetPrivate();
 
 	const TSharedRef<FUpdateContextPrivate> Context = MakeShared<FUpdateContextPrivate>(Public);
 	Context->bOnlyUpdateIfNotGenerated = true;
@@ -2946,7 +2946,7 @@ FTexturePlatformData* MutableCreateImagePlatformData(mu::Ptr<const mu::Image> Mu
 
 	// Reduce final texture size if we surpass the max size we can generate.
 	UCustomizableObjectSystem* System = UCustomizableObjectSystem::GetInstance();
-	FCustomizableObjectSystemPrivate* SystemPrivate = System ? System->GetPrivate() : nullptr;
+	UCustomizableObjectSystemPrivate* SystemPrivate = System ? System->GetPrivate() : nullptr;
 
 	int32 MaxTextureSizeToGenerate = SystemPrivate ? SystemPrivate->MaxTextureSizeToGenerate : 0;
 
@@ -5997,7 +5997,7 @@ void UCustomizableInstancePrivate::BuildMaterials(const TSharedRef<FUpdateContex
 	{
 		MUTABLE_CPUPROFILER_SCOPE(BuildMaterials_Exchange);
 
-		FCustomizableObjectSystemPrivate* CustomizableObjectSystem = UCustomizableObjectSystem::GetInstance()->GetPrivate();
+		UCustomizableObjectSystemPrivate* CustomizableObjectSystem = UCustomizableObjectSystem::GetInstance()->GetPrivate();
 		TexturesToRelease.Empty();
 
 		for (const FGeneratedTexture& Texture : NewGeneratedTextures)
diff --git a/Engine/Plugins/Experimental/Mutable/Source/CustomizableObject/Private/MuCO/CustomizableObjectInstanceDescriptor.cpp b/Engine/Plugins/Experimental/Mutable/Source/CustomizableObject/Private/MuCO/CustomizableObjectInstanceDescriptor.cpp
--- a/Engine/Plugins/Experimental/Mutable/Source/CustomizableObject/Private/MuCO/CustomizableObjectInstanceDescriptor.cpp
+++ b/Engine/Plugins/Experimental/Mutable/Source/CustomizableObject/Private/MuCO/CustomizableObjectInstanceDescriptor.cpp
@@ -2369,7 +2369,7 @@ void FCustomizableObjectInstanceDescriptor::SetRandomValues(const int32 InRandom
 		VectorParameter.ParameterValue.A = Random.GetFraction();
 	}
 
-	const FCustomizableObjectSystemPrivate* SystemPrivate = UCustomizableObjectSystem::GetInstance()->GetPrivate();
+	const UCustomizableObjectSystemPrivate* SystemPrivate = UCustomizableObjectSystem::GetInstance()->GetPrivate();
 
 	TArray<FName> PossibleValues;
 
diff --git a/Engine/Plugins/Experimental/Mutable/Source/CustomizableObject/Private/MuCO/CustomizableObjectMipDataProvider.cpp b/Engine/Plugins/Experimental/Mutable/Source/CustomizableObject/Private/MuCO/CustomizableObjectMipDataProvider.cpp
--- a/Engine/Plugins/Experimental/Mutable/Source/CustomizableObject/Private/MuCO/CustomizableObjectMipDataProvider.cpp
+++ b/Engine/Plugins/Experimental/Mutable/Source/CustomizableObject/Private/MuCO/CustomizableObjectMipDataProvider.cpp
@@ -33,7 +33,7 @@ FMutableUpdateContext::FMutableUpdateContext(
 {
 	if (Parameters)
 	{
-		const FCustomizableObjectSystemPrivate* Private = UCustomizableObjectSystem::GetInstance()->GetPrivate();
+		const UCustomizableObjectSystemPrivate* Private = UCustomizableObjectSystem::GetInstance()->GetPrivate();
 		Private->GetImageProviderChecked()->CacheImages(*Parameters);
 	}
 }
@@ -44,7 +44,7 @@ FMutableUpdateContext::~FMutableUpdateContext()
 	if (Parameters &&
 		UCustomizableObjectSystem::IsCreated())
 	{
-		const FCustomizableObjectSystemPrivate* Private = UCustomizableObjectSystem::GetInstance()->GetPrivate();
+		const UCustomizableObjectSystemPrivate* Private = UCustomizableObjectSystem::GetInstance()->GetPrivate();
 		Private->GetImageProviderChecked()->UnCacheImages(*Parameters);
 	}
 }
@@ -345,7 +345,7 @@ int32 FMutableTextureMipDataProvider::GetMips(const FTextureUpdateContext& Conte
 		}
 	}
 
-	FCustomizableObjectSystemPrivate* CustomizableObjectSystem = UCustomizableObjectSystem::GetInstance()->GetPrivate();
+	UCustomizableObjectSystemPrivate* CustomizableObjectSystem = UCustomizableObjectSystem::GetInstance()->GetPrivate();
 	if (CustomizableObjectSystem)
 	{
 		if (OperationData.IsValid())
diff --git a/Engine/Plugins/Experimental/Mutable/Source/CustomizableObject/Private/MuCO/CustomizableObjectSystem.cpp b/Engine/Plugins/Experimental/Mutable/Source/CustomizableObject/Private/MuCO/CustomizableObjectSystem.cpp
--- a/Engine/Plugins/Experimental/Mutable/Source/CustomizableObject/Private/MuCO/CustomizableObjectSystem.cpp
+++ b/Engine/Plugins/Experimental/Mutable/Source/CustomizableObject/Private/MuCO/CustomizableObjectSystem.cpp
@@ -58,7 +58,7 @@ DECLARE_CYCLE_STAT(TEXT("MutableTask"), STAT_MutableTask, STATGROUP_Game);
 #define UE_MUTABLE_UPDATE_REGION		TEXT("Mutable Update")
 
 
-UCustomizableObjectSystem* FCustomizableObjectSystemPrivate::SSystem = nullptr;
+UCustomizableObjectSystem* UCustomizableObjectSystemPrivate::SSystem = nullptr;
 
 static TAutoConsoleVariable<int32> CVarWorkingMemory(
 	TEXT("mutable.WorkingMemory"),
@@ -106,13 +106,13 @@ TAutoConsoleVariable<bool> CVarRollbackFixModelDiskStreamerDataRace(
 	TEXT("When true, use the new code path that fixes a data race in UnrealMutableModelDiskStreamer."));
 
 
-int32 FCustomizableObjectSystemPrivate::SkeletalMeshMinLodQualityLevel = -1;
+int32 UCustomizableObjectSystemPrivate::SkeletalMeshMinLodQualityLevel = -1;
 
 static void CVarMutableSinkFunction()
 {
 	if (UCustomizableObjectSystem::IsCreated())
 	{
-		FCustomizableObjectSystemPrivate* PrivateSystem = UCustomizableObjectSystem::GetInstance()->GetPrivate();
+		UCustomizableObjectSystemPrivate* PrivateSystem = UCustomizableObjectSystem::GetInstance()->GetPrivate();
 
 		static const IConsoleVariable* CVarSkeletalMeshMinLodQualityLevelCVarName = IConsoleManager::Get().FindConsoleVariable(TEXT("r.SkeletalMesh.MinLodQualityLevel"));
 		PrivateSystem->SkeletalMeshMinLodQualityLevel = CVarSkeletalMeshMinLodQualityLevelCVarName ? CVarSkeletalMeshMinLodQualityLevelCVarName->GetInt() : INDEX_NONE;
@@ -296,23 +296,23 @@ void FMutablePendingInstanceWork::AddIDRelease(mu::Instance::ID IDToRelease)
 
 UCustomizableObjectSystem* UCustomizableObjectSystem::GetInstance()
 {
-	if (!FCustomizableObjectSystemPrivate::SSystem)
+	if (!UCustomizableObjectSystemPrivate::SSystem)
 	{
 		UE_LOG(LogMutable, Log, TEXT("Creating Mutable Customizable Object System."));
 
 		check(IsInGameThread());
 
-		FCustomizableObjectSystemPrivate::SSystem = NewObject<UCustomizableObjectSystem>(UCustomizableObjectSystem::StaticClass());
-		check(FCustomizableObjectSystemPrivate::SSystem != nullptr);
-		checkf(!GUObjectArray.IsDisregardForGC(FCustomizableObjectSystemPrivate::SSystem), TEXT("Mutable was initialized too early in the UE4 init process, for instance, in the constructor of a default UObject."));
-		FCustomizableObjectSystemPrivate::SSystem->AddToRoot();
-		checkf(!GUObjectArray.IsDisregardForGC(FCustomizableObjectSystemPrivate::SSystem), TEXT("Mutable was initialized too early in the UE4 init process, for instance, in the constructor of a default UObject."));
-		FCustomizableObjectSystemPrivate::SSystem->InitSystem();
+		UCustomizableObjectSystemPrivate::SSystem = NewObject<UCustomizableObjectSystem>(UCustomizableObjectSystem::StaticClass());
+		check(UCustomizableObjectSystemPrivate::SSystem != nullptr);
+		checkf(!GUObjectArray.IsDisregardForGC(UCustomizableObjectSystemPrivate::SSystem), TEXT("Mutable was initialized too early in the UE4 init process, for instance, in the constructor of a default UObject."));
+		UCustomizableObjectSystemPrivate::SSystem->AddToRoot();
+		checkf(!GUObjectArray.IsDisregardForGC(UCustomizableObjectSystemPrivate::SSystem), TEXT("Mutable was initialized too early in the UE4 init process, for instance, in the constructor of a default UObject."));
+		UCustomizableObjectSystemPrivate::SSystem->InitSystem();
 
-		//FCoreUObjectDelegates::PurgePendingReleaseSkeletalMesh.AddUObject(FCustomizableObjectSystemPrivate::SSystem, &UCustomizableObjectSystem::PurgePendingReleaseSkeletalMesh);
+		//FCoreUObjectDelegates::PurgePendingReleaseSkeletalMesh.AddUObject(UCustomizableObjectSystemPrivate::SSystem, &UCustomizableObjectSystem::PurgePendingReleaseSkeletalMesh);
 	}
 
-	return FCustomizableObjectSystemPrivate::SSystem;
+	return UCustomizableObjectSystemPrivate::SSystem;
 }
 
 
@@ -413,26 +413,26 @@ void UCustomizableObjectSystem::LogShowData(bool bFullInfo, bool ShowMaterialInf
 }
 
 
-FCustomizableObjectSystemPrivate* UCustomizableObjectSystem::GetPrivate()
+UCustomizableObjectSystemPrivate* UCustomizableObjectSystem::GetPrivate()
 {
 	return Private.Get();
 }
 
 
-const FCustomizableObjectSystemPrivate* UCustomizableObjectSystem::GetPrivate() const
+const UCustomizableObjectSystemPrivate* UCustomizableObjectSystem::GetPrivate() const
 {
 	return Private.Get();
 }
 
 
-FCustomizableObjectSystemPrivate* UCustomizableObjectSystem::GetPrivateChecked()
+UCustomizableObjectSystemPrivate* UCustomizableObjectSystem::GetPrivateChecked()
 {
 	check(Private)
 	return Private.Get();
 }
 
 
-const FCustomizableObjectSystemPrivate* UCustomizableObjectSystem::GetPrivateChecked() const
+const UCustomizableObjectSystemPrivate* UCustomizableObjectSystem::GetPrivateChecked() const
 {
 	check(Private)
 	return Private.Get();
@@ -447,14 +447,14 @@ FStreamableManager& UCustomizableObjectSystem::GetStreamableManager()
 
 bool UCustomizableObjectSystem::IsCreated()
 {
-	return FCustomizableObjectSystemPrivate::SSystem != 0;
+	return UCustomizableObjectSystemPrivate::SSystem != 0;
 }
 
 
 void UCustomizableObjectSystem::InitSystem()
 {
 	// Everything initialized in Init() instead of constructor to prevent the default UCustomizableObjectSystem from registering a tick function
-	Private = MakeShareable(new FCustomizableObjectSystemPrivate());
+	Private = NewObject<UCustomizableObjectSystemPrivate>(this, FName("Private"));
 	check(Private != nullptr);
 	Private->NewCompilerFunc = nullptr;
 
@@ -480,7 +480,7 @@ void UCustomizableObjectSystem::InitSystem()
 	check(pSettings);
 	pSettings->SetProfile(false);
 	pSettings->SetWorkingMemoryBytes(Private->LastWorkingMemoryBytes);
-	Private->ExtensionDataStreamer = MakeShared<FUnrealExtensionDataStreamer>(Private.ToSharedRef());
+	Private->ExtensionDataStreamer = MakeShared<FUnrealExtensionDataStreamer>(Private);
 	Private->MutableSystem = new mu::System(pSettings, Private->ExtensionDataStreamer);
 	check(Private->MutableSystem);
 
@@ -530,7 +530,7 @@ void UCustomizableObjectSystem::BeginDestroy()
 #endif
 
 	// It could be null, for the default object.
-	if (Private.IsValid())
+	if (Private)
 	{
 
 #if !UE_SERVER
@@ -564,7 +564,7 @@ void UCustomizableObjectSystem::BeginDestroy()
 
 		Private->CurrentInstanceBeingUpdated = nullptr;
 
-		FCustomizableObjectSystemPrivate::SSystem = nullptr;
+		UCustomizableObjectSystemPrivate::SSystem = nullptr;
 
 		Private = nullptr;
 	}
@@ -579,7 +579,7 @@ FString UCustomizableObjectSystem::GetDesc()
 }
 
 
-FCustomizableObjectCompilerBase* (*FCustomizableObjectSystemPrivate::NewCompilerFunc)() = nullptr;
+FCustomizableObjectCompilerBase* (*UCustomizableObjectSystemPrivate::NewCompilerFunc)() = nullptr;
 
 
 FCustomizableObjectCompilerBase* UCustomizableObjectSystem::GetNewCompiler()
@@ -602,40 +602,21 @@ void UCustomizableObjectSystem::SetNewCompilerFunc(FCustomizableObjectCompilerBa
 }
 
 
-int32 FCustomizableObjectSystemPrivate::EnableMutableAnimInfoDebugging = 0;
+int32 UCustomizableObjectSystemPrivate::EnableMutableAnimInfoDebugging = 0;
 
 static FAutoConsoleVariableRef CVarEnableMutableAnimInfoDebugging(
-	TEXT("mutable.EnableMutableAnimInfoDebugging"), FCustomizableObjectSystemPrivate::EnableMutableAnimInfoDebugging,
+	TEXT("mutable.EnableMutableAnimInfoDebugging"), UCustomizableObjectSystemPrivate::EnableMutableAnimInfoDebugging,
 	TEXT("If set to 1 or greater print on screen the animation info of the pawn's Customizable Object Instance. Anim BPs, slots and tags will be displayed."
 	"If the root Customizable Object is recompiled after this command is run, the used skeletal meshes will also be displayed."),
 	ECVF_Default);
 
 
-void FCustomizableObjectSystemPrivate::AddGameThreadTask(const FMutableTask& Task)
+void UCustomizableObjectSystemPrivate::AddGameThreadTask(const FMutableTask& Task)
 {
 	PendingTasks.Enqueue(Task);
 }
 
 
-void FCustomizableObjectSystemPrivate::AddReferencedObjects(FReferenceCollector& Collector)
-{
-	if (CurrentInstanceBeingUpdated)
-	{
-		Collector.AddReferencedObject(CurrentInstanceBeingUpdated);
-	}
-
-#if WITH_EDITORONLY_DATA
-	Collector.AddReferencedObject(EditorImageProvider);
-#endif
-}
-
-
-FString FCustomizableObjectSystemPrivate::GetReferencerName() const
-{
-	return TEXT("FCustomizableObjectSystemPrivate");
-}
-
-
 TAutoConsoleVariable<bool> CVarCleanupTextureCache(
 	TEXT("mutable.EnableCleanupCache"),
 	true,
@@ -643,7 +624,7 @@ TAutoConsoleVariable<bool> CVarCleanupTextureCache(
 	ECVF_Scalability);
 
 
-void FCustomizableObjectSystemPrivate::CleanupCache()
+void UCustomizableObjectSystemPrivate::CleanupCache()
 {
 	check(IsInGameThread());
 
@@ -685,7 +666,7 @@ void FCustomizableObjectSystemPrivate::CleanupCache()
 }
 
 
-FMutableResourceCache& FCustomizableObjectSystemPrivate::GetObjectCache(const UCustomizableObject* Object)
+FMutableResourceCache& UCustomizableObjectSystemPrivate::GetObjectCache(const UCustomizableObject* Object)
 {
 	check(IsInGameThread());
 
@@ -707,50 +688,50 @@ FMutableResourceCache& FCustomizableObjectSystemPrivate::GetObjectCache(const UC
 }
 
 
-int32 FCustomizableObjectSystemPrivate::EnableMutableProgressiveMipStreaming = 1;
+int32 UCustomizableObjectSystemPrivate::EnableMutableProgressiveMipStreaming = 1;
 
 // Warning! If this is enabled, do not get references to the textures generated by Mutable! They are owned by Mutable and could become invalid at any moment
 static FAutoConsoleVariableRef CVarEnableMutableProgressiveMipStreaming(
-	TEXT("mutable.EnableMutableProgressiveMipStreaming"), FCustomizableObjectSystemPrivate::EnableMutableProgressiveMipStreaming,
+	TEXT("mutable.EnableMutableProgressiveMipStreaming"), UCustomizableObjectSystemPrivate::EnableMutableProgressiveMipStreaming,
 	TEXT("If set to 1 or greater use progressive Mutable Mip streaming for Mutable textures. If disabled, all mips will always be generated and spending memory. In that case, on Desktop platforms they will be stored in CPU memory, on other platforms textures will be non-streaming."),
 	ECVF_Default);
 
 
-int32 FCustomizableObjectSystemPrivate::EnableMutableLiveUpdate = 1;
+int32 UCustomizableObjectSystemPrivate::EnableMutableLiveUpdate = 1;
 
 static FAutoConsoleVariableRef CVarEnableMutableLiveUpdate(
-	TEXT("mutable.EnableMutableLiveUpdate"), FCustomizableObjectSystemPrivate::EnableMutableLiveUpdate,
+	TEXT("mutable.EnableMutableLiveUpdate"), UCustomizableObjectSystemPrivate::EnableMutableLiveUpdate,
 	TEXT("If set to 1 or greater Mutable can use the live update mode if set in the current Mutable state. If disabled, it will never use live update mode even if set in the current Mutable state."),
 	ECVF_Default);
 
 
-int32 FCustomizableObjectSystemPrivate::EnableReuseInstanceTextures = 1;
+int32 UCustomizableObjectSystemPrivate::EnableReuseInstanceTextures = 1;
 
 static FAutoConsoleVariableRef CVarEnableMutableReuseInstanceTextures(
-	TEXT("mutable.EnableReuseInstanceTextures"), FCustomizableObjectSystemPrivate::EnableReuseInstanceTextures,
+	TEXT("mutable.EnableReuseInstanceTextures"), UCustomizableObjectSystemPrivate::EnableReuseInstanceTextures,
 	TEXT("If set to 1 or greater and set in the corresponding setting in the current Mutable state, Mutable can reuse instance UTextures (only uncompressed and not streaming, so set the options in the state) and their resources between updates when they are modified. If geometry or state is changed they cannot be reused."),
 	ECVF_Default);
 
 
-int32 FCustomizableObjectSystemPrivate::EnableOnlyGenerateRequestedLODs = 1;
+int32 UCustomizableObjectSystemPrivate::EnableOnlyGenerateRequestedLODs = 1;
 
 static FAutoConsoleVariableRef CVarEnableOnlyGenerateRequestedLODs(
-	TEXT("mutable.EnableOnlyGenerateRequestedLODs"), FCustomizableObjectSystemPrivate::EnableOnlyGenerateRequestedLODs,
+	TEXT("mutable.EnableOnlyGenerateRequestedLODs"), UCustomizableObjectSystemPrivate::EnableOnlyGenerateRequestedLODs,
 	TEXT("If 1 or greater, Only the RequestedLODLevels will be generated. If 0, all LODs will be build."),
 	ECVF_Default);
 
-int32 FCustomizableObjectSystemPrivate::EnableSkipGenerateResidentMips = 1;
+int32 UCustomizableObjectSystemPrivate::EnableSkipGenerateResidentMips = 1;
 
 static FAutoConsoleVariableRef CVarSkipGenerateResidentMips(
-	TEXT("mutable.EnableSkipGenerateResidentMips"), FCustomizableObjectSystemPrivate::EnableSkipGenerateResidentMips,
+	TEXT("mutable.EnableSkipGenerateResidentMips"), UCustomizableObjectSystemPrivate::EnableSkipGenerateResidentMips,
 	TEXT("If 1 or greater, resident mip generation will be optional. If 0, resident mips will be always generated"),
 	ECVF_Default);
 
-int32 FCustomizableObjectSystemPrivate::MaxTextureSizeToGenerate = 0;
+int32 UCustomizableObjectSystemPrivate::MaxTextureSizeToGenerate = 0;
 
 FAutoConsoleVariableRef CVarMaxTextureSizeToGenerate(
 	TEXT("Mutable.MaxTextureSizeToGenerate"),
-	FCustomizableObjectSystemPrivate::MaxTextureSizeToGenerate,
+	UCustomizableObjectSystemPrivate::MaxTextureSizeToGenerate,
 	TEXT("Max texture size on Mutable textures. Mip 0 will be the first mip with max size equal or less than MaxTextureSizeToGenerate."
 		"If a texture doesn't have small enough mips, mip 0 will be the last mip available."));
 
@@ -768,7 +749,7 @@ void FinishUpdateGlobal(const TSharedRef<FUpdateContextPrivate>& Context)
 	UCustomizableObjectInstance* Instance = Context->Instance.Get();
 
 	UCustomizableObjectSystem* System = UCustomizableObjectSystem::GetInstance();
-	FCustomizableObjectSystemPrivate* SystemPrivate = System ? System->GetPrivate() : nullptr;
+	UCustomizableObjectSystemPrivate* SystemPrivate = System ? System->GetPrivate() : nullptr;
 
 	if (Instance)
 	{
@@ -923,7 +904,7 @@ void UpdateSkeletalMesh(const TSharedRef<FUpdateContextPrivate>& Context)
 }
 
 
-void FCustomizableObjectSystemPrivate::GetMipStreamingConfig(const UCustomizableObjectInstance& Instance, bool& bOutNeverStream, int32& OutMipsToSkip) const
+void UCustomizableObjectSystemPrivate::GetMipStreamingConfig(const UCustomizableObjectInstance& Instance, bool& bOutNeverStream, int32& OutMipsToSkip) const
 {
 	bOutNeverStream = false;
 
@@ -962,25 +943,25 @@ void FCustomizableObjectSystemPrivate::GetMipStreamingConfig(const UCustomizable
 }
 
 
-bool FCustomizableObjectSystemPrivate::IsReplaceDiscardedWithReferenceMeshEnabled() const
+bool UCustomizableObjectSystemPrivate::IsReplaceDiscardedWithReferenceMeshEnabled() const
 {
 	return bReplaceDiscardedWithReferenceMesh;
 }
 
 
-void FCustomizableObjectSystemPrivate::SetReplaceDiscardedWithReferenceMeshEnabled(bool bIsEnabled)
+void UCustomizableObjectSystemPrivate::SetReplaceDiscardedWithReferenceMeshEnabled(bool bIsEnabled)
 {
 	bReplaceDiscardedWithReferenceMesh = bIsEnabled;
 }
 
 
-int32 FCustomizableObjectSystemPrivate::GetNumSkeletalMeshes() const
+int32 UCustomizableObjectSystemPrivate::GetNumSkeletalMeshes() const
 {
 	return NumSkeletalMeshes;
 }
 
 
-void FCustomizableObjectSystemPrivate::AddTextureReference(const FMutableImageCacheKey& TextureId)
+void UCustomizableObjectSystemPrivate::AddTextureReference(const FMutableImageCacheKey& TextureId)
 {
 	uint32& CountRef = TextureReferenceCount.FindOrAdd(TextureId);
 
@@ -988,7 +969,7 @@ void FCustomizableObjectSystemPrivate::AddTextureReference(const FMutableImageCa
 }
 
 
-bool FCustomizableObjectSystemPrivate::RemoveTextureReference(const FMutableImageCacheKey& TextureId)
+bool UCustomizableObjectSystemPrivate::RemoveTextureReference(const FMutableImageCacheKey& TextureId)
 {
 	uint32* CountPtr = TextureReferenceCount.Find(TextureId);
 
@@ -1013,7 +994,7 @@ bool FCustomizableObjectSystemPrivate::RemoveTextureReference(const FMutableImag
 }
 
 
-bool FCustomizableObjectSystemPrivate::TextureHasReferences(const FMutableImageCacheKey& TextureId) const
+bool UCustomizableObjectSystemPrivate::TextureHasReferences(const FMutableImageCacheKey& TextureId) const
 {
 	const uint32* CountPtr = TextureReferenceCount.Find(TextureId);
 
@@ -1026,7 +1007,7 @@ bool FCustomizableObjectSystemPrivate::TextureHasReferences(const FMutableImageC
 }
 
 
-EUpdateRequired FCustomizableObjectSystemPrivate::IsUpdateRequired(const UCustomizableObjectInstance& Instance, bool bOnlyUpdateIfNotGenerated, bool bOnlyUpdateIfLODs, bool bIgnoreCloseDist) const
+EUpdateRequired UCustomizableObjectSystemPrivate::IsUpdateRequired(const UCustomizableObjectInstance& Instance, bool bOnlyUpdateIfNotGenerated, bool bOnlyUpdateIfLODs, bool bIgnoreCloseDist) const
 {
 	UCustomizableObjectSystem* System = UCustomizableObjectSystem::GetInstance();
 	const UCustomizableInstancePrivate* const Private = Instance.GetPrivate();
@@ -1081,7 +1062,7 @@ EUpdateRequired FCustomizableObjectSystemPrivate::IsUpdateRequired(const UCustom
 }
 
 
-EQueuePriorityType FCustomizableObjectSystemPrivate::GetUpdatePriority(const UCustomizableObjectInstance& Instance,	bool bForceHighPriority) const
+EQueuePriorityType UCustomizableObjectSystemPrivate::GetUpdatePriority(const UCustomizableObjectInstance& Instance,	bool bForceHighPriority) const
 {
 	const UCustomizableInstancePrivate* InstancePrivate = Instance.GetPrivate();
 		
@@ -1120,7 +1101,7 @@ EQueuePriorityType FCustomizableObjectSystemPrivate::GetUpdatePriority(const UCu
 }
 
 
-void FCustomizableObjectSystemPrivate::EnqueueUpdateSkeletalMesh(const TSharedRef<FUpdateContextPrivate>& Context)
+void UCustomizableObjectSystemPrivate::EnqueueUpdateSkeletalMesh(const TSharedRef<FUpdateContextPrivate>& Context)
 {
 	MUTABLE_CPUPROFILER_SCOPE(FCustomizableObjectSystemPrivate::EnqueueUpdateSkeletalMesh);
 	check(IsInGameThread());
@@ -1255,7 +1236,7 @@ void FCustomizableObjectSystemPrivate::EnqueueUpdateSkeletalMesh(const TSharedRe
 }
 
 
-void FCustomizableObjectSystemPrivate::InitDiscardResourcesSkeletalMesh(UCustomizableObjectInstance* InCustomizableObjectInstance)
+void UCustomizableObjectSystemPrivate::InitDiscardResourcesSkeletalMesh(UCustomizableObjectInstance* InCustomizableObjectInstance)
 {
 	check(IsInGameThread());
 
@@ -1267,7 +1248,7 @@ void FCustomizableObjectSystemPrivate::InitDiscardResourcesSkeletalMesh(UCustomi
 }
 
 
-void FCustomizableObjectSystemPrivate::InitInstanceIDRelease(mu::Instance::ID IDToRelease)
+void UCustomizableObjectSystemPrivate::InitInstanceIDRelease(mu::Instance::ID IDToRelease)
 {
 	check(IsInGameThread());
 
@@ -1277,7 +1258,7 @@ void FCustomizableObjectSystemPrivate::InitInstanceIDRelease(mu::Instance::ID ID
 
 bool UCustomizableObjectSystem::IsReplaceDiscardedWithReferenceMeshEnabled() const
 {
-	if (Private.IsValid())
+	if (Private)
 	{
 		return Private->IsReplaceDiscardedWithReferenceMeshEnabled();
 	}
@@ -1288,7 +1269,7 @@ bool UCustomizableObjectSystem::IsReplaceDiscardedWithReferenceMeshEnabled() con
 
 void UCustomizableObjectSystem::SetReplaceDiscardedWithReferenceMeshEnabled(bool bIsEnabled)
 {
-	if (Private.IsValid())
+	if (Private)
 	{
 		Private->SetReplaceDiscardedWithReferenceMeshEnabled(bIsEnabled);
 	}
@@ -1527,7 +1508,7 @@ void UCustomizableObjectSystem::ClearCurrentMutableOperation()
 }
 
 
-void FCustomizableObjectSystemPrivate::UpdateMemoryLimit()
+void UCustomizableObjectSystemPrivate::UpdateMemoryLimit()
 {
 	// This must run on game thread, and when the mutable thread is not running
 	check(IsInGameThread());
@@ -1596,7 +1577,7 @@ namespace impl
 	void CreateMutableInstance(const TSharedRef<FUpdateContextPrivate>& Operation)
 	{
 		UCustomizableObjectSystem* System = UCustomizableObjectSystem::GetInstanceChecked(); // Save since UCustomizableObjectSystem::BeginDestroy always waits for all tasks to finish
-		const FCustomizableObjectSystemPrivate* SystemPrivate = System->GetPrivateChecked();
+		const UCustomizableObjectSystemPrivate* SystemPrivate = System->GetPrivateChecked();
 		
 		Operation->UpdateStartBytes = mu::FGlobalMemoryCounter::GetCounter();
 		mu::FGlobalMemoryCounter::Zero();
@@ -1910,7 +1891,7 @@ namespace impl
 	{
 		MUTABLE_CPUPROFILER_SCOPE(Subtask_Mutable_GetImages)
 
-		const FCustomizableObjectSystemPrivate* CustomizableObjectSystemPrivateData = UCustomizableObjectSystem::GetInstanceChecked()->GetPrivateChecked();
+		const UCustomizableObjectSystemPrivate* CustomizableObjectSystemPrivateData = UCustomizableObjectSystem::GetInstanceChecked()->GetPrivateChecked();
 		mu::System* System = CustomizableObjectSystemPrivateData->MutableSystem.get();
 		check(System != nullptr);
 
@@ -2262,7 +2243,7 @@ namespace impl
 			return;
 		}
 
-		FCustomizableObjectSystemPrivate * CustomizableObjectSystemPrivateData = System->GetPrivateChecked();
+		UCustomizableObjectSystemPrivate * CustomizableObjectSystemPrivateData = System->GetPrivateChecked();
 
 		// Actual work
 		// TODO MTBL-391: Review This hotfix
@@ -2349,7 +2330,7 @@ namespace impl
 			}
 		} // if (!bInstanceValid)
 
-		FCustomizableObjectSystemPrivate* CustomizableObjectSystemPrivateData = System->GetPrivateChecked();
+		UCustomizableObjectSystemPrivate* CustomizableObjectSystemPrivateData = System->GetPrivateChecked();
 
 		// Next Task: Release Mutable. We need this regardless if we cancel or not
 		//-------------------------------------------------------------		
@@ -2465,7 +2446,7 @@ namespace impl
 		// When protecting textures there mustn't be any left from a previous update
 		check(System->ProtectedCachedTextures.Num() == 0);
 
-		FCustomizableObjectSystemPrivate* SystemPrivateData = System->GetPrivateChecked();
+		UCustomizableObjectSystemPrivate* SystemPrivateData = System->GetPrivateChecked();
 
 		// TODO: If this is the first code that runs after the CO program has finished AND if it's
 		// guaranteed that the next CO program hasn't started yet, we need to call ClearActiveObject
@@ -2549,7 +2530,7 @@ namespace impl
 		MUTABLE_CPUPROFILER_SCOPE(Task_Game_ReleaseInstanceID)
 
 		UCustomizableObjectSystem* System = UCustomizableObjectSystem::GetInstanceChecked();
-		FCustomizableObjectSystemPrivate* SystemPrivateData = System->GetPrivateChecked();
+		UCustomizableObjectSystemPrivate* SystemPrivateData = System->GetPrivateChecked();
 
 		const mu::Ptr<mu::System> MutableSystem = SystemPrivateData->MutableSystem;
 
@@ -2572,7 +2553,7 @@ namespace impl
 		MUTABLE_CPUPROFILER_SCOPE(Task_Game_LockMeshCache);
 
 		UCustomizableObjectSystem* System = UCustomizableObjectSystem::GetInstanceChecked();
-		FCustomizableObjectSystemPrivate* SystemPrivate = System->GetPrivateChecked();
+		UCustomizableObjectSystemPrivate* SystemPrivate = System->GetPrivateChecked();
 
 		const UCustomizableObject* CustomizableObject = Operation->Instance->GetCustomizableObject();
 		UCustomizableObjectPrivate* CustomizableObjectPrivate = CustomizableObject->GetPrivate();
@@ -2619,7 +2600,7 @@ namespace impl
 
 
 		UCustomizableObjectSystem* System = UCustomizableObjectSystem::GetInstanceChecked(); // Save since UCustomizableObjectSystem::BeginDestroy always waits for all tasks to finish
-		FCustomizableObjectSystemPrivate* SystemPrivate = System->GetPrivateChecked();
+		UCustomizableObjectSystemPrivate* SystemPrivate = System->GetPrivateChecked();
 		
 		CreateMutableInstance(Operation);
 		FixLODs(Operation);
@@ -2737,7 +2718,7 @@ namespace impl
 			return;
 		}
 
-		FCustomizableObjectSystemPrivate* SystemPrivateData = System->GetPrivateChecked();
+		UCustomizableObjectSystemPrivate* SystemPrivateData = System->GetPrivateChecked();
 
 		SystemPrivateData->CurrentInstanceBeingUpdated = CandidateInstance;
 
@@ -2962,7 +2943,7 @@ bool UCustomizableObjectSystem::Tick(float DeltaTime)
 	return true;
 #endif
 
-	if (!Private.IsValid())
+	if (!Private)
 	{
 		return true;
 	}
@@ -3694,7 +3675,7 @@ void UCustomizableObjectSystem::ClearImageCache()
 }
 
 
-bool FCustomizableObjectSystemPrivate::IsMutableAnimInfoDebuggingEnabled() const
+bool UCustomizableObjectSystemPrivate::IsMutableAnimInfoDebuggingEnabled() const
 { 
 #if WITH_EDITORONLY_DATA
 	return EnableMutableAnimInfoDebugging > 0;
@@ -3704,14 +3685,14 @@ bool FCustomizableObjectSystemPrivate::IsMutableAnimInfoDebuggingEnabled() const
 }
 
 
-FUnrealMutableImageProvider* FCustomizableObjectSystemPrivate::GetImageProviderChecked() const
+FUnrealMutableImageProvider* UCustomizableObjectSystemPrivate::GetImageProviderChecked() const
 {
 	check(ImageProvider)
 	return ImageProvider.Get();
 }
 
 
-void FCustomizableObjectSystemPrivate::StartUpdateSkeletalMesh(const TSharedRef<FUpdateContextPrivate>& Context)
+void UCustomizableObjectSystemPrivate::StartUpdateSkeletalMesh(const TSharedRef<FUpdateContextPrivate>& Context)
 {
 	Context->UpdateStarted = true;
 	TRACE_BEGIN_REGION(UE_MUTABLE_UPDATE_REGION);
@@ -3726,7 +3707,7 @@ void FCustomizableObjectSystemPrivate::StartUpdateSkeletalMesh(const TSharedRef<
 }
 
 
-bool FCustomizableObjectSystemPrivate::IsUpdating(const UCustomizableObjectInstance& Instance) const
+bool UCustomizableObjectSystemPrivate::IsUpdating(const UCustomizableObjectInstance& Instance) const
 {
 	if (CurrentMutableOperation && CurrentMutableOperation->Instance.Get() == &Instance)
 	{
@@ -3742,7 +3723,7 @@ bool FCustomizableObjectSystemPrivate::IsUpdating(const UCustomizableObjectInsta
 }
 
 
-void FCustomizableObjectSystemPrivate::UpdateStats()
+void UCustomizableObjectSystemPrivate::UpdateStats()
 {
 	NumSkeletalMeshes = 0;
 	
diff --git a/Engine/Plugins/Experimental/Mutable/Source/CustomizableObject/Private/MuCO/CustomizableObjectSystemPrivate.h b/Engine/Plugins/Experimental/Mutable/Source/CustomizableObject/Private/MuCO/CustomizableObjectSystemPrivate.h
--- a/Engine/Plugins/Experimental/Mutable/Source/CustomizableObject/Private/MuCO/CustomizableObjectSystemPrivate.h
+++ b/Engine/Plugins/Experimental/Mutable/Source/CustomizableObject/Private/MuCO/CustomizableObjectSystemPrivate.h
@@ -520,8 +520,11 @@ struct FMutableReleasePlatformOperationData
 };
 
 
-class FCustomizableObjectSystemPrivate : public FGCObject
+UCLASS()
+class UCustomizableObjectSystemPrivate : public UObject
 {
+	GENERATED_BODY()
+	
 public:
 	// Singleton for the unreal mutable system.
 	static UCustomizableObjectSystem* SSystem;
@@ -572,10 +575,6 @@ public:
 
 	void AddGameThreadTask(const FMutableTask& Task);
 
-	/** FSerializableObject interface */
-	virtual void AddReferencedObjects(FReferenceCollector& Collector) override;
-	virtual FString GetReferencerName() const override;
-
 	// Remove references to cached objects that have been deleted in the unreal
 	// side, and cannot be cached anyway.
 	// This should only happen in the game thread
@@ -620,7 +619,7 @@ public:
 
 	TMap<FMutableImageCacheKey, uint32> TextureReferenceCount; // Keeps a count of texture usage to decide if they have to be blocked from GC during an update
 
-	// This is protected from GC by AddReferencedObjects
+	UPROPERTY(Transient)
 	TObjectPtr<UCustomizableObjectInstance> CurrentInstanceBeingUpdated = nullptr;
 
 	TSharedPtr<FUpdateContextPrivate> CurrentMutableOperation = nullptr;
@@ -651,6 +650,7 @@ public:
 	
 #if WITH_EDITORONLY_DATA
 	/** Mutable default image provider. Used by the COIEditor and Instance/Descriptor APIs. */
+	UPROPERTY(Transient)
 	TObjectPtr<UEditorImageProvider> EditorImageProvider = nullptr;
 #endif
 
diff --git a/Engine/Plugins/Experimental/Mutable/Source/CustomizableObject/Private/MuCO/UnrealExtensionDataStreamer.cpp b/Engine/Plugins/Experimental/Mutable/Source/CustomizableObject/Private/MuCO/UnrealExtensionDataStreamer.cpp
--- a/Engine/Plugins/Experimental/Mutable/Source/CustomizableObject/Private/MuCO/UnrealExtensionDataStreamer.cpp
+++ b/Engine/Plugins/Experimental/Mutable/Source/CustomizableObject/Private/MuCO/UnrealExtensionDataStreamer.cpp
@@ -10,11 +10,11 @@
 #include "MuCO/CustomizableObjectSystemPrivate.h"
 #include "Templates/SharedPointer.h"
 
-FUnrealExtensionDataStreamer::FUnrealExtensionDataStreamer(const TSharedRef<FCustomizableObjectSystemPrivate>& InSystemPrivate)
-	: SystemPrivate(InSystemPrivate)
+FUnrealExtensionDataStreamer::FUnrealExtensionDataStreamer(const TWeakObjectPtr<UCustomizableObjectSystemPrivate>& InSystemPrivateWeak)
 {
 	check(IsInGameThread());
 
+	SystemPrivate = InSystemPrivateWeak;
 	Mutex = new FCriticalSection();
 }
 
@@ -152,13 +152,20 @@ TSharedRef<const mu::FExtensionDataLoadHandle> FUnrealExtensionDataStreamer::Sta
 // If the requested object is already loaded, this function can return nullptr and still be
 // considered successful.
 TSharedPtr<FStreamableHandle> FUnrealExtensionDataStreamer::StartLoadOnGameThread(
-	const TSharedRef<FCustomizableObjectSystemPrivate>& SystemPrivate,
+	const TWeakObjectPtr<UCustomizableObjectSystemPrivate>& SystemPrivateWeak,
 	const TWeakObjectPtr<UCustomizableObject>& ObjectToLoadFor,
 	const TSharedRef<mu::FExtensionDataLoadHandle>& LoadHandle)
 {
 	check(IsInGameThread());
 	check(LoadHandle->Data->Origin == mu::ExtensionData::EOrigin::ConstantStreamed);
 
+	UCustomizableObjectSystemPrivate* SystemPrivate = SystemPrivateWeak.Get();
+	if (!SystemPrivate)
+	{
+		LoadHandle->LoadState = mu::FExtensionDataLoadHandle::ELoadState::FailedToLoad;
+		return nullptr;
+	}
+	
 	UCustomizableObject* Object = ObjectToLoadFor.Get();
 	if (!Object)
 	{
@@ -196,7 +203,8 @@ TSharedPtr<FStreamableHandle> FUnrealExtensionDataStreamer::StartLoadOnGameThrea
 	TArray<FSoftObjectPath> TargetsToStream;
 	TargetsToStream.Add(StreamedData.GetPath().ToSoftObjectPath());
 
-	TFunction<void()> OnLoadComplete = [SystemPrivate, ObjectToLoadFor, LoadHandle]()
+	check(SystemPrivate->ExtensionDataStreamer);
+	TFunction<void()> OnLoadComplete = [ExtensionDataStreamer = SystemPrivate->ExtensionDataStreamer, ObjectToLoadFor, LoadHandle]()
 	{
 		UCustomizableObject* Object = ObjectToLoadFor.Get();
 		if (!Object)
@@ -205,8 +213,7 @@ TSharedPtr<FStreamableHandle> FUnrealExtensionDataStreamer::StartLoadOnGameThrea
 			return;
 		}
 
-		check(SystemPrivate->ExtensionDataStreamer);
-		SystemPrivate->ExtensionDataStreamer->NotifyLoadCompleted(Object, LoadHandle);
+		ExtensionDataStreamer->NotifyLoadCompleted(Object, LoadHandle);
 	};
 
 	const bool bManageActiveHandle = false;
diff --git a/Engine/Plugins/Experimental/Mutable/Source/CustomizableObject/Public/MuCO/UnrealExtensionDataStreamer.h b/Engine/Plugins/Experimental/Mutable/Source/CustomizableObject/Public/MuCO/UnrealExtensionDataStreamer.h
--- a/Engine/Plugins/Experimental/Mutable/Source/CustomizableObject/Public/MuCO/UnrealExtensionDataStreamer.h
+++ b/Engine/Plugins/Experimental/Mutable/Source/CustomizableObject/Public/MuCO/UnrealExtensionDataStreamer.h
@@ -8,14 +8,14 @@
 #include "MuR/ExtensionDataStreamer.h"
 #include "Templates/SharedPointerFwd.h"
 
-class FCustomizableObjectSystemPrivate;
+class UCustomizableObjectSystemPrivate;
 class UCustomizableObject;
 
 /** An implementation of mu::ExtensionDataStreamer designed to work with the Customizable Object integration */
 class CUSTOMIZABLEOBJECT_API FUnrealExtensionDataStreamer : public mu::ExtensionDataStreamer
 {
 public:
-	FUnrealExtensionDataStreamer(const TSharedRef<FCustomizableObjectSystemPrivate>& InSystemPrivate);
+	FUnrealExtensionDataStreamer(const TWeakObjectPtr<UCustomizableObjectSystemPrivate>& InSystemPrivateWeak);
 	virtual ~FUnrealExtensionDataStreamer();
 
 	/** Not copyable because TFuture isn't copyable. This shouldn't be limiting in practise. */
@@ -58,14 +58,14 @@ public:
 private:
 	/** Queued up on the Game thread to start the async load */
 	static TSharedPtr<FStreamableHandle> StartLoadOnGameThread(
-		const TSharedRef<FCustomizableObjectSystemPrivate>& SystemPrivate,
+		const TWeakObjectPtr<UCustomizableObjectSystemPrivate>& SystemPrivateWeak,
 		const TWeakObjectPtr<UCustomizableObject>& ObjectToLoadFor,
 		const TSharedRef<mu::FExtensionDataLoadHandle>& LoadHandle);
 
 	/** Called when the async load has finished loading the Extension Data */
 	void NotifyLoadCompleted(UCustomizableObject* Object, const TSharedRef<mu::FExtensionDataLoadHandle>& LoadHandle);
 
-	TSharedRef<FCustomizableObjectSystemPrivate> SystemPrivate;
+	TWeakObjectPtr<UCustomizableObjectSystemPrivate> SystemPrivate; // Can not be a TStrongObjectPtr since then it will create a cycle b TSharedPtr and UObjects.
 
 	/** Guards PendingLoads, ShouldCancelPtr and ActiveObject */
 	FCriticalSection* Mutex = nullptr;
