[Mutable] Add CVar to disable Mutable. Remove CompilationDisabled code.
Split the UCustomizableInstancePrivateData::DiscardResourcesAndSetReferenceSkeletalMesh into two separate methods.
Simplify the DiscardInstance() method and fix its bUpdating condition.

CVar Mutable.Enabled with values true/false.
 - false - Disables compilation, generation, and streaming and removes the system ticker.
 - true - All core features are enabled.

#rnx

diff --git a/Engine/Plugins/Experimental/Mutable/Source/CustomizableObject/Classes/MuCO/CustomizableObject.h b/Engine/Plugins/Experimental/Mutable/Source/CustomizableObject/Classes/MuCO/CustomizableObject.h
--- a/Engine/Plugins/Experimental/Mutable/Source/CustomizableObject/Classes/MuCO/CustomizableObject.h
+++ b/Engine/Plugins/Experimental/Mutable/Source/CustomizableObject/Classes/MuCO/CustomizableObject.h
@@ -1280,7 +1280,7 @@ public:
 	FString FindIntParameterValueName(int32 ParamIndex, int32 ParamValue) const;
 
 	//
-	USkeletalMesh* GetRefSkeletalMesh(int32 ComponentIndex = 0);
+	USkeletalMesh* GetRefSkeletalMesh(int32 ComponentIndex = 0) const;
 	
 	//
 	FMutableRefSkeletalMeshData* GetRefSkeletalMeshData(int32 ComponentIndex = 0);
diff --git a/Engine/Plugins/Experimental/Mutable/Source/CustomizableObject/Classes/MuCO/CustomizableObjectSystem.h b/Engine/Plugins/Experimental/Mutable/Source/CustomizableObject/Classes/MuCO/CustomizableObjectSystem.h
--- a/Engine/Plugins/Experimental/Mutable/Source/CustomizableObject/Classes/MuCO/CustomizableObjectSystem.h
+++ b/Engine/Plugins/Experimental/Mutable/Source/CustomizableObject/Classes/MuCO/CustomizableObjectSystem.h
@@ -56,7 +56,7 @@ extern TAutoConsoleVariable<bool> CVarRollbackFixModelDiskStreamerDataRace;
 struct FEditorCompileSettings 
 {
 	// General case
-	bool bDisableCompilation;
+	bool bIsMutableEnabled = true;
 
 	// Auto Compile 
 	bool bEnableAutomaticCompilation = true;
@@ -200,6 +200,10 @@ public:
 	// Return true if the singleton has been created. It is different than GetInstance in that GetInstance will create it if it doesn't exist.
 	static bool IsCreated();
 
+	/** Returns the current status of Mutable. Only when active is it possible to compile COs, generate instances, and stream textures.
+	  * @return True if Mutable is enabled. */
+	static bool IsActive();
+
 	// Begin UObject interface.
 	virtual void BeginDestroy() override;
 	virtual FString GetDesc() override;
@@ -231,9 +235,6 @@ public:
 	// Called whenever the Mutable Editor Settings change, copying the new value of the current needed settings to the Customizable Object System
 	void EditorSettingsChanged(const FEditorCompileSettings& InEditorSettings);
 
-	// If compilation is disabled, Customizable Objects won't be compiled in the editor
-	bool IsCompilationDisabled() const;
-
 	// If true, uncompiled Customizable Objects will be compiled whenever an instance update is required
 	bool IsAutoCompileEnabled() const;
 
diff --git a/Engine/Plugins/Experimental/Mutable/Source/CustomizableObject/Internal/MuCO/CustomizableInstancePrivateData.h b/Engine/Plugins/Experimental/Mutable/Source/CustomizableObject/Internal/MuCO/CustomizableInstancePrivateData.h
--- a/Engine/Plugins/Experimental/Mutable/Source/CustomizableObject/Internal/MuCO/CustomizableInstancePrivateData.h
+++ b/Engine/Plugins/Experimental/Mutable/Source/CustomizableObject/Internal/MuCO/CustomizableInstancePrivateData.h
@@ -231,11 +231,16 @@ public:
 	// The following method is basically copied from PostEditChangeProperty and/or SkeletalMesh.cpp to be able to replicate PostEditChangeProperty without the editor
 	void PostEditChangePropertyWithoutEditor(USkeletalMesh* SkeletalMesh);
 	
-	void DiscardResourcesAndSetReferenceSkeletalMesh(UCustomizableObjectInstance* Instance);
+	/** Calls ReleaseResources on all SkeletalMeshes generated by this instance and invalidates the generated data.
+	  * It should not be called if the meshes are still in use or shared with other instances. */
+	void DiscardResources();
 
 	// Releases all the mutable resources this instance holds, should only be called when it is not going to be used any more.
 	void ReleaseMutableResources(bool bCalledFromBeginDestroy, const UCustomizableObjectInstance& Instance);
 
+	/** Set the reference SkeletalMesh, or an empty mesh, to all actors using this instance. */
+	void SetDefaultSkeletalMesh(bool bSetEmptyMesh = false) const;
+
 	const TArray<FAnimInstanceOverridePhysicsAsset>* GetGeneratedPhysicsAssetsForAnimInstance(TSubclassOf<UAnimInstance> AnimInstance) const;
 
 	/** */
diff --git a/Engine/Plugins/Experimental/Mutable/Source/CustomizableObject/Private/MuCO/CustomizableObject.cpp b/Engine/Plugins/Experimental/Mutable/Source/CustomizableObject/Private/MuCO/CustomizableObject.cpp
--- a/Engine/Plugins/Experimental/Mutable/Source/CustomizableObject/Private/MuCO/CustomizableObject.cpp
+++ b/Engine/Plugins/Experimental/Mutable/Source/CustomizableObject/Private/MuCO/CustomizableObject.cpp
@@ -871,8 +871,8 @@ bool UCustomizableObject::ConditionalAutoCompile()
 		return false;
 	}
 
-	// Don't compile if we're running game, Compile and/or if AutoCompile is disabled
-	if (IsRunningGame() || System->IsCompilationDisabled() || !System->IsAutoCompileEnabled())
+	// Don't compile if we're running game or if Mutable or AutoCompile is disabled.
+	if (IsRunningGame() || !System->IsActive() || !System->IsAutoCompileEnabled())
 	{
 		System->AddUncompiledCOWarning(*this);
 		return false;
@@ -1080,7 +1080,7 @@ void UCustomizableObject::AddUncompiledCOWarning(const FString& AdditionalLoggin
 }
 
 
-USkeletalMesh* UCustomizableObject::GetRefSkeletalMesh(int32 ComponentIndex)
+USkeletalMesh* UCustomizableObject::GetRefSkeletalMesh(int32 ComponentIndex) const
 {
 #if WITH_EDITORONLY_DATA
 	if (ReferenceSkeletalMeshes.IsValidIndex(ComponentIndex))
diff --git a/Engine/Plugins/Experimental/Mutable/Source/CustomizableObject/Private/MuCO/CustomizableObjectInstance.cpp b/Engine/Plugins/Experimental/Mutable/Source/CustomizableObject/Private/MuCO/CustomizableObjectInstance.cpp
--- a/Engine/Plugins/Experimental/Mutable/Source/CustomizableObject/Private/MuCO/CustomizableObjectInstance.cpp
+++ b/Engine/Plugins/Experimental/Mutable/Source/CustomizableObject/Private/MuCO/CustomizableObjectInstance.cpp
@@ -2754,8 +2754,14 @@ bool UCustomizableObjectInstance::IsSelectedParameterProfileDirty() const
 //}
 
 
-void UCustomizableInstancePrivate::DiscardResourcesAndSetReferenceSkeletalMesh(UCustomizableObjectInstance* Instance)
+void UCustomizableInstancePrivate::DiscardResources()
 {
+	UCustomizableObjectInstance* Instance = Cast<UCustomizableObjectInstance>(GetOuter());
+	if (!Instance)
+	{
+		return;
+	}
+
 	if (SkeletalMeshStatus == ESkeletalMeshStatus::Success)
 	{
 		for (int32 Component = 0; Component < Instance->SkeletalMeshes.Num(); ++Component)
@@ -2774,29 +2780,49 @@ void UCustomizableInstancePrivate::DiscardResourcesAndSetReferenceSkeletalMesh(U
 	
 	Instance->SkeletalMeshes.Reset();
 	DescriptorRuntimeHash = FDescriptorRuntimeHash();
+}
+
+
+void UCustomizableInstancePrivate::SetDefaultSkeletalMesh(bool bSetEmptyMesh) const
+{
+	UCustomizableObjectInstance* Instance = Cast<UCustomizableObjectInstance>(GetOuter());
+	if (!Instance || !Instance->GetCustomizableObject())
+	{
+		return;
+	}
+
+	UCustomizableObject* CustomizableObject = Instance->GetCustomizableObject();
+	const int32 NumComponents = CustomizableObject->GetComponentCount();
 
 	for (TObjectIterator<UCustomizableObjectInstanceUsage> It; It; ++It)
 	{
 		UCustomizableObjectInstanceUsage* CustomizableObjectInstanceUsage = *It;
+		if (!IsValid(CustomizableObjectInstanceUsage) || CustomizableObjectInstanceUsage->GetCustomizableObjectInstance() != Instance)
+		{
+			continue;
+		}
 
 #if WITH_EDITOR
-		if (IsValid(CustomizableObjectInstanceUsage) && CustomizableObjectInstanceUsage->IsNetMode(NM_DedicatedServer))
+		if (CustomizableObjectInstanceUsage->IsNetMode(NM_DedicatedServer))
 		{
 			continue;
 		}
 #endif
 
-		if (IsValid(CustomizableObjectInstanceUsage) && CustomizableObjectInstanceUsage->GetCustomizableObjectInstance() == Instance)
+		USkeletalMesh* SkeletalMesh = nullptr;
+		
+		FMutableRefSkeletalMeshData* RefSkeletalMeshData = CustomizableObject->GetRefSkeletalMeshData(CustomizableObjectInstanceUsage->GetComponentIndex());
+		if (!bSetEmptyMesh && RefSkeletalMeshData)
 		{
-			UCustomizableObject* CustomizableObject = Instance->GetCustomizableObject();
-			bool bReplaceDiscardedWithReferenceMesh = UCustomizableObjectSystem::GetInstance()->GetPrivate()->IsReplaceDiscardedWithReferenceMeshEnabled();
-			CustomizableObjectInstanceUsage->SetSkeletalMesh(CustomizableObject && bReplaceDiscardedWithReferenceMesh ? CustomizableObject->GetRefSkeletalMesh(CustomizableObjectInstanceUsage->GetComponentIndex()) : nullptr);
+			// Force load the reference mesh if necessary. 
+			SkeletalMesh = TSoftObjectPtr<USkeletalMesh>(RefSkeletalMeshData->SkeletalMeshAssetPath).LoadSynchronous();
 		}
+		
+		CustomizableObjectInstanceUsage->SetSkeletalMesh(SkeletalMesh);
 	}
 }
 
 
-
 namespace
 {
 	inline float unpack_uint8(uint8 i)
@@ -6350,7 +6376,7 @@ TSubclassOf<UAnimInstance> UCustomizableObjectInstance::GetAnimBP(int32 Componen
 	
 	if (!ComponentData)
 	{
-		FString ErrorMsg = FString::Printf(TEXT("Tried to access and invalid component index [%d] in a Mutable Instance."), ComponentIndex);
+		FString ErrorMsg = FString::Printf(TEXT("Tried to access an invalid component index [%d] in a Mutable Instance."), ComponentIndex);
 		UE_LOG(LogMutable, Error, TEXT("%s"), *ErrorMsg);
 #if WITH_EDITOR
 		FMessageLogModule& MessageLogModule = FModuleManager::LoadModuleChecked<FMessageLogModule>("MessageLog");
@@ -6396,7 +6422,7 @@ void UCustomizableObjectInstance::ForEachAnimInstance(int32 ComponentIndex, FEac
 	
 	if (!ComponentData)
 	{
-		FString ErrorMsg = FString::Printf(TEXT("Tried to access and invalid component index [%d] in a Mutable Instance."), ComponentIndex);
+		FString ErrorMsg = FString::Printf(TEXT("Tried to access an invalid component index [%d] in a Mutable Instance."), ComponentIndex);
 		UE_LOG(LogMutable, Error, TEXT("%s"), *ErrorMsg);
 #if WITH_EDITOR
 		FMessageLogModule& MessageLogModule = FModuleManager::LoadModuleChecked<FMessageLogModule>("MessageLog");
@@ -6543,7 +6569,7 @@ void UCustomizableObjectInstance::ForEachAnimInstance(int32 ComponentIndex, FEac
 	
 	if (!ComponentData)
 	{
-		FString ErrorMsg = FString::Printf(TEXT("Tried to access and invalid component index [%d] in a Mutable Instance."), ComponentIndex);
+		FString ErrorMsg = FString::Printf(TEXT("Tried to access an invalid component index [%d] in a Mutable Instance."), ComponentIndex);
 		UE_LOG(LogMutable, Error, TEXT("%s"), *ErrorMsg);
 #if WITH_EDITOR
 		FMessageLogModule& MessageLogModule = FModuleManager::LoadModuleChecked<FMessageLogModule>("MessageLog");
diff --git a/Engine/Plugins/Experimental/Mutable/Source/CustomizableObject/Private/MuCO/CustomizableObjectMipDataProvider.cpp b/Engine/Plugins/Experimental/Mutable/Source/CustomizableObject/Private/MuCO/CustomizableObjectMipDataProvider.cpp
--- a/Engine/Plugins/Experimental/Mutable/Source/CustomizableObject/Private/MuCO/CustomizableObjectMipDataProvider.cpp
+++ b/Engine/Plugins/Experimental/Mutable/Source/CustomizableObject/Private/MuCO/CustomizableObjectMipDataProvider.cpp
@@ -288,7 +288,14 @@ namespace impl
 
 int32 FMutableTextureMipDataProvider::GetMips(const FTextureUpdateContext& Context, int32 StartingMipIndex, const FTextureMipInfoArray& MipInfos, const FTextureUpdateSyncOptions& SyncOptions)
 {
-	MUTABLE_CPUPROFILER_SCOPE(FMutableTextureMipDataProvider::GetMips)
+	MUTABLE_CPUPROFILER_SCOPE(FMutableTextureMipDataProvider::GetMips);
+
+	if (!UCustomizableObjectSystem::IsActive())
+	{
+		// Mutable is disabled. Skip all mip operations and mark the update task as completed.
+		AdvanceTo(ETickState::Done, ETickThread::None);
+		return CurrentFirstLODIdx;
+	}
 
 #if WITH_EDITOR
 	check(Context.Texture->HasPendingInitOrStreaming());
diff --git a/Engine/Plugins/Experimental/Mutable/Source/CustomizableObject/Private/MuCO/CustomizableObjectSystem.cpp b/Engine/Plugins/Experimental/Mutable/Source/CustomizableObject/Private/MuCO/CustomizableObjectSystem.cpp
--- a/Engine/Plugins/Experimental/Mutable/Source/CustomizableObject/Private/MuCO/CustomizableObjectSystem.cpp
+++ b/Engine/Plugins/Experimental/Mutable/Source/CustomizableObject/Private/MuCO/CustomizableObjectSystem.cpp
@@ -60,6 +60,14 @@ DECLARE_CYCLE_STAT(TEXT("MutableTask"), STAT_MutableTask, STATGROUP_Game);
 
 UCustomizableObjectSystem* UCustomizableObjectSystemPrivate::SSystem = nullptr;
 
+bool bIsMutableEnabled = true;
+
+static FAutoConsoleVariableRef CVarMutableEnabled(
+	TEXT("Mutable.Enabled"),
+	bIsMutableEnabled,
+	TEXT("true/false - Disabling Mutable will turn off CO compilation, mesh generation, and texture streaming and will remove the system ticker. "),
+	FConsoleVariableDelegate::CreateStatic(&UCustomizableObjectSystemPrivate::OnMutableEnabledChanged));
+
 static TAutoConsoleVariable<int32> CVarWorkingMemory(
 	TEXT("mutable.WorkingMemory"),
 #if !PLATFORM_DESKTOP
@@ -450,6 +458,11 @@ bool UCustomizableObjectSystem::IsCreated()
 	return UCustomizableObjectSystemPrivate::SSystem != 0;
 }
 
+bool UCustomizableObjectSystem::IsActive()
+{
+	return IsCreated() && bIsMutableEnabled;
+}
+
 
 void UCustomizableObjectSystem::InitSystem()
 {
@@ -460,19 +473,9 @@ void UCustomizableObjectSystem::InitSystem()
 
 	Private->bReplaceDiscardedWithReferenceMesh = false;
 
-	const IConsoleVariable* CVarSupport16BitBoneIndex = IConsoleManager::Get().FindConsoleVariable(TEXT("r.GPUSkin.Support16BitBoneIndex"));
-	Private->bSupport16BitBoneIndex = CVarSupport16BitBoneIndex ? CVarSupport16BitBoneIndex->GetBool() : false;
-
-	CVarMutableSinkFunction();
-
 	Private->CurrentMutableOperation = nullptr;
 	Private->CurrentInstanceBeingUpdated = nullptr;
 
-#if !UE_SERVER
-	Private->TickDelegate = FTickerDelegate::CreateUObject(this, &UCustomizableObjectSystem::Tick);
-	Private->TickDelegateHandle = FTSTicker::GetCoreTicker().AddTicker(Private->TickDelegate, 0.f);
-#endif // !UE_SERVER
-
 	Private->LastWorkingMemoryBytes = CVarWorkingMemory.GetValueOnGameThread() * 1024;
 	Private->LastGeneratedResourceCacheSize = CVarGeneratedResourcesCacheSize.GetValueOnGameThread();
 
@@ -500,16 +503,18 @@ void UCustomizableObjectSystem::InitSystem()
 	RegisterImageProvider(Private->EditorImageProvider);
 #endif
 	
-#if WITH_EDITOR
-	if (!IsRunningGame())
-	{
-		FEditorDelegates::PreBeginPIE.AddUObject(this, &UCustomizableObjectSystem::OnPreBeginPIE);
-	}
-#endif
-
 	DefaultInstanceLODManagement = NewObject<UCustomizableInstanceLODManagement>();
 	check(DefaultInstanceLODManagement != nullptr);
 	CurrentInstanceLODManagement = DefaultInstanceLODManagement;
+
+	// This CVar is constant for the lifespan of the program. Read its value once. 
+	const IConsoleVariable* CVarSupport16BitBoneIndex = IConsoleManager::Get().FindConsoleVariable(TEXT("r.GPUSkin.Support16BitBoneIndex"));
+	Private->bSupport16BitBoneIndex = CVarSupport16BitBoneIndex ? CVarSupport16BitBoneIndex->GetBool() : false;
+
+	// Read non-constant CVars and do work if required.
+	CVarMutableSinkFunction();
+
+	Private->OnMutableEnabledChanged();
 }
 
 
@@ -1117,6 +1122,15 @@ void UCustomizableObjectSystemPrivate::EnqueueUpdateSkeletalMesh(const TSharedRe
 	const bool bIsPlayerOrNearIt = InstancePrivate->HasCOInstanceFlags(UsedByPlayerOrNearIt);
 	UE_LOG(LogMutable, Log, TEXT("Enqueue UpdateSkeletalMesh Async. Instance=%d, Frame=%d, Priority=%d, dist=%f, bIsPlayerOrNearIt=%d"), InstanceId, GFrameNumber, static_cast<int32>(Priority), Distance, bIsPlayerOrNearIt);				
 	
+	if (!bIsMutableEnabled)
+	{
+		// Mutable is disabled. Set the reference SkeletalMesh and finish the update with success to avoid breaking too many things.
+		Context->UpdateResult = EUpdateResult::Success;
+		InstancePrivate->SetDefaultSkeletalMesh();
+		FinishUpdateGlobal(Context);
+		return;
+	}
+
 	if (!Instance->CanUpdateInstance())
 	{
 		Context->UpdateResult = EUpdateResult::Error;
@@ -1124,8 +1138,6 @@ void UCustomizableObjectSystemPrivate::EnqueueUpdateSkeletalMesh(const TSharedRe
 		return;
 	}
 
-	UCustomizableObjectSystem* System = UCustomizableObjectSystem::GetInstance();
-
 	const EUpdateRequired UpdateRequired = IsUpdateRequired(*Instance, Context->bOnlyUpdateIfNotGenerated, false, Context->bIgnoreCloseDist);
 	switch (UpdateRequired)
 	{
@@ -1223,7 +1235,7 @@ void UCustomizableObjectSystemPrivate::EnqueueUpdateSkeletalMesh(const TSharedRe
 
 	case EUpdateRequired::Discard:
 	{
-		System->GetPrivate()->InitDiscardResourcesSkeletalMesh(Instance);
+		InitDiscardResourcesSkeletalMesh(Instance);
 
 		Context->UpdateResult = EUpdateResult::ErrorDiscarded;
 		FinishUpdateGlobal(Context);
@@ -1403,15 +1415,10 @@ bool UCustomizableObjectSystem::CheckIfDiskOrMipUpdateOperationsPending(const UC
 void UCustomizableObjectSystem::EditorSettingsChanged(const FEditorCompileSettings& InEditorSettings)
 {
 	EditorSettings = InEditorSettings;
-}
-
 
-bool UCustomizableObjectSystem::IsCompilationDisabled() const
-{
-	return EditorSettings.bDisableCompilation;
+	CVarMutableEnabled->Set(InEditorSettings.bIsMutableEnabled);
 }
 
-
 bool UCustomizableObjectSystem::IsAutoCompileEnabled() const
 {
 	return EditorSettings.bEnableAutomaticCompilation;
@@ -2979,7 +2986,7 @@ bool UCustomizableObjectSystem::Tick(float DeltaTime)
 	Private->UpdateStats();
 	
 	// Get a new operation if we aren't working on one
-	if (!Private->CurrentMutableOperation)
+	if (!Private->CurrentMutableOperation && bIsMutableEnabled)
 	{
 		// Reset the instance relevancy
 		// The RequestedUpdates only refer to LOD changes. User Customization and discards are handled separately
@@ -3151,7 +3158,14 @@ bool UCustomizableObjectSystem::Tick(float DeltaTime)
 	Private->MutableTaskGraph.Tick();
 
 	Private->LogBenchmarkUtil.UpdateStats(); // Must to be the last thing to perform
-	
+
+	if (!bIsMutableEnabled && !Private->CurrentMutableOperation)
+	{
+		// Mutable has been disabled. Unregister the ticker if there is no CurrentMutableOperation.
+		FTSTicker::GetCoreTicker().RemoveTicker(Private->TickDelegateHandle);
+		Private->TickDelegateHandle.Reset();
+	}
+
 	return true;
 }
 
@@ -3165,6 +3179,10 @@ TAutoConsoleVariable<int32> CVarMaxNumInstancesToDiscardPerTick(
 
 void UCustomizableObjectSystem::DiscardInstances()
 {
+	MUTABLE_CPUPROFILER_SCOPE(DiscardInstances);
+
+	check(IsInGameThread());
+
 	// Handle instance discards
 	int32 NumInstancesDiscarded = 0;
 	const int32 DiscardLimitPerTick = CVarMaxNumInstancesToDiscardPerTick.GetValueOnGameThread();
@@ -3173,28 +3191,20 @@ void UCustomizableObjectSystem::DiscardInstances()
 		Iterator && NumInstancesDiscarded < DiscardLimitPerTick;
 		++Iterator)
 	{
-		MUTABLE_CPUPROFILER_SCOPE(OperationDiscard);
 
 		UCustomizableObjectInstance* COI = Iterator->CustomizableObjectInstance.Get();
-
-		const bool bUpdating = Private->CurrentMutableOperation && Private->CurrentMutableOperation->Instance != Iterator->CustomizableObjectInstance;
-		if (COI && !bUpdating)
+		
+		const bool bUpdating = Private->CurrentMutableOperation && Private->CurrentMutableOperation->Instance == Iterator->CustomizableObjectInstance;
+		if (COI && COI->GetPrivate() && !bUpdating)
 		{
 			UCustomizableInstancePrivate* COIPrivateData = COI ? COI->GetPrivate() : nullptr;
 
 			// Only discard resources if the instance is still out range (it could have got closer to the player since the task was queued)
 			if (!CurrentInstanceLODManagement->IsOnlyUpdateCloseCustomizableObjectsEnabled() ||
-				!COI ||
-				((COIPrivateData != nullptr) &&
-					(COIPrivateData->LastMinSquareDistFromComponentToPlayer > FMath::Square(CurrentInstanceLODManagement->GetOnlyUpdateCloseCustomizableObjectsDist()))
-					)
-				)
+				COIPrivateData->LastMinSquareDistFromComponentToPlayer > FMath::Square(CurrentInstanceLODManagement->GetOnlyUpdateCloseCustomizableObjectsDist()))
 			{
-				if (COI && COI->IsValidLowLevel())
-				{
-					check(COIPrivateData != nullptr);
-					COIPrivateData->DiscardResourcesAndSetReferenceSkeletalMesh(COI);
-				}
+				COIPrivateData->DiscardResources();
+				COIPrivateData->SetDefaultSkeletalMesh(!IsReplaceDiscardedWithReferenceMeshEnabled());
 			}
 		}
 
@@ -3457,7 +3467,7 @@ void UCustomizableObjectSystem::SetReleaseMutableTexturesImmediately(bool bRelea
 
 void UCustomizableObjectSystem::OnPreBeginPIE(const bool bIsSimulatingInEditor)
 {
-	if (!EditorSettings.bCompileRootObjectsOnStartPIE || IsRunningGame() || IsCompilationDisabled())
+	if (!EditorSettings.bCompileRootObjectsOnStartPIE || IsRunningGame())
 	{
 		return;
 	}
@@ -3526,7 +3536,7 @@ void UCustomizableObjectSystem::StartNextRecompile()
 void UCustomizableObjectSystem::RecompileCustomizableObjectAsync(const FAssetData& InAssetData,
 	const UCustomizableObject* InObject)
 {
-	if (IsRunningGame() || IsCompilationDisabled())
+	if (!IsActive() || IsRunningGame())
 	{
 		return;
 	}
@@ -3548,7 +3558,7 @@ void UCustomizableObjectSystem::RecompileCustomizableObjectAsync(const FAssetDat
 
 void UCustomizableObjectSystem::RecompileCustomizableObjects(const TArray<FAssetData>& InObjects)
 {
-	if (IsRunningGame() || IsCompilationDisabled())
+	if (!IsActive() || IsRunningGame())
 	{
 		return;
 	}
@@ -3692,6 +3702,42 @@ FUnrealMutableImageProvider* UCustomizableObjectSystemPrivate::GetImageProviderC
 }
 
 
+void UCustomizableObjectSystemPrivate::OnMutableEnabledChanged(IConsoleVariable* MutableEnabled)
+{
+	if (!UCustomizableObjectSystem::IsCreated())
+	{
+		return;
+	}
+
+	UCustomizableObjectSystem* System = UCustomizableObjectSystem::GetInstance();
+	UCustomizableObjectSystemPrivate* SystemPrivate = System->GetPrivateChecked();
+
+	if (bIsMutableEnabled)
+	{
+#if !UE_SERVER
+		if (!SystemPrivate->TickDelegateHandle.IsValid())
+		{
+			SystemPrivate->TickDelegate = FTickerDelegate::CreateUObject(System, &UCustomizableObjectSystem::Tick);
+			SystemPrivate->TickDelegateHandle = FTSTicker::GetCoreTicker().AddTicker(SystemPrivate->TickDelegate, 0.f);
+		}
+#endif // !UE_SERVER
+
+#if WITH_EDITOR
+		if (!IsRunningGame() && !FEditorDelegates::PreBeginPIE.IsBoundToObject(System))
+		{
+			FEditorDelegates::PreBeginPIE.AddUObject(System, &UCustomizableObjectSystem::OnPreBeginPIE);
+		}
+#endif
+	}
+	else
+	{
+#if WITH_EDITOR
+		FEditorDelegates::PreBeginPIE.RemoveAll(System);
+#endif
+	}
+}
+
+
 void UCustomizableObjectSystemPrivate::StartUpdateSkeletalMesh(const TSharedRef<FUpdateContextPrivate>& Context)
 {
 	Context->UpdateStarted = true;
diff --git a/Engine/Plugins/Experimental/Mutable/Source/CustomizableObject/Private/MuCO/CustomizableObjectSystemPrivate.h b/Engine/Plugins/Experimental/Mutable/Source/CustomizableObject/Private/MuCO/CustomizableObjectSystemPrivate.h
--- a/Engine/Plugins/Experimental/Mutable/Source/CustomizableObject/Private/MuCO/CustomizableObjectSystemPrivate.h
+++ b/Engine/Plugins/Experimental/Mutable/Source/CustomizableObject/Private/MuCO/CustomizableObjectSystemPrivate.h
@@ -628,6 +628,10 @@ public:
 	FTSTicker::FDelegateHandle TickDelegateHandle;
 	FTickerDelegate TickDelegate;
 
+	/** Change the current status of Mutable. Enabling/Disabling core features.	
+	 * Disabling Mutable will turn off compilation, generation, and streaming and will remove the system ticker. */
+	static void OnMutableEnabledChanged(IConsoleVariable* CVar = nullptr);
+
 	/** Update the last set amount of internal memory Mutable can use to build objects. */
 	void UpdateMemoryLimit();
 
diff --git a/Engine/Plugins/Experimental/Mutable/Source/CustomizableObjectEditor/Private/MuCOE/CustomizableObjectCompileRunnable.cpp b/Engine/Plugins/Experimental/Mutable/Source/CustomizableObjectEditor/Private/MuCOE/CustomizableObjectCompileRunnable.cpp
--- a/Engine/Plugins/Experimental/Mutable/Source/CustomizableObjectEditor/Private/MuCOE/CustomizableObjectCompileRunnable.cpp
+++ b/Engine/Plugins/Experimental/Mutable/Source/CustomizableObjectEditor/Private/MuCOE/CustomizableObjectCompileRunnable.cpp
@@ -23,7 +23,6 @@ class ITargetPlatform;
 FCustomizableObjectCompileRunnable::FCustomizableObjectCompileRunnable(mu::Ptr<mu::Node> Root)
 	: MutableRoot(Root)
 	, bThreadCompleted(false)
-	, MutableIsDisabled(false)
 {
 	PrepareUnrealCompression();
 }
@@ -38,13 +37,6 @@ uint32 FCustomizableObjectCompileRunnable::Run()
 	uint32 Result = 1;
 	ErrorMsg = FString();
 
-	if (MutableIsDisabled)
-	{
-		bThreadCompleted = true;
-		UE_LOG(LogMutable, Verbose, TEXT("PROFILE: [ %16.8f ] FCustomizableObjectCompileRunnable::Run end. NOTE: Mutable compile is deactivated in Editor. To reactivate it, go to Project Settings -> Plugins -> Mutable and unmark the option Disable Mutable Compile In Editor"), FPlatformTime::Seconds());
-		return true;
-	}
-
 	// Translate CO compile options into mu::CompilerOptions
 	mu::Ptr<mu::CompilerOptions> CompilerOptions = new mu::CompilerOptions();
 
diff --git a/Engine/Plugins/Experimental/Mutable/Source/CustomizableObjectEditor/Private/MuCOE/CustomizableObjectCompileRunnable.h b/Engine/Plugins/Experimental/Mutable/Source/CustomizableObjectEditor/Private/MuCOE/CustomizableObjectCompileRunnable.h
--- a/Engine/Plugins/Experimental/Mutable/Source/CustomizableObjectEditor/Private/MuCOE/CustomizableObjectCompileRunnable.h
+++ b/Engine/Plugins/Experimental/Mutable/Source/CustomizableObjectEditor/Private/MuCOE/CustomizableObjectCompileRunnable.h
@@ -65,9 +65,6 @@ public:
 
 	// Whether the thread has finished running
 	bool bThreadCompleted;
-
-	// If Mutable compile is disabled, to immediately finish the Run() method
-	bool MutableIsDisabled;
 };
 
 
diff --git a/Engine/Plugins/Experimental/Mutable/Source/CustomizableObjectEditor/Private/MuCOE/CustomizableObjectCompiler.cpp b/Engine/Plugins/Experimental/Mutable/Source/CustomizableObjectEditor/Private/MuCOE/CustomizableObjectCompiler.cpp
--- a/Engine/Plugins/Experimental/Mutable/Source/CustomizableObjectEditor/Private/MuCOE/CustomizableObjectCompiler.cpp
+++ b/Engine/Plugins/Experimental/Mutable/Source/CustomizableObjectEditor/Private/MuCOE/CustomizableObjectCompiler.cpp
@@ -112,31 +112,6 @@ bool FCustomizableObjectCompiler::Tick()
 	return bUpdated;
 }
 
-void FCustomizableObjectCompiler::MutableIsDisabledCase(UCustomizableObject* Object)
-{
-	if (Object->IsLocked())
-	{
-		return;
-	}
-
-	if (!Object->IsLocked())
-	{
-		UCustomizableObjectSystem::GetInstance()->LockObject(Object);
-	}
-
-	CurrentObject = Object;
-
-	if (CompileTask.IsValid()) // Don't start compilation if there's a compilation running
-	{
-		return;
-	}
-
-	UE_LOG(LogMutable, Log, TEXT("Mutable has been disabled. To reenable it, please deactivate the Disable Mutale Option in the Plugins -> Mutable option"), FPlatformTime::Seconds());
-	CompileTask = MakeShareable(new FCustomizableObjectCompileRunnable(nullptr));
-	CompileTask->MutableIsDisabled = true;
-	LaunchMutableCompile(false);
-}
-
 
 bool FCustomizableObjectCompiler::IsRootObject(const UCustomizableObject* Object) const
 {
@@ -191,15 +166,15 @@ void FCustomizableObjectCompiler::Compile(UCustomizableObject& Object, const FCo
 {
 	TRACE_BEGIN_REGION(UE_MUTABLE_COMPILE_REGION);
 
-	UCustomizableObjectSystem* System = UCustomizableObjectSystem::GetInstance();
-	check(System);
-
-	if (System->IsCompilationDisabled())
+	if (!UCustomizableObjectSystem::IsActive())
 	{
-		CompileInternal(&Object, InOptions, bAsync);
+		UE_LOG(LogMutable, Warning, TEXT("Failed to compile Customizable Object [%s]. Mutable is disabled. To enable it set the CVar Mutable.Enabled to true."), *Object.GetName());
+		SetCompilationState(ECustomizableObjectCompilationState::Failed);
 		return;
 	}
 
+	UCustomizableObjectSystem* System = UCustomizableObjectSystem::GetInstanceChecked();
+
 	UE_LOG(LogMutable, Verbose, TEXT("PROFILE: [ %16.8f ] Preload asynchronously assets start."), FPlatformTime::Seconds());
 
 	FString Message = FString::Printf(TEXT("Customizable Object %s is already being compiled or updated. Please wait a few seconds and try again."), *Object.GetName());
@@ -853,12 +828,6 @@ void FCustomizableObjectCompiler::CompileInternal(UCustomizableObject* Object, c
 
 	if (!Object) return;
 
-	if(UCustomizableObjectSystem::GetInstance()->IsCompilationDisabled())
-	{
-		MutableIsDisabledCase(Object);
-		return;
-	}
-
 	if (bAsync && CompileTask.IsValid()) // Don't start compilation if there's a compilation running
 	{
 		UE_LOG(LogMutable, Log, TEXT("FCustomizableObjectCompiler::Compile An object is already being compiled."));
diff --git a/Engine/Plugins/Experimental/Mutable/Source/CustomizableObjectEditor/Private/MuCOE/CustomizableObjectCompiler.h b/Engine/Plugins/Experimental/Mutable/Source/CustomizableObjectEditor/Private/MuCOE/CustomizableObjectCompiler.h
--- a/Engine/Plugins/Experimental/Mutable/Source/CustomizableObjectEditor/Private/MuCOE/CustomizableObjectCompiler.h
+++ b/Engine/Plugins/Experimental/Mutable/Source/CustomizableObjectEditor/Private/MuCOE/CustomizableObjectCompiler.h
@@ -31,9 +31,6 @@ public:
 	/** Check for pending compilation process. Returns true if an object has been updated. */
 	CUSTOMIZABLEOBJECTEDITOR_API virtual bool Tick() override;
 
-	/** In case Mutable Compile is disabled, fake the compilation process */
-	CUSTOMIZABLEOBJECTEDITOR_API void MutableIsDisabledCase(UCustomizableObject* Object);
-
 	/** Provided a CO object it provides the root CO it is connected. In other words : it returns the root of the entire
 	 * mutable graph.
 	 * @param InObject Customizable object whose root CO we are asking for.
diff --git a/Engine/Plugins/Experimental/Mutable/Source/CustomizableObjectEditor/Private/MuCOE/CustomizableObjectEditor.cpp b/Engine/Plugins/Experimental/Mutable/Source/CustomizableObjectEditor/Private/MuCOE/CustomizableObjectEditor.cpp
--- a/Engine/Plugins/Experimental/Mutable/Source/CustomizableObjectEditor/Private/MuCOE/CustomizableObjectEditor.cpp
+++ b/Engine/Plugins/Experimental/Mutable/Source/CustomizableObjectEditor/Private/MuCOE/CustomizableObjectEditor.cpp
@@ -741,12 +741,6 @@ UCustomizableObjectInstance* FCustomizableObjectEditor::GetPreviewInstance()
 
 void FCustomizableObjectEditor::CompileObjectUserPressedButton()
 {
-	if (UCustomizableObjectSystem::GetInstance()->IsCompilationDisabled())
-	{
-		UE_LOG(LogMutable, Warning, TEXT("Mutable compile is disabled in Editor. To enable it, go to Project Settings -> Plugins -> Mutable and unmark the option Disable Mutable Compile In Editor"));
-		return;
-	}
-
 	Compiler.ClearAllCompileOnlySelectedOption();
 	CompileObject();
 }
@@ -754,12 +748,6 @@ void FCustomizableObjectEditor::CompileObjectUserPressedButton()
 
 void FCustomizableObjectEditor::CompileOnlySelectedObjectUserPressedButton()
 {
-	if (UCustomizableObjectSystem::GetInstance()->IsCompilationDisabled())
-	{
-		UE_LOG(LogMutable, Warning, TEXT("Mutable compile is disabled in Editor. To enable it, go to Project Settings -> Plugins -> Mutable and unmark the option Disable Mutable Compile In Editor"));
-		return;
-	}
-
 	if (PreviewInstance)
 	{
 		Compiler.ClearAllCompileOnlySelectedOption();
@@ -783,13 +771,13 @@ void FCustomizableObjectEditor::BindCommands()
 	ToolkitCommands->MapAction(
 		Commands.Compile,
 		FExecuteAction::CreateSP(this, &FCustomizableObjectEditor::CompileObjectUserPressedButton),
-		FCanExecuteAction(),
+		FCanExecuteAction::CreateStatic(&UCustomizableObjectSystem::IsActive),
 		FIsActionChecked());
 
 	ToolkitCommands->MapAction(
 		Commands.CompileOnlySelected,
 		FExecuteAction::CreateSP(this, &FCustomizableObjectEditor::CompileOnlySelectedObjectUserPressedButton),
-		FCanExecuteAction(),
+		FCanExecuteAction::CreateStatic(&UCustomizableObjectSystem::IsActive),
 		FIsActionChecked());
 
 	// Compile and options
diff --git a/Engine/Plugins/Experimental/Mutable/Source/CustomizableObjectEditor/Private/MuCOE/CustomizableObjectEditorActions.cpp b/Engine/Plugins/Experimental/Mutable/Source/CustomizableObjectEditor/Private/MuCOE/CustomizableObjectEditorActions.cpp
--- a/Engine/Plugins/Experimental/Mutable/Source/CustomizableObjectEditor/Private/MuCOE/CustomizableObjectEditorActions.cpp
+++ b/Engine/Plugins/Experimental/Mutable/Source/CustomizableObjectEditor/Private/MuCOE/CustomizableObjectEditorActions.cpp
@@ -21,8 +21,8 @@ FCustomizableObjectEditorCommands::FCustomizableObjectEditorCommands()
 
 void FCustomizableObjectEditorCommands::RegisterCommands()
 {
-	UI_COMMAND(Compile, "Compile", "Compile the source graph of the customizable object and update the previews.", EUserInterfaceActionType::Button, FInputChord());
-	UI_COMMAND(CompileOnlySelected, "Compile Only Selected", "Compile the source graph of the customizable object and update the previews, only for the selected options in the preview. The rest of options are discarded. If they are selected, press again this button to see their effect in the preview.", EUserInterfaceActionType::Button, FInputChord());
+	UI_COMMAND(Compile, "Compile", "Compile the source graph of the customizable object and update the previews. \nActive if the CVar Mutable.Enabled is set to true.", EUserInterfaceActionType::Button, FInputChord());
+	UI_COMMAND(CompileOnlySelected, "Compile Only Selected", "Compile the source graph of the customizable object and update the previews, only for the selected options in the preview. The rest of options are discarded. If they are selected, press again this button to see their effect in the preview. \nActive if the CVar Mutable.Enabled is set to true.", EUserInterfaceActionType::Button, FInputChord());
 	UI_COMMAND(ResetCompileOptions, "Reset Compilation Options", "Set reasonable defaults for the compilation options.", EUserInterfaceActionType::Button, FInputChord());
 	UI_COMMAND(CompileOptions_UseDiskCompilation, "Enable compiling using the disk as memory.", "This is very slow but supports compiling huge objects. It requires a lot of free space in the OS disk.", EUserInterfaceActionType::ToggleButton, FInputChord());
 	UI_COMMAND(Debug, "Debug", "Open the CustomizableObject debugger tab for this object.", EUserInterfaceActionType::Button, FInputChord());
diff --git a/Engine/Plugins/Experimental/Mutable/Source/CustomizableObjectEditor/Private/MuCOE/CustomizableObjectEditorModule.cpp b/Engine/Plugins/Experimental/Mutable/Source/CustomizableObjectEditor/Private/MuCOE/CustomizableObjectEditorModule.cpp
--- a/Engine/Plugins/Experimental/Mutable/Source/CustomizableObjectEditor/Private/MuCOE/CustomizableObjectEditorModule.cpp
+++ b/Engine/Plugins/Experimental/Mutable/Source/CustomizableObjectEditor/Private/MuCOE/CustomizableObjectEditorModule.cpp
@@ -269,7 +269,7 @@ bool FCustomizableObjectEditorModule::HandleSettingsSaved()
 		CustomizableObjectSettings->SaveConfig();
 		
 		FEditorCompileSettings CompileSettings;
-		CompileSettings.bDisableCompilation = CustomizableObjectSettings->bDisableMutableCompileInEditor;
+		CompileSettings.bIsMutableEnabled = !CustomizableObjectSettings->bDisableMutableCompileInEditor;
 		CompileSettings.bEnableAutomaticCompilation = CustomizableObjectSettings->bEnableAutomaticCompilation;
 		CompileSettings.bCompileObjectsSynchronously = CustomizableObjectSettings->bCompileObjectsSynchronously;
 		CompileSettings.bCompileRootObjectsOnStartPIE = CustomizableObjectSettings->bCompileRootObjectsOnStartPIE;
@@ -304,7 +304,7 @@ void FCustomizableObjectEditorModule::RegisterSettings()
 			if (CustomizableObjectSettings != nullptr)
 			{
 				FEditorCompileSettings CompileSettings;
-				CompileSettings.bDisableCompilation = CustomizableObjectSettings->bDisableMutableCompileInEditor;
+				CompileSettings.bIsMutableEnabled = !CustomizableObjectSettings->bDisableMutableCompileInEditor;
 				CompileSettings.bEnableAutomaticCompilation = CustomizableObjectSettings->bEnableAutomaticCompilation;
 				CompileSettings.bCompileObjectsSynchronously = CustomizableObjectSettings->bCompileObjectsSynchronously;
 				CompileSettings.bCompileRootObjectsOnStartPIE = CustomizableObjectSettings->bCompileRootObjectsOnStartPIE;
diff --git a/Engine/Plugins/Experimental/Mutable/Source/CustomizableObjectEditor/Private/MuCOE/CustomizableObjectInstanceEditor.cpp b/Engine/Plugins/Experimental/Mutable/Source/CustomizableObjectEditor/Private/MuCOE/CustomizableObjectInstanceEditor.cpp
--- a/Engine/Plugins/Experimental/Mutable/Source/CustomizableObjectEditor/Private/MuCOE/CustomizableObjectInstanceEditor.cpp
+++ b/Engine/Plugins/Experimental/Mutable/Source/CustomizableObjectEditor/Private/MuCOE/CustomizableObjectInstanceEditor.cpp
@@ -358,56 +358,39 @@ FText FCustomizableObjectInstanceEditor::GetBaseToolkitName() const
 void FCustomizableObjectInstanceEditor::CreatePreviewInstance()
 {
 	check(CustomizableObjectInstance);
-	UCustomizableObject* CustomizableObject = CustomizableObjectInstance->GetCustomizableObject();
-	
-	if (!CustomizableObject || (!CustomizableObject->IsCompiled() && !UCustomizableObjectSystem::GetInstance()->IsCompilationDisabled()))
+	if (!CustomizableObjectInstance->CanUpdateInstance())
 	{
 		return;
 	}
 
 	// Bind update delegate
-	HelperCallback = NewObject<UUpdateClassWrapperClass>(GetTransientPackage());
+	HelperCallback = NewObject<UUpdateClassWrapperClass>();
 	HelperCallback->Delegate.BindSP(this, &FCustomizableObjectInstanceEditor::OnUpdatePreviewInstance);
 	CustomizableObjectInstance->UpdatedDelegate.AddDynamic(HelperCallback, &UUpdateClassWrapperClass::DelegatedCallback);
 
 	PreviewStaticMeshComponent = nullptr;
 
 	// Create a SkeletalMeshComponent for each component in the CO
-	int32 NumMeshComponents = CustomizableObject->GetComponentCount();
+	const int32 NumMeshComponents = CustomizableObjectInstance->GetNumComponents();
 	PreviewSkeletalMeshComponents.AddZeroed(NumMeshComponents);
 	PreviewCustomizableSkeletalComponents.AddZeroed(NumMeshComponents);
 
-	if (CustomizableObject->ReferenceSkeletalMeshes.Num())
+	for (int32 ComponentIndex = 0; ComponentIndex < NumMeshComponents; ++ComponentIndex)
 	{
-		for (int32 ComponentIndex = 0; ComponentIndex < NumMeshComponents; ++ComponentIndex)
-		{
-			PreviewCustomizableSkeletalComponents[ComponentIndex] = NewObject<UCustomizableSkeletalComponent>(UCustomizableSkeletalComponent::StaticClass());
+		PreviewCustomizableSkeletalComponents[ComponentIndex] = NewObject<UCustomizableSkeletalComponent>();
+		check(PreviewCustomizableSkeletalComponents[ComponentIndex]);
 
-			if (PreviewCustomizableSkeletalComponents[ComponentIndex])
-			{
-				PreviewCustomizableSkeletalComponents[ComponentIndex]->CustomizableObjectInstance = CustomizableObjectInstance;
-				PreviewCustomizableSkeletalComponents[ComponentIndex]->ComponentIndex = ComponentIndex;
+		PreviewCustomizableSkeletalComponents[ComponentIndex]->CustomizableObjectInstance = CustomizableObjectInstance;
+		PreviewCustomizableSkeletalComponents[ComponentIndex]->ComponentIndex = ComponentIndex;
 
-				PreviewSkeletalMeshComponents[ComponentIndex] = NewObject<UDebugSkelMeshComponent>(GetTransientPackage(), NAME_None, RF_Transient);
+		PreviewSkeletalMeshComponents[ComponentIndex] = NewObject<UDebugSkelMeshComponent>();
+		check(PreviewSkeletalMeshComponents[ComponentIndex]);
 
-				if (PreviewSkeletalMeshComponents[ComponentIndex])
-				{
-					PreviewCustomizableSkeletalComponents[ComponentIndex]->AttachToComponent(PreviewSkeletalMeshComponents[ComponentIndex], FAttachmentTransformRules::KeepRelativeTransform);
-				}
-				else
-				{
-					return;
-				}
-			}
-			else
-			{
-				return;
-			}
-		}
-
-		CustomizableObjectInstance->UpdateSkeletalMeshAsync(true, true);
-		Viewport->SetPreviewComponents(PreviewSkeletalMeshComponents);
+		PreviewCustomizableSkeletalComponents[ComponentIndex]->AttachToComponent(PreviewSkeletalMeshComponents[ComponentIndex], FAttachmentTransformRules::KeepRelativeTransform);
 	}
+
+	CustomizableObjectInstance->UpdateSkeletalMeshAsync(true, true);
+	Viewport->SetPreviewComponents(PreviewSkeletalMeshComponents);
 }
 
 
