[Mutable] Make FDescriptorRuntimeHash private.

- Merged FDescriptorRuntimeHash and FDescriptorHash

#rnx

diff --git a/Engine/Plugins/Experimental/Mutable/Source/CustomizableObject/Classes/MuCO/CustomizableObjectInstance.h b/Engine/Plugins/Experimental/Mutable/Source/CustomizableObject/Classes/MuCO/CustomizableObjectInstance.h
--- a/Engine/Plugins/Experimental/Mutable/Source/CustomizableObject/Classes/MuCO/CustomizableObjectInstance.h
+++ b/Engine/Plugins/Experimental/Mutable/Source/CustomizableObject/Classes/MuCO/CustomizableObjectInstance.h
@@ -662,12 +662,6 @@ public:
 
 	const TArray<uint16>& GetRequestedLODsPerComponent() const;
 
-	/** Return the UCustomizableObjectInstance::Descriptor hash on the last update request. */
-	FDescriptorRuntimeHash GetDescriptorRuntimeHash() const;
-
-	/** Return the UCustomizableObjectInstance::Descriptor hash on the last successful update. */
-	FDescriptorRuntimeHash GetUpdateDescriptorRuntimeHash() const;
-
 #if WITH_EDITOR
 	const TArray<TObjectPtr<UTexture2D>>& GetTextureParameterDeclarations() const;
 #endif
diff --git a/Engine/Plugins/Experimental/Mutable/Source/CustomizableObject/Internal/MuCO/CustomizableInstancePrivateData.h b/Engine/Plugins/Experimental/Mutable/Source/CustomizableObject/Internal/MuCO/CustomizableInstancePrivateData.h
--- a/Engine/Plugins/Experimental/Mutable/Source/CustomizableObject/Internal/MuCO/CustomizableInstancePrivateData.h
+++ b/Engine/Plugins/Experimental/Mutable/Source/CustomizableObject/Internal/MuCO/CustomizableInstancePrivateData.h
@@ -9,6 +9,7 @@
 #include "MuCO/CustomizableObjectInstance.h"
 #include "MuR/Instance.h"
 #include "GameplayTagContainer.h"
+#include "MuCO/DescriptorHash.h"
 #include "UObject/Package.h"
 
 #include "CustomizableInstancePrivateData.generated.h"
@@ -153,6 +154,7 @@ struct FAnimBpGeneratedPhysicsAssets
 	TArray<FAnimInstanceOverridePhysicsAsset> AnimInstancePropertyIndexAndPhysicsAssets;
 };
 
+
 UCLASS()
 class UCustomizableInstancePrivate : public UObject
 {
@@ -327,10 +329,10 @@ private:
 
 public:
 	/** Hash of the UCustomizableObjectInstance::Descriptor on the last update request. */
-	FDescriptorRuntimeHash UpdateDescriptorRuntimeHash;
+	FDescriptorHash UpdateDescriptorHash;
 	
 	/** Hash of the UCustomizableObjectInstance::Descriptor on the last successful update. */
-	FDescriptorRuntimeHash DescriptorRuntimeHash;
+	FDescriptorHash DescriptorHash;
 
 	/** Status of the generated Skeletal Mesh. Not to be confused with the Update Result. */
 	ESkeletalMeshStatus SkeletalMeshStatus = ESkeletalMeshStatus::NotGenerated;
diff --git a/Engine/Plugins/Experimental/Mutable/Source/CustomizableObject/Internal/MuCO/DescriptorHash.h b/Engine/Plugins/Experimental/Mutable/Source/CustomizableObject/Internal/MuCO/DescriptorHash.h
new file mode 100644
--- /dev/null
+++ b/Engine/Plugins/Experimental/Mutable/Source/CustomizableObject/Internal/MuCO/DescriptorHash.h
@@ -0,0 +1,44 @@
+// Copyright Epic Games, Inc. All Rights Reserved.
+
+#pragma once
+
+#include "MuCO/CustomizableObjectInstanceDescriptor.h"
+
+
+/** Hash of the Descriptor.
+* Can change and is not backwards compatible. Do not serialize. */
+class CUSTOMIZABLEOBJECT_API FDescriptorHash
+{
+public:
+	FDescriptorHash() = default;
+
+	explicit FDescriptorHash(const FCustomizableObjectInstanceDescriptor& Descriptor);
+
+	bool operator==(const FDescriptorHash& Other) const;
+
+	bool operator!=(const FDescriptorHash& Other) const;
+	
+	/** Return true if this Hash is a subset of the other Hash (i.e., this Descriptor is a subset of the other Descriptor). */
+	bool IsSubset(const FDescriptorHash& Other) const;
+
+	void UpdateMinMaxLOD(int32 InMinLOD, int32 InMaxLOD);
+
+	int32 GetMinLOD() const;
+
+	int32 GetMaxLOD() const;
+
+	void UpdateRequestedLODs(const TArray<uint16>& InRequestedLODs);
+
+	const TArray<uint16>& GetRequestedLODs() const;
+
+	FString ToString() const;
+
+private:
+	uint32 Hash = 0;
+
+	int32 MinLOD = 0;
+	int32 MaxLOD = INT32_MAX;
+
+	// Array of bitmasks that indicate which LODs of each component have been requested
+	TArray<uint16> RequestedLODsPerComponent;
+};
diff --git a/Engine/Plugins/Experimental/Mutable/Source/CustomizableObject/Private/MuCO/CustomizableObjectInstance.cpp b/Engine/Plugins/Experimental/Mutable/Source/CustomizableObject/Private/MuCO/CustomizableObjectInstance.cpp
--- a/Engine/Plugins/Experimental/Mutable/Source/CustomizableObject/Private/MuCO/CustomizableObjectInstance.cpp
+++ b/Engine/Plugins/Experimental/Mutable/Source/CustomizableObject/Private/MuCO/CustomizableObjectInstance.cpp
@@ -2779,7 +2779,7 @@ void UCustomizableInstancePrivate::DiscardResources()
 	InvalidateGeneratedData();
 	
 	Instance->SkeletalMeshes.Reset();
-	DescriptorRuntimeHash = FDescriptorRuntimeHash();
+	DescriptorHash = FDescriptorHash();
 }
 
 
@@ -2869,17 +2869,6 @@ void SetTexturePropertiesFromMutableImageProps(UTexture2D* Texture, const FMutab
 }
 
 
-FDescriptorRuntimeHash UCustomizableObjectInstance::GetDescriptorRuntimeHash() const
-{
-	return GetPrivate()->DescriptorRuntimeHash;
-}
-
-
-FDescriptorRuntimeHash UCustomizableObjectInstance::GetUpdateDescriptorRuntimeHash() const
-{
-	return GetPrivate()->UpdateDescriptorRuntimeHash;
-}
-
 #if WITH_EDITOR
 const TArray<TObjectPtr<UTexture2D>>& UCustomizableObjectInstance::GetTextureParameterDeclarations() const
 {
@@ -2914,12 +2903,12 @@ void FMutableUpdateCandidate::ApplyLODUpdateParamsToInstance(FUpdateContextPriva
 
 	CustomizableObjectInstance->Descriptor.RequestedLODLevels = RequestedLODLevels;
 
-	CustomizableObjectInstance->GetPrivate()->UpdateDescriptorRuntimeHash.UpdateMinMaxLOD(MinLOD, MaxLOD);
-	CustomizableObjectInstance->GetPrivate()->UpdateDescriptorRuntimeHash.UpdateRequestedLODs(CustomizableObjectInstance->Descriptor.RequestedLODLevels);
+	CustomizableObjectInstance->GetPrivate()->UpdateDescriptorHash.UpdateMinMaxLOD(MinLOD, MaxLOD);
+	CustomizableObjectInstance->GetPrivate()->UpdateDescriptorHash.UpdateRequestedLODs(CustomizableObjectInstance->Descriptor.RequestedLODLevels);
 
 	if (Context)
 	{
-		Context->InstanceDescriptorRuntimeHash = CustomizableObjectInstance->GetPrivate()->UpdateDescriptorRuntimeHash;
+		Context->InstanceDescriptorHash = CustomizableObjectInstance->GetPrivate()->UpdateDescriptorHash;
 		Context->CurrentMinLOD = MinLOD;
 		Context->CurrentMaxLOD = MaxLOD;
 		Context->RequestedLODs = RequestedLODLevels;
@@ -6215,9 +6204,9 @@ void UCustomizableObjectInstance::SetRequestedLODs(int32 InMinLOD, int32 InMaxLO
 		const int32 ComponentCount = GetNumComponents();
 		MutableUpdateCandidate.RequestedLODLevels.SetNumZeroed(ComponentCount);
 
-		const TArray<uint16>& GeneratedLODsPerComponent = GetPrivate()->DescriptorRuntimeHash.GetRequestedLODs();
+		const TArray<uint16>& GeneratedLODsPerComponent = GetPrivate()->DescriptorHash.GetRequestedLODs();
 
-		const bool bIgnoreGeneratedLODs = GetPrivate()->DescriptorRuntimeHash != GetPrivate()->UpdateDescriptorRuntimeHash || GeneratedLODsPerComponent.Num() != ComponentCount;
+		const bool bIgnoreGeneratedLODs = GetPrivate()->DescriptorHash != GetPrivate()->UpdateDescriptorHash || GeneratedLODsPerComponent.Num() != ComponentCount;
 
 		if (bMinMaxLODChanged || bIgnoreGeneratedLODs || Descriptor.GetRequestedLODLevels() != InRequestedLODsPerComponent)
 		{
diff --git a/Engine/Plugins/Experimental/Mutable/Source/CustomizableObject/Private/MuCO/CustomizableObjectInstanceDescriptor.cpp b/Engine/Plugins/Experimental/Mutable/Source/CustomizableObject/Private/MuCO/CustomizableObjectInstanceDescriptor.cpp
--- a/Engine/Plugins/Experimental/Mutable/Source/CustomizableObject/Private/MuCO/CustomizableObjectInstanceDescriptor.cpp
+++ b/Engine/Plugins/Experimental/Mutable/Source/CustomizableObject/Private/MuCO/CustomizableObjectInstanceDescriptor.cpp
@@ -2572,169 +2572,5 @@ void FCustomizableObjectInstanceDescriptor::CreateParametersLookupTable()
 	}
 }
 
-FDescriptorHash::FDescriptorHash(const FCustomizableObjectInstanceDescriptor& Descriptor)
-{
-#if WITH_EDITORONLY_DATA
-	if (Descriptor.CustomizableObject)
-	{
-		Hash = HashCombine(Hash, GetTypeHash(Descriptor.CustomizableObject->GetPathName()));
-		Hash = HashCombine(Hash, GetTypeHash(Descriptor.CustomizableObject->GetCompilationGuid()));
-	}
-#endif
-
-	for (const FCustomizableObjectBoolParameterValue& Value : Descriptor.BoolParameters)
-	{
-		Hash = HashCombine(Hash, GetTypeHash(Value));
-	}	
-
-	for (const FCustomizableObjectIntParameterValue& Value : Descriptor.IntParameters)
-	{
-		Hash = HashCombine(Hash, GetTypeHash(Value));
-	}
-	
-	for (const FCustomizableObjectFloatParameterValue& Value : Descriptor.FloatParameters)
-	{
-		Hash = HashCombine(Hash, GetTypeHash(Value));
-	}
-
-	for (const FCustomizableObjectTextureParameterValue& Value : Descriptor.TextureParameters)
-	{
-		Hash = HashCombine(Hash, GetTypeHash(Value));
-	}
-
-	for (const FCustomizableObjectVectorParameterValue& Value : Descriptor.VectorParameters)
-	{
-		Hash = HashCombine(Hash, GetTypeHash(Value));
-	}
-
-	for (const FCustomizableObjectProjectorParameterValue& Value : Descriptor.ProjectorParameters)
-	{
-		Hash = HashCombine(Hash, GetTypeHash(Value));
-	}
-	
-	Hash = HashCombine(Hash, GetTypeHash(Descriptor.State));
-	Hash = HashCombine(Hash, GetTypeHash(Descriptor.bBuildParameterRelevancy));
-
-	for (const TTuple<FName, FMultilayerProjector>& Pair : Descriptor.MultilayerProjectors)
-	{
-		// Hash = HashCombine(Hash, GetTypeHash(Pair.Key)); // Already hashed by the FMultilayerProjector.
-		Hash = HashCombine(Hash, GetTypeHash(Pair.Value));
-	}
-}
-
-
-bool FDescriptorHash::operator==(const FDescriptorHash& Other) const
-{
-	return Hash == Other.Hash;
-}
-
-
-bool FDescriptorHash::operator!=(const FDescriptorHash& Other) const
-{
-	return Hash != Other.Hash;	
-}
-
-
-bool FDescriptorHash::operator<(const FDescriptorHash& Other) const
-{
-	return Hash < Other.Hash;
-}
-
-
-FDescriptorRuntimeHash::FDescriptorRuntimeHash(const FCustomizableObjectInstanceDescriptor& Descriptor) :
-	FDescriptorHash(Descriptor)
-{
-	MinLOD = Descriptor.MinLOD;
-	MaxLOD = Descriptor.MaxLOD;	
-
-	RequestedLODsPerComponent = Descriptor.RequestedLODLevels;
-}
-
-
-bool FDescriptorRuntimeHash::IsSubset(const FDescriptorRuntimeHash& Other) const
-{
-	if (FDescriptorHash::operator!=(Other) || MinLOD < Other.MinLOD || MaxLOD != Other.MaxLOD)
-	{
-		return false;
-	}
-
-	if (RequestedLODsPerComponent == Other.RequestedLODsPerComponent)
-	{
-		return true;
-	}
-
-	if (RequestedLODsPerComponent.Num() != Other.RequestedLODsPerComponent.Num())
-	{
-		return false;
-	}
-	
-	for (int32 ComponentIndex = 0; ComponentIndex < RequestedLODsPerComponent.Num(); ++ComponentIndex)
-	{
-		int32 RequestedLODs = RequestedLODsPerComponent[ComponentIndex];
-		int32 OtherRequestedLODs = Other.RequestedLODsPerComponent[ComponentIndex];
-		for (uint16 LODIndex = MinLOD; LODIndex <= MaxLOD; ++LODIndex)
-		{
-			// To be a subset all bits set in RequestedLODs must be set in OtherRequestedLODs. OtherRequestedLODs can have additional bits set
-			if (RequestedLODs & (1 << LODIndex) && !(OtherRequestedLODs & (1 << LODIndex)))
-			{
-				return false;
-			}
-		}
-	}
-
-	return true;
-}
-
-
-void FDescriptorRuntimeHash::UpdateMinMaxLOD(const int32 InMinLOD, const int32 InMaxLOD)
-{
-	MinLOD = InMinLOD;
-	MaxLOD = InMaxLOD;	
-}
-
-
-int32 FDescriptorRuntimeHash::GetMinLOD() const
-{
-	return MinLOD;
-}
-
-
-int32 FDescriptorRuntimeHash::GetMaxLOD() const
-{
-	return MaxLOD;
-}
-
-
-void FDescriptorRuntimeHash::UpdateRequestedLODs(const TArray<uint16>& InRequestedLODs)
-{
-	RequestedLODsPerComponent = InRequestedLODs;
-}
-
-
-const TArray<uint16>& FDescriptorRuntimeHash::GetRequestedLODs() const
-{
-	return RequestedLODsPerComponent;
-}
-
-
-FString FDescriptorRuntimeHash::ToString() const
-{
-	TStringBuilder<150> Builder;
-
-	Builder.Appendf(TEXT("(Hash=%u,"), Hash);
-	Builder.Appendf(TEXT("MinLOD=%i,"), MinLOD);
-	Builder.Appendf(TEXT("MaxLOD=%i,"), MaxLOD);
-	Builder.Appendf(TEXT("RequestredLODLevels=["));
-
-	for (const uint16 RequestedLODs : RequestedLODsPerComponent)
-	{
-		Builder.Appendf(TEXT("%i,"), RequestedLODs);
-	}
-	
-	Builder.Appendf(TEXT("])"));
-
-	return Builder.ToString();
-}
-
 
 #undef RETURN_ON_CO_UNCOMPILED
diff --git a/Engine/Plugins/Experimental/Mutable/Source/CustomizableObject/Private/MuCO/CustomizableObjectSystem.cpp b/Engine/Plugins/Experimental/Mutable/Source/CustomizableObject/Private/MuCO/CustomizableObjectSystem.cpp
--- a/Engine/Plugins/Experimental/Mutable/Source/CustomizableObject/Private/MuCO/CustomizableObjectSystem.cpp
+++ b/Engine/Plugins/Experimental/Mutable/Source/CustomizableObject/Private/MuCO/CustomizableObjectSystem.cpp
@@ -136,8 +136,8 @@ FUpdateContextPrivate::FUpdateContextPrivate(UCustomizableObjectInstance& InInst
 	check(InInstance.GetCustomizableObject());
 
 	Instance = &InInstance;
-	InstanceDescriptorRuntimeHash = FDescriptorRuntimeHash(Instance->GetDescriptor());
-	InInstance.GetPrivate()->UpdateDescriptorRuntimeHash = InstanceDescriptorRuntimeHash; // TODO GMTFuture Remove on MTBL-1409
+	InstanceDescriptorHash = FDescriptorHash(Instance->GetDescriptor());
+	InInstance.GetPrivate()->UpdateDescriptorHash = InstanceDescriptorHash; // TODO GMTFuture Remove on MTBL-1409
 	State = InInstance.GetState();
 	bBuildParameterRelevancy = InInstance.GetBuildParameterRelevancy();
 	Parameters = InInstance.GetDescriptor().GetParameters();
@@ -765,7 +765,7 @@ void FinishUpdateGlobal(const TSharedRef<FUpdateContextPrivate>& Context)
 		case EUpdateResult::Success:
 			PrivateInstance->SkeletalMeshStatus = ESkeletalMeshStatus::Success;
 
-			PrivateInstance->DescriptorRuntimeHash = Context->InstanceDescriptorRuntimeHash;
+			PrivateInstance->DescriptorHash = Context->InstanceDescriptorHash;
 
 			// Delegates must be called only after updating the Instance flags.
 			Instance->UpdatedDelegate.Broadcast(Instance);
@@ -1156,7 +1156,7 @@ void UCustomizableObjectSystemPrivate::EnqueueUpdateSkeletalMesh(const TSharedRe
 
 		if (const FMutablePendingInstanceUpdate* QueueElem = MutablePendingInstanceWork.GetUpdate(Instance))
 		{
-			if (InstancePrivate->UpdateDescriptorRuntimeHash.IsSubset(FDescriptorRuntimeHash(QueueElem->Context->InstanceDescriptorRuntimeHash)))
+			if (InstancePrivate->UpdateDescriptorHash.IsSubset(FDescriptorHash(QueueElem->Context->InstanceDescriptorHash)))
 			{
 				Context->UpdateResult = EUpdateResult::ErrorOptimized;
 				FinishUpdateGlobal(Context);			
@@ -1166,14 +1166,14 @@ void UCustomizableObjectSystemPrivate::EnqueueUpdateSkeletalMesh(const TSharedRe
 
 		if (CurrentMutableOperation &&
 			Instance == CurrentMutableOperation->Instance &&
-			InstancePrivate->UpdateDescriptorRuntimeHash.IsSubset(CurrentMutableOperation->InstanceDescriptorRuntimeHash))
+			InstancePrivate->UpdateDescriptorHash.IsSubset(CurrentMutableOperation->InstanceDescriptorHash))
 		{
 			Context->UpdateResult = EUpdateResult::ErrorOptimized;
 			FinishUpdateGlobal(Context);
 			return; // The requested update is equal to the running update.
 		}
 	
-		if (InstancePrivate->UpdateDescriptorRuntimeHash.IsSubset(InstancePrivate->DescriptorRuntimeHash) &&
+		if (InstancePrivate->UpdateDescriptorHash.IsSubset(InstancePrivate->DescriptorHash) &&
 			!(CurrentMutableOperation &&
 			Instance == CurrentMutableOperation->Instance)) // This condition is necessary because even if the descriptor is a subset, it will be replaced by the CurrentMutableOperation
 		{
@@ -1219,7 +1219,7 @@ void UCustomizableObjectSystemPrivate::EnqueueUpdateSkeletalMesh(const TSharedRe
 			{
 				FString String = TEXT("DESCRIPTOR DEBUG PRINT\n");
 				String += "================================\n";				
-				String += FString::Printf(TEXT("=== DESCRIPTOR HASH ===\n%s\n"), *InstancePrivate->UpdateDescriptorRuntimeHash.ToString());
+				String += FString::Printf(TEXT("=== DESCRIPTOR HASH ===\n%s\n"), *InstancePrivate->UpdateDescriptorHash.ToString());
 				String += FString::Printf(TEXT("=== DESCRIPTOR ===\n%s"), *Instance->GetDescriptor().ToString());
 				String += "================================";
 				
@@ -1653,7 +1653,7 @@ namespace impl
 			}
 		}
 
-		Operation->InstanceDescriptorRuntimeHash.UpdateRequestedLODs(Operation->RequestedLODs);
+		Operation->InstanceDescriptorHash.UpdateRequestedLODs(Operation->RequestedLODs);
 	}
 	
 	
@@ -2673,7 +2673,7 @@ namespace impl
 		}
 
 		// Skip update, the requested update is equal to the running update.
-		if (Operation->InstanceDescriptorRuntimeHash.IsSubset(CandidateInstance->GetDescriptorRuntimeHash()))
+		if (Operation->InstanceDescriptorHash.IsSubset(CandidateInstance->GetPrivate()->DescriptorHash))
 		{
 			System->ClearCurrentMutableOperation();
 
@@ -2787,8 +2787,8 @@ namespace impl
 		
 		// Task: Mutable Update and GetMesh
 		//-------------------------------------------------------------
-		Operation->CurrentMinLOD = Operation->InstanceDescriptorRuntimeHash.GetMinLOD();
-		Operation->CurrentMaxLOD = Operation->InstanceDescriptorRuntimeHash.GetMaxLOD();
+		Operation->CurrentMinLOD = Operation->InstanceDescriptorHash.GetMinLOD();
+		Operation->CurrentMaxLOD = Operation->InstanceDescriptorHash.GetMaxLOD();
 		Operation->InstanceID = Operation->bLiveUpdateMode ? CandidateInstancePrivateData->LiveUpdateModeInstanceID : 0;
 		Operation->bUseMeshCache = CustomizableObject->IsMeshCacheEnabled() && !Operation->bLiveUpdateMode && CVarEnableMeshCache.GetValueOnGameThread();
 #if WITH_EDITOR
@@ -2839,7 +2839,7 @@ namespace impl
 			bIsInEditorViewport)
 		{
 			Operation->RequestedLODs.Init(MAX_uint8, Operation->NumComponents);
-			Operation->InstanceDescriptorRuntimeHash.UpdateRequestedLODs(Operation->RequestedLODs);
+			Operation->InstanceDescriptorHash.UpdateRequestedLODs(Operation->RequestedLODs);
 		}
 
 #ifdef MUTABLE_USE_NEW_TASKGRAPH
diff --git a/Engine/Plugins/Experimental/Mutable/Source/CustomizableObject/Private/MuCO/CustomizableObjectSystemPrivate.h b/Engine/Plugins/Experimental/Mutable/Source/CustomizableObject/Private/MuCO/CustomizableObjectSystemPrivate.h
--- a/Engine/Plugins/Experimental/Mutable/Source/CustomizableObject/Private/MuCO/CustomizableObjectSystemPrivate.h
+++ b/Engine/Plugins/Experimental/Mutable/Source/CustomizableObject/Private/MuCO/CustomizableObjectSystemPrivate.h
@@ -4,6 +4,7 @@
 
 #include "LogBenchmarkUtil.h"
 #include "Containers/Queue.h"
+#include "MuCO/DescriptorHash.h"
 #include "MuCO/CustomizableObject.h"
 #include "MuCO/CustomizableObjectInstance.h"
 #include "MuCO/CustomizableObjectExtension.h"
@@ -431,7 +432,7 @@ public:
 	TWeakObjectPtr<UCustomizableObjectInstance> Instance;
 
 	/** Hash of the UCustomizableObjectInstance::Descriptor at the time of the update request. */
-	FDescriptorRuntimeHash InstanceDescriptorRuntimeHash;
+	FDescriptorHash InstanceDescriptorHash;
 			
 	/** Instance parameters at the time of the operation request. */
 	mu::ParametersPtr Parameters; 
diff --git a/Engine/Plugins/Experimental/Mutable/Source/CustomizableObject/Private/MuCO/DescriptorHash.cpp b/Engine/Plugins/Experimental/Mutable/Source/CustomizableObject/Private/MuCO/DescriptorHash.cpp
new file mode 100644
--- /dev/null
+++ b/Engine/Plugins/Experimental/Mutable/Source/CustomizableObject/Private/MuCO/DescriptorHash.cpp
@@ -0,0 +1,162 @@
+// Copyright Epic Games, Inc. All Rights Reserved.
+
+#include "MuCO/DescriptorHash.h"
+
+#include "MuCO/CustomizableObject.h"
+
+
+FDescriptorHash::FDescriptorHash(const FCustomizableObjectInstanceDescriptor& Descriptor)
+{
+#if WITH_EDITORONLY_DATA
+	if (Descriptor.CustomizableObject)
+	{
+		Hash = HashCombine(Hash, GetTypeHash(Descriptor.CustomizableObject->GetPathName()));
+		Hash = HashCombine(Hash, GetTypeHash(Descriptor.CustomizableObject->GetCompilationGuid()));
+	}
+#endif
+
+	for (const FCustomizableObjectBoolParameterValue& Value : Descriptor.BoolParameters)
+	{
+		Hash = HashCombine(Hash, GetTypeHash(Value));
+	}	
+
+	for (const FCustomizableObjectIntParameterValue& Value : Descriptor.IntParameters)
+	{
+		Hash = HashCombine(Hash, GetTypeHash(Value));
+	}
+	
+	for (const FCustomizableObjectFloatParameterValue& Value : Descriptor.FloatParameters)
+	{
+		Hash = HashCombine(Hash, GetTypeHash(Value));
+	}
+
+	for (const FCustomizableObjectTextureParameterValue& Value : Descriptor.TextureParameters)
+	{
+		Hash = HashCombine(Hash, GetTypeHash(Value));
+	}
+
+	for (const FCustomizableObjectVectorParameterValue& Value : Descriptor.VectorParameters)
+	{
+		Hash = HashCombine(Hash, GetTypeHash(Value));
+	}
+
+	for (const FCustomizableObjectProjectorParameterValue& Value : Descriptor.ProjectorParameters)
+	{
+		Hash = HashCombine(Hash, GetTypeHash(Value));
+	}
+	
+	Hash = HashCombine(Hash, GetTypeHash(Descriptor.State));
+	Hash = HashCombine(Hash, GetTypeHash(Descriptor.bBuildParameterRelevancy));
+
+	for (const TTuple<FName, FMultilayerProjector>& Pair : Descriptor.MultilayerProjectors)
+	{
+		// Hash = HashCombine(Hash, GetTypeHash(Pair.Key)); // Already hashed by the FMultilayerProjector.
+		Hash = HashCombine(Hash, GetTypeHash(Pair.Value));
+	}
+	
+	MinLOD = Descriptor.MinLOD;
+	MaxLOD = Descriptor.MaxLOD;	
+
+	RequestedLODsPerComponent = Descriptor.RequestedLODLevels;
+}
+
+
+bool FDescriptorHash::operator==(const FDescriptorHash& Other) const
+{
+	return Hash == Other.Hash &&
+		MinLOD == Other.MinLOD &&
+		MaxLOD == Other.MaxLOD &&
+		RequestedLODsPerComponent == Other.RequestedLODsPerComponent;
+}
+
+
+bool FDescriptorHash::operator!=(const FDescriptorHash& Other) const
+{
+	return !(*this == Other);
+}
+
+
+bool FDescriptorHash::IsSubset(const FDescriptorHash& Other) const
+{
+	if (Hash != Other.Hash || MinLOD < Other.MinLOD || MaxLOD != Other.MaxLOD)
+	{
+		return false;
+	}
+
+	if (RequestedLODsPerComponent == Other.RequestedLODsPerComponent)
+	{
+		return true;
+	}
+
+	if (RequestedLODsPerComponent.Num() != Other.RequestedLODsPerComponent.Num())
+	{
+		return false;
+	}
+	
+	for (int32 ComponentIndex = 0; ComponentIndex < RequestedLODsPerComponent.Num(); ++ComponentIndex)
+	{
+		const int32 RequestedLODs = RequestedLODsPerComponent[ComponentIndex];
+		const int32 OtherRequestedLODs = Other.RequestedLODsPerComponent[ComponentIndex];
+		for (uint16 LODIndex = MinLOD; LODIndex <= MaxLOD; ++LODIndex)
+		{
+			// To be a subset all bits set in RequestedLODs must be set in OtherRequestedLODs. OtherRequestedLODs can have additional bits set
+			if (RequestedLODs & (1 << LODIndex) && !(OtherRequestedLODs & (1 << LODIndex)))
+			{
+				return false;
+			}
+		}
+	}
+
+	return true;
+}
+
+
+void FDescriptorHash::UpdateMinMaxLOD(const int32 InMinLOD, const int32 InMaxLOD)
+{
+	MinLOD = InMinLOD;
+	MaxLOD = InMaxLOD;	
+}
+
+
+int32 FDescriptorHash::GetMinLOD() const
+{
+	return MinLOD;
+}
+
+
+int32 FDescriptorHash::GetMaxLOD() const
+{
+	return MaxLOD;
+}
+
+
+void FDescriptorHash::UpdateRequestedLODs(const TArray<uint16>& InRequestedLODs)
+{
+	RequestedLODsPerComponent = InRequestedLODs;
+}
+
+
+const TArray<uint16>& FDescriptorHash::GetRequestedLODs() const
+{
+	return RequestedLODsPerComponent;
+}
+
+
+FString FDescriptorHash::ToString() const
+{
+	TStringBuilder<150> Builder;
+
+	Builder.Appendf(TEXT("(Hash=%u,"), Hash);
+	Builder.Appendf(TEXT("MinLOD=%i,"), MinLOD);
+	Builder.Appendf(TEXT("MaxLOD=%i,"), MaxLOD);
+	Builder.Appendf(TEXT("RequestredLODLevels=["));
+
+	for (const uint16 RequestedLODs : RequestedLODsPerComponent)
+	{
+		Builder.Appendf(TEXT("%i,"), RequestedLODs);
+	}
+	
+	Builder.Appendf(TEXT("])"));
+
+	return Builder.ToString();
+}
diff --git a/Engine/Plugins/Experimental/Mutable/Source/CustomizableObject/Public/MuCO/CustomizableObjectInstanceDescriptor.h b/Engine/Plugins/Experimental/Mutable/Source/CustomizableObject/Public/MuCO/CustomizableObjectInstanceDescriptor.h
--- a/Engine/Plugins/Experimental/Mutable/Source/CustomizableObject/Public/MuCO/CustomizableObjectInstanceDescriptor.h
+++ b/Engine/Plugins/Experimental/Mutable/Source/CustomizableObject/Public/MuCO/CustomizableObjectInstanceDescriptor.h
@@ -16,7 +16,6 @@ class UCustomizableInstancePrivate;
 class UCustomizableObject;
 class UCustomizableObjectInstance;
 class FDescriptorHash;
-class FDescriptorRuntimeHash;
 class FMutableUpdateCandidate;
 
 typedef TMap<const UCustomizableObjectInstance*, FMutableUpdateCandidate> FMutableInstanceUpdateMap;
@@ -401,7 +400,6 @@ private:
 	
 	// Friends
 	friend FDescriptorHash;
-	friend FDescriptorRuntimeHash;
 	friend UCustomizableObjectInstance;
 	friend UCustomizableInstancePrivate;
 	friend FMultilayerProjector;
@@ -409,56 +407,6 @@ private:
 };
 
 
-/** Hash of the Descriptor. Hashes everything except runtime information. */
-class CUSTOMIZABLEOBJECT_API FDescriptorHash
-{
-public:
-	FDescriptorHash() = default;
-
-	explicit FDescriptorHash(const FCustomizableObjectInstanceDescriptor& Descriptor);
-
-	bool operator==(const FDescriptorHash& Other) const;
-
-	bool operator!=(const FDescriptorHash& Other) const;
-
-	bool operator<(const FDescriptorHash& Other) const;
-
-protected:
-	uint32 Hash = 0;
-};
-
-
-/** Hash of the Descriptor. Hashes everything including runtime information. */
-class CUSTOMIZABLEOBJECT_API FDescriptorRuntimeHash : public FDescriptorHash
-{
-public:
-	FDescriptorRuntimeHash() = default;
-
-	explicit FDescriptorRuntimeHash(const FCustomizableObjectInstanceDescriptor& Descriptor);
-
-	/** Return true if this Hash is a subset of the other Hash (i.e., this Descriptor is a subset of the other Descriptor). */
-	bool IsSubset(const FDescriptorRuntimeHash& Other) const;
-
-	void UpdateMinMaxLOD(int32 InMinLOD, int32 InMaxLOD);
-
-	int32 GetMinLOD() const;
-
-	int32 GetMaxLOD() const;
-
-	void UpdateRequestedLODs(const TArray<uint16>& InRequestedLODs);
-
-	const TArray<uint16>& GetRequestedLODs() const;
-
-	FString ToString() const;
-
-private:
-	int32 MinLOD = 0;
-	int32 MaxLOD = INT32_MAX;
-
-	// Array of bitmasks that indicate which LODs of each component have been requested
-	TArray<uint16> RequestedLODsPerComponent;
-};
-
 #if UE_ENABLE_INCLUDE_ORDER_DEPRECATED_IN_5_2
 #include "MuCO/CustomizableObject.h"
 #include "MuCO/CustomizableObjectParameterTypeDefinitions.h"
