[Mutable] Fix Text conflict from NSLOCTEXT macro for namespace in SMutableGraphViewer.cpp

#rnx

diff --git a/Engine/Plugins/Experimental/Mutable/Source/CustomizableObjectEditor/Private/MuCOE/SMutableGraphViewer.cpp b/Engine/Plugins/Experimental/Mutable/Source/CustomizableObjectEditor/Private/MuCOE/SMutableGraphViewer.cpp
--- a/Engine/Plugins/Experimental/Mutable/Source/CustomizableObjectEditor/Private/MuCOE/SMutableGraphViewer.cpp
+++ b/Engine/Plugins/Experimental/Mutable/Source/CustomizableObjectEditor/Private/MuCOE/SMutableGraphViewer.cpp
@@ -310,9 +310,9 @@ TSharedRef<SWidget> SMutableGraphViewer::GenerateCompileOptionsMenuContent()
 
 		// Optimisation level
 		CompileOptimizationStrings.Empty();
-		CompileOptimizationStrings.Add(MakeShareable(new FString(NSLOCTEXT("CustomizableObjectEditor", "OptimizationNone", "None").ToString())));
-		CompileOptimizationStrings.Add(MakeShareable(new FString(NSLOCTEXT("CustomizableObjectEditor", "OptimizationMin", "Minimal").ToString())));
-		CompileOptimizationStrings.Add(MakeShareable(new FString(NSLOCTEXT("CustomizableObjectEditor", "OptimizationMax", "Maximum").ToString())));
+		CompileOptimizationStrings.Add(MakeShareable(new FString(NSLOCTEXT("CustomizableObjectEditor", "Debugger_OptimizationNone", "None").ToString())));
+		CompileOptimizationStrings.Add(MakeShareable(new FString(NSLOCTEXT("CustomizableObjectEditor", "Debugger_OptimizationMin", "Minimal").ToString())));
+		CompileOptimizationStrings.Add(MakeShareable(new FString(NSLOCTEXT("CustomizableObjectEditor", "Debugger_OptimizationMax", "Maximum").ToString())));
 
 		CompileOptions.OptimizationLevel = FMath::Min(CompileOptions.OptimizationLevel, CompileOptimizationStrings.Num() - 1);
 
