[Mutable] Changing the materials of a mesh that has been initialized will recreate the render state of the components using it. Prevents invalid access from the render pipeline to GCed materials.
- Remove CVarTicksUntilMaterialRelease.

#rnx

diff --git a/Engine/Plugins/Experimental/Mutable/Source/CustomizableObject/Classes/MuCO/CustomizableObjectSystem.h b/Engine/Plugins/Experimental/Mutable/Source/CustomizableObject/Classes/MuCO/CustomizableObjectSystem.h
--- a/Engine/Plugins/Experimental/Mutable/Source/CustomizableObject/Classes/MuCO/CustomizableObjectSystem.h
+++ b/Engine/Plugins/Experimental/Mutable/Source/CustomizableObject/Classes/MuCO/CustomizableObjectSystem.h
@@ -264,16 +264,10 @@ public:
 
 	void PurgePendingReleaseSkeletalMesh();
 
-	void AddPendingReleaseMaterials(TArray<TObjectPtr<UMaterialInterface>>& InMaterials);
-	void TickPendingReleaseMaterials();
-
 private:
 	UPROPERTY()
 	TArray<FPendingReleaseSkeletalMeshInfo> PendingReleaseSkeletalMesh;
 
-	UPROPERTY(Transient)
-	TArray<FPendingReleaseMaterialsInfo> PendingReleaseMaterials;
-
 public:
 	/** [Texture Parameters] Get a list of all the possible values for external texture parameters according to the various providers registered with RegisterImageProvider. */
 	TArray<FCustomizableObjectExternalTexture> GetTextureParameterValues();
diff --git a/Engine/Plugins/Experimental/Mutable/Source/CustomizableObject/Private/MuCO/CustomizableObjectInstance.cpp b/Engine/Plugins/Experimental/Mutable/Source/CustomizableObject/Private/MuCO/CustomizableObjectInstance.cpp
--- a/Engine/Plugins/Experimental/Mutable/Source/CustomizableObject/Private/MuCO/CustomizableObjectInstance.cpp
+++ b/Engine/Plugins/Experimental/Mutable/Source/CustomizableObject/Private/MuCO/CustomizableObjectInstance.cpp
@@ -48,6 +48,8 @@
 #if WITH_EDITOR
 #include "Logging/MessageLog.h"
 #include "MessageLogModule.h"
+#include "UnrealEdMisc.h"
+#include "Subsystems/AssetEditorSubsystem.h"
 #endif
 
 
@@ -5406,10 +5408,10 @@ void UCustomizableInstancePrivate::BuildMaterials(const TSharedRef<FUpdateContex
 	
 	const bool bReuseTextures = OperationData->bReuseInstanceTextures;
 
-	const FInstanceUpdateData::FLOD& FirstLOD = OperationData->InstanceUpdateData.LODs[OperationData->CurrentMinLOD];
-	const int32 NumComponents = FirstLOD.ComponentCount;
+	TArray<bool> RecreateRenderStateOnComponent;
+	RecreateRenderStateOnComponent.Init(false, OperationData->NumComponents);
 
-	for (int32 ComponentIndex = 0; ComponentIndex < NumComponents; ++ComponentIndex)
+	for (int32 ComponentIndex = 0; ComponentIndex < OperationData->NumComponents; ++ComponentIndex)
 	{
 		if (!Public->SkeletalMeshes.IsValidIndex(ComponentIndex))
 		{
@@ -5429,28 +5431,6 @@ void UCustomizableInstancePrivate::BuildMaterials(const TSharedRef<FUpdateContex
 
 		ComponentsData[ComponentIndex].OverrideMaterials.Reset();
 
-		if (!bUseOverrideMaterialsOnly)
-		{
-			{
-				// TEMP: Keep a reference to the previous materials for n frames to avoid GC of materials in use in the render thread.
-				// TODO: MTBL-1632 - Implement a proper fix to replace this hotfix
-				const TArray<FSkeletalMaterial>& Materials = SkeletalMesh->GetMaterials();
-				const int32 NumMaterials = Materials.Num();
-
-				TArray<TObjectPtr<UMaterialInterface>> MaterialsToRelease;
-				MaterialsToRelease.Reserve(NumMaterials);
-
-				for (const FSkeletalMaterial& Material : Materials)
-				{
-					MaterialsToRelease.Add(Material.MaterialInterface);
-				}
-
-				UCustomizableObjectSystem::GetInstance()->AddPendingReleaseMaterials(MaterialsToRelease);
-			}
-
-			SkeletalMesh->GetMaterials().Reset();
-		}
-
 		TArray<FSkeletalMaterial> Materials;
 
 		// Maps serializations of FMutableMaterialPlaceholder to Created Dynamic Material instances, used to reuse materials across LODs
@@ -5968,14 +5948,6 @@ void UCustomizableInstancePrivate::BuildMaterials(const TSharedRef<FUpdateContex
 			}
 		}
 
-		if (!bUseOverrideMaterialsOnly)
-		{
-			SkeletalMesh->SetMaterials(Materials);
-		}
-
-		// Ensure the number of materials is the same on both sides when using overrides. 
-		check(SkeletalMesh->GetMaterials().Num() == Materials.Num());
-
 		{
 			// Copy data from valid LODs into the skipped ones.
 			int32 LastValidLODIndex = OperationData->CurrentMaxLOD;
@@ -6006,9 +5978,65 @@ void UCustomizableInstancePrivate::BuildMaterials(const TSharedRef<FUpdateContex
 				}
 			}
 		}
+
+		if (!bUseOverrideMaterialsOnly)
+		{
+			// Force recreate render state after replacing the materials to avoid a crash in the render pipeline if the old materials are GCed while in use.
+			RecreateRenderStateOnComponent[ComponentIndex] = SkeletalMesh->GetResourceForRendering()->IsInitialized() && SkeletalMesh->GetMaterials() != Materials;
+
+			SkeletalMesh->SetMaterials(Materials);
+
+#if WITH_EDITOR
+			if (RecreateRenderStateOnComponent[ComponentIndex])
+			{
+				// Close all open editors for this mesh to invalidate viewports.
+				GEditor->GetEditorSubsystem<UAssetEditorSubsystem>()->CloseAllEditorsForAsset(SkeletalMesh);
+			}
+#endif
+		}
+
+		// Ensure the number of materials is the same on both sides when using overrides. 
+		check(SkeletalMesh->GetMaterials().Num() == Materials.Num());
+	}
+
+	// Force recreate render state if the mesh is reused and the materials have changed.
+	// TODO: MTBL-1697 Remove after merging ConvertResources and Callbacks.
+	if (RecreateRenderStateOnComponent.Find(true) != INDEX_NONE)
+	{
+		MUTABLE_CPUPROFILER_SCOPE(BuildMaterials_RecreateRenderState);
+
+		for (TObjectIterator<UCustomizableObjectInstanceUsage> It; It; ++It)
+		{
+			UCustomizableObjectInstanceUsage* CustomizableObjectInstanceUsage = *It;
+
+			if (!IsValid(CustomizableObjectInstanceUsage) || CustomizableObjectInstanceUsage->GetCustomizableObjectInstance() != Public)
+			{
+				continue;
+			}
+
+#if WITH_EDITOR
+			if (CustomizableObjectInstanceUsage->IsNetMode(NM_DedicatedServer))
+			{
+				continue;
+			}
+#endif
+
+			const int32 ComponentIndex = CustomizableObjectInstanceUsage->GetComponentIndex();
+			if (!RecreateRenderStateOnComponent.IsValidIndex(ComponentIndex) || !RecreateRenderStateOnComponent[ComponentIndex])
+			{
+				continue;
+			}
+
+			USkeletalMeshComponent* AttachedParent = CustomizableObjectInstanceUsage->GetAttachParent();
+			if (!AttachedParent || AttachedParent->GetSkeletalMeshAsset() != Public->SkeletalMeshes[ComponentIndex])
+			{
+				continue;
+			}
+
+			AttachedParent->RecreateRenderState_Concurrent();
+		}
 	}
 
-	
 	{
 		MUTABLE_CPUPROFILER_SCOPE(BuildMaterials_Exchange);
 
diff --git a/Engine/Plugins/Experimental/Mutable/Source/CustomizableObject/Private/MuCO/CustomizableObjectSystem.cpp b/Engine/Plugins/Experimental/Mutable/Source/CustomizableObject/Private/MuCO/CustomizableObjectSystem.cpp
--- a/Engine/Plugins/Experimental/Mutable/Source/CustomizableObject/Private/MuCO/CustomizableObjectSystem.cpp
+++ b/Engine/Plugins/Experimental/Mutable/Source/CustomizableObject/Private/MuCO/CustomizableObjectSystem.cpp
@@ -1464,35 +1464,6 @@ void UCustomizableObjectSystem::PurgePendingReleaseSkeletalMesh()
 	}
 }
 
-static int32 TicksUntilMaterialRelease = 4;
-
-static FAutoConsoleVariableRef CVarTicksUntilMaterialRelease(
-	TEXT("mutable.TicksUntilMaterialRelease"), TicksUntilMaterialRelease,
-	TEXT("If higher than zero, the number of ticks to keep discarded materials guarded from being GCed."),
-	ECVF_Default);
-
-void UCustomizableObjectSystem::AddPendingReleaseMaterials(TArray<TObjectPtr<UMaterialInterface>>& InMaterials)
-{
-	if (TicksUntilMaterialRelease > 0)
-	{
-		FPendingReleaseMaterialsInfo& PendingMaterials = PendingReleaseMaterials.AddDefaulted_GetRef();
-		PendingMaterials.Materials = InMaterials;
-		PendingMaterials.TicksUntilRelease = TicksUntilMaterialRelease;
-	}
-}
-
-
-void UCustomizableObjectSystem::TickPendingReleaseMaterials()
-{
-	for (TArray<FPendingReleaseMaterialsInfo>::TIterator PendingRelease = PendingReleaseMaterials.CreateIterator(); PendingRelease; ++PendingRelease)
-	{
-		if (--PendingRelease->TicksUntilRelease <= 0)
-		{
-			PendingRelease.RemoveCurrent();
-		}
-	}
-}
-
 
 void UCustomizableObjectSystem::AddPendingReleaseSkeletalMesh(USkeletalMesh* SkeletalMesh)
 {
@@ -2980,8 +2951,6 @@ bool UCustomizableObjectSystem::Tick(float DeltaTime)
 		return true;
 	}
 #endif
-	
-	TickPendingReleaseMaterials();
 
 	Private->UpdateStats();
 	
