[Mutable] Remove CompilationGuid. Reset DescriptorHash when invalidating the generated resources of an instance.

#rnx

diff --git a/Engine/Plugins/Experimental/Mutable/Source/CustomizableObject/Classes/MuCO/CustomizableObject.h b/Engine/Plugins/Experimental/Mutable/Source/CustomizableObject/Classes/MuCO/CustomizableObject.h
--- a/Engine/Plugins/Experimental/Mutable/Source/CustomizableObject/Classes/MuCO/CustomizableObject.h
+++ b/Engine/Plugins/Experimental/Mutable/Source/CustomizableObject/Classes/MuCO/CustomizableObject.h
@@ -1590,18 +1590,12 @@ public:
 	UPROPERTY(Transient)
 	FGuid Identifier;
 
-	/** Unique Identifier - Regenerated each time the object is compiled. */
-	UPROPERTY(Transient)
-	FGuid CompilationGuid;
-
 	ECustomizableObjectCompilationState CompilationState = ECustomizableObjectCompilationState::None;
 
 	FPostCompileDelegate PostCompileDelegate;
 
 	void PostCompile();
 
-	FGuid GetCompilationGuid() const;
-
 #endif
 
 	/** Create a new instance of this object. The instance parameters will be initialized with the object default values. */
diff --git a/Engine/Plugins/Experimental/Mutable/Source/CustomizableObject/Private/MuCO/CustomizableObject.cpp b/Engine/Plugins/Experimental/Mutable/Source/CustomizableObject/Private/MuCO/CustomizableObject.cpp
--- a/Engine/Plugins/Experimental/Mutable/Source/CustomizableObject/Private/MuCO/CustomizableObject.cpp
+++ b/Engine/Plugins/Experimental/Mutable/Source/CustomizableObject/Private/MuCO/CustomizableObject.cpp
@@ -1265,16 +1265,8 @@ FString UCustomizableObject::GetStateParameterName(int32 StateIndex, int32 Param
 #if WITH_EDITORONLY_DATA
 void UCustomizableObject::PostCompile()
 {
-	CompilationGuid = FGuid::NewGuid();
-
 	PostCompileDelegate.Broadcast();
 }
-
-
-FGuid UCustomizableObject::GetCompilationGuid() const
-{
-	return CompilationGuid;
-}
 #endif
 
 
diff --git a/Engine/Plugins/Experimental/Mutable/Source/CustomizableObject/Private/MuCO/CustomizableObjectInstance.cpp b/Engine/Plugins/Experimental/Mutable/Source/CustomizableObject/Private/MuCO/CustomizableObjectInstance.cpp
--- a/Engine/Plugins/Experimental/Mutable/Source/CustomizableObject/Private/MuCO/CustomizableObjectInstance.cpp
+++ b/Engine/Plugins/Experimental/Mutable/Source/CustomizableObject/Private/MuCO/CustomizableObjectInstance.cpp
@@ -192,6 +192,9 @@ void UCustomizableInstancePrivate::InvalidateGeneratedData()
 {
 	SkeletalMeshStatus = ESkeletalMeshStatus::NotGenerated;
 
+	DescriptorHash = FDescriptorHash();
+	UpdateDescriptorHash = FDescriptorHash();
+
 	// Init Component Data
 	FCustomizableInstanceComponentData TemplateComponentData;
 	TemplateComponentData.LastMeshIdPerLOD.Init(MAX_uint64, MAX_MESH_LOD_COUNT);
@@ -2781,7 +2784,6 @@ void UCustomizableInstancePrivate::DiscardResources()
 	InvalidateGeneratedData();
 	
 	Instance->SkeletalMeshes.Reset();
-	DescriptorHash = FDescriptorHash();
 }
 
 
@@ -6191,6 +6193,11 @@ void UCustomizableObjectInstance::SetRequestedLODs(int32 InMinLOD, int32 InMaxLO
 		return;
 	}
 
+	if (GetPrivate()->SkeletalMeshStatus == ESkeletalMeshStatus::Error)
+	{
+		return;
+	}
+
 	if (CVarPreserveUserLODsOnFirstGeneration.GetValueOnGameThread() &&
 		GetCustomizableObject()->IsPreserveUserLODsOnFirstGeneration() &&
 		GetPrivate()->SkeletalMeshStatus != ESkeletalMeshStatus::Success)
diff --git a/Engine/Plugins/Experimental/Mutable/Source/CustomizableObject/Private/MuCO/DescriptorHash.cpp b/Engine/Plugins/Experimental/Mutable/Source/CustomizableObject/Private/MuCO/DescriptorHash.cpp
--- a/Engine/Plugins/Experimental/Mutable/Source/CustomizableObject/Private/MuCO/DescriptorHash.cpp
+++ b/Engine/Plugins/Experimental/Mutable/Source/CustomizableObject/Private/MuCO/DescriptorHash.cpp
@@ -11,7 +11,6 @@ FDescriptorHash::FDescriptorHash(const FCustomizableObjectInstanceDescriptor& De
 	if (Descriptor.CustomizableObject)
 	{
 		Hash = HashCombine(Hash, GetTypeHash(Descriptor.CustomizableObject->GetPathName()));
-		Hash = HashCombine(Hash, GetTypeHash(Descriptor.CustomizableObject->GetCompilationGuid()));
 	}
 #endif
 
