Silence false V621, V654, and V1078 warnings mostly caused by TStaticArray, placement new, or population of the array inside a lambda

diff --git a/Engine/Plugins/Enterprise/DatasmithCADImporter/Source/CADLibrary/Private/MeshDescriptionHelper.cpp b/Engine/Plugins/Enterprise/DatasmithCADImporter/Source/CADLibrary/Private/MeshDescriptionHelper.cpp
--- a/Engine/Plugins/Enterprise/DatasmithCADImporter/Source/CADLibrary/Private/MeshDescriptionHelper.cpp
+++ b/Engine/Plugins/Enterprise/DatasmithCADImporter/Source/CADLibrary/Private/MeshDescriptionHelper.cpp
@@ -523,7 +523,7 @@ bool FillMesh(const FMeshConversionContext& MeshConversionContext, FBodyMesh& Bo
 			{
 				// compute normals of Symmetric vertex
 				FMatrix44f SymmetricMatrix = FDatasmithUtils::GetSymmetricMatrix(MeshConversionContext.MeshParameters.SymmetricOrigin, MeshConversionContext.MeshParameters.SymmetricNormal);
-				for (const FVertexInstanceID& VertexInstanceID : MeshVertexInstanceIDs)
+				for (const FVertexInstanceID& VertexInstanceID : MeshVertexInstanceIDs) //-V1078
 				{
 					VertexInstanceNormals[VertexInstanceID] = SymmetricMatrix.TransformVector(VertexInstanceNormals[VertexInstanceID]);
 				}
diff --git a/Engine/Plugins/Enterprise/DatasmithFBXImporter/Source/DatasmithVREDTranslator/Private/DatasmithVREDImporter.cpp b/Engine/Plugins/Enterprise/DatasmithFBXImporter/Source/DatasmithVREDTranslator/Private/DatasmithVREDImporter.cpp
--- a/Engine/Plugins/Enterprise/DatasmithFBXImporter/Source/DatasmithVREDTranslator/Private/DatasmithVREDImporter.cpp
+++ b/Engine/Plugins/Enterprise/DatasmithFBXImporter/Source/DatasmithVREDTranslator/Private/DatasmithVREDImporter.cpp
@@ -1997,7 +1997,7 @@ bool FDatasmithVREDImporter::SendSceneToDatasmith()
 		// clips to not really care if nodes were split or not
 		TArray<FCombinedAnimBlock> CombinedBlocks = CombineBlocks(IntermediateScene->AnimNodes);
 
-		for (FCombinedAnimBlock& CombinedBlock : CombinedBlocks)
+		for (FCombinedAnimBlock& CombinedBlock : CombinedBlocks) //-V1078
 		{
 			TSharedPtr<IDatasmithLevelSequenceElement> ConvertedBlock = ConvertAnimBlock(CombinedBlock);
 			if (ConvertedBlock.IsValid())
diff --git a/Engine/Plugins/Experimental/AnimNext/Source/AnimNextUncookedOnly/Private/UncookedOnlyUtils.cpp b/Engine/Plugins/Experimental/AnimNext/Source/AnimNextUncookedOnly/Private/UncookedOnlyUtils.cpp
--- a/Engine/Plugins/Experimental/AnimNext/Source/AnimNextUncookedOnly/Private/UncookedOnlyUtils.cpp
+++ b/Engine/Plugins/Experimental/AnimNext/Source/AnimNextUncookedOnly/Private/UncookedOnlyUtils.cpp
@@ -2059,7 +2059,7 @@ void FUtils::CompileSchedule(UAnimNextSchedule* InSchedule)
 
 		// Build sources that need to run (thread-safe or not) at each index
 		TMap<FInsertionLocation, TArray<FAnimNextScheduleExternalParameterSource>> InsertionMap;
-		for(const TPair<FName, FParameterTracker>& TrackedParameterPair : TrackerMap)
+		for(const TPair<FName, FParameterTracker>& TrackedParameterPair : TrackerMap) //-V1078
 		{
 			// If an index/container was not set up, then the update of the set of parameters is not a pre-requisite of a task, so we just insert the task at the
 			// start of the schedule as the only requirement is that the parameters exist 
diff --git a/Engine/Plugins/Experimental/BlueprintStats/Source/BlueprintStats/Private/BlueprintStatsModule.cpp b/Engine/Plugins/Experimental/BlueprintStats/Source/BlueprintStats/Private/BlueprintStatsModule.cpp
--- a/Engine/Plugins/Experimental/BlueprintStats/Source/BlueprintStats/Private/BlueprintStatsModule.cpp
+++ b/Engine/Plugins/Experimental/BlueprintStats/Source/BlueprintStats/Private/BlueprintStatsModule.cpp
@@ -68,7 +68,7 @@ void FBlueprintStatsModule::DumpBlueprintStats()
 
 	// Now merge them
 	FBlueprintStatRecord Aggregate(NULL);
-	for (const FBlueprintStatRecord& SourceRecord : Records)
+	for (const FBlueprintStatRecord& SourceRecord : Records) //-V1078
 	{
 		Aggregate.MergeAnotherRecordIn(SourceRecord);
 	}
diff --git a/Engine/Plugins/Experimental/ChaosClothAssetEditor/Source/ChaosClothAssetDataflowNodes/Private/ChaosClothAsset/RemeshNode.cpp b/Engine/Plugins/Experimental/ChaosClothAssetEditor/Source/ChaosClothAssetDataflowNodes/Private/ChaosClothAsset/RemeshNode.cpp
--- a/Engine/Plugins/Experimental/ChaosClothAssetEditor/Source/ChaosClothAssetDataflowNodes/Private/ChaosClothAsset/RemeshNode.cpp
+++ b/Engine/Plugins/Experimental/ChaosClothAssetEditor/Source/ChaosClothAssetDataflowNodes/Private/ChaosClothAsset/RemeshNode.cpp
@@ -118,7 +118,7 @@ namespace UE::Chaos::ClothAsset::Private
 
 				const FIntVector2& InnerStitch = InnerSeam[InnerStitchID];
 
-				for (const int32& SeamVertex : SeamEdgeVertices)
+				for (const int32& SeamVertex : SeamEdgeVertices) //-V1078
 				{
 					if (SeamVertex == InnerStitch[0] || SeamVertex == InnerStitch[1])
 					{
diff --git a/Engine/Plugins/Experimental/GPULightmass/Source/GPULightmass/Private/LightmapRenderer.cpp b/Engine/Plugins/Experimental/GPULightmass/Source/GPULightmass/Private/LightmapRenderer.cpp
--- a/Engine/Plugins/Experimental/GPULightmass/Source/GPULightmass/Private/LightmapRenderer.cpp
+++ b/Engine/Plugins/Experimental/GPULightmass/Source/GPULightmass/Private/LightmapRenderer.cpp
@@ -2062,7 +2062,7 @@ void FLightmapRenderer::Finalize(FRDGBuilder& GraphBuilder)
 
 	TStaticArray<FRDGTextureUAVRef, 3> ScratchTilePoolLayerUAVs;
 
-	for (int32 Index = 0; Index < ScratchTilePoolLayerUAVs.Num(); ++Index)
+	for (int32 Index = 0; Index < ScratchTilePoolLayerUAVs.Num(); ++Index) //-V621 //-V654
 	{
 		ScratchTilePoolLayerUAVs[Index] = GraphBuilder.CreateUAV(GraphBuilder.RegisterExternalTexture(ScratchTilePoolGPU->PooledRenderTargets[Index]));
 	}
diff --git a/Engine/Plugins/Experimental/Mutable/Source/MutableRuntime/Private/MuR/OpMeshReshape.h b/Engine/Plugins/Experimental/Mutable/Source/MutableRuntime/Private/MuR/OpMeshReshape.h
--- a/Engine/Plugins/Experimental/Mutable/Source/MutableRuntime/Private/MuR/OpMeshReshape.h
+++ b/Engine/Plugins/Experimental/Mutable/Source/MutableRuntime/Private/MuR/OpMeshReshape.h
@@ -226,7 +226,7 @@ namespace mu
 
 		// Geometric linear regression of top, bottom and rings centroids.
 		FVector3f Centroid = FVector3f::Zero();
-		for (const FVector3f& C : Centroids)
+		for (const FVector3f& C : Centroids) //-V1078
 		{
 			Centroid += C;
 		}
@@ -234,7 +234,7 @@ namespace mu
 		constexpr float OneOverNumCentroids = 1.0f / static_cast<float>(NumCentroids);
 		Centroid *= OneOverNumCentroids;
 		
-		for (FVector3f& C : Centroids)
+		for (FVector3f& C : Centroids) //-V1078
 		{
 			C -= Centroid;
 		}
@@ -244,7 +244,7 @@ namespace mu
 		for (int32 Iter = 0; Iter < NumIters; ++Iter)
 		{
 			FVector3f IterDirRefinement = Direction;
-			for (const FVector3f& C : Centroids)
+			for (const FVector3f& C : Centroids) //-V1078
 			{
 				IterDirRefinement += C * FVector3f::DotProduct(Direction, C);
 			}
@@ -253,7 +253,7 @@ namespace mu
 		}
 
 		// Project centroids to the line described by Direction and Centroid.
-		for (FVector3f& C : Centroids)
+		for (FVector3f& C : Centroids) //-V1078
 		{
 			C = Centroid + Direction * FVector3f::DotProduct(C, Direction);
 		}
@@ -312,7 +312,7 @@ namespace mu
 	
 		// Geometric linear regression of top, bottom and ring centroids.
 		FVector3f Centroid = FVector3f::Zero();
-		for (const FVector3f& C : Centroids)
+		for (const FVector3f& C : Centroids) //-V1078
 		{
 			Centroid += C;
 		}
@@ -320,7 +320,7 @@ namespace mu
 		constexpr float OneOverNumCentroids = 1.0f / static_cast<float>(NumCentroids);
 		Centroid *= OneOverNumCentroids;
 		
-		for (FVector3f& C : Centroids)
+		for (FVector3f& C : Centroids) //-V1078
 		{
 			C -= Centroid;
 		}
@@ -330,7 +330,7 @@ namespace mu
 		for (int32 Iter = 0; Iter < NumIters; ++Iter)
 		{
 			FVector3f IterDirRefinement = Direction;
-			for (const FVector3f& C : Centroids)
+			for (const FVector3f& C : Centroids) //-V1078
 			{
 				IterDirRefinement += C * FVector3f::DotProduct(Direction, C);
 			}
@@ -339,7 +339,7 @@ namespace mu
 		}
 
 		// Project centroids to the line described by Direction and Centroid.
-		for (FVector3f& C : Centroids)
+		for (FVector3f& C : Centroids) //-V1078
 		{
 			C = Centroid + Direction * FVector3f::DotProduct(C, Direction);
 		}
diff --git a/Engine/Plugins/Experimental/PCG/Source/PCG/Private/PCGComponent.cpp b/Engine/Plugins/Experimental/PCG/Source/PCG/Private/PCGComponent.cpp
--- a/Engine/Plugins/Experimental/PCG/Source/PCG/Private/PCGComponent.cpp
+++ b/Engine/Plugins/Experimental/PCG/Source/PCG/Private/PCGComponent.cpp
@@ -1045,7 +1045,7 @@ FPCGTaskId UPCGComponent::CreateCleanupTask(bool bRemoveComponents, const TArray
 							return true;
 						});
 
-						for (const FFolder& FolderToDelete : SubfoldersToDelete)
+						for (const FFolder& FolderToDelete : SubfoldersToDelete) //-V1078
 						{
 							FActorFolders::Get().DeleteFolder(*ThisWorld, FolderToDelete);
 						}
diff --git a/Engine/Plugins/Experimental/PCG/Source/PCG/Private/Tests/Elements/Metadata/PCGMetadataBreakTransformTest.cpp b/Engine/Plugins/Experimental/PCG/Source/PCG/Private/Tests/Elements/Metadata/PCGMetadataBreakTransformTest.cpp
--- a/Engine/Plugins/Experimental/PCG/Source/PCG/Private/Tests/Elements/Metadata/PCGMetadataBreakTransformTest.cpp
+++ b/Engine/Plugins/Experimental/PCG/Source/PCG/Private/Tests/Elements/Metadata/PCGMetadataBreakTransformTest.cpp
@@ -259,7 +259,7 @@ bool FPCGMetadataBreakTransformTest::RunTest(const FString& Parameters)
 	// Setup error catching. This error should happen only twice (invalid type)
 	AddExpectedError(TEXT("Attribute/Property 'Float' from pin"), EAutomationExpectedErrorFlags::Contains, 2);
 
-	for (PairWhatData& PairTestData : AllTestData)
+	for (PairWhatData& PairTestData : AllTestData) //-V1078
 	{
 		AddInfo(PairTestData.Key);
 		PCGTestsCommon::FTestData* TestData = PairTestData.Value;
diff --git a/Engine/Plugins/Experimental/PCG/Source/PCG/Private/Tests/Elements/Metadata/PCGMetadataBreakVectorTest.cpp b/Engine/Plugins/Experimental/PCG/Source/PCG/Private/Tests/Elements/Metadata/PCGMetadataBreakVectorTest.cpp
--- a/Engine/Plugins/Experimental/PCG/Source/PCG/Private/Tests/Elements/Metadata/PCGMetadataBreakVectorTest.cpp
+++ b/Engine/Plugins/Experimental/PCG/Source/PCG/Private/Tests/Elements/Metadata/PCGMetadataBreakVectorTest.cpp
@@ -325,7 +325,7 @@ bool FPCGMetadataBreakVectorTest::RunTest(const FString& Parameters)
 	// Setup error catching. This error should happen only twice (invalid type)
 	AddExpectedError(TEXT("Attribute/Property 'Float' from pin"), EAutomationExpectedErrorFlags::Contains, 2);
 
-	for (PairWhatData& PairTestData : AllTestData)
+	for (PairWhatData& PairTestData : AllTestData) //-V1078
 	{
 		AddInfo(PairTestData.Key);
 		PCGTestsCommon::FTestData* TestData = PairTestData.Value;
diff --git a/Engine/Plugins/Experimental/PCG/Source/PCG/Private/Tests/Elements/Metadata/PCGMetadataMakeTransformTest.cpp b/Engine/Plugins/Experimental/PCG/Source/PCG/Private/Tests/Elements/Metadata/PCGMetadataMakeTransformTest.cpp
--- a/Engine/Plugins/Experimental/PCG/Source/PCG/Private/Tests/Elements/Metadata/PCGMetadataMakeTransformTest.cpp
+++ b/Engine/Plugins/Experimental/PCG/Source/PCG/Private/Tests/Elements/Metadata/PCGMetadataMakeTransformTest.cpp
@@ -226,7 +226,7 @@ bool FPCGMetadataMakeTransformTest::RunTest(const FString& Parameters)
 	AllTestData[0] = PairWhatData("Testing with point data as input", &TestDataSpatial);
 	AllTestData[1] = PairWhatData("Testing with param data as input", &TestDataParams);
 
-	for (PairWhatData& PairTestData : AllTestData)
+	for (PairWhatData& PairTestData : AllTestData) //-V1078
 	{
 		AddInfo(PairTestData.Key);
 		PCGTestsCommon::FTestData* TestData = PairTestData.Value;
diff --git a/Engine/Plugins/FX/Niagara/Source/Niagara/Private/NiagaraUvQuadTree.cpp b/Engine/Plugins/FX/Niagara/Source/Niagara/Private/NiagaraUvQuadTree.cpp
--- a/Engine/Plugins/FX/Niagara/Source/Niagara/Private/NiagaraUvQuadTree.cpp
+++ b/Engine/Plugins/FX/Niagara/Source/Niagara/Private/NiagaraUvQuadTree.cpp
@@ -104,7 +104,7 @@ int32 FNiagaraUvQuadTree::FSubTree::Freeze(const FNiagaraUvQuadTree & QuadTree,
 	if (bInternal)
 	{
 		FChildArray ChildOffsets;
-		for (int32 ChildIt = 0; ChildIt < ChildOffsets.Num(); ++ChildIt)
+		for (int32 ChildIt = 0; ChildIt < ChildOffsets.Num(); ++ChildIt) //-V621 //-V654
 		{
 			ChildOffsets[ChildIt] = QuadTree.ChildTrees[SubTreeIndices[ChildIt]].Freeze(QuadTree, Ar);
 		}
@@ -112,7 +112,7 @@ int32 FNiagaraUvQuadTree::FSubTree::Freeze(const FNiagaraUvQuadTree & QuadTree,
 		const int64 EndOffset = Ar.Tell();
 		Ar.Seek(StartOffset);
 
-		for (int32 ChildIt = 0; ChildIt < ChildOffsets.Num(); ++ChildIt)
+		for (int32 ChildIt = 0; ChildIt < ChildOffsets.Num(); ++ChildIt) //-V621 //-V654
 		{
 			Ar << ChildOffsets[ChildIt];
 		}
diff --git a/Engine/Plugins/Runtime/Metasound/Source/MetasoundFrontend/Private/MetasoundFrontendTransform.cpp b/Engine/Plugins/Runtime/Metasound/Source/MetasoundFrontend/Private/MetasoundFrontendTransform.cpp
--- a/Engine/Plugins/Runtime/Metasound/Source/MetasoundFrontend/Private/MetasoundFrontendTransform.cpp
+++ b/Engine/Plugins/Runtime/Metasound/Source/MetasoundFrontend/Private/MetasoundFrontendTransform.cpp
@@ -443,7 +443,7 @@ namespace Metasound
 				FVector2D NextLocation = { 0.0f, 0.0f };
 				{
 					int32 NumBeforeDefined = 1;
-					for (const TPair<int32, FNodeHandle>& Pair : SortOrderToName)
+					for (const TPair<int32, FNodeHandle>& Pair : SortOrderToName) //-V1078
 					{
 						const FConstNodeHandle& NodeHandle = Pair.Value;
 						const FName NodeName = NodeHandle->GetNodeName();
@@ -468,7 +468,7 @@ namespace Metasound
 				}
 
 				// Iterate through sorted map in sequence, slotting in new locations after existing swapped nodes with predefined locations.
-				for (TPair<int32, FNodeHandle>& Pair : SortOrderToName)
+				for (TPair<int32, FNodeHandle>& Pair : SortOrderToName) //-V1078
 				{
 					FNodeHandle& NodeHandle = Pair.Value;
 					const FName NodeName = NodeHandle->GetNodeName();
diff --git a/Engine/Plugins/Runtime/OpenXR/Source/OpenXRHMD/Private/OpenXRHMD_Swapchain.cpp b/Engine/Plugins/Runtime/OpenXR/Source/OpenXRHMD/Private/OpenXRHMD_Swapchain.cpp
--- a/Engine/Plugins/Runtime/OpenXR/Source/OpenXRHMD/Private/OpenXRHMD_Swapchain.cpp
+++ b/Engine/Plugins/Runtime/OpenXR/Source/OpenXRHMD/Private/OpenXRHMD_Swapchain.cpp
@@ -143,7 +143,7 @@ uint8 FOpenXRSwapchain::GetNearestSupportedSwapchainFormat(XrSession InSession,
 	// Search for a fallback format in order of preference (first element in the array has the highest preference).
 	uint8 FallbackFormat = 0;
 	uint32 FallbackPlatformFormat = 0;
-	for (int64_t Format : Formats)
+	for (int64_t Format : Formats) //-V1078
 	{
 		if (RequestedFormat == PF_DepthStencil)
 		{
diff --git a/Engine/Plugins/Runtime/OpenXRViveTracker/Source/OpenXRViveTracker/Private/OpenXRViveTracker.cpp b/Engine/Plugins/Runtime/OpenXRViveTracker/Source/OpenXRViveTracker/Private/OpenXRViveTracker.cpp
--- a/Engine/Plugins/Runtime/OpenXRViveTracker/Source/OpenXRViveTracker/Private/OpenXRViveTracker.cpp
+++ b/Engine/Plugins/Runtime/OpenXRViveTracker/Source/OpenXRViveTracker/Private/OpenXRViveTracker.cpp
@@ -227,7 +227,7 @@ const void* FOpenXRViveTracker::OnCreateSession(XrInstance InInstance, XrSystemI
 
 	// Add LiveLink poses for trackers that have no role assigned
 	uint32 TrackerIndex = 0;
-	for (const XrViveTrackerPathsHTCX& TrackerPath : TrackerPaths)
+	for (const XrViveTrackerPathsHTCX& TrackerPath : TrackerPaths) //-V1078
 	{
 		if (TrackerPath.rolePath == XR_NULL_PATH)
 		{
diff --git a/Engine/Plugins/Runtime/XR/XRScribe/Source/XRScribe/Private/XRScribeEmulationLayer.cpp b/Engine/Plugins/Runtime/XR/XRScribe/Source/XRScribe/Private/XRScribeEmulationLayer.cpp
--- a/Engine/Plugins/Runtime/XR/XRScribe/Source/XRScribe/Private/XRScribeEmulationLayer.cpp
+++ b/Engine/Plugins/Runtime/XR/XRScribe/Source/XRScribe/Private/XRScribeEmulationLayer.cpp
@@ -297,7 +297,7 @@ void FOpenXREmulationLayer::PostLoadActions()
 
 	TArray<TStaticArray<ANSICHAR, XR_MAX_API_LAYER_NAME_SIZE>> EmulatedLayers;
 	SupportedEmulatedLayers = CaptureDecoder.GetApiLayerProperties().FilterByPredicate([EmulatedLayers](const XrApiLayerProperties& Layer) {
-		for (const TStaticArray<ANSICHAR, XR_MAX_API_LAYER_NAME_SIZE>& EmulatedLayerName : EmulatedLayers)
+		for (const TStaticArray<ANSICHAR, XR_MAX_API_LAYER_NAME_SIZE>& EmulatedLayerName : EmulatedLayers) //-V1078
 		{
 			if (FCStringAnsi::Strcmp(Layer.layerName, EmulatedLayerName.GetData()) == 0)
 			{
diff --git a/Engine/Plugins/VirtualProduction/CaptureManager/Source/CaptureProtocolStack/Private/Tests/Export/TestExportResponse.cpp b/Engine/Plugins/VirtualProduction/CaptureManager/Source/CaptureProtocolStack/Private/Tests/Export/TestExportResponse.cpp
--- a/Engine/Plugins/VirtualProduction/CaptureManager/Source/CaptureProtocolStack/Private/Tests/Export/TestExportResponse.cpp
+++ b/Engine/Plugins/VirtualProduction/CaptureManager/Source/CaptureProtocolStack/Private/Tests/Export/TestExportResponse.cpp
@@ -11,7 +11,7 @@
 TStaticArray<uint8, 16> CreateHash()
 {
 	TStaticArray<uint8, 16> StaticArray;
-	for (int32 Index = 0; Index < StaticArray.Num(); ++Index)
+	for (int32 Index = 0; Index < StaticArray.Num(); ++Index) //-V621 //-V654
 	{
 		StaticArray[Index] = Index;
 	}
diff --git a/Engine/Plugins/VirtualProduction/DMX/DMXEngine/Source/DMXRuntime/Private/MVR/DMXMVRGeneralSceneDescription.cpp b/Engine/Plugins/VirtualProduction/DMX/DMXEngine/Source/DMXRuntime/Private/MVR/DMXMVRGeneralSceneDescription.cpp
--- a/Engine/Plugins/VirtualProduction/DMX/DMXEngine/Source/DMXRuntime/Private/MVR/DMXMVRGeneralSceneDescription.cpp
+++ b/Engine/Plugins/VirtualProduction/DMX/DMXEngine/Source/DMXRuntime/Private/MVR/DMXMVRGeneralSceneDescription.cpp
@@ -128,7 +128,7 @@ void UDMXMVRGeneralSceneDescription::WriteDMXLibraryToGeneralSceneDescription(co
 	const TArray<int32> FixtureIDsInUse = GetNumericalFixtureIDsInUse(DMXLibrary);
 	int32 NextFixtureID = FixtureIDsInUse.Num() > 0 ? FixtureIDsInUse.Last() + 1 : 1;
 	
-	for (UDMXMVRFixtureNode* FixtureNode : UE::DMX::DMXMVRGeneralSceneDescription::HotFix::Private::NewNodes)
+	for (UDMXMVRFixtureNode* FixtureNode : UE::DMX::DMXMVRGeneralSceneDescription::HotFix::Private::NewNodes) //-V1078
 	{
 		FixtureNode->FixtureID = FString::FromInt(NextFixtureID);
 		NextFixtureID++;
diff --git a/Engine/Source/Editor/UnrealEd/Private/Cooker/Algo/GraphConvert.cpp b/Engine/Source/Editor/UnrealEd/Private/Cooker/Algo/GraphConvert.cpp
--- a/Engine/Source/Editor/UnrealEd/Private/Cooker/Algo/GraphConvert.cpp
+++ b/Engine/Source/Editor/UnrealEd/Private/Cooker/Algo/GraphConvert.cpp
@@ -637,7 +637,7 @@ bool FCondensationGraphTest::RunTest(const FString& Parameters)
 						}
 					}
 				}
-				for (FVertex ReachableVertex : Reachable)
+				for (FVertex ReachableVertex : Reachable) //-V1078
 				{
 					if (ReachableVertex < Root)
 					{
diff --git a/Engine/Source/Editor/UnrealEd/Private/Kismet2/ReloadUtilities.cpp b/Engine/Source/Editor/UnrealEd/Private/Kismet2/ReloadUtilities.cpp
--- a/Engine/Source/Editor/UnrealEd/Private/Kismet2/ReloadUtilities.cpp
+++ b/Engine/Source/Editor/UnrealEd/Private/Kismet2/ReloadUtilities.cpp
@@ -1111,7 +1111,7 @@ void FReload::Reinstance()
 	FBlueprintEditorUtils::UpdateEnumsInNodes(ChangedEnums, OnNodeFoundOrUpdated);
 
 	// Update all the nodes before we could possibly recompile
-	for (TPair<UBlueprint*, FBlueprintUpdateInfo>& KVP : ModifiedBlueprints)
+	for (TPair<UBlueprint*, FBlueprintUpdateInfo>& KVP : ModifiedBlueprints) //-V1078
 	{
 		UBlueprint* Blueprint = KVP.Key;
 		FBlueprintUpdateInfo& Info = KVP.Value;
@@ -1135,7 +1135,7 @@ void FReload::Reinstance()
 	rch.ReinstanceClasses(ClassesToReinstance);
 
 	// Recompile blueprints if they haven't already been recompiled)
-	for (TPair<UBlueprint*, FBlueprintUpdateInfo>& KVP : ModifiedBlueprints)
+	for (TPair<UBlueprint*, FBlueprintUpdateInfo>& KVP : ModifiedBlueprints) //-V1078
 	{
 		UBlueprint* Blueprint = KVP.Key;
 		FBlueprintUpdateInfo& Info = KVP.Value;
diff --git a/Engine/Source/Editor/UnrealEd/Private/WorldPartition/WorldPartitionRenameDuplicateBuilder.cpp b/Engine/Source/Editor/UnrealEd/Private/WorldPartition/WorldPartitionRenameDuplicateBuilder.cpp
--- a/Engine/Source/Editor/UnrealEd/Private/WorldPartition/WorldPartitionRenameDuplicateBuilder.cpp
+++ b/Engine/Source/Editor/UnrealEd/Private/WorldPartition/WorldPartitionRenameDuplicateBuilder.cpp
@@ -278,7 +278,7 @@ bool UWorldPartitionRenameDuplicateBuilder::RunInternal(UWorld* World, const FCe
 					}, false);
 
 					// Move dependant objects into the new actor package
-					for (UObject* DependantObject : DependantObjects)
+					for (UObject* DependantObject : DependantObjects) //-V1078
 					{
 						DependantObject->Rename(nullptr, Actor->GetExternalPackage(), REN_NonTransactional | REN_DontCreateRedirectors | REN_ForceNoResetLoaders | REN_DoNotDirty);
 					}
@@ -329,7 +329,7 @@ bool UWorldPartitionRenameDuplicateBuilder::RunInternal(UWorld* World, const FCe
 					}, false);
 
 					// Move back dependant objects into the previous actor package
-					for (UObject* DependantObject : DependantObjects)
+					for (UObject* DependantObject : DependantObjects) //-V1078
 					{
 						DependantObject->Rename(nullptr, Actor->GetExternalPackage(), REN_NonTransactional | REN_DontCreateRedirectors | REN_ForceNoResetLoaders | REN_DoNotDirty);
 					}
diff --git a/Engine/Source/Runtime/Core/Private/Experimental/Graph/GraphConvert.cpp b/Engine/Source/Runtime/Core/Private/Experimental/Graph/GraphConvert.cpp
--- a/Engine/Source/Runtime/Core/Private/Experimental/Graph/GraphConvert.cpp
+++ b/Engine/Source/Runtime/Core/Private/Experimental/Graph/GraphConvert.cpp
@@ -744,7 +744,7 @@ bool FCondensationGraphTest_Core::RunTest(const FString& Parameters)
 						}
 					}
 				}
-				for (FVertex ReachableVertex : Reachable)
+				for (FVertex ReachableVertex : Reachable) //-V1078
 				{
 					if (ReachableVertex < Root)
 					{
diff --git a/Engine/Source/Runtime/CoreUObject/Private/UObject/UObjectGlobals.cpp b/Engine/Source/Runtime/CoreUObject/Private/UObject/UObjectGlobals.cpp
--- a/Engine/Source/Runtime/CoreUObject/Private/UObject/UObjectGlobals.cpp
+++ b/Engine/Source/Runtime/CoreUObject/Private/UObject/UObjectGlobals.cpp
@@ -2415,7 +2415,7 @@ void EndLoad(FUObjectSerializeContext* LoadContext, TArray<UPackage*>* OutLoaded
 	if (OutLoadedPackages)
 	{
 		OutLoadedPackages->Reserve(LoadedPackages.Num());
-		for (UPackage* Package : LoadedPackages)
+		for (UPackage* Package : LoadedPackages) //-V1078
 		{
 			OutLoadedPackages->Add(Package);
 		}
diff --git a/Engine/Source/Runtime/Engine/Private/WorldPartition/RuntimeHashSet/WorldPartitionRuntimeHashSetDebugDraw.cpp b/Engine/Source/Runtime/Engine/Private/WorldPartition/RuntimeHashSet/WorldPartitionRuntimeHashSetDebugDraw.cpp
--- a/Engine/Source/Runtime/Engine/Private/WorldPartition/RuntimeHashSet/WorldPartitionRuntimeHashSetDebugDraw.cpp
+++ b/Engine/Source/Runtime/Engine/Private/WorldPartition/RuntimeHashSet/WorldPartitionRuntimeHashSetDebugDraw.cpp
@@ -132,7 +132,7 @@ bool UWorldPartitionRuntimeHashSet::Draw2D(FWorldPartitionDraw2DContext& DrawCon
 
 		if (FilteredCells.Num())
 		{
-			for (const UWorldPartitionRuntimeCell* Cell : FilteredCells)
+			for (const UWorldPartitionRuntimeCell* Cell : FilteredCells) //-V1078
 			{
 				const FVector2D CellBoundsSize = FVector2D(Cell->GetCellBounds().GetSize());
 				const FVector2D CellBoundsMin = FVector2D(Cell->GetCellBounds().Min);
diff --git a/Engine/Source/Runtime/Projects/Private/PluginManager.cpp b/Engine/Source/Runtime/Projects/Private/PluginManager.cpp
--- a/Engine/Source/Runtime/Projects/Private/PluginManager.cpp
+++ b/Engine/Source/Runtime/Projects/Private/PluginManager.cpp
@@ -1436,7 +1436,7 @@ bool FPluginManager::ConfigureEnabledPlugins()
 					for (const FString& PotentialPluginName : PluginsToConfigure)
 					{
 						bool bMatchesAnyWildcard = false;
-						for (const FString& WildcardPluginName : WildcardPluginsList)
+						for (const FString& WildcardPluginName : WildcardPluginsList) //-V1078
 						{
 							if (PotentialPluginName.MatchesWildcard(WildcardPluginName))
 							{
@@ -1588,7 +1588,7 @@ bool FPluginManager::ConfigureEnabledPlugins()
 				SCOPED_BOOT_TIMING("ConfigureTargetEnabledPlugins");
 
 				TArray<FString> TargetEnabledPlugins = { UBT_TARGET_ENABLED_PLUGINS };
-				for (const FString& TargetEnabledPlugin : TargetEnabledPlugins)
+				for (const FString& TargetEnabledPlugin : TargetEnabledPlugins) //-V1078
 				{
 					if (!ConfiguredPluginNames.Contains(TargetEnabledPlugin))
 					{
@@ -1606,7 +1606,7 @@ bool FPluginManager::ConfigureEnabledPlugins()
 				SCOPED_BOOT_TIMING("ConfigureTargetDisabledPlugins");
 
 				TArray<FString> TargetDisabledPlugins = { UBT_TARGET_DISABLED_PLUGINS };
-				for (const FString& TargetDisabledPlugin : TargetDisabledPlugins)
+				for (const FString& TargetDisabledPlugin : TargetDisabledPlugins) //-V1078
 				{
 					if (!ConfiguredPluginNames.Contains(TargetDisabledPlugin))
 					{
