[Mutable] Enable resize and mipmap memory optimizations for block compressed images.

- mutable.EnableImageResizeMemoryOptimizations, mutable.EnableVectorImplementationForSmallResizes and mutable.EnableCompressedMipGenerationMemoryOptimizations cvars are now set to true by default.

#rnx

diff --git a/Engine/Plugins/Experimental/Mutable/Source/MutableRuntime/Private/MuR/OpImageMipmap.cpp b/Engine/Plugins/Experimental/Mutable/Source/MutableRuntime/Private/MuR/OpImageMipmap.cpp
--- a/Engine/Plugins/Experimental/Mutable/Source/MutableRuntime/Private/MuR/OpImageMipmap.cpp
+++ b/Engine/Plugins/Experimental/Mutable/Source/MutableRuntime/Private/MuR/OpImageMipmap.cpp
@@ -8,7 +8,7 @@
 namespace
 {
 
-bool bEnableCompressedMipGenerationMemoryOptimizations = false;
+bool bEnableCompressedMipGenerationMemoryOptimizations = true;
 static FAutoConsoleVariableRef CVarEnableCompressedMipGenerationMemoryOptimizations (
 	TEXT("mutable.EnableCompressedMipGenerationMemoryOptimizations"),
 	bEnableCompressedMipGenerationMemoryOptimizations,
diff --git a/Engine/Plugins/Experimental/Mutable/Source/MutableRuntime/Private/MuR/OpImageResize.cpp b/Engine/Plugins/Experimental/Mutable/Source/MutableRuntime/Private/MuR/OpImageResize.cpp
--- a/Engine/Plugins/Experimental/Mutable/Source/MutableRuntime/Private/MuR/OpImageResize.cpp
+++ b/Engine/Plugins/Experimental/Mutable/Source/MutableRuntime/Private/MuR/OpImageResize.cpp
@@ -7,7 +7,7 @@
 namespace
 {
 
-bool bEnableImageResizeMemoryOptimizations = false;
+bool bEnableImageResizeMemoryOptimizations = true;
 static FAutoConsoleVariableRef CVarEnableImageResizeMemoryOptimizations (
 	TEXT("mutable.EnableImageResizeMemoryOptimizations"),
 	bEnableImageResizeMemoryOptimizations,
@@ -15,7 +15,7 @@ static FAutoConsoleVariableRef CVarEnableImageResizeMemoryOptimizations (
 	ECVF_Default);
 
 
-bool bEnableVectorImplementationForSmallResizes = false;
+bool bEnableVectorImplementationForSmallResizes = true;
 static FAutoConsoleVariableRef CVarEnableVectorImplementationForSmallResizes (
 	TEXT("mutable.EnableVectorImplementationForSmallResizes"),
 	bEnableVectorImplementationForSmallResizes,
