[Mutable] Fix possible unecessary recompilations due to Participating Objects.

#rnx

diff --git a/Engine/Plugins/Experimental/Mutable/Source/CustomizableObject/Classes/MuCO/CustomizableObject.h b/Engine/Plugins/Experimental/Mutable/Source/CustomizableObject/Classes/MuCO/CustomizableObject.h
--- a/Engine/Plugins/Experimental/Mutable/Source/CustomizableObject/Classes/MuCO/CustomizableObject.h
+++ b/Engine/Plugins/Experimental/Mutable/Source/CustomizableObject/Classes/MuCO/CustomizableObject.h
@@ -1392,7 +1392,7 @@ private:
 	TMap<FString, int32> ParameterPropertiesLookupTable;
 
 	// Rebuild ParameterProperties from the current compiled model.
-	void UpdateParameterPropertiesFromModel();
+	void UpdateParameterPropertiesFromModel(const TSharedPtr<mu::Model>& Model);
 
 	// This is a manual version number for the binary blobs in this asset.
 	// Increasing it invalidates all the previously compiled models.
@@ -1411,8 +1411,9 @@ public:
 
 #if WITH_EDITOR
 	
-	// Compile the object if Automatic Compilation is enabled and the object can be compiled.
-	// Automatic compilation can be enabled/disabled in the Mutable's Plugin Settings.
+	/** Compile the object if Automatic Compilation is enabled and has not been already compiled.
+	 * Automatic compilation can be enabled/disabled in the Mutable's Plugin Settings.
+	 * @return true if compiled */
 	bool ConditionalAutoCompile();
 	
 	// Add a profile that stores the values of the parameters used by the CustomInstance.
@@ -1448,8 +1449,6 @@ public:
 	bool IsCachedCookedPlatformDataLoaded(const ITargetPlatform* TargetPlatform) override;
 	// End UObject interface.
 
-	// Rebuild HashToStreamableBlocks and ParameterProperties from the current compiled model.
-	void UpdateCompiledDataFromModel();
 	
 	/** Generic Save/Load methods to write/read compiled data */
 	void SaveCompiledData(FArchive& Ar, bool bSkipEditorOnlyData = false);
@@ -1655,13 +1654,6 @@ private:
 	UPROPERTY()
 	TObjectPtr<UCustomizableObjectBulk> BulkData;
 
-#if WITH_EDITORONLY_DATA
-	/** List of external objects that if changed, a compilation is required.
-	 * The FGuid is the the UPackage::Guid, which is regenerated each time the packages is saved. */
-	UPROPERTY()
-	TMap<TObjectPtr<const UObject>, FGuid> ParticipatingObjects;
-#endif
-
 	UPROPERTY()
 	TObjectPtr<UCustomizableObjectPrivate> Private;
 };
diff --git a/Engine/Plugins/Experimental/Mutable/Source/CustomizableObject/Internal/MuCO/CustomizableObjectPrivate.h b/Engine/Plugins/Experimental/Mutable/Source/CustomizableObject/Internal/MuCO/CustomizableObjectPrivate.h
--- a/Engine/Plugins/Experimental/Mutable/Source/CustomizableObject/Internal/MuCO/CustomizableObjectPrivate.h
+++ b/Engine/Plugins/Experimental/Mutable/Source/CustomizableObject/Internal/MuCO/CustomizableObjectPrivate.h
@@ -2,6 +2,7 @@
 
 #pragma once
 
+#include "MuCO/StateMachine.h"
 #include "Templates/SharedPointer.h"
 #include "UObject/WeakObjectPtr.h"
 #include "MuR/Types.h"
@@ -28,6 +29,33 @@ private:
 	TMap<TArray<mu::FResourceID>, TWeakObjectPtr<USkeletalMesh>> GeneratedMeshes;
 };
 
+struct FCustomizableObjectStatusTypes
+{
+	enum class EState : uint8
+	{
+		Loading = 0, // Waiting for PostLoad and Asset Registry to finish.
+		ModelLoaded, // Model loaded correctly.
+		NoModel, // No model (due to no model not found and automatic compilations disabled).
+		// Compiling, // Compiling the CO.
+
+		Count,
+	};
+	
+	static constexpr EState StartState = EState::Loading;
+
+	static constexpr bool ValidTransitions[3][3] =
+	{
+		// TO
+		// Loading, ModelLoaded, NoModel // FROM
+		{false,   true,        true},  // Loading
+		{false,   true,        true},  // ModelLoaded
+		{false,   true,        true},  // NoModel
+	};
+};
+
+
+using FCustomizableObjectStatus = FStateMachine<FCustomizableObjectStatusTypes>;
+
 
 UCLASS()
 class UCustomizableObjectPrivate : public UObject
@@ -54,9 +82,17 @@ public:
 
 	bool bModelCompiledForCook = false;
 	TArray<FString> CachedPlatformNames;
+
+	/** List of external packages that if changed, a compilation is required.
+	 * Key is the package name. Value is the the UPackage::Guid, which is regenerated each time the packages is saved.
+	 *
+	 * Updated each time the CO is compiled and saved in the Derived Data. */
+	TMap<FName, FGuid> ParticipatingObjects;
 #endif
 
 	/** Cache of generated SkeletalMeshes */
 	FMeshCache MeshCache;
+
+	FCustomizableObjectStatus Status;
 };
 
diff --git a/Engine/Plugins/Experimental/Mutable/Source/CustomizableObject/Private/MuCO/CustomizableObject.cpp b/Engine/Plugins/Experimental/Mutable/Source/CustomizableObject/Private/MuCO/CustomizableObject.cpp
--- a/Engine/Plugins/Experimental/Mutable/Source/CustomizableObject/Private/MuCO/CustomizableObject.cpp
+++ b/Engine/Plugins/Experimental/Mutable/Source/CustomizableObject/Private/MuCO/CustomizableObject.cpp
@@ -232,27 +232,19 @@ void UCustomizableObject::PostLoad()
 		}
 	}
 
-
-#endif
-
-#if WITH_EDITORONLY_DATA
-	for (TTuple<TObjectPtr<const UObject>, FGuid>& ParticipatingObject : ParticipatingObjects)
+	if (!IsRunningCookCommandlet())
 	{
-		if (!ParticipatingObject.Key) // Object no longer exists.
+		ITargetPlatformManagerModule& TargetPlatformManager = GetTargetPlatformManagerRef();
+		const ITargetPlatform* RunningPlatform = TargetPlatformManager.GetRunningTargetPlatform();
+
+		const FAssetRegistryModule& AssetRegistryModule = FModuleManager::LoadModuleChecked<FAssetRegistryModule>("AssetRegistry");
+		if (AssetRegistryModule.Get().IsLoadingAssets())
 		{
-			continue;
+			AssetRegistryModule.Get().OnFilesLoaded().AddUObject(this, &UCustomizableObject::LoadCompiledDataFromDisk);
 		}
-		
-		PRAGMA_DISABLE_DEPRECATION_WARNINGS
-		const FGuid PackageGuid = ParticipatingObject.Key.GetPackage()->GetGuid();
-		PRAGMA_ENABLE_DEPRECATION_WARNINGS
-		
-		if (PackageGuid != ParticipatingObject.Value)
+		else
 		{
-			SetModel(nullptr);
-			
-			UE_LOG(LogMutable, Display, TEXT("Forcing recompilation due to changes in %s."), *ParticipatingObject.Key->GetFullName());
-			break;
+			LoadCompiledDataFromDisk();
 		}
 	}
 #endif
@@ -290,11 +282,6 @@ void UCustomizableObject::Serialize(FArchive& Ar_Asset)
 		// Can't remove this or saved customizable objects will fail to load
 		int64 InternalVersion = CurrentSupportedVersion;
 		Ar_Asset << InternalVersion;
-		
-		if (Ar_Asset.IsLoading())
-		{
-			LoadCompiledDataFromDisk();
-		}
 	}
 #else
 	if (Ar_Asset.IsLoading())
@@ -368,33 +355,6 @@ void UCustomizableObject::ClearCompiledData()
 }
 
 
-void UCustomizableObject::UpdateCompiledDataFromModel()
-{
-	TSharedPtr<mu::Model, ESPMode::ThreadSafe> Model = Private->GetModel();
-
-	// Generate a map that using the resource id tells the offset and size of the resource inside the bulk data
-	if (Model)
-	{
-		uint64 Offset = 0;
-
-		const int32 NumStreamingFiles = Model->GetRomCount();
-		HashToStreamableBlock.Empty(NumStreamingFiles);
-
-		for (size_t FileIndex = 0; FileIndex < NumStreamingFiles; ++FileIndex)
-		{
-			const uint32 ResourceId = Model->GetRomId(FileIndex);
-			const uint32 ResourceSize = Model->GetRomSize(FileIndex);
-
-			HashToStreamableBlock.Add(ResourceId, FMutableStreamableBlock({0, Offset, ResourceSize }));
-			Offset += ResourceSize;
-		}
-	}
-
-	// Generate ParameterProperties and IntParameterLookUpTable
-	UpdateParameterPropertiesFromModel();
-}
-
-
 void SerializeStreamedResources(FArchive& Ar, UObject* Object, TArray<FCustomizableObjectStreamedResourceData>& StreamedResources, bool bIsCooking)
 {
 	if (Ar.IsSaving())
@@ -576,6 +536,7 @@ void UCustomizableObject::SaveCompiledData(FArchive& MemoryWriter, bool bIsCooki
 	{
 		MemoryWriter << CustomizableObjectPathMap;
 		MemoryWriter << GroupNodeMap;
+		MemoryWriter << GetPrivate()->ParticipatingObjects;
 	}
 
 	MemoryWriter << LODSettings.NumLODsInRoot;
@@ -680,11 +641,42 @@ void UCustomizableObject::LoadCompiledData(FArchive& MemoryReader, const ITarget
 
 		MemoryReader << HashToStreamableBlock;
 
+		bool bForceRecompilation = false;
+
 		// All Editor Only data must be loaded here
 		if (!bIsCooking)
 		{
 			MemoryReader << CustomizableObjectPathMap;
 			MemoryReader << GroupNodeMap;
+
+			TMap<FName, FGuid>& ParticipatingObjects = GetPrivate()->ParticipatingObjects;
+			MemoryReader << ParticipatingObjects;
+
+			const FAssetRegistryModule& AssetRegistryModule = FModuleManager::LoadModuleChecked<FAssetRegistryModule>("AssetRegistry");
+			
+			for (TTuple<FName, FGuid>& ParticipatingObject : ParticipatingObjects)
+			{				
+				FAssetPackageData AssetPackageData;
+				const UE::AssetRegistry::EExists Result = AssetRegistryModule.Get().TryGetAssetPackageData(ParticipatingObject.Key, AssetPackageData);
+				if (Result == UE::AssetRegistry::EExists::Exists)
+				{
+					PRAGMA_DISABLE_DEPRECATION_WARNINGS
+					const FGuid PackageGuid = AssetPackageData.PackageGuid;
+					PRAGMA_ENABLE_DEPRECATION_WARNINGS
+		
+					bForceRecompilation = PackageGuid != ParticipatingObject.Value;
+				}
+				else
+				{
+					bForceRecompilation = true;
+				}
+		
+				if (bForceRecompilation)
+                {
+                	UE_LOG(LogMutable, Display, TEXT("Forcing recompilation due to changes in %s."), *ParticipatingObject.Key.ToString());
+					break;
+                }
+			}
 		}
 
 		MemoryReader << LODSettings.NumLODsInRoot;
@@ -698,7 +690,7 @@ void UCustomizableObject::LoadCompiledData(FArchive& MemoryReader, const ITarget
 		bool bModelSerialized = false;
 		MemoryReader << bModelSerialized;
 
-		if (bModelSerialized)
+		if (bModelSerialized && !bForceRecompilation)
 		{
 			UnrealMutableInputStream stream(MemoryReader);
 			mu::InputArchive arch(&stream);
@@ -707,10 +699,11 @@ void UCustomizableObject::LoadCompiledData(FArchive& MemoryReader, const ITarget
 			Private->SetModel(Model, Identifier);
 		}
 	}
-
-	UpdateParameterPropertiesFromModel();
+	
+	UpdateParameterPropertiesFromModel(GetModel());
 }
 
+
 void UCustomizableObject::LoadCompiledDataFromDisk()
 {
 	ITargetPlatformManagerModule& TargetPlatformManager = GetTargetPlatformManagerRef();
@@ -772,6 +765,14 @@ void UCustomizableObject::LoadCompiledDataFromDisk()
 		delete CompiledDataFileHandle;
 		delete StreamableDataFileHandle;
 	}
+
+	UCustomizableObjectPrivate* CustomizableObjectPrivate = GetPrivate();
+
+	if (!CustomizableObjectPrivate->GetModel()) // Not failed to load the model
+	{
+		CustomizableObjectPrivate->Status.NextState(FCustomizableObjectStatusTypes::EState::NoModel);
+		ConditionalAutoCompile(); // Recompile if no model has been loaded
+	}
 }
 
 
@@ -852,6 +853,12 @@ bool UCustomizableObject::ConditionalAutoCompile()
 		return true;
 	}
 
+	// Model has not loaded yet
+	if (GetPrivate()->Status.Get() == FCustomizableObjectStatusTypes::EState::Loading)
+	{
+		return false;
+	}
+
 	UCustomizableObjectSystem* System = UCustomizableObjectSystem::GetInstance();
 	if (!System || !System->IsValidLowLevel() || System->HasAnyFlags(RF_BeginDestroyed))
 	{
@@ -890,7 +897,7 @@ bool UCustomizableObject::ConditionalAutoCompile()
 	}
 	else
 	{
-		const FAssetRegistryModule& AssetRegistryModule = FModuleManager::LoadModuleChecked<FAssetRegistryModule>("AssetRegistry");
+		const FAssetRegistryModule& AssetRegistryModule = FModuleManager::GetModuleChecked<FAssetRegistryModule>(TEXT("AssetRegistry"));
 		const FAssetData AssetData = AssetRegistryModule.Get().GetAssetByObjectPath(UE_MUTABLE_OBJECTPATH(GetPathName()));
 		System->RecompileCustomizableObjectAsync(AssetData, this);
 	}
@@ -1044,10 +1051,11 @@ void UCustomizableObject::LoadEmbeddedData(FArchive& Ar)
 		UnrealMutableInputStream stream(Ar);
 		mu::InputArchive arch(&stream);
 		TSharedPtr<mu::Model, ESPMode::ThreadSafe> Model = mu::Model::StaticUnserialise( arch );
-		Private->SetModel( Model, FGuid());
 
 		// Create parameter properties
-		UpdateParameterPropertiesFromModel();
+		UpdateParameterPropertiesFromModel(Model);
+
+		Private->SetModel( Model, FGuid());
 	}
 }
 
@@ -1292,9 +1300,28 @@ TSharedPtr<mu::Model, ESPMode::ThreadSafe> UCustomizableObject::GetModel() const
 #if WITH_EDITOR
 void UCustomizableObject::SetModel(TSharedPtr<mu::Model, ESPMode::ThreadSafe> Model)
 {
+	// Generate a map that using the resource id tells the offset and size of the resource inside the bulk data
+	if (Model)
+	{
+		uint64 Offset = 0;
+
+		const int32 NumStreamingFiles = Model->GetRomCount();
+		HashToStreamableBlock.Empty(NumStreamingFiles);
+
+		for (size_t FileIndex = 0; FileIndex < NumStreamingFiles; ++FileIndex)
+		{
+			const uint32 ResourceId = Model->GetRomId(FileIndex);
+			const uint32 ResourceSize = Model->GetRomSize(FileIndex);
+
+			HashToStreamableBlock.Add(ResourceId, FMutableStreamableBlock({0, Offset, ResourceSize }));
+			Offset += ResourceSize;
+		}
+	}
+
+	// Generate ParameterProperties and IntParameterLookUpTable
+	UpdateParameterPropertiesFromModel(Model);
+
 	Private->SetModel(Model, GenerateIdentifier(*this));
-	
-	UpdateCompiledDataFromModel();
 }
 
 void UCustomizableObject::SetBoneNamesArray(const TArray<FName>& InBoneNames)
@@ -1379,11 +1406,11 @@ const FString & UCustomizableObject::GetParameterName(int32 ParamIndex) const
 }
 
 
-void UCustomizableObject::UpdateParameterPropertiesFromModel()
+void UCustomizableObject::UpdateParameterPropertiesFromModel(const TSharedPtr<mu::Model>& Model)
 {
-	if (Private->GetModel())
+	if (Model)
 	{
-		mu::ParametersPtr MutableParameters = mu::Model::NewParameters(Private->GetModel());
+		mu::ParametersPtr MutableParameters = mu::Model::NewParameters(Model);
 		int paramCount = MutableParameters->GetCount();
 
 		ParameterProperties.Reset(paramCount);
@@ -1905,18 +1932,28 @@ void FMeshCache::Add(const TArray<mu::FResourceID>& Key, USkeletalMesh* Value)
 
 void UCustomizableObjectPrivate::SetModel(const TSharedPtr<mu::Model, ESPMode::ThreadSafe>& Model, const FGuid Id)
 {
+	if (MutableModel == Model
+#if WITH_EDITOR
+		&& Identifier == Id
+#endif
+		)
+	{
+		return;
+	}
+	
 #if WITH_EDITOR
 	if (MutableModel)
 	{
 		MutableModel->Invalidate();
 	}
+
+	Identifier = Id;
 #endif
 	
 	MutableModel = Model;
 
-#if WITH_EDITOR
-	Identifier = Id;
-#endif
+	using EState = FCustomizableObjectStatus::EState;
+	Status.NextState(Model ? EState::ModelLoaded : EState::NoModel);
 }
 
 
@@ -1932,14 +1969,6 @@ TSharedPtr<const mu::Model, ESPMode::ThreadSafe> UCustomizableObjectPrivate::Get
 }
 
 
-#if WITH_EDITORONLY_DATA
-TMap<TObjectPtr<const UObject>, FGuid>& UCustomizableObjectPrivate::GetParticipatingObjects(UCustomizableObject& Public)
-{
-	return Public.ParticipatingObjects;
-}
-#endif
-
-
 //-------------------------------------------------------------------------------------------------
 //-------------------------------------------------------------------------------------------------
 //-------------------------------------------------------------------------------------------------
diff --git a/Engine/Plugins/Experimental/Mutable/Source/CustomizableObject/Public/MuCO/StateMachine.h b/Engine/Plugins/Experimental/Mutable/Source/CustomizableObject/Public/MuCO/StateMachine.h
new file mode 100644
--- /dev/null
+++ b/Engine/Plugins/Experimental/Mutable/Source/CustomizableObject/Public/MuCO/StateMachine.h
@@ -0,0 +1,48 @@
+// Copyright Epic Games, Inc. All Rights Reserved.
+
+#pragma once
+
+#include "HAL/Platform.h"
+#include "Delegates/Delegate.h" 
+
+
+/** Simple state machine with transitions check.
+ *
+ * Type T must implement the following:
+ * - EState: Enum of valid states.
+ * - StartState: Variable defining the initial state.
+ * - ValidTransitions: 2D bool array of valid transitions. */
+template<typename T>
+class FStateMachine
+{
+public:
+	using EState = typename T::EState;
+
+	static_assert(static_cast<int32>(EState::Count) * static_cast<int32>(EState::Count) == sizeof(T::ValidTransitions), "Size of EState and ValidTransitions do not match.");
+	
+	void NextState(EState State)
+	{
+		check(T::ValidTransitions[static_cast<int32>(CurrentState)][static_cast<int32>(State)])
+		
+		OnStateChangedDelegate.Broadcast(CurrentState, State);
+		CurrentState = State;
+	}
+
+	EState Get() const
+	{
+		return CurrentState;
+	}
+
+	DECLARE_MULTICAST_DELEGATE_TwoParams(FOnStateChangedDelegate, EState, EState)
+
+	FOnStateChangedDelegate& GetOnStateChangedDelegate()
+	{
+		return OnStateChangedDelegate;
+	}
+	
+private:
+	FOnStateChangedDelegate OnStateChangedDelegate;
+
+	EState CurrentState = T::StartState;
+};
+
diff --git a/Engine/Plugins/Experimental/Mutable/Source/CustomizableObjectEditor/Private/MuCOE/CustomizableObjectCompiler.cpp b/Engine/Plugins/Experimental/Mutable/Source/CustomizableObjectEditor/Private/MuCOE/CustomizableObjectCompiler.cpp
--- a/Engine/Plugins/Experimental/Mutable/Source/CustomizableObjectEditor/Private/MuCOE/CustomizableObjectCompiler.cpp
+++ b/Engine/Plugins/Experimental/Mutable/Source/CustomizableObjectEditor/Private/MuCOE/CustomizableObjectCompiler.cpp
@@ -713,6 +713,12 @@ void FCustomizableObjectCompiler::AddCachedReferencers(const FName& PathName, TA
 			ArrayAssetData.Add(ArrayAssetDataTemp[i]);
 		}
 	}
+
+	// Required to make compilations deterministic within editor runs.
+	ArrayReferenceNames.Sort([](const FName& A, const FName& B)
+	{
+		return A.LexicalLess(B);
+	});
 }
 
 
@@ -1150,23 +1156,10 @@ void FCustomizableObjectCompiler::CompileInternal(UCustomizableObject* Object, c
 
 		if (!ParamNamesToSelectedOptions.Num())
 		{
-			GenerationContext.ParticipatingObjects.Remove(Object); // Remove self CO reference.
-
-			checkCode
-			(
-				for (auto Pair : GenerationContext.ParticipatingObjects)
-				{
-					check(Pair.Key.GetPackage() != Object->GetPackage()); // Adding a reference to itself will always force a recompilation after the CO is saved, so avoid doing that.
-				}
-			);
+			const FName PackageName = Object->GetPackage()->GetFName();
+			GenerationContext.ParticipatingObjects.Remove(PackageName); // Remove self CO reference.
 			
-			TMap<TObjectPtr<const UObject>, FGuid>& ParticipatingObjects = Object->GetPrivate()->GetParticipatingObjects(*Object);
-			if (ParticipatingObjects != GenerationContext.ParticipatingObjects)
-			{
-				Object->MarkPackageDirty();
-			}
-
-			ParticipatingObjects = MoveTemp(GenerationContext.ParticipatingObjects);
+			Object->GetPrivate()->ParticipatingObjects = MoveTemp(GenerationContext.ParticipatingObjects);
 		}
 
 		if (CompileTask.IsValid()) // Don't start compilation if there's a compilation running
diff --git a/Engine/Plugins/Experimental/Mutable/Source/CustomizableObjectEditor/Private/MuCOE/CustomizableObjectEditor.cpp b/Engine/Plugins/Experimental/Mutable/Source/CustomizableObjectEditor/Private/MuCOE/CustomizableObjectEditor.cpp
--- a/Engine/Plugins/Experimental/Mutable/Source/CustomizableObjectEditor/Private/MuCOE/CustomizableObjectEditor.cpp
+++ b/Engine/Plugins/Experimental/Mutable/Source/CustomizableObjectEditor/Private/MuCOE/CustomizableObjectEditor.cpp
@@ -53,6 +53,7 @@
 #include "ScopedTransaction.h"
 #include "MuCOE/SCustomizableInstanceProperties.h"
 #include "MuCO/CustomizableInstancePrivateData.h"
+#include "MuCO/CustomizableObjectPrivate.h"
 #include "UObject/EnumProperty.h"
 #include "Widgets/Docking/SDockTab.h"
 #include "Widgets/Input/SSearchBox.h"
@@ -187,7 +188,9 @@ FCustomizableObjectEditor::~FCustomizableObjectEditor()
 		Compiler.ForceFinishBeforeStartCompilation(CustomizableObject);
 	}
 
-	FCoreUObjectDelegates::OnObjectModified.Remove(OnObjectModifiedHandle);
+	FCoreUObjectDelegates::OnObjectModified.RemoveAll(this);
+
+	CustomizableObject->GetPrivate()->Status.GetOnStateChangedDelegate().RemoveAll(this);
 
 	FAssetRegistryModule& AssetRegistryModule = FModuleManager::LoadModuleChecked<FAssetRegistryModule>(TEXT("AssetRegistry"));
 	AssetRegistryModule.Get().OnFilesLoaded().RemoveAll(this);
@@ -205,16 +208,6 @@ void FCustomizableObjectEditor::InitCustomizableObjectEditor(const EToolkitMode:
 	CustomSettings = NewObject<UCustomSettings>();
 	CustomSettings->SetEditor(SharedThis(this));
 
-	FAssetRegistryModule& AssetRegistryModule = FModuleManager::LoadModuleChecked<FAssetRegistryModule>(TEXT("AssetRegistry"));
-	if (AssetRegistryModule.Get().IsLoadingAssets())
-	{
-		AssetRegistryModule.Get().OnFilesLoaded().AddRaw(this, &FCustomizableObjectEditor::OnAssetRegistryLoadComplete);
-	}
-	else
-	{
-		AssetRegistryLoaded = true;
-	}
-
 	HelperCallback = nullptr;
 
 	// Support undo/redo
@@ -249,7 +242,6 @@ void FCustomizableObjectEditor::InitCustomizableObjectEditor(const EToolkitMode:
 		.CustomizableObjectEditor(SharedThis(this));
 
 	Viewport->SetCustomizableObject(CustomizableObject);
-	Viewport->SetAssetRegistryLoaded(AssetRegistryLoaded);
 	ViewportClient = Viewport->GetViewportClient();
 
 	// \TODO: Create only when needed?
@@ -339,17 +331,13 @@ void FCustomizableObjectEditor::InitCustomizableObjectEditor(const EToolkitMode:
 	// Clears selection highlight.
 	OnObjectPropertySelectionChanged(NULL);
 	OnInstancePropertySelectionChanged(NULL);
-	OnObjectModifiedHandle = FCoreUObjectDelegates::OnObjectModified.AddRaw(this, &FCustomizableObjectEditor::OnObjectModified);
+	FCoreUObjectDelegates::OnObjectModified.AddRaw(this, &FCustomizableObjectEditor::OnObjectModified);
 
-	// Compile for the first time if necessary
-	if (!CustomizableObject->IsCompiled())
-	{
-		CompileObject();
-	}
-	else
-	{
-		CreatePreviewInstance();
-	}
+	UCustomizableObjectPrivate* CustomizableObjectPrivate = CustomizableObject->GetPrivate();
+
+	CustomizableObjectPrivate->Status.GetOnStateChangedDelegate().AddRaw(this, &FCustomizableObjectEditor::OnCustomizableObjectStatusChanged);
+	const FCustomizableObjectStatusTypes::EState CurrentStatus = CustomizableObjectPrivate->Status.Get();
+	OnCustomizableObjectStatusChanged(CurrentStatus, CurrentStatus);
 }
 
 
@@ -420,8 +408,6 @@ void FCustomizableObjectEditor::CreatePreviewInstance()
 	HelperCallback = NewObject<UUpdateClassWrapper>(GetTransientPackage());
 	PreviewInstance->UpdatedDelegate.AddDynamic(HelperCallback, &UUpdateClassWrapper::DelegatedCallback);
 
-	PreviewStaticMeshComponent = nullptr;
-
 	if (CustomizableObject->ReferenceSkeletalMeshes.Num())
 	{
 		CreatePreviewComponents();
@@ -430,24 +416,8 @@ void FCustomizableObjectEditor::CreatePreviewInstance()
 		{
 			HelperCallback->Delegate.BindSP(this, &FCustomizableObjectEditor::OnUpdatePreviewInstance);
 
-			if (AssetRegistryLoaded)
-			{
-				// Asset loading works in the Main Thread, there's a risk the sync loading could put to sleep the thread while
-				// waiting for the asset registry to load the content (that never gets executed)
-				Viewport->SetAssetRegistryLoaded(true);
-				PreviewInstance->SetBuildParameterRelevancy(true);
-				PreviewInstance->UpdateSkeletalMeshAsync(true, true);
-			}
-			else
-			{
-				FNotificationInfo Info(LOCTEXT("CustomizableObjectCompileTryLater", "Please wait until asset registry loads all assets"));
-				Info.bFireAndForget = true;
-				Info.bUseThrobber = true;
-				Info.FadeOutDuration = 1.0f;
-				Info.ExpireDuration = 2.0f;
-				FSlateNotificationManager::Get().AddNotification(Info);
-				UpdateSkeletalMeshAfterAssetLoaded = true;
-			}
+			PreviewInstance->SetBuildParameterRelevancy(true);
+			PreviewInstance->UpdateSkeletalMeshAsync(true, true);
 		}
 		else
 		{
@@ -476,7 +446,6 @@ void FCustomizableObjectEditor::AddReferencedObjects( FReferenceCollector& Colle
 	Collector.AddReferencedObject( CustomizableObject );
 	Collector.AddReferencedObject( PreviewInstance );
 	Collector.AddReferencedObjects( PreviewCustomizableSkeletalComponents );
-	Collector.AddReferencedObject( PreviewStaticMeshComponent );
 	Collector.AddReferencedObjects( PreviewSkeletalMeshComponents );
 	Collector.AddReferencedObject( HelperCallback );
 	Collector.AddReferencedObject( ProjectorParameter );
@@ -1474,9 +1443,9 @@ void FCustomizableObjectEditor::CompileObject()
 	UE_LOG(LogMutable, Verbose, TEXT("PROFILE: -----------------------------------------------------------"));
 	UE_LOG(LogMutable, Verbose, TEXT("PROFILE: [ %16.8f ] FCustomizableObjectEditor::CompileObject start."), FPlatformTime::Seconds());
 
-	if (!AssetRegistryLoaded)
+	if (CustomizableObject->GetPrivate()->Status.Get() == FCustomizableObjectStatus::EState::Loading)
 	{
-		FNotificationInfo Info(LOCTEXT("CustomizableObjectCompileTryLater", "Please wait until asset registry loads all assets"));
+		FNotificationInfo Info(LOCTEXT("CustomizableObjectCompileTryLater", "Please wait until Customizable Object is loaded"));
 		Info.bFireAndForget = true;
 		Info.bUseThrobber = true;
 		Info.FadeOutDuration = 1.0f;
@@ -1986,26 +1955,6 @@ UEdGraphNode* FCustomizableObjectEditor::CreateCommentBox(const FVector2D& InTar
 }
 
 
-void FCustomizableObjectEditor::OnAssetRegistryLoadComplete()
-{
-	AssetRegistryLoaded = true;
-
-	Viewport->SetAssetRegistryLoaded(true);
-
-	if (UpdateSkeletalMeshAfterAssetLoaded)
-	{
-		UpdateSkeletalMeshAfterAssetLoaded = false;
-
-		PreviewInstance->UpdateSkeletalMeshAsync(true, true);		
-	}
-
-	if (!CustomizableObject->IsCompiled())
-	{
-		CompileObject();
-	}
-}
-
-
 TSharedPtr<SCustomizableObjectEditorAdvancedPreviewSettings> FCustomizableObjectEditor::GetCustomizableObjectEditorAdvancedPreviewSettings()
 {
 	return CustomizableObjectEditorAdvancedPreviewSettings;
@@ -2549,12 +2498,6 @@ void FCustomizableObjectEditor::UpdateGraphNodeProperties()
 }
 
 
-bool FCustomizableObjectEditor::GetAssetRegistryLoaded()
-{
-	return AssetRegistryLoaded;
-}
-
-
 void FCustomizableObjectEditor::OpenTextureAnalyzerTab()
 {
 	TabManager->TryInvokeTab(TextureAnalyzerTabId);
@@ -2782,4 +2725,23 @@ void FCustomizableObjectEditor::CreatePreviewComponents()
 }
 
 
+void FCustomizableObjectEditor::OnCustomizableObjectStatusChanged(FCustomizableObjectStatus::EState, const FCustomizableObjectStatus::EState NextState)
+{
+	switch (NextState)
+	{
+	case FCustomizableObjectStatus::EState::ModelLoaded: 
+		CreatePreviewInstance();
+		break;
+		
+	case FCustomizableObjectStatus::EState::NoModel:
+		CustomizableObject->ConditionalAutoCompile();
+		break;
+
+	default:
+		break;
+	}
+}
+
+
+
 #undef LOCTEXT_NAMESPACE
diff --git a/Engine/Plugins/Experimental/Mutable/Source/CustomizableObjectEditor/Private/MuCOE/CustomizableObjectEditor.h b/Engine/Plugins/Experimental/Mutable/Source/CustomizableObjectEditor/Private/MuCOE/CustomizableObjectEditor.h
--- a/Engine/Plugins/Experimental/Mutable/Source/CustomizableObjectEditor/Private/MuCOE/CustomizableObjectEditor.h
+++ b/Engine/Plugins/Experimental/Mutable/Source/CustomizableObjectEditor/Private/MuCOE/CustomizableObjectEditor.h
@@ -9,6 +9,7 @@
 #include "MuCOE/CustomizableObjectCompiler.h"
 #include "MuCOE/ICustomizableObjectEditor.h"
 #include "TickableEditorObject.h"
+#include "MuCO/CustomizableObjectPrivate.h"
 #include "Widgets/Input/SNumericDropDown.h"
 
 #include "CustomizableObjectEditor.generated.h"
@@ -247,9 +248,6 @@ public:
 
 	virtual void UpdateGraphNodeProperties();
 
-	/** Getter of AssetRegistryLoaded */
-	bool GetAssetRegistryLoaded();
-
 	/** Callback to notify the editor when the PreviewInstance has been updated */
 	void OnUpdatePreviewInstance();
 
@@ -336,9 +334,6 @@ private:
 
 	virtual UEdGraphNode* CreateCommentBox(const FVector2D& NodePos) override;
 
-	/** Callback for the asset registry initial load */
-	void OnAssetRegistryLoadComplete();
-
 	/** Updates the visibility of PreviewSkeletalMeshComponent */
 	void UpdatePreviewVisibility();
 
@@ -363,6 +358,8 @@ private:
 	void CreatePreviewComponents();
 
 public:
+	void OnCustomizableObjectStatusChanged(FCustomizableObjectStatus::EState PreviousState, FCustomizableObjectStatus::EState NextState);
+
 	// Helpers to get the absolute parent of a Customizable Object
 	static UCustomizableObject* GetAbsoluteCOParent(const UCustomizableObjectNodeObject* const Root);
 	static void AddCachedReferencers(const FName& PathName, TArray<FName>& ArrayReferenceNames, TArray<FAssetData>& ArrayAssetData);
@@ -387,7 +384,6 @@ private:
 	TObjectPtr<UCustomizableObject> CustomizableObject;
 	TObjectPtr<UCustomizableObjectInstance> PreviewInstance = nullptr;
 	TArray<TObjectPtr<UCustomizableSkeletalComponent>> PreviewCustomizableSkeletalComponents;
-	TObjectPtr<UStaticMeshComponent> PreviewStaticMeshComponent = nullptr;
 	TArray<TObjectPtr<UDebugSkelMeshComponent>> PreviewSkeletalMeshComponents;
 
 	/** Object compiler */
@@ -416,15 +412,6 @@ private:
 
 	/** Widget to select which node pins are visible. */
 	TSharedPtr<class SCustomizableObjectNodePinViewer> NodePinViewer;
-
-	/** Handle for the OnObjectModified event */
-	FDelegateHandle OnObjectModifiedHandle;
-
-	/** Flag to know when the asset registry initial loading has completed */
-	bool AssetRegistryLoaded = false;
-
-	/** Flag to know whether the lasts steps when making a new preview instance need to be done when asset registry completes asset loading */
-	bool UpdateSkeletalMeshAfterAssetLoaded = false;
 	
 	/** UObject class to be able to use the update callback */
 	TObjectPtr<UUpdateClassWrapper> HelperCallback;
@@ -438,7 +425,7 @@ private:
 	/** Texture Analyzer table widget which shows the information of the transient textures used in the customizable object instance */
 	TSharedPtr<class SCustomizableObjecEditorTextureAnalyzer> TextureAnalyzer;
 
-	/** Performance report widget to test and analyze the current cusomizable object resource demands */
+	/** Performance report widget to test and analyze the current customizable object resource demands */
 	TSharedPtr<class SCustomizableObjecEditorPerformanceReport> PerformanceReport;
 
 	/** Widget to explore all the tags related with the Customizable Object open in the editor */
diff --git a/Engine/Plugins/Experimental/Mutable/Source/CustomizableObjectEditor/Private/MuCOE/CustomizableObjectEditorViewportClient.cpp b/Engine/Plugins/Experimental/Mutable/Source/CustomizableObjectEditor/Private/MuCOE/CustomizableObjectEditorViewportClient.cpp
--- a/Engine/Plugins/Experimental/Mutable/Source/CustomizableObjectEditor/Private/MuCOE/CustomizableObjectEditorViewportClient.cpp
+++ b/Engine/Plugins/Experimental/Mutable/Source/CustomizableObjectEditor/Private/MuCOE/CustomizableObjectEditorViewportClient.cpp
@@ -39,6 +39,7 @@
 #include "ScopedTransaction.h"
 #include "SkeletalDebugRendering.h"
 #include "UnrealWidget.h"
+#include "MuCO/CustomizableObjectPrivate.h"
 #include "Widgets/Input/SButton.h"
 #include "Widgets/Input/SEditableTextBox.h"
 #include "Widgets/Layout/SUniformGridPanel.h"
@@ -58,7 +59,6 @@ FCustomizableObjectEditorViewportClient::FCustomizableObjectEditorViewportClient
 	, CustomizableObjectEditorPtr(InCustomizableObjectEditor)
 	, CustomizableObject(nullptr)
 	, BakingOverwritePermission(false)
-	, AssetRegistryLoaded(false)
 {
 	// load config
 	ConfigOption = UPersonaOptions::StaticClass()->GetDefaultObject<UPersonaOptions>();
@@ -1632,9 +1632,9 @@ void FCustomizableObjectEditorViewportClient::BakeInstance()
 //-------------------------------------------------------------------------------------------------
 void FCustomizableObjectEditorViewportClient::BakeInstance(UCustomizableObjectInstance* InInstance)
 {
-	if (!AssetRegistryLoaded)
+	if (CustomizableObject->GetPrivate()->Status.Get() == FCustomizableObjectStatus::EState::Loading)
 	{
-		FNotificationInfo Info(NSLOCTEXT("CustomizableObjectEditor", "CustomizableObjectCompileTryLater", "Please wait until asset registry loads all assets"));
+		FNotificationInfo Info(NSLOCTEXT("CustomizableObjectEditor", "CustomizableObjectCompileTryLater", "Please wait unitl Customizable Object is loaded"));
 		Info.bFireAndForget = true;
 		Info.bUseThrobber = true;
 		Info.FadeOutDuration = 1.0f;
@@ -2366,12 +2366,6 @@ void FCustomizableObjectEditorViewportClient::AddReferencedObjects(FReferenceCol
 }
 
 
-void FCustomizableObjectEditorViewportClient::SetAssetRegistryLoaded(bool Value)
-{
-	AssetRegistryLoaded = Value;
-}
-
-
 void FCustomizableObjectEditorViewportClient::SetAdvancedShowFlagsForScene(const bool bAdvancedShowFlags)
 {
 	if (bAdvancedShowFlags)
diff --git a/Engine/Plugins/Experimental/Mutable/Source/CustomizableObjectEditor/Private/MuCOE/CustomizableObjectInstanceEditor.cpp b/Engine/Plugins/Experimental/Mutable/Source/CustomizableObjectEditor/Private/MuCOE/CustomizableObjectInstanceEditor.cpp
--- a/Engine/Plugins/Experimental/Mutable/Source/CustomizableObjectEditor/Private/MuCOE/CustomizableObjectInstanceEditor.cpp
+++ b/Engine/Plugins/Experimental/Mutable/Source/CustomizableObjectEditor/Private/MuCOE/CustomizableObjectInstanceEditor.cpp
@@ -197,7 +197,6 @@ void UCustomSettings::SetEditor(TSharedPtr<ICustomizableObjectInstanceEditor> In
 FCustomizableObjectInstanceEditor::FCustomizableObjectInstanceEditor()
 {
 	CustomizableObjectInstance = nullptr;
-	PreviewStaticMeshComponent = nullptr;
 	HelperCallback = nullptr;
 	PoseAsset = nullptr;
 }
@@ -212,8 +211,6 @@ FCustomizableObjectInstanceEditor::~FCustomizableObjectInstanceEditor()
 		HelperCallback = nullptr;
 	}
 
-	CustomizableObjectInstance = nullptr;
-
 	for (UCustomizableSkeletalComponent* PreviewCustomizableSkeletalComponent : PreviewCustomizableSkeletalComponents)
 	{
 		if (PreviewCustomizableSkeletalComponent)
@@ -224,7 +221,6 @@ FCustomizableObjectInstanceEditor::~FCustomizableObjectInstanceEditor()
 
 	PreviewCustomizableSkeletalComponents.Reset();
 	PreviewSkeletalMeshComponents.Reset();
-	PreviewStaticMeshComponent = nullptr;
 
 	CustomizableInstanceDetailsView.Reset();
 	Viewport.Reset();
@@ -235,8 +231,10 @@ FCustomizableObjectInstanceEditor::~FCustomizableObjectInstanceEditor()
 		Compiler->ForceFinishCompilation();
 	}
 
-	FAssetRegistryModule& AssetRegistryModule = FModuleManager::LoadModuleChecked<FAssetRegistryModule>(TEXT("AssetRegistry"));
-	AssetRegistryModule.Get().OnFilesLoaded().RemoveAll(this);
+	if (const UCustomizableObject* CustomizableObject = CustomizableObjectInstance->GetCustomizableObject())
+	{
+		CustomizableObject->GetPrivate()->Status.GetOnStateChangedDelegate().RemoveAll(this);		
+	}
 
 	FCoreUObjectDelegates::OnObjectPropertyChanged.RemoveAll(this);
 }
@@ -275,31 +273,6 @@ void FCustomizableObjectInstanceEditor::InitCustomizableObjectInstanceEditor( co
 	bOnlyRelevantParameters = InCustomizableObjectInstance->bShowOnlyRelevantParameters;
 	bOnlyRuntimeParameters = InCustomizableObjectInstance->bShowOnlyRuntimeParameters;
 
-	// Check if the asset registry has finished loading assets before compiling the object or updating the instance
-	FAssetRegistryModule& AssetRegistryModule = FModuleManager::LoadModuleChecked<FAssetRegistryModule>(TEXT("AssetRegistry"));
-	if (AssetRegistryModule.Get().IsLoadingAssets())
-	{
-		AssetRegistryModule.Get().OnFilesLoaded().AddRaw(this, &FCustomizableObjectInstanceEditor::OnAssetRegistryLoadComplete);
-		AssetRegistryLoaded = false;
-	}
-	else
-	{
-		UCustomizableObject* CustomizableObject = CustomizableObjectInstance->GetCustomizableObject();
-
-		// Compile for the first time if necessary
-		if (CustomizableObject && !CustomizableObject->IsCompiled())
-		{
-			
-			CompileObject(CustomizableObject);
-		}
-		else if (CustomizableObject)
-		{
-			CreatePreviewInstance();
-		}
-		
-		Viewport->SetAssetRegistryLoaded(true);
-	}
-
 	FAdvancedPreviewSceneModule& AdvancedPreviewSceneModule = FModuleManager::LoadModuleChecked<FAdvancedPreviewSceneModule>("AdvancedPreviewScene");
 
 	TSharedPtr<FAdvancedPreviewScene> AdvancedPreviewScene = StaticCastSharedPtr<FAdvancedPreviewScene>(Viewport->GetPreviewScene());
@@ -340,6 +313,15 @@ void FCustomizableObjectInstanceEditor::InitCustomizableObjectInstanceEditor( co
 
 	// Clears selection highlight.
 	OnInstancePropertySelectionChanged(NULL);
+	
+	if (const UCustomizableObject* CustomizableObject = CustomizableObjectInstance->GetCustomizableObject())
+	{
+		UCustomizableObjectPrivate* CustomizableObjectPrivate = CustomizableObject->GetPrivate();
+
+		CustomizableObjectPrivate->Status.GetOnStateChangedDelegate().AddRaw(this, &FCustomizableObjectInstanceEditor::OnCustomizableObjectStatusChanged);
+		const FCustomizableObjectStatusTypes::EState CurrentStatus = CustomizableObjectPrivate->Status.Get();
+		OnCustomizableObjectStatusChanged(CurrentStatus, CurrentStatus);
+	}
 }
 
 
@@ -368,12 +350,10 @@ void FCustomizableObjectInstanceEditor::CreatePreviewInstance()
 	HelperCallback->Delegate.BindSP(this, &FCustomizableObjectInstanceEditor::OnUpdatePreviewInstance);
 	CustomizableObjectInstance->UpdatedDelegate.AddDynamic(HelperCallback, &UUpdateClassWrapperClass::DelegatedCallback);
 
-	PreviewStaticMeshComponent = nullptr;
-
 	// Create a SkeletalMeshComponent for each component in the CO
 	const int32 NumMeshComponents = CustomizableObjectInstance->GetNumComponents();
-	PreviewSkeletalMeshComponents.AddZeroed(NumMeshComponents);
-	PreviewCustomizableSkeletalComponents.AddZeroed(NumMeshComponents);
+	PreviewSkeletalMeshComponents.SetNum(NumMeshComponents);
+	PreviewCustomizableSkeletalComponents.SetNum(NumMeshComponents);
 
 	for (int32 ComponentIndex = 0; ComponentIndex < NumMeshComponents; ++ComponentIndex)
 	{
@@ -394,30 +374,6 @@ void FCustomizableObjectInstanceEditor::CreatePreviewInstance()
 }
 
 
-void FCustomizableObjectInstanceEditor::OnAssetRegistryLoadComplete()
-{
-	AssetRegistryLoaded = true;
-	Viewport->SetAssetRegistryLoaded(true);
-
-	check(CustomizableObjectInstance);
-	UCustomizableObject* CustomizableObject = CustomizableObjectInstance->GetCustomizableObject();
-	if (!CustomizableObject)
-	{
-		return;
-	}
-
-	// Compile for the first time if necessary
-	if (!CustomizableObject->IsCompiled())
-	{
-		CompileObject(CustomizableObject);
-	}
-	else
-	{
-		CreatePreviewInstance();
-	}
-}
-
-
 void FCustomizableObjectInstanceEditor::UpdatePreviewVisibility()
 {
 	const bool bEnableVisibility = CustomizableObjectInstance->GetPrivate()->SkeletalMeshStatus == ESkeletalMeshStatus::Success;
@@ -479,8 +435,6 @@ void FCustomizableObjectInstanceEditor::AddReferencedObjects( FReferenceCollecto
 		Collector.AddReferencedObject(PreviewCustomizableSkeletalComponent);
 	}
 
-	Collector.AddReferencedObject( PreviewStaticMeshComponent );
-
 	for (auto& PreviewSkeletalMeshComponent : PreviewSkeletalMeshComponents)
 	{
 		Collector.AddReferencedObject(PreviewSkeletalMeshComponent);
@@ -687,12 +641,6 @@ void FCustomizableObjectInstanceEditor::HideGizmoProjectorParameter()
 }
 
 
-bool FCustomizableObjectInstanceEditor::GetAssetRegistryLoaded()
-{
-	return AssetRegistryLoaded;
-}
-
-
 void FCustomizableObjectInstanceEditor::OnInstancePropertySelectionChanged(FProperty* InProperty)
 {
 	Viewport->GetViewportClient()->Invalidate();
@@ -735,24 +683,6 @@ void FCustomizableObjectInstanceEditor::OnUpdatePreviewInstance()
 }
 
 
-void FCustomizableObjectInstanceEditor::CompileObject(UCustomizableObject* Object)
-{
-	if (!Object || !Object->Source)
-	{
-		return;
-	}
-	
-	if (!Compiler)
-	{
-		Compiler = MakeUnique<FCustomizableObjectCompiler>();
-	}
-	
-	FCompilationOptions Options = Object->CompileOptions;
-	Options.bSilentCompilation = false;
-	Compiler->Compile(*Object, Options, true);
-}
-
-
 bool FCustomizableObjectInstanceEditor::IsTickable(void) const
 {
 	return true;
@@ -802,6 +732,32 @@ TStatId FCustomizableObjectInstanceEditor::GetStatId() const
 	RETURN_QUICK_DECLARE_CYCLE_STAT(FCustomizableObjectInstanceEditor, STATGROUP_Tickables);
 }
 
+
+void FCustomizableObjectInstanceEditor::OnCustomizableObjectStatusChanged(FCustomizableObjectStatus::EState, const FCustomizableObjectStatus::EState NextState)
+{
+	switch (NextState)
+	{
+	case FCustomizableObjectStatus::EState::ModelLoaded:
+		{
+			CreatePreviewInstance();
+			break;			
+		}
+		
+	case FCustomizableObjectStatus::EState::NoModel:
+		{
+			UCustomizableObject* CustomizableObject = CustomizableObjectInstance->GetCustomizableObject();
+			check(CustomizableObject)
+		
+			CustomizableObject->ConditionalAutoCompile();	
+			break;			
+		}
+
+	default:
+		break;
+	}	
+}
+
+
 void FCustomizableObjectInstanceEditor::OpenTextureAnalyzerTab()
 {
 	TabManager->TryInvokeTab(TextureAnalyzerTabId);
diff --git a/Engine/Plugins/Experimental/Mutable/Source/CustomizableObjectEditor/Private/MuCOE/CustomizableObjectInstanceEditor.h b/Engine/Plugins/Experimental/Mutable/Source/CustomizableObjectEditor/Private/MuCOE/CustomizableObjectInstanceEditor.h
--- a/Engine/Plugins/Experimental/Mutable/Source/CustomizableObjectEditor/Private/MuCOE/CustomizableObjectInstanceEditor.h
+++ b/Engine/Plugins/Experimental/Mutable/Source/CustomizableObjectEditor/Private/MuCOE/CustomizableObjectInstanceEditor.h
@@ -5,6 +5,7 @@
 #include "MuCOE/ICustomizableObjectInstanceEditor.h"
 #include "MuCOE/CustomizableObjectEditorViewportLights.h"
 #include "TickableEditorObject.h"
+#include "MuCO/CustomizableObjectPrivate.h"
 
 #include "CustomizableObjectInstanceEditor.generated.h"
 
@@ -184,20 +185,16 @@ public:
 	virtual void ShowGizmoProjectorParameter(const FString& ParamName, int32 RangeIndex) override;
 	virtual void HideGizmoProjectorParameter() override;
 	
-	/** Getter of AssetRegistryLoaded */
-	bool GetAssetRegistryLoaded();
-
 	/** Callback to notify the editor when the PreviewInstance has been updated */
 	void OnUpdatePreviewInstance();
 
-	/** Compile the customizable object. */
-	void CompileObject(UCustomizableObject* Object);
-
 	/** FTickableGameObject interface */
 	virtual bool IsTickable(void) const override;
 	virtual void Tick(float InDeltaTime) override;
 	virtual TStatId GetStatId() const override;
 
+	void OnCustomizableObjectStatusChanged(FCustomizableObjectStatus::EState PreviousState, FCustomizableObjectStatus::EState NextState);
+	
 private:
 
 	TSharedRef<SDockTab> SpawnTab_Viewport(const FSpawnTabArgs& Args);
@@ -213,9 +210,6 @@ private:
 	/** Callback when selection changes in the Property Tree. */
 	void OnInstancePropertySelectionChanged(FProperty* InProperty);
 
-	/** Callback for the asset registry initial load */
-	void OnAssetRegistryLoadComplete();
-
 	/** Updates the visibility of PreviewSkeletalMeshComponent */
 	void UpdatePreviewVisibility();
 
@@ -289,9 +283,6 @@ private:
 	/** Handle for the OnObjectModified event */
 	FDelegateHandle OnObjectModifiedHandle;
 
-	/** Flag to know when the asset registry initial loading has completed */
-	bool AssetRegistryLoaded = true;
-
 	/** UObject class to be able to use the update callback */
 	TObjectPtr<UUpdateClassWrapperClass> HelperCallback;
 
diff --git a/Engine/Plugins/Experimental/Mutable/Source/CustomizableObjectEditor/Private/MuCOE/GenerateMutableSource/GenerateMutableSource.cpp b/Engine/Plugins/Experimental/Mutable/Source/CustomizableObjectEditor/Private/MuCOE/GenerateMutableSource/GenerateMutableSource.cpp
--- a/Engine/Plugins/Experimental/Mutable/Source/CustomizableObjectEditor/Private/MuCOE/GenerateMutableSource/GenerateMutableSource.cpp
+++ b/Engine/Plugins/Experimental/Mutable/Source/CustomizableObjectEditor/Private/MuCOE/GenerateMutableSource/GenerateMutableSource.cpp
@@ -137,7 +137,7 @@ void FMutableGraphGenerationContext::AddParticipatingObject(const UObject& InObj
 	const FGuid PackageGuid = InObject.GetPackage()->GetGuid();
 	PRAGMA_ENABLE_DEPRECATION_WARNINGS
 
-	ParticipatingObjects.Add(&InObject, PackageGuid);
+	ParticipatingObjects.Add(InObject.GetPackage()->GetFName(), PackageGuid);
 }
 
 
diff --git a/Engine/Plugins/Experimental/Mutable/Source/CustomizableObjectEditor/Private/MuCOE/GenerateMutableSource/GenerateMutableSource.h b/Engine/Plugins/Experimental/Mutable/Source/CustomizableObjectEditor/Private/MuCOE/GenerateMutableSource/GenerateMutableSource.h
--- a/Engine/Plugins/Experimental/Mutable/Source/CustomizableObjectEditor/Private/MuCOE/GenerateMutableSource/GenerateMutableSource.h
+++ b/Engine/Plugins/Experimental/Mutable/Source/CustomizableObjectEditor/Private/MuCOE/GenerateMutableSource/GenerateMutableSource.h
@@ -476,10 +476,10 @@ struct FMutableGraphGenerationContext
 	FMutableGraphGenerationContext(UCustomizableObject* CustomizableObject, class FCustomizableObjectCompiler* InCompiler, const FCompilationOptions& InOptions);
 	~FMutableGraphGenerationContext();
 
-	/** See UCustomizableObject::ParticipatingObjects. */
+	/** See FCustomizableObjectPrivateData::ParticipatingObjects. */
 	void AddParticipatingObject(const UObject& Object);
 
-	/** See UCustomizableObject::ParticipatingObjects. */
+	/** See FCustomizableObjectPrivateData::ParticipatingObjects. */
 	template<typename Type>
 	void AddParticipatingObject(const TArray<Type>& InObjects)
 	{
@@ -782,8 +782,8 @@ struct FMutableGraphGenerationContext
 	TArray<FCustomizableObjectResourceData> AlwaysLoadedExtensionData;
 	TArray<UCustomizableObjectResourceDataContainer*> StreamedExtensionData;
 
-	/** See UCustomizableObject::ParticipatingObjects. */
-	TMap<TObjectPtr<const UObject>, FGuid> ParticipatingObjects;
+	/** See FCustomizableObjectPrivateData::ParticipatingObjects. */
+	TMap<FName, FGuid> ParticipatingObjects;
 
 	/** Map to relate a Composite Data Table Row and its original DataTable */
 	TMap<UDataTable*,TMap<FName, TArray<UDataTable*>>> CompositeDataTableRowToOriginalDataTableMap;
diff --git a/Engine/Plugins/Experimental/Mutable/Source/CustomizableObjectEditor/Private/MuCOE/SCustomizableObjectEditorViewport.cpp b/Engine/Plugins/Experimental/Mutable/Source/CustomizableObjectEditor/Private/MuCOE/SCustomizableObjectEditorViewport.cpp
--- a/Engine/Plugins/Experimental/Mutable/Source/CustomizableObjectEditor/Private/MuCOE/SCustomizableObjectEditorViewport.cpp
+++ b/Engine/Plugins/Experimental/Mutable/Source/CustomizableObjectEditor/Private/MuCOE/SCustomizableObjectEditorViewport.cpp
@@ -1031,12 +1031,6 @@ void SCustomizableObjectEditorViewportTabBody::SetCustomizableObject(UCustomizab
 }
 
 
-void SCustomizableObjectEditorViewportTabBody::SetAssetRegistryLoaded(bool Value)
-{
-	AssetRegistryLoaded = Value;
-	LevelViewportClient->SetAssetRegistryLoaded(Value);
-}
-
 FLinearColor SCustomizableObjectEditorViewportTabBody::GetViewportBackgroundColor() const
 {
 	return LevelViewportClient->GetBackgroundColor();
diff --git a/Engine/Plugins/Experimental/Mutable/Source/CustomizableObjectEditor/Private/MuCOE/SCustomizableObjectEditorViewport.h b/Engine/Plugins/Experimental/Mutable/Source/CustomizableObjectEditor/Private/MuCOE/SCustomizableObjectEditorViewport.h
--- a/Engine/Plugins/Experimental/Mutable/Source/CustomizableObjectEditor/Private/MuCOE/SCustomizableObjectEditorViewport.h
+++ b/Engine/Plugins/Experimental/Mutable/Source/CustomizableObjectEditor/Private/MuCOE/SCustomizableObjectEditorViewport.h
@@ -206,9 +206,6 @@ public:
 	/** Setter of CustomizableObject */
 	void SetCustomizableObject(UCustomizableObject* CustomizableObjectParameter);
 
-	/** Setter of AssetRegistryLoaded and LevelViewportClient::AssetRegistryLoaded */
-	void SetAssetRegistryLoaded(bool Value);
-
 	/** Called to bring up the screenshot UI */
 	void OnTakeHighResScreenshot();
 
diff --git a/Engine/Plugins/Experimental/Mutable/Source/CustomizableObjectEditor/Private/MuCOE/SCustomizableObjectEditorViewportToolBar.cpp b/Engine/Plugins/Experimental/Mutable/Source/CustomizableObjectEditor/Private/MuCOE/SCustomizableObjectEditorViewportToolBar.cpp
--- a/Engine/Plugins/Experimental/Mutable/Source/CustomizableObjectEditor/Private/MuCOE/SCustomizableObjectEditorViewportToolBar.cpp
+++ b/Engine/Plugins/Experimental/Mutable/Source/CustomizableObjectEditor/Private/MuCOE/SCustomizableObjectEditorViewportToolBar.cpp
@@ -21,7 +21,10 @@
 #include "SViewportToolBarComboMenu.h"
 #include "Settings/LevelEditorViewportSettings.h"
 #include "UnrealEdGlobals.h"
+#include "AssetRegistry/AssetRegistryModule.h"
 #include "MuCO/CustomizableInstancePrivateData.h"
+#include "MuCO/CustomizableObject.h"
+#include "MuCO/CustomizableObjectPrivate.h"
 #include "MuCO/CustomizableObjectSystem.h"
 #include "Widgets/Input/SButton.h"
 #include "Widgets/Input/SMenuAnchor.h"
@@ -190,9 +193,22 @@ FText SCustomizableObjectEditorViewportToolBar::GetCompileErrorOverlayText() con
 		return {};
 	}
 
-	if (!Editor->GetAssetRegistryLoaded())
+	bool bLoading = false;
+
+	const UObject* EditingObject = (*Editor->GetObjectsCurrentlyBeingEdited())[0];
+	if (const UCustomizableObject* CustomizableObject = Cast<UCustomizableObject>(EditingObject))
+	{
+		bLoading = CustomizableObject->GetPrivate()->Status.Get() == FCustomizableObjectStatusTypes::EState::Loading;
+	}
+	else if (const UCustomizableObjectInstance* Instance = Cast<UCustomizableObjectInstance>(EditingObject))
+	{
+		const UCustomizableObject* InstanceCustomizableObject = Instance->GetCustomizableObject();
+		bLoading = InstanceCustomizableObject && InstanceCustomizableObject->GetPrivate()->Status.Get() == FCustomizableObjectStatusTypes::EState::Loading;
+	}
+
+	if (bLoading)
 	{
-		return LOCTEXT("LoadingAssets", "Loading Assets");
+		return LOCTEXT("LoadingAssets", "Loading Customizable Object");
 	}
 
 	const UCustomizableObjectInstance* Instance = Editor->GetPreviewInstance();
@@ -201,7 +217,7 @@ FText SCustomizableObjectEditorViewportToolBar::GetCompileErrorOverlayText() con
 		return LOCTEXT("NoPreviewInstance", "No Preview Instance");
 	}
 
-	UCustomizableObjectSystem* System = UCustomizableObjectSystem::GetInstanceChecked();
+	const UCustomizableObjectSystem* System = UCustomizableObjectSystem::GetInstanceChecked();
 	
 	if (System->IsUpdating(Instance))
 	{
diff --git a/Engine/Plugins/Experimental/Mutable/Source/CustomizableObjectEditor/Public/MuCOE/CustomizableObjectEditorViewportClient.h b/Engine/Plugins/Experimental/Mutable/Source/CustomizableObjectEditor/Public/MuCOE/CustomizableObjectEditorViewportClient.h
--- a/Engine/Plugins/Experimental/Mutable/Source/CustomizableObjectEditor/Public/MuCOE/CustomizableObjectEditorViewportClient.h
+++ b/Engine/Plugins/Experimental/Mutable/Source/CustomizableObjectEditor/Public/MuCOE/CustomizableObjectEditorViewportClient.h
@@ -225,9 +225,6 @@ public:
 	}
 	// End of FSerializableObject interface
 
-	/** Setter of AssetRegistryLoaded */
-	void SetAssetRegistryLoaded(bool Value);
-
 	/** Show pero LOD geometric information of the instance */
 	void ShowInstanceGeometryInformation(FCanvas* InCanvas);
 
@@ -341,9 +338,6 @@ private:
 	/** To know if the user has given permission to overwrite files when baking if already present */
 	bool BakingOverwritePermission;
 
-	/** Flag to know when the asset registry initial loading has completed, value set by the Customizable Object / Customizable Object Instance editor */
-	bool AssetRegistryLoaded;
-
 	/** To know if the orbital camera is being used or not */
 	bool bActivateOrbitalCamera;
 
diff --git a/Engine/Plugins/Experimental/Mutable/Source/CustomizableObjectEditor/Public/MuCOE/ICustomizableObjectInstanceEditor.h b/Engine/Plugins/Experimental/Mutable/Source/CustomizableObjectEditor/Public/MuCOE/ICustomizableObjectInstanceEditor.h
--- a/Engine/Plugins/Experimental/Mutable/Source/CustomizableObjectEditor/Public/MuCOE/ICustomizableObjectInstanceEditor.h
+++ b/Engine/Plugins/Experimental/Mutable/Source/CustomizableObjectEditor/Public/MuCOE/ICustomizableObjectInstanceEditor.h
@@ -31,9 +31,6 @@ public:
 
 	/** Refreshes everything in the Customizable Object Instance Editor. */
 	virtual void RefreshTool() = 0;
-
-	/** Getter of AssetRegistryLoaded */
-	virtual bool GetAssetRegistryLoaded() = 0;
 	
 	virtual void SetPoseAsset(UPoseAsset* PoseAssetParameter) {}
 
