[Mutable] Move FMergedSkeletons cache to the CO's private.
Remove CO's Identifier and bModelCompiledForCook.

#rnx

diff --git a/Engine/Plugins/Experimental/Mutable/Source/CustomizableObject/Classes/MuCO/CustomizableObject.h b/Engine/Plugins/Experimental/Mutable/Source/CustomizableObject/Classes/MuCO/CustomizableObject.h
--- a/Engine/Plugins/Experimental/Mutable/Source/CustomizableObject/Classes/MuCO/CustomizableObject.h
+++ b/Engine/Plugins/Experimental/Mutable/Source/CustomizableObject/Classes/MuCO/CustomizableObject.h
@@ -48,24 +48,6 @@ namespace mu
 
 CUSTOMIZABLEOBJECT_API DECLARE_LOG_CATEGORY_EXTERN(LogMutable, Log, All);
 
-struct FMergedSkeleton
-{
-	FMergedSkeleton(TObjectPtr<USkeleton> InSkeleton, const int32 InComponentIndex, const TArray<uint16>& InSkeletonIds);
-
-	// Merged skeleton
-	TWeakObjectPtr<USkeleton> Skeleton;
-	
-	// Component Index and Ids of the Skeletons involved in the merge. Used to identify the skeleton
-	uint16 ComponentIndex;
-	TArray<uint16> SkeletonIds;
-
-	bool operator==(const FMergedSkeleton& Other) const
-	{
-		return ComponentIndex == Other.ComponentIndex
-			&& SkeletonIds == Other.SkeletonIds;
-	}
-};
-
 
 USTRUCT()
 struct FFParameterOptionsTags
@@ -1287,15 +1269,6 @@ public:
 	
 	TSoftObjectPtr<UMaterialInterface> GetReferencedMaterialAssetPtr(uint32 Index);
 
-	// Return a valid Skeletons if cached. ComponentIndex and SkeletonIds are used as a key to find the skeleton.
-	TObjectPtr<USkeleton> GetCachedMergedSkeleton(int32 ComponentIndex, const TArray<uint16>& SkeletonIds) const;
-
-	// Add merged skeleton to the cache. It'll be cached as a TWeakObjPtr.
-	void CacheMergedSkeleton(const int32 ComponentIndex, const TArray<uint16>& SkeletonIds, TObjectPtr<USkeleton> Skeleton);
-
-	// Remove skeletons that have been destroyed by the garbage collector from the cache.
-	void UnCacheInvalidSkeletons();
-
 private:
 
 	
@@ -1585,10 +1558,6 @@ public:
 	UPROPERTY()
 	bool bIsChildObject = false;
 
-	/** Unique Identifier - Deterministic. Used to locate Model and Streamable data on disk. Should not be modified. */
-	UPROPERTY(Transient)
-	FGuid Identifier;
-
 	ECustomizableObjectCompilationState CompilationState = ECustomizableObjectCompilationState::None;
 
 	FPostCompileDelegate PostCompileDelegate;
@@ -1646,9 +1615,6 @@ private:
 	UPROPERTY()
 	TArray<FName> BoneNames;
 
-	/** Cache of merged skeletons */
-	TArray<FMergedSkeleton> MergedSkeletons;
-
 	/** BulkData that stores all in-game resources used by Mutable when generating instances.
 	  * Only valid in packaged builds */
 	UPROPERTY()
diff --git a/Engine/Plugins/Experimental/Mutable/Source/CustomizableObject/Internal/MuCO/CustomizableObjectPrivate.h b/Engine/Plugins/Experimental/Mutable/Source/CustomizableObject/Internal/MuCO/CustomizableObjectPrivate.h
--- a/Engine/Plugins/Experimental/Mutable/Source/CustomizableObject/Internal/MuCO/CustomizableObjectPrivate.h
+++ b/Engine/Plugins/Experimental/Mutable/Source/CustomizableObject/Internal/MuCO/CustomizableObjectPrivate.h
@@ -16,6 +16,7 @@
 namespace mu { class Model; }
 class UCustomizableObject;
 class USkeletalMesh;
+class USkeleton;
 
 
 class FMeshCache
@@ -57,6 +58,18 @@ struct FCustomizableObjectStatusTypes
 using FCustomizableObjectStatus = FStateMachine<FCustomizableObjectStatusTypes>;
 
 
+class FSkeletonCache
+{
+public:
+	USkeleton* Get(const TArray<uint16>& Key);
+
+	void Add(const TArray<uint16>& Key, USkeleton* Value);
+
+private:
+	TMap<TArray<uint16>, TWeakObjectPtr<USkeleton>> MergedSkeletons;
+};
+
+
 UCLASS()
 class UCustomizableObjectPrivate : public UObject
 {
@@ -80,7 +93,6 @@ public:
 #if WITH_EDITOR
 	FGuid Identifier;
 
-	bool bModelCompiledForCook = false;
 	TArray<FString> CachedPlatformNames;
 
 	/** List of external packages that if changed, a compilation is required.
@@ -93,6 +105,9 @@ public:
 	/** Cache of generated SkeletalMeshes */
 	FMeshCache MeshCache;
 
+	/** Cache of merged Skeletons */
+	FSkeletonCache SkeletonCache;
+
 	FCustomizableObjectStatus Status;
 };
 
diff --git a/Engine/Plugins/Experimental/Mutable/Source/CustomizableObject/Private/MuCO/CustomizableObject.cpp b/Engine/Plugins/Experimental/Mutable/Source/CustomizableObject/Private/MuCO/CustomizableObject.cpp
--- a/Engine/Plugins/Experimental/Mutable/Source/CustomizableObject/Private/MuCO/CustomizableObject.cpp
+++ b/Engine/Plugins/Experimental/Mutable/Source/CustomizableObject/Private/MuCO/CustomizableObject.cpp
@@ -104,14 +104,13 @@ void UCustomizableObject::GetAssetRegistryTags(TArray<FAssetRegistryTag>& OutTag
 void UCustomizableObject::GetAssetRegistryTags(FAssetRegistryTagsContext Context) const
 {
 	int32 isRoot = 0;
-#if WITH_EDITOR
+
 	FCustomizableObjectCompilerBase* Compiler = UCustomizableObjectSystem::GetInstance()->GetNewCompiler();
 	if (Compiler)
 	{
 		isRoot = Compiler->IsRootObject(this) ? 1 : 0;
 		delete Compiler;
 	}
-#endif
 
 	Context.AddTag(FAssetRegistryTag("IsRoot", FString::FromInt(isRoot), FAssetRegistryTag::TT_Numerical));
 	Super::GetAssetRegistryTags(Context);
@@ -326,6 +325,16 @@ bool UCustomizableObject::IsCachedCookedPlatformDataLoaded( const ITargetPlatfor
 }
 
 
+FGuid GenerateIdentifier(const UCustomizableObject& CustomizableObject)
+{
+	// Generate the Identifier using the path and name of the asset
+	uint32 FullPathHash = GetTypeHash(CustomizableObject.GetFullName());
+	uint32 OutermostHash = GetTypeHash(GetNameSafe(CustomizableObject.GetOutermost()));
+	uint32 OuterHash = GetTypeHash(CustomizableObject.GetName());
+	return FGuid(0, FullPathHash, OutermostHash, OuterHash);
+}
+
+
 void UCustomizableObject::ClearCompiledData()
 {
 	ReferenceSkeletalMeshesData.Empty();
@@ -696,7 +705,7 @@ void UCustomizableObject::LoadCompiledData(FArchive& MemoryReader, const ITarget
 			mu::InputArchive arch(&stream);
 			TSharedPtr<mu::Model, ESPMode::ThreadSafe> Model = mu::Model::StaticUnserialise( arch );
 
-			Private->SetModel(Model, Identifier);
+			Private->SetModel(Model, GenerateIdentifier(*this));
 		}
 	}
 	
@@ -948,23 +957,10 @@ FString UCustomizableObject::GetCompiledDataFolderPath() const
 }
 
 
-FGuid GenerateIdentifier(const UCustomizableObject& CustomizableObject)
-{
-	// Generate the Identifier using the path and name of the asset
-	uint32 FullPathHash = GetTypeHash(CustomizableObject.GetFullName());
-	uint32 OutermostHash = GetTypeHash(GetNameSafe(CustomizableObject.GetOutermost()));
-	uint32 OuterHash = GetTypeHash(CustomizableObject.GetName());
-	return FGuid(0, FullPathHash, OutermostHash, OuterHash);
-}
-
-
 FString UCustomizableObject::GetCompiledDataFileName(bool bIsModel, const ITargetPlatform* InTargetPlatform, bool bIsDiskStreamer)
 {
-	// Generate the Identifier using the path and name of the asset
-	Identifier = GenerateIdentifier(*this);
-
 	const FString PlatformName = InTargetPlatform ? InTargetPlatform->PlatformName() : FPlatformProperties::PlatformName();
-	const FString FileIdentifier = bIsDiskStreamer ? GetPrivate()->Identifier.ToString() : Identifier.ToString();
+	const FString FileIdentifier = bIsDiskStreamer ? GetPrivate()->Identifier.ToString() : GenerateIdentifier(*this).ToString();
 	const FString Extension = bIsModel ? TEXT("_M.mut") : TEXT("_S.mut");
 	return PlatformName + FileIdentifier + Extension;
 }
@@ -1149,43 +1145,6 @@ TSoftObjectPtr<USkeleton> UCustomizableObject::GetReferencedSkeletonAssetPtr( ui
 }
 
 
-TObjectPtr<USkeleton> UCustomizableObject::GetCachedMergedSkeleton(const int32 ComponentIndex, const TArray<uint16>& SkeletonIds) const
-{
-	const FMergedSkeleton* MergedSkeleton = MergedSkeletons.FindByPredicate([&ComponentIndex, &SkeletonIds](const FMergedSkeleton& MergedSkeleton) 
-		{ return MergedSkeleton.ComponentIndex == ComponentIndex && MergedSkeleton.SkeletonIds == SkeletonIds; });
-
-	if (MergedSkeleton && MergedSkeleton->Skeleton.IsValid())
-	{
-		return MergedSkeleton->Skeleton.Get();
-	}
-
-	return nullptr;
-}
-
-
-void UCustomizableObject::CacheMergedSkeleton(const int32 ComponentIndex, const TArray<uint16>& SkeletonIds, TObjectPtr<USkeleton> Skeleton)
-{
-	if (Skeleton)
-	{
-		MergedSkeletons.AddUnique({ Skeleton, ComponentIndex, SkeletonIds });
-	}
-}
-
-
-void UCustomizableObject::UnCacheInvalidSkeletons()
-{
-	for (int32 SkeletonIndex = MergedSkeletons.Num() - 1; SkeletonIndex >= 0; --SkeletonIndex)
-	{
-		FMergedSkeleton& MergedSkeleton = MergedSkeletons[SkeletonIndex];
-		if (!MergedSkeleton.Skeleton.IsValid())
-		{
-			MergedSkeletons.RemoveSingleSwap(MergedSkeleton);
-		}
-	}
-}
-
-
-
 int32 UCustomizableObject::FindState( const FString& Name ) const
 {
 	int32 Result = -1;
@@ -1930,6 +1889,33 @@ void FMeshCache::Add(const TArray<mu::FResourceID>& Key, USkeletalMesh* Value)
 }
 
 
+USkeleton* FSkeletonCache::Get(const TArray<uint16>& Key)
+{
+	const TWeakObjectPtr<USkeleton>* Result = MergedSkeletons.Find(Key);
+	return Result ? Result->Get() : nullptr;
+}
+
+
+void FSkeletonCache::Add(const TArray<uint16>& Key, USkeleton* Value)
+{
+	if (!Value)
+	{
+		return;
+	}
+
+	MergedSkeletons.Add(Key, Value);
+
+	// Remove invalid SkeletalMeshes from the cache.
+	for (auto SkeletonIterator = MergedSkeletons.CreateIterator(); SkeletonIterator; ++SkeletonIterator)
+	{
+		if (SkeletonIterator.Value().IsStale())
+		{
+			SkeletonIterator.RemoveCurrent();
+		}
+	}
+}
+
+
 void UCustomizableObjectPrivate::SetModel(const TSharedPtr<mu::Model, ESPMode::ThreadSafe>& Model, const FGuid Id)
 {
 	if (MutableModel == Model
@@ -1969,18 +1955,6 @@ TSharedPtr<const mu::Model, ESPMode::ThreadSafe> UCustomizableObjectPrivate::Get
 }
 
 
-//-------------------------------------------------------------------------------------------------
-//-------------------------------------------------------------------------------------------------
-//-------------------------------------------------------------------------------------------------
-FMergedSkeleton::FMergedSkeleton(TObjectPtr<USkeleton> InSkeleton, const int32 InComponentIndex, const TArray<uint16>& InSkeletonIds)
-{
-	check(InSkeleton);
-	Skeleton = InSkeleton;
-	ComponentIndex = InComponentIndex;
-	SkeletonIds = InSkeletonIds;
-}
-
-
 //-------------------------------------------------------------------------------------------------
 //-------------------------------------------------------------------------------------------------
 //-------------------------------------------------------------------------------------------------
diff --git a/Engine/Plugins/Experimental/Mutable/Source/CustomizableObject/Private/MuCO/CustomizableObjectInstance.cpp b/Engine/Plugins/Experimental/Mutable/Source/CustomizableObject/Private/MuCO/CustomizableObjectInstance.cpp
--- a/Engine/Plugins/Experimental/Mutable/Source/CustomizableObject/Private/MuCO/CustomizableObjectInstance.cpp
+++ b/Engine/Plugins/Experimental/Mutable/Source/CustomizableObject/Private/MuCO/CustomizableObjectInstance.cpp
@@ -1101,7 +1101,7 @@ USkeleton* UCustomizableInstancePrivate::MergeSkeletons(UCustomizableObject& Cus
 		}
 
 		// Add Skeleton to the cache
-		CustomizableObject.CacheMergedSkeleton(ComponentIndex, ReferencedSkeletons.SkeletonIds, FinalSkeleton);
+		CustomizableObject.GetPrivate()->SkeletonCache.Add(ReferencedSkeletons.SkeletonIds, FinalSkeleton);
 		ReferencedSkeletons.SkeletonIds.Empty();
 	}
 	
@@ -4866,8 +4866,6 @@ FGraphEventRef UCustomizableInstancePrivate::LoadAdditionalAssetsAsync(const TSh
 		}
 	}
 
-	// Clear invalid skeletons from the MergedSkeletons cache
-	CustomizableObject->UnCacheInvalidSkeletons();
 
 	// Load Skeletons required by the SubMeshes of the newly generated Mesh, will be merged later
 	for (FInstanceUpdateData::FSkeletonData& SkeletonData : OperationData->InstanceUpdateData.Skeletons)
@@ -4880,7 +4878,7 @@ FGraphEventRef UCustomizableInstancePrivate::LoadAdditionalAssetsAsync(const TSh
 		}
 
 		// Reuse merged Skeleton if cached
-		ComponentData->Skeletons.Skeleton = CustomizableObject->GetCachedMergedSkeleton(SkeletonData.ComponentIndex, SkeletonData.SkeletonIds);
+		ComponentData->Skeletons.Skeleton = CustomizableObject->GetPrivate()->SkeletonCache.Get(SkeletonData.SkeletonIds);
 		if (ComponentData->Skeletons.Skeleton)
 		{
 			ComponentData->Skeletons.SkeletonIds.Empty();
