[Mutable] Improve the visibility of generated resources in the editor and move the generated Skeletal Meshes to the instance's private.

#rnx

diff --git a/Engine/Plugins/Experimental/Mutable/Source/CustomizableObject/Classes/MuCO/CustomizableObjectInstance.h b/Engine/Plugins/Experimental/Mutable/Source/CustomizableObject/Classes/MuCO/CustomizableObjectInstance.h
--- a/Engine/Plugins/Experimental/Mutable/Source/CustomizableObject/Classes/MuCO/CustomizableObjectInstance.h
+++ b/Engine/Plugins/Experimental/Mutable/Source/CustomizableObject/Classes/MuCO/CustomizableObjectInstance.h
@@ -139,10 +139,6 @@ public:
 
 	TMap<FString, bool> ParamNameToExpandedMap; // Used to check whether a mutable param is expanded in the editor to show its child params
 
-	/** The generated skeletal meshes for this Instance, one for each component */
-	UPROPERTY(Transient, VisibleAnywhere, Category = CustomizableSkeletalMesh)
-	TArray< TObjectPtr<USkeletalMesh> > SkeletalMeshes;
-
 #if WITH_EDITOR
 	/** During editor, always remember the duration of the last update in the mutable runtime, for profiling. */
 	int32 LastUpdateMutableRuntimeCycles = 0;
diff --git a/Engine/Plugins/Experimental/Mutable/Source/CustomizableObject/Internal/MuCO/CustomizableInstancePrivateData.h b/Engine/Plugins/Experimental/Mutable/Source/CustomizableObject/Internal/MuCO/CustomizableInstancePrivateData.h
--- a/Engine/Plugins/Experimental/Mutable/Source/CustomizableObject/Internal/MuCO/CustomizableInstancePrivateData.h
+++ b/Engine/Plugins/Experimental/Mutable/Source/CustomizableObject/Internal/MuCO/CustomizableInstancePrivateData.h
@@ -163,7 +163,11 @@ public:
 
 	UCustomizableInstancePrivate();
 
-	UPROPERTY( Transient )
+	/** The generated skeletal meshes for this Instance, one for each component */
+	UPROPERTY(Transient, VisibleAnywhere, Category = NoCategory)
+	TArray<TObjectPtr<USkeletalMesh>> SkeletalMeshes;
+
+	UPROPERTY(Transient, VisibleAnywhere, Category = NoCategory)
 	TArray<FGeneratedMaterial> GeneratedMaterials;
 
 	UPROPERTY( Transient )
@@ -231,7 +235,7 @@ public:
 	int32 GetNumLODsAvailable() const { return NumLODsAvailable; }
 	
 	// The following method is basically copied from PostEditChangeProperty and/or SkeletalMesh.cpp to be able to replicate PostEditChangeProperty without the editor
-	void PostEditChangePropertyWithoutEditor(USkeletalMesh* SkeletalMesh);
+	void PostEditChangePropertyWithoutEditor();
 	
 	/** Calls ReleaseResources on all SkeletalMeshes generated by this instance and invalidates the generated data.
 	  * It should not be called if the meshes are still in use or shared with other instances. */
@@ -247,7 +251,7 @@ public:
 
 	/** */
 #if WITH_EDITORONLY_DATA
-	static void RegenerateImportedModel(USkeletalMesh* SkeletalMesh);
+	void RegenerateImportedModels();
 #endif
 
 private:
@@ -306,13 +310,13 @@ public:
 	TArray< TObjectPtr<UPhysicsAsset> > ClothingPhysicsAssets;
 
 	// To keep loaded AnimBPs referenced and prevent GC
-	UPROPERTY(Transient, Category = CustomizableObjectInstance, editfixedsize, VisibleAnywhere)
+	UPROPERTY(Transient, Category = Animation, editfixedsize, VisibleAnywhere)
 	TArray<TSubclassOf<UAnimInstance>> GatheredAnimBPs;
 
-	UPROPERTY(Transient, Category = CustomizableObjectInstance, editfixedsize, VisibleAnywhere)
+	UPROPERTY(Transient, Category = Animation, editfixedsize, VisibleAnywhere)
 	FGameplayTagContainer AnimBPGameplayTags;
 
-	UPROPERTY(Transient, Category = CustomizableObjectInstance, editfixedsize, VisibleAnywhere)
+	UPROPERTY(Transient, Category = Animation, editfixedsize, VisibleAnywhere)
 	TMap<TSubclassOf<UAnimInstance>, FAnimBpGeneratedPhysicsAssets> AnimBpPhysicsAssets;
 
 	// The pass-through textures that will be loaded during an update
diff --git a/Engine/Plugins/Experimental/Mutable/Source/CustomizableObject/Private/MuCO/CustomizableInstanceLODManagement.cpp b/Engine/Plugins/Experimental/Mutable/Source/CustomizableObject/Private/MuCO/CustomizableInstanceLODManagement.cpp
--- a/Engine/Plugins/Experimental/Mutable/Source/CustomizableObject/Private/MuCO/CustomizableInstanceLODManagement.cpp
+++ b/Engine/Plugins/Experimental/Mutable/Source/CustomizableObject/Private/MuCO/CustomizableInstanceLODManagement.cpp
@@ -363,9 +363,8 @@ void UCustomizableInstanceLODManagement::UpdateInstanceDistsAndLODs(FMutableInst
 #if WITH_EDITOR
 					// If the instance is generated but the component doesn't have a mesh, set it.
 					// Can happen when duplicating instances in the editor.
-					if (COI->SkeletalMeshes.IsValidIndex(ComponentIndex) && 
-						Parent->GetSkeletalMeshAsset() == nullptr &&
-						COI->SkeletalMeshes[ComponentIndex])
+					const USkeletalMesh* SkeletalMesh = COI->GetSkeletalMesh(ComponentIndex);
+					if (SkeletalMesh && !Parent->GetSkeletalMeshAsset())
 					{
 						// As the instance is already generated, this will be very fast and just set the mesh and call the delegates
 						COI->UpdateSkeletalMeshAsync();
diff --git a/Engine/Plugins/Experimental/Mutable/Source/CustomizableObject/Private/MuCO/CustomizableObjectInstance.cpp b/Engine/Plugins/Experimental/Mutable/Source/CustomizableObject/Private/MuCO/CustomizableObjectInstance.cpp
--- a/Engine/Plugins/Experimental/Mutable/Source/CustomizableObject/Private/MuCO/CustomizableObjectInstance.cpp
+++ b/Engine/Plugins/Experimental/Mutable/Source/CustomizableObject/Private/MuCO/CustomizableObjectInstance.cpp
@@ -191,6 +191,7 @@ void UCustomizableInstancePrivate::SetLastMeshId(int32 ComponentIndex, int32 LOD
 void UCustomizableInstancePrivate::InvalidateGeneratedData()
 {
 	SkeletalMeshStatus = ESkeletalMeshStatus::NotGenerated;
+	SkeletalMeshes.Reset();
 
 	DescriptorHash = FDescriptorHash();
 	UpdateDescriptorHash = FDescriptorHash();
@@ -783,19 +784,19 @@ bool UCustomizableObjectInstance::IsParameterRelevant(const FString& ParamName)
 }
 
 
-void UCustomizableInstancePrivate::PostEditChangePropertyWithoutEditor(USkeletalMesh* InSkeletalMesh)
+void UCustomizableInstancePrivate::PostEditChangePropertyWithoutEditor()
 {
-	MUTABLE_CPUPROFILER_SCOPE(UCustomizableInstancePrivateData::PostEditChangePropertyWithoutEditor);
-
-	if (!InSkeletalMesh) return;
-
-	FSkeletalMeshRenderData* Resource = InSkeletalMesh->GetResourceForRendering();
-	if (!Resource) return;
+	MUTABLE_CPUPROFILER_SCOPE(UCustomizableInstancePrivate::PostEditChangePropertyWithoutEditor);
 
+	for (USkeletalMesh* SkeletalMesh : SkeletalMeshes)
 	{
-		MUTABLE_CPUPROFILER_SCOPE(InitResources);
-		// reinitialize resource
-		InSkeletalMesh->InitResources();
+		if (SkeletalMesh && SkeletalMesh->GetResourceForRendering())
+		{
+			MUTABLE_CPUPROFILER_SCOPE(InitResources);
+
+			// reinitialize resources
+			SkeletalMesh->InitResources();
+		}
 	}
 }
 
@@ -1769,7 +1770,7 @@ bool UCustomizableInstancePrivate::DoComponentsNeedUpdate(UCustomizableObjectIns
 		}
 
 		// Update the component if there is a mesh and it shouldn't, or the other way around.
-		const bool bHasSkeletalMesh = Public->SkeletalMeshes.IsValidIndex(ComponentIndex) && Public->SkeletalMeshes[ComponentIndex];
+		const bool bHasSkeletalMesh = SkeletalMeshes.IsValidIndex(ComponentIndex) && SkeletalMeshes[ComponentIndex];
 		OutComponentNeedsUpdate[ComponentIndex] = (ComponentWithMesh[ComponentIndex] != bHasSkeletalMesh);
 
 		const FCustomizableInstanceComponentData* ComponentData = GetComponentData(ComponentIndex);
@@ -1792,7 +1793,7 @@ bool UCustomizableInstancePrivate::DoComponentsNeedUpdate(UCustomizableObjectIns
 
 bool UCustomizableInstancePrivate::UpdateSkeletalMesh_PostBeginUpdate0(UCustomizableObjectInstance* Public, const TSharedRef<FUpdateContextPrivate>& OperationData)
 {
-	MUTABLE_CPUPROFILER_SCOPE(UCustomizableInstancePrivateData::UpdateSkeletalMesh_PostBeginUpdate0)
+	MUTABLE_CPUPROFILER_SCOPE(UCustomizableInstancePrivate::UpdateSkeletalMesh_PostBeginUpdate0)
 
 	bool bHasInvalidMesh = false;
 
@@ -1804,8 +1805,6 @@ bool UCustomizableInstancePrivate::UpdateSkeletalMesh_PostBeginUpdate0(UCustomiz
 	{
 		UE_LOG(LogMutable, Warning, TEXT("Failed to generate SkeletalMesh for CO Instance %s. CO [%s]"), *Public->GetName(), *GetNameSafe(Public->GetCustomizableObject()));
 
-		Public->SkeletalMeshes.Reset(); // What about all the references
-
 		InvalidateGeneratedData();
 
 		OperationData->UpdateResult = EUpdateResult::Error;
@@ -1823,7 +1822,7 @@ bool UCustomizableInstancePrivate::UpdateSkeletalMesh_PostBeginUpdate0(UCustomiz
 
 	TextureReuseCache.Empty(); // Sections may have changed, so invalidate the texture reuse cache because it's indexed by section
 
-	TArray<TObjectPtr<USkeletalMesh>> OldSkeletalMeshes = Public->SkeletalMeshes;
+	TArray<TObjectPtr<USkeletalMesh>> OldSkeletalMeshes = SkeletalMeshes;
 
 	bool bSuccess = true;
 
@@ -1831,7 +1830,7 @@ bool UCustomizableInstancePrivate::UpdateSkeletalMesh_PostBeginUpdate0(UCustomiz
 	check(CustomizableObject);
 
 	// Initialize the maximum number of SkeletalMeshes we could possibly have. 
-	Public->SkeletalMeshes.Init(nullptr, CustomizableObject->GetComponentCount());
+	SkeletalMeshes.Init(nullptr, CustomizableObject->GetComponentCount());
 
 	const int32 NumMutableComponents = OperationData->InstanceUpdateData.Components.Num();
 	for (int32 ComponentIndex = 0; ComponentIndex < NumMutableComponents; ++ComponentIndex)
@@ -1843,14 +1842,14 @@ bool UCustomizableInstancePrivate::UpdateSkeletalMesh_PostBeginUpdate0(UCustomiz
 		{
 			if (OldSkeletalMeshes.IsValidIndex(Component.Id))
 			{
-				Public->SkeletalMeshes[Component.Id] = OldSkeletalMeshes[Component.Id];
+				SkeletalMeshes[Component.Id] = OldSkeletalMeshes[Component.Id];
 			}
 
 			continue;
 		}
 
 		// Check if we have initialized the component
-		if (Public->SkeletalMeshes[Component.Id])
+		if (SkeletalMeshes[Component.Id])
 		{
 			continue;
 		}
@@ -1859,7 +1858,7 @@ bool UCustomizableInstancePrivate::UpdateSkeletalMesh_PostBeginUpdate0(UCustomiz
 		{
 			if (USkeletalMesh* CachedMesh = CustomizableObject->GetPrivate()->MeshCache.Get(OperationData->MeshDescriptors[Component.Id]))
 			{
-				Public->SkeletalMeshes[Component.Id] = CachedMesh;
+				SkeletalMeshes[Component.Id] = CachedMesh;
 				continue;
 			}
 		}
@@ -1882,9 +1881,9 @@ bool UCustomizableInstancePrivate::UpdateSkeletalMesh_PostBeginUpdate0(UCustomiz
 		// Create and initialize the SkeletalMesh for this component
 		MUTABLE_CPUPROFILER_SCOPE(ConstructMesh);
 
-		Public->SkeletalMeshes[Component.Id] = NewObject<USkeletalMesh>(GetTransientPackage(), NAME_None, RF_Transient);
+		SkeletalMeshes[Component.Id] = NewObject<USkeletalMesh>(GetTransientPackage(), NAME_None, RF_Transient);
 
-		USkeletalMesh* SkeletalMesh = Public->SkeletalMeshes[Component.Id];
+		USkeletalMesh* SkeletalMesh = SkeletalMeshes[Component.Id];
 		check(SkeletalMesh);
 
 		// Set up the default information any mesh from this component will have (LODArrayInfos, RenderData, Mesh settings, etc). 
@@ -2061,7 +2060,7 @@ bool UCustomizableInstancePrivate::UpdateSkeletalMesh_PostBeginUpdate0(UCustomiz
 	{
 		MUTABLE_CPUPROFILER_SCOPE(BuildSkeletalMeshData);
 
-		const int32 NumComponents = Public->SkeletalMeshes.Num();
+		const int32 NumComponents = SkeletalMeshes.Num();
 		for (int32 ComponentIndex = 0; bSuccess && ComponentIndex < NumComponents; ++ComponentIndex)
 		{
 			if (OperationData->bUseMeshCache)
@@ -2078,7 +2077,7 @@ bool UCustomizableInstancePrivate::UpdateSkeletalMesh_PostBeginUpdate0(UCustomiz
 				continue;
 			}
 
-			USkeletalMesh* SkeletalMesh = Public->SkeletalMeshes[ComponentIndex];
+			USkeletalMesh* SkeletalMesh = SkeletalMeshes[ComponentIndex];
 
 			if (!SkeletalMesh)
 			{
@@ -2114,8 +2113,6 @@ bool UCustomizableInstancePrivate::UpdateSkeletalMesh_PostBeginUpdate0(UCustomiz
 
 	if (!bSuccess)
 	{
-		Public->SkeletalMeshes.Reset();
-
 		InvalidateGeneratedData();
 	}
 
@@ -2769,12 +2766,12 @@ void UCustomizableInstancePrivate::DiscardResources()
 
 	if (SkeletalMeshStatus == ESkeletalMeshStatus::Success)
 	{
-		for (int32 Component = 0; Component < Instance->SkeletalMeshes.Num(); ++Component)
+		for (int32 Component = 0; Component < SkeletalMeshes.Num(); ++Component)
 		{
-			if (Instance->SkeletalMeshes[Component] && Instance->SkeletalMeshes[Component]->IsValidLowLevel())
+			if (SkeletalMeshes[Component] && SkeletalMeshes[Component]->IsValidLowLevel())
 			{
-				Instance->SkeletalMeshes[Component]->ReleaseResources();
-				Instance->SkeletalMeshes[Component] = nullptr;
+				SkeletalMeshes[Component]->ReleaseResources();
+				SkeletalMeshes[Component] = nullptr;
 			}
 		}
 
@@ -2783,7 +2780,6 @@ void UCustomizableInstancePrivate::DiscardResources()
 	
 	InvalidateGeneratedData();
 	
-	Instance->SkeletalMeshes.Reset();
 }
 
 
@@ -3206,7 +3202,7 @@ FTexturePlatformData* MutableCreateImagePlatformData(mu::Ptr<const mu::Image> Mu
 
 void ConvertImage(UTexture2D* Texture, mu::Ptr<const mu::Image> MutableImage, const FMutableModelImageProperties& Props, int OnlyLOD, int32 ExtractChannel)
 {
-	MUTABLE_CPUPROFILER_SCOPE(UCustomizableInstancePrivateData::ConvertImage);
+	MUTABLE_CPUPROFILER_SCOPE(UCustomizableInstancePrivate::ConvertImage);
 
 	SetTexturePropertiesFromMutableImageProps(Texture, Props, false);
 
@@ -3285,7 +3281,7 @@ void ConvertImage(UTexture2D* Texture, mu::Ptr<const mu::Image> MutableImage, co
 
 void UCustomizableInstancePrivate::InitSkeletalMeshData(const TSharedRef<FUpdateContextPrivate>& OperationData, USkeletalMesh* SkeletalMesh, const FMutableRefSkeletalMeshData* RefSkeletalMeshData, const UCustomizableObject& CustomizableObject, int32 ComponentIndex)
 {
-	MUTABLE_CPUPROFILER_SCOPE(UCustomizableInstancePrivateData::InitSkeletalMesh);
+	MUTABLE_CPUPROFILER_SCOPE(UCustomizableInstancePrivate::InitSkeletalMesh);
 
 	check(SkeletalMesh);
 	check(RefSkeletalMeshData);
@@ -3370,7 +3366,7 @@ void UCustomizableInstancePrivate::InitSkeletalMeshData(const TSharedRef<FUpdate
 
 bool UCustomizableInstancePrivate::BuildSkeletonData(const TSharedRef<FUpdateContextPrivate>& OperationData, USkeletalMesh& SkeletalMesh, const FMutableRefSkeletalMeshData& RefSkeletalMeshData, UCustomizableObject& CustomizableObject, int32 ComponentIndex)
 {
-	MUTABLE_CPUPROFILER_SCOPE(UCustomizableInstancePrivateData::BuildSkeletonData);
+	MUTABLE_CPUPROFILER_SCOPE(UCustomizableInstancePrivate::BuildSkeletonData);
 
 	const TObjectPtr<USkeleton> Skeleton = MergeSkeletons(CustomizableObject, RefSkeletalMeshData, ComponentIndex);
 	if (!Skeleton)
@@ -3500,7 +3496,7 @@ bool UCustomizableInstancePrivate::BuildSkeletonData(const TSharedRef<FUpdateCon
 void UCustomizableInstancePrivate::BuildMeshSockets(const TSharedRef<FUpdateContextPrivate>& OperationData, USkeletalMesh* SkeletalMesh, const FMutableRefSkeletalMeshData* RefSkeletalMeshData, UCustomizableObjectInstance* Public, mu::MeshPtrConst MutableMesh)
 {
 	// Build mesh sockets.
-	MUTABLE_CPUPROFILER_SCOPE(UCustomizableInstancePrivateData::BuildMeshSockets);
+	MUTABLE_CPUPROFILER_SCOPE(UCustomizableInstancePrivate::BuildMeshSockets);
 
 	check(SkeletalMesh);
 	check(RefSkeletalMeshData);
@@ -3603,7 +3599,7 @@ void UCustomizableInstancePrivate::BuildMeshSockets(const TSharedRef<FUpdateCont
 
 void UCustomizableInstancePrivate::BuildOrCopyElementData(const TSharedRef<FUpdateContextPrivate>& OperationData, USkeletalMesh* SkeletalMesh, UCustomizableObjectInstance* CustomizableObjectInstance, int32 ComponentIndex)
 {
-	MUTABLE_CPUPROFILER_SCOPE(UCustomizableInstancePrivateData::BuildOrCopyElementData);
+	MUTABLE_CPUPROFILER_SCOPE(UCustomizableInstancePrivate::BuildOrCopyElementData);
 
 	for (int32 LODIndex = OperationData->CurrentMaxLOD; LODIndex >= FirstLODAvailable; --LODIndex)
 	{
@@ -3631,7 +3627,7 @@ void UCustomizableInstancePrivate::BuildOrCopyElementData(const TSharedRef<FUpda
 
 void UCustomizableInstancePrivate::BuildOrCopyMorphTargetsData(const TSharedRef<FUpdateContextPrivate>& OperationData, USkeletalMesh* SkeletalMesh, const USkeletalMesh* LastUpdateSkeletalMesh, UCustomizableObjectInstance* CustomizableObjectInstance, int32 ComponentIndex)
 {
-	MUTABLE_CPUPROFILER_SCOPE(UCustomizableInstancePrivateData::BuildOrCopyMorphTargetsData);
+	MUTABLE_CPUPROFILER_SCOPE(UCustomizableInstancePrivate::BuildOrCopyMorphTargetsData);
 
 	if (!SkeletalMesh || !CustomizableObjectInstance)
 	{
@@ -3783,7 +3779,7 @@ namespace
 
 void UCustomizableInstancePrivate::BuildOrCopyClothingData(const TSharedRef<FUpdateContextPrivate>&OperationData, USkeletalMesh * SkeletalMesh, const USkeletalMesh* LastUpdateSkeletalMesh, UCustomizableObjectInstance * CustomizableObjectInstance, int32 ComponentIndex)
 {
-	MUTABLE_CPUPROFILER_SCOPE(UCustomizableInstancePrivateData::BuildOrCopyClothingData);
+	MUTABLE_CPUPROFILER_SCOPE(UCustomizableInstancePrivate::BuildOrCopyClothingData);
 
 	if (!SkeletalMesh)
 	{
@@ -4649,7 +4645,7 @@ void UCustomizableInstancePrivate::BuildOrCopyClothingData(const TSharedRef<FUpd
 
 bool UCustomizableInstancePrivate::BuildOrCopyRenderData(const TSharedRef<FUpdateContextPrivate>& OperationData, USkeletalMesh* SkeletalMesh, const USkeletalMesh* LastUpdateSkeletalMesh, UCustomizableObjectInstance* Public, int32 ComponentIndex)
 {
-	MUTABLE_CPUPROFILER_SCOPE(UCustomizableInstancePrivateData::BuildOrCopyRenderData);
+	MUTABLE_CPUPROFILER_SCOPE(UCustomizableInstancePrivate::BuildOrCopyRenderData);
 
 	UCustomizableObject* CustomizableObject = Public->GetCustomizableObject();
 	check(CustomizableObject);
@@ -4819,7 +4815,7 @@ FAutoConsoleVariableRef CVarMutableHighPriorityLoading(
 
 FGraphEventRef UCustomizableInstancePrivate::LoadAdditionalAssetsAsync(const TSharedRef<FUpdateContextPrivate>& OperationData, UCustomizableObjectInstance* Public, FStreamableManager& StreamableManager)
 {
-	MUTABLE_CPUPROFILER_SCOPE(UCustomizableInstancePrivateData::LoadAdditionalAssetsAsync);
+	MUTABLE_CPUPROFILER_SCOPE(UCustomizableInstancePrivate::LoadAdditionalAssetsAsync);
 
 	UCustomizableObject* CustomizableObject = Public->GetCustomizableObject();
 
@@ -5103,7 +5099,7 @@ const TArray<TObjectPtr<UMaterialInterface>>* UCustomizableObjectInstance::GetOv
 
 void UCustomizableInstancePrivate::AdditionalAssetsAsyncLoaded(UCustomizableObjectInstance* Public)
 {
-	MUTABLE_CPUPROFILER_SCOPE(UCustomizableInstancePrivateData::AdditionalAssetsAsyncLoaded);
+	MUTABLE_CPUPROFILER_SCOPE(UCustomizableInstancePrivate::AdditionalAssetsAsyncLoaded);
 
 	UCustomizableObject* CustomizableObject = Public->GetCustomizableObject();
 
@@ -5378,7 +5374,7 @@ FAutoConsoleVariableRef CVarMutableReuseMaterialInstances(
 
 void UCustomizableInstancePrivate::BuildMaterials(const TSharedRef<FUpdateContextPrivate>& OperationData, UCustomizableObjectInstance* Public)
 {
-	MUTABLE_CPUPROFILER_SCOPE(UCustomizableInstancePrivateData::BuildMaterials)
+	MUTABLE_CPUPROFILER_SCOPE(UCustomizableInstancePrivate::BuildMaterials)
 
 	UCustomizableObject* CustomizableObject = Public->GetCustomizableObject();
 
@@ -5413,13 +5409,13 @@ void UCustomizableInstancePrivate::BuildMaterials(const TSharedRef<FUpdateContex
 
 	for (int32 ComponentIndex = 0; ComponentIndex < OperationData->NumComponents; ++ComponentIndex)
 	{
-		if (!Public->SkeletalMeshes.IsValidIndex(ComponentIndex))
+		if (!SkeletalMeshes.IsValidIndex(ComponentIndex))
 		{
 			// This could happen if there are no meshes at all.
 			continue;
 		}
 
-		USkeletalMesh* SkeletalMesh = Public->SkeletalMeshes[ComponentIndex];
+		USkeletalMesh* SkeletalMesh = SkeletalMeshes[ComponentIndex];
 
 		if (!SkeletalMesh)
 		{
@@ -6028,7 +6024,7 @@ void UCustomizableInstancePrivate::BuildMaterials(const TSharedRef<FUpdateContex
 			}
 
 			USkeletalMeshComponent* AttachedParent = CustomizableObjectInstanceUsage->GetAttachParent();
-			if (!AttachedParent || AttachedParent->GetSkeletalMeshAsset() != Public->SkeletalMeshes[ComponentIndex])
+			if (!AttachedParent || AttachedParent->GetSkeletalMeshAsset() != SkeletalMeshes[ComponentIndex])
 			{
 				continue;
 			}
@@ -6297,15 +6293,17 @@ const TArray<uint16>& UCustomizableObjectInstance::GetRequestedLODsPerComponent(
 
 USkeletalMesh* UCustomizableObjectInstance::GetSkeletalMesh(int32 ComponentIndex) const
 {
-	return SkeletalMeshes.IsValidIndex(ComponentIndex) ? SkeletalMeshes[ComponentIndex] : nullptr;
+	check(PrivateData);
+	return PrivateData->SkeletalMeshes.IsValidIndex(ComponentIndex) ? PrivateData->SkeletalMeshes[ComponentIndex] : nullptr;
 }
 
 
 bool UCustomizableObjectInstance::HasAnySkeletalMesh() const
 {
-	for (int32 Index = 0; Index < SkeletalMeshes.Num(); ++Index)
+	check(PrivateData);
+	for (int32 Index = 0; Index < PrivateData->SkeletalMeshes.Num(); ++Index)
 	{
-		if (SkeletalMeshes[Index])
+		if (PrivateData->SkeletalMeshes[Index])
 		{
 			return true;
 		}
@@ -6653,199 +6651,219 @@ TSet<UAssetUserData*> UCustomizableObjectInstance::GetMergedAssetUserData(int32
 
 #if WITH_EDITORONLY_DATA
 
-void UCustomizableInstancePrivate::RegenerateImportedModel(USkeletalMesh* SkeletalMesh)
+void UCustomizableInstancePrivate::RegenerateImportedModels()
 {
-	FSkeletalMeshRenderData* SkelResource = SkeletalMesh->GetResourceForRendering();
-	if (!SkelResource)
-	{
-		return;
-	}
+	MUTABLE_CPUPROFILER_SCOPE(RegenerateImportedModels);
 
-	for (UClothingAssetBase* ClothingAssetBase : SkeletalMesh->GetMeshClothingAssets())
+	for (USkeletalMesh* SkeletalMesh : SkeletalMeshes)
 	{
-		if (!ClothingAssetBase)
-		{
-			continue;
-		}
-
-		UClothingAssetCommon* ClothAsset = Cast<UClothingAssetCommon>(ClothingAssetBase);
-
-		if (!ClothAsset)
+		if (!SkeletalMesh)
 		{
 			continue;
 		}
 
-		if (!ClothAsset->LodData.Num())
+		FSkeletalMeshRenderData* SkelResource = SkeletalMesh->GetResourceForRendering();
+		if (!SkelResource || SkelResource->IsInitialized())
 		{
 			continue;
 		}
 
-		for (FClothLODDataCommon& ClothLodData : ClothAsset->LodData)
+		for (UClothingAssetBase* ClothingAssetBase : SkeletalMesh->GetMeshClothingAssets())
 		{
-			ClothLodData.PointWeightMaps.Empty(16);
-			for (TPair<uint32, FPointWeightMap>& WeightMap : ClothLodData.PhysicalMeshData.WeightMaps)
+			if (!ClothingAssetBase)
 			{
-				if (WeightMap.Value.Num())
-				{
-					FPointWeightMap& PointWeightMap = ClothLodData.PointWeightMaps.AddDefaulted_GetRef();
-					PointWeightMap.Initialize(WeightMap.Value, WeightMap.Key);
-				}
+				continue;
 			}
-		}
-	}
-
-	FSkeletalMeshModel* ImportedModel = SkeletalMesh->GetImportedModel();
-	ImportedModel->bGuidIsHash = false;
-	ImportedModel->SkeletalMeshModelGUID = FGuid::NewGuid();
-
-	ImportedModel->LODModels.Empty();
 
-	int32 OriginalIndex = 0;
-	for (int32 LODIndex = 0; LODIndex < SkelResource->LODRenderData.Num(); ++LODIndex)
-	{
-		ImportedModel->LODModels.Add(new FSkeletalMeshLODModel());
+			UClothingAssetCommon* ClothAsset = Cast<UClothingAssetCommon>(ClothingAssetBase);
 
-		FSkeletalMeshLODRenderData& LODModel = SkelResource->LODRenderData[LODIndex];
-		int32 CurrentSectionInitialVertex = 0;
+			if (!ClothAsset)
+			{
+				continue;
+			}
 
-		ImportedModel->LODModels[LODIndex].ActiveBoneIndices = LODModel.ActiveBoneIndices;
-		ImportedModel->LODModels[LODIndex].NumTexCoords = LODModel.GetNumTexCoords();
-		ImportedModel->LODModels[LODIndex].RequiredBones = LODModel.RequiredBones;
-		ImportedModel->LODModels[LODIndex].NumVertices = LODModel.GetNumVertices();
+			if (!ClothAsset->LodData.Num())
+			{
+				continue;
+			}
 
-		// Indices
-		if (LODModel.MultiSizeIndexContainer.IsIndexBufferValid())
-		{
-			const int32 NumIndices = LODModel.MultiSizeIndexContainer.GetIndexBuffer()->Num();
-			ImportedModel->LODModels[LODIndex].IndexBuffer.SetNum(NumIndices);
-			for (int32 Index = 0; Index < NumIndices; ++Index)
+			for (FClothLODDataCommon& ClothLodData : ClothAsset->LodData)
 			{
-				ImportedModel->LODModels[LODIndex].IndexBuffer[Index] = LODModel.MultiSizeIndexContainer.GetIndexBuffer()->Get(Index);
+				ClothLodData.PointWeightMaps.Empty(16);
+				for (TPair<uint32, FPointWeightMap>& WeightMap : ClothLodData.PhysicalMeshData.WeightMaps)
+				{
+					if (WeightMap.Value.Num())
+					{
+						FPointWeightMap& PointWeightMap = ClothLodData.PointWeightMaps.AddDefaulted_GetRef();
+						PointWeightMap.Initialize(WeightMap.Value, WeightMap.Key);
+					}
+				}
 			}
 		}
 
-		ImportedModel->LODModels[LODIndex].Sections.SetNum(LODModel.RenderSections.Num());
+		FSkeletalMeshModel* ImportedModel = SkeletalMesh->GetImportedModel();
+		ImportedModel->bGuidIsHash = false;
+		ImportedModel->SkeletalMeshModelGUID = FGuid::NewGuid();
+
+		ImportedModel->LODModels.Empty();
 
-		for (int SectionIndex = 0; SectionIndex < LODModel.RenderSections.Num(); ++SectionIndex)
+		int32 OriginalIndex = 0;
+		for (int32 LODIndex = 0; LODIndex < SkelResource->LODRenderData.Num(); ++LODIndex)
 		{
-			const FSkelMeshRenderSection& RenderSection = LODModel.RenderSections[SectionIndex];
-			FSkelMeshSection& ImportedSection = ImportedModel->LODModels[LODIndex].Sections[SectionIndex];
+			ImportedModel->LODModels.Add(new FSkeletalMeshLODModel());
 
-			ImportedSection.CorrespondClothAssetIndex = RenderSection.CorrespondClothAssetIndex;
-			ImportedSection.ClothingData = RenderSection.ClothingData;
+			FSkeletalMeshLODRenderData& LODModel = SkelResource->LODRenderData[LODIndex];
+			int32 CurrentSectionInitialVertex = 0;
 
-			if (RenderSection.ClothMappingDataLODs.Num())
+			ImportedModel->LODModels[LODIndex].ActiveBoneIndices = LODModel.ActiveBoneIndices;
+			ImportedModel->LODModels[LODIndex].NumTexCoords = LODModel.GetNumTexCoords();
+			ImportedModel->LODModels[LODIndex].RequiredBones = LODModel.RequiredBones;
+			ImportedModel->LODModels[LODIndex].NumVertices = LODModel.GetNumVertices();
+
+			// Indices
+			if (LODModel.MultiSizeIndexContainer.IsIndexBufferValid())
 			{
-				ImportedSection.ClothMappingDataLODs.SetNum(1);
-				ImportedSection.ClothMappingDataLODs[0] = RenderSection.ClothMappingDataLODs[0];
+				const int32 NumIndices = LODModel.MultiSizeIndexContainer.GetIndexBuffer()->Num();
+				ImportedModel->LODModels[LODIndex].IndexBuffer.SetNum(NumIndices);
+				for (int32 Index = 0; Index < NumIndices; ++Index)
+				{
+					ImportedModel->LODModels[LODIndex].IndexBuffer[Index] = LODModel.MultiSizeIndexContainer.GetIndexBuffer()->Get(Index);
+				}
 			}
 
-			// Vertices
-			ImportedSection.NumVertices = RenderSection.NumVertices;
-			ImportedSection.SoftVertices.Empty(RenderSection.NumVertices);
-			ImportedSection.SoftVertices.AddUninitialized(RenderSection.NumVertices);
-			ImportedSection.bUse16BitBoneIndex = LODModel.DoesVertexBufferUse16BitBoneIndex();
+			ImportedModel->LODModels[LODIndex].Sections.SetNum(LODModel.RenderSections.Num());
 
-			for (uint32 i = 0; i < RenderSection.NumVertices; ++i)
+			for (int SectionIndex = 0; SectionIndex < LODModel.RenderSections.Num(); ++SectionIndex)
 			{
-				const FPositionVertex* PosPtr = static_cast<const FPositionVertex*>(LODModel.StaticVertexBuffers.PositionVertexBuffer.GetVertexData());
-				PosPtr += (CurrentSectionInitialVertex + i);
-
-				check(!LODModel.StaticVertexBuffers.StaticMeshVertexBuffer.GetUseHighPrecisionTangentBasis());
-				const FPackedNormal* TangentPtr = static_cast<const FPackedNormal*>(LODModel.StaticVertexBuffers.StaticMeshVertexBuffer.GetTangentData());
-				TangentPtr += ((CurrentSectionInitialVertex + i) * 2);
-
-				check(LODModel.StaticVertexBuffers.StaticMeshVertexBuffer.GetUseFullPrecisionUVs());
+				const FSkelMeshRenderSection& RenderSection = LODModel.RenderSections[SectionIndex];
+				FSkelMeshSection& ImportedSection = ImportedModel->LODModels[LODIndex].Sections[SectionIndex];
 
-				using UVsVectorType = typename TDecay<decltype(DeclVal<FSoftSkinVertex>().UVs[0])>::Type;
+				ImportedSection.CorrespondClothAssetIndex = RenderSection.CorrespondClothAssetIndex;
+				ImportedSection.ClothingData = RenderSection.ClothingData;
 
-				const UVsVectorType* TexCoordPosPtr = static_cast<const UVsVectorType*>(LODModel.StaticVertexBuffers.StaticMeshVertexBuffer.GetTexCoordData());
-				const uint32 NumTexCoords = LODModel.StaticVertexBuffers.StaticMeshVertexBuffer.GetNumTexCoords();
-				TexCoordPosPtr += ((CurrentSectionInitialVertex + i) * NumTexCoords);
-
-				FSoftSkinVertex& Vertex = ImportedSection.SoftVertices[i];
-				for (int32 j = 0; j < RenderSection.MaxBoneInfluences; ++j)
+				if (RenderSection.ClothMappingDataLODs.Num())
 				{
-					Vertex.InfluenceBones[j] = LODModel.SkinWeightVertexBuffer.GetBoneIndex(CurrentSectionInitialVertex + i, j);
-					Vertex.InfluenceWeights[j] = LODModel.SkinWeightVertexBuffer.GetBoneWeight(CurrentSectionInitialVertex + i, j);
+					ImportedSection.ClothMappingDataLODs.SetNum(1);
+					ImportedSection.ClothMappingDataLODs[0] = RenderSection.ClothMappingDataLODs[0];
 				}
 
-				for (int32 j = RenderSection.MaxBoneInfluences; j < MAX_TOTAL_INFLUENCES; ++j)
+				// Vertices
+				ImportedSection.NumVertices = RenderSection.NumVertices;
+				ImportedSection.SoftVertices.Empty(RenderSection.NumVertices);
+				ImportedSection.SoftVertices.AddUninitialized(RenderSection.NumVertices);
+				ImportedSection.bUse16BitBoneIndex = LODModel.DoesVertexBufferUse16BitBoneIndex();
+
+				for (uint32 i = 0; i < RenderSection.NumVertices; ++i)
 				{
-					Vertex.InfluenceBones[j] = 0;
-					Vertex.InfluenceWeights[j] = 0;
-				}
+					const FPositionVertex* PosPtr = static_cast<const FPositionVertex*>(LODModel.StaticVertexBuffers.PositionVertexBuffer.GetVertexData());
+					PosPtr += (CurrentSectionInitialVertex + i);
 
+					check(!LODModel.StaticVertexBuffers.StaticMeshVertexBuffer.GetUseHighPrecisionTangentBasis());
+					const FPackedNormal* TangentPtr = static_cast<const FPackedNormal*>(LODModel.StaticVertexBuffers.StaticMeshVertexBuffer.GetTangentData());
+					TangentPtr += ((CurrentSectionInitialVertex + i) * 2);
 
-				Vertex.Color = FColor::White;
+					check(LODModel.StaticVertexBuffers.StaticMeshVertexBuffer.GetUseFullPrecisionUVs());
 
-				Vertex.Position = PosPtr->Position;
+					using UVsVectorType = typename TDecay<decltype(DeclVal<FSoftSkinVertex>().UVs[0])>::Type;
 
-				Vertex.TangentX = TangentPtr[0].ToFVector3f();
-				Vertex.TangentZ = TangentPtr[1].ToFVector3f();
-				float TangentSign = TangentPtr[1].Vector.W == 0 ? -1.f : 1.f;
-				Vertex.TangentY = FVector3f::CrossProduct(Vertex.TangentZ, Vertex.TangentX) * TangentSign;
+					const UVsVectorType* TexCoordPosPtr = static_cast<const UVsVectorType*>(LODModel.StaticVertexBuffers.StaticMeshVertexBuffer.GetTexCoordData());
+					const uint32 NumTexCoords = LODModel.StaticVertexBuffers.StaticMeshVertexBuffer.GetNumTexCoords();
+					TexCoordPosPtr += ((CurrentSectionInitialVertex + i) * NumTexCoords);
 
-				Vertex.UVs[0] = TexCoordPosPtr[0];
-				Vertex.UVs[1] = NumTexCoords > 1 ? TexCoordPosPtr[1] : UVsVectorType::ZeroVector;
-				Vertex.UVs[2] = NumTexCoords > 2 ? TexCoordPosPtr[2] : UVsVectorType::ZeroVector;
-				Vertex.UVs[3] = NumTexCoords > 3 ? TexCoordPosPtr[3] : UVsVectorType::ZeroVector;
-			}
+					FSoftSkinVertex& Vertex = ImportedSection.SoftVertices[i];
+					for (int32 j = 0; j < RenderSection.MaxBoneInfluences; ++j)
+					{
+						Vertex.InfluenceBones[j] = LODModel.SkinWeightVertexBuffer.GetBoneIndex(CurrentSectionInitialVertex + i, j);
+						Vertex.InfluenceWeights[j] = LODModel.SkinWeightVertexBuffer.GetBoneWeight(CurrentSectionInitialVertex + i, j);
+					}
+
+					for (int32 j = RenderSection.MaxBoneInfluences; j < MAX_TOTAL_INFLUENCES; ++j)
+					{
+						Vertex.InfluenceBones[j] = 0;
+						Vertex.InfluenceWeights[j] = 0;
+					}
 
-			CurrentSectionInitialVertex += RenderSection.NumVertices;
 
-			// Triangles
-			ImportedSection.NumTriangles = RenderSection.NumTriangles;
-			ImportedSection.BaseIndex = RenderSection.BaseIndex;
-			ImportedSection.BaseVertexIndex = RenderSection.BaseVertexIndex;
-			ImportedSection.BoneMap = RenderSection.BoneMap;
+					Vertex.Color = FColor::White;
 
-			const TArray<int32>& LODMaterialMap = SkeletalMesh->GetLODInfoArray()[LODIndex].LODMaterialMap;
+					Vertex.Position = PosPtr->Position;
 
-			if (LODMaterialMap.IsValidIndex(RenderSection.MaterialIndex))
-			{
-				ImportedSection.MaterialIndex = LODMaterialMap[RenderSection.MaterialIndex];
-			}
-			else
-			{
-				// The material should have been in the LODMaterialMap
-				ensureMsgf(false, TEXT("Unexpected material index in UCustomizableInstancePrivate::RegenerateImportedModel"));
+					Vertex.TangentX = TangentPtr[0].ToFVector3f();
+					Vertex.TangentZ = TangentPtr[1].ToFVector3f();
+					float TangentSign = TangentPtr[1].Vector.W == 0 ? -1.f : 1.f;
+					Vertex.TangentY = FVector3f::CrossProduct(Vertex.TangentZ, Vertex.TangentX) * TangentSign;
+
+					Vertex.UVs[0] = TexCoordPosPtr[0];
+					Vertex.UVs[1] = NumTexCoords > 1 ? TexCoordPosPtr[1] : UVsVectorType::ZeroVector;
+					Vertex.UVs[2] = NumTexCoords > 2 ? TexCoordPosPtr[2] : UVsVectorType::ZeroVector;
+					Vertex.UVs[3] = NumTexCoords > 3 ? TexCoordPosPtr[3] : UVsVectorType::ZeroVector;
+				}
+
+				CurrentSectionInitialVertex += RenderSection.NumVertices;
+
+				// Triangles
+				ImportedSection.NumTriangles = RenderSection.NumTriangles;
+				ImportedSection.BaseIndex = RenderSection.BaseIndex;
+				ImportedSection.BaseVertexIndex = RenderSection.BaseVertexIndex;
+				ImportedSection.BoneMap = RenderSection.BoneMap;
 
-				// Fallback index, may shift materials around sections
-				if (SkeletalMesh->GetMaterials().IsValidIndex(RenderSection.MaterialIndex))
+				const TArray<int32>& LODMaterialMap = SkeletalMesh->GetLODInfoArray()[LODIndex].LODMaterialMap;
+
+				if (LODMaterialMap.IsValidIndex(RenderSection.MaterialIndex))
 				{
-					ImportedSection.MaterialIndex = RenderSection.MaterialIndex;
+					ImportedSection.MaterialIndex = LODMaterialMap[RenderSection.MaterialIndex];
 				}
 				else
 				{
-					ImportedSection.MaterialIndex = 0;
-				}
-			}
+					// The material should have been in the LODMaterialMap
+					ensureMsgf(false, TEXT("Unexpected material index in UCustomizableInstancePrivate::RegenerateImportedModel"));
+
+					// Fallback index, may shift materials around sections
+					if (SkeletalMesh->GetMaterials().IsValidIndex(RenderSection.MaterialIndex))
+					{
+						ImportedSection.MaterialIndex = LODMaterialMap[RenderSection.MaterialIndex];
+					}
+					else
+					{
+						// The material should have been in the LODMaterialMap
+						ensureMsgf(false, TEXT("Unexpected material index in UCustomizableInstancePrivate::RegenerateImportedModel"));
+
+						// Fallback index, may shift materials around sections
+						if (SkeletalMesh->GetMaterials().IsValidIndex(RenderSection.MaterialIndex))
+						{
+							ImportedSection.MaterialIndex = RenderSection.MaterialIndex;
+						}
+						else
+						{
+							ImportedSection.MaterialIndex = 0;
+						}
+					}
 
-			ImportedSection.MaxBoneInfluences = RenderSection.MaxBoneInfluences;
-			ImportedSection.OriginalDataSectionIndex = OriginalIndex++;
+					ImportedSection.MaxBoneInfluences = RenderSection.MaxBoneInfluences;
+					ImportedSection.OriginalDataSectionIndex = OriginalIndex++;
 
-			FSkelMeshSourceSectionUserData& SectionUserData = ImportedModel->LODModels[LODIndex].UserSectionsData.FindOrAdd(ImportedSection.OriginalDataSectionIndex);
+					FSkelMeshSourceSectionUserData& SectionUserData = ImportedModel->LODModels[LODIndex].UserSectionsData.FindOrAdd(ImportedSection.OriginalDataSectionIndex);
 
-			SectionUserData.CorrespondClothAssetIndex = RenderSection.CorrespondClothAssetIndex;
-			SectionUserData.ClothingData.AssetGuid = RenderSection.ClothingData.AssetGuid;
-			SectionUserData.ClothingData.AssetLodIndex = RenderSection.ClothingData.AssetLodIndex;
-			SectionUserData.bCastShadow = RenderSection.bCastShadow;
-		}
+					SectionUserData.CorrespondClothAssetIndex = RenderSection.CorrespondClothAssetIndex;
+					SectionUserData.ClothingData.AssetGuid = RenderSection.ClothingData.AssetGuid;
+					SectionUserData.ClothingData.AssetLodIndex = RenderSection.ClothingData.AssetLodIndex;
+					SectionUserData.bCastShadow = RenderSection.bCastShadow;
+				}
 
-		ImportedModel->LODModels[LODIndex].SyncronizeUserSectionsDataArray();
+				ImportedModel->LODModels[LODIndex].SyncronizeUserSectionsDataArray();
 
-		// DDC keys
-		const USkeletalMeshLODSettings* LODSettings = SkeletalMesh->GetLODSettings();
-		const bool bValidLODSettings = LODSettings && LODSettings->GetNumberOfSettings() > LODIndex;
-		const FSkeletalMeshLODGroupSettings* SkeletalMeshLODGroupSettings = bValidLODSettings ? &LODSettings->GetSettingsForLODLevel(LODIndex) : nullptr;
+				// DDC keys
+				const USkeletalMeshLODSettings* LODSettings = SkeletalMesh->GetLODSettings();
+				const bool bValidLODSettings = LODSettings && LODSettings->GetNumberOfSettings() > LODIndex;
+				const FSkeletalMeshLODGroupSettings* SkeletalMeshLODGroupSettings = bValidLODSettings ? &LODSettings->GetSettingsForLODLevel(LODIndex) : nullptr;
 
-		FSkeletalMeshLODInfo* LODInfo = SkeletalMesh->GetLODInfo(LODIndex);
-		LODInfo->BuildGUID = LODInfo->ComputeDeriveDataCacheKey(SkeletalMeshLODGroupSettings);
+				FSkeletalMeshLODInfo* LODInfo = SkeletalMesh->GetLODInfo(LODIndex);
+				LODInfo->BuildGUID = LODInfo->ComputeDeriveDataCacheKey(SkeletalMeshLODGroupSettings);
 
-		ImportedModel->LODModels[LODIndex].BuildStringID = ImportedModel->LODModels[LODIndex].GetLODModelDeriveDataKey();
+				ImportedModel->LODModels[LODIndex].BuildStringID = ImportedModel->LODModels[LODIndex].GetLODModelDeriveDataKey();
+			}
+		}
 	}
-
 }
 
 #endif
diff --git a/Engine/Plugins/Experimental/Mutable/Source/CustomizableObject/Private/MuCO/CustomizableObjectSystem.cpp b/Engine/Plugins/Experimental/Mutable/Source/CustomizableObject/Private/MuCO/CustomizableObjectSystem.cpp
--- a/Engine/Plugins/Experimental/Mutable/Source/CustomizableObject/Private/MuCO/CustomizableObjectSystem.cpp
+++ b/Engine/Plugins/Experimental/Mutable/Source/CustomizableObject/Private/MuCO/CustomizableObjectSystem.cpp
@@ -886,19 +886,19 @@ void UpdateSkeletalMesh(const TSharedRef<FUpdateContextPrivate>& Context)
 
 		if (IsValid(CustomizableObjectInstanceUsage) &&
 			(CustomizableObjectInstanceUsage->GetCustomizableObjectInstance() == CustomizableObjectInstance) &&
-			CustomizableObjectInstance->SkeletalMeshes.IsValidIndex(CustomizableObjectInstanceUsage->GetComponentIndex())
+			CustomizableObjectInstancePrivateData->SkeletalMeshes.IsValidIndex(CustomizableObjectInstanceUsage->GetComponentIndex())
 		   )
 		{
 			MUTABLE_CPUPROFILER_SCOPE(UpdateSkeletalMesh_SetSkeletalMesh);
 
 			const bool bIsCreatingSkeletalMesh = CustomizableObjectInstancePrivateData->HasCOInstanceFlags(CreatingSkeletalMesh); //TODO MTBL-391: Review
-			CustomizableObjectInstanceUsage->SetSkeletalMesh(CustomizableObjectInstance->SkeletalMeshes[CustomizableObjectInstanceUsage->GetComponentIndex()], false, bIsCreatingSkeletalMesh);
+			CustomizableObjectInstanceUsage->SetSkeletalMesh(CustomizableObjectInstancePrivateData->SkeletalMeshes[CustomizableObjectInstanceUsage->GetComponentIndex()], false, bIsCreatingSkeletalMesh);
 
 			if (CustomizableObjectInstancePrivateData->HasCOInstanceFlags(ReplacePhysicsAssets))
 			{
 				CustomizableObjectInstanceUsage->SetPhysicsAsset(
-					CustomizableObjectInstance->SkeletalMeshes[CustomizableObjectInstanceUsage->GetComponentIndex()] ?
-					CustomizableObjectInstance->SkeletalMeshes[CustomizableObjectInstanceUsage->GetComponentIndex()]->GetPhysicsAsset() : nullptr);
+					CustomizableObjectInstancePrivateData->SkeletalMeshes[CustomizableObjectInstanceUsage->GetComponentIndex()] ?
+					CustomizableObjectInstancePrivateData->SkeletalMeshes[CustomizableObjectInstanceUsage->GetComponentIndex()]->GetPhysicsAsset() : nullptr);
 			}
 		}
 	}
@@ -2293,17 +2293,11 @@ namespace impl
 				// This used to be CustomizableObjectInstance::UpdateSkeletalMesh_PostBeginUpdate2
 				{
 					MUTABLE_CPUPROFILER_SCOPE(UpdateSkeletalMesh_PostBeginUpdate2);
-
-					for (TObjectPtr<USkeletalMesh>& SkeletalMesh : CustomizableObjectInstance->SkeletalMeshes)
-					{
-						if (SkeletalMesh && SkeletalMesh->GetResourceForRendering() && !SkeletalMesh->GetResourceForRendering()->IsInitialized())
-						{
-#if WITH_EDITOR
-							UCustomizableInstancePrivate::RegenerateImportedModel(SkeletalMesh);
+					
+#if WITH_EDITORONLY_DATA
+					CustomizableInstancePrivateData->RegenerateImportedModels();
 #endif
-							CustomizableInstancePrivateData->PostEditChangePropertyWithoutEditor(SkeletalMesh);
-						}
-					}
+					CustomizableInstancePrivateData->PostEditChangePropertyWithoutEditor();
 				}
 			}
 		} // if (!bInstanceValid)
@@ -3749,7 +3743,7 @@ void UCustomizableObjectSystemPrivate::UpdateStats()
 			continue;
 		}
 
-		NumSkeletalMeshes += Instance->SkeletalMeshes.Num();
+		NumSkeletalMeshes += Instance->GetPrivate()->SkeletalMeshes.Num();
 	}
 }
 
diff --git a/Engine/Plugins/Experimental/Mutable/Source/CustomizableObject/Private/MuCO/CustomizableObjectSystemPrivate.h b/Engine/Plugins/Experimental/Mutable/Source/CustomizableObject/Private/MuCO/CustomizableObjectSystemPrivate.h
--- a/Engine/Plugins/Experimental/Mutable/Source/CustomizableObject/Private/MuCO/CustomizableObjectSystemPrivate.h
+++ b/Engine/Plugins/Experimental/Mutable/Source/CustomizableObject/Private/MuCO/CustomizableObjectSystemPrivate.h
@@ -198,10 +198,10 @@ struct FGeneratedTexture
 
 	FMutableImageCacheKey Key;
 
-	UPROPERTY(Category = CustomizableObjectInstance, VisibleAnywhere)
+	UPROPERTY(Category = NoCategory, VisibleAnywhere)
 	FString Name;
 
-	UPROPERTY(Category = CustomizableObjectInstance, VisibleAnywhere)
+	UPROPERTY(Category = NoCategory, VisibleAnywhere)
 	TObjectPtr<UTexture> Texture = nullptr;
 
 	bool operator==(const FGeneratedTexture& Other) const = default;
@@ -213,10 +213,10 @@ struct FGeneratedMaterial
 {
 	GENERATED_USTRUCT_BODY();
 
-	UPROPERTY()
+	UPROPERTY(Category = NoCategory, VisibleAnywhere)
 	TObjectPtr<UMaterialInterface> MaterialInterface;
 
-	UPROPERTY(Category = CustomizedMaterial, VisibleAnywhere)
+	UPROPERTY(Category = NoCategory, VisibleAnywhere)
 	TArray< FGeneratedTexture > Textures;
 
 	// Surface or SharedSurface Id
diff --git a/Engine/Plugins/Experimental/Mutable/Source/CustomizableObject/Private/MuCO/LogBenchmarkUtil.cpp b/Engine/Plugins/Experimental/Mutable/Source/CustomizableObject/Private/MuCO/LogBenchmarkUtil.cpp
--- a/Engine/Plugins/Experimental/Mutable/Source/CustomizableObject/Private/MuCO/LogBenchmarkUtil.cpp
+++ b/Engine/Plugins/Experimental/Mutable/Source/CustomizableObject/Private/MuCO/LogBenchmarkUtil.cpp
@@ -99,9 +99,9 @@ void FLogBenchmarkUtil::GetInstancesStats(int32& OutNumInstances, int32& OutNumB
 			++OutNumBuiltInstances;
 		}
 		
-		for (int32 ComponentIndex = 0; ComponentIndex < Instance->SkeletalMeshes.Num(); ++ComponentIndex)
+		for (int32 ComponentIndex = 0; ComponentIndex < Instance->GetNumComponents(); ++ComponentIndex)
 		{
-			if (Instance->SkeletalMeshes[ComponentIndex] && Instance->SkeletalMeshes[ComponentIndex]->GetResourceForRendering())
+			if (Instance->GetSkeletalMesh(ComponentIndex))
 			{
 				++OutNumAllocatedSkeletalMeshes;
 			}
diff --git a/Engine/Plugins/Experimental/Mutable/Source/CustomizableObject/Private/MuCO/LogInformationUtil.cpp b/Engine/Plugins/Experimental/Mutable/Source/CustomizableObject/Private/MuCO/LogInformationUtil.cpp
--- a/Engine/Plugins/Experimental/Mutable/Source/CustomizableObject/Private/MuCO/LogInformationUtil.cpp
+++ b/Engine/Plugins/Experimental/Mutable/Source/CustomizableObject/Private/MuCO/LogInformationUtil.cpp
@@ -170,9 +170,9 @@ void LogInformationUtil::LogShowInstanceDataFull(const UCustomizableObjectInstan
 	MessageChunk += FString::Printf(TEXT("        bIsCreatingSkeletalMesh = %d\n"), InstanceFlags & ECOInstanceFlags::CreatingSkeletalMesh);
 	LogData += MessageChunk;
 
-	for (int32 ComponentIndex = 0; ComponentIndex < CustomizableObjectInstance->SkeletalMeshes.Num(); ++ComponentIndex)
+	for (int32 ComponentIndex = 0; ComponentIndex < CustomizableObjectInstance->GetNumComponents(); ++ComponentIndex)
 	{
-		if ((CustomizableObjectInstance->GetSkeletalMesh(ComponentIndex) != nullptr) && (CustomizableObjectInstance->GetSkeletalMesh(ComponentIndex)->GetResourceForRendering()))
+		if (CustomizableObjectInstance->GetSkeletalMesh(ComponentIndex))
 		{
 			if (CustomizableObjectInstance->GetCurrentMinLOD() < 1)
 			{
diff --git a/Engine/Plugins/Experimental/Mutable/Source/CustomizableObjectEditor/Private/MuCOE/CustomizableInstanceDetails.cpp b/Engine/Plugins/Experimental/Mutable/Source/CustomizableObjectEditor/Private/MuCOE/CustomizableInstanceDetails.cpp
--- a/Engine/Plugins/Experimental/Mutable/Source/CustomizableObjectEditor/Private/MuCOE/CustomizableInstanceDetails.cpp
+++ b/Engine/Plugins/Experimental/Mutable/Source/CustomizableObjectEditor/Private/MuCOE/CustomizableInstanceDetails.cpp
@@ -6,6 +6,7 @@
 #include "DetailWidgetRow.h"
 #include "IDetailsView.h"
 #include "MuCO/CustomizableObjectInstance.h"
+#include "MuCO/CustomizableInstancePrivateData.h"
 #include "MuCOE/SCustomizableInstanceProperties.h"
 
 class UObject;
@@ -28,8 +29,18 @@ void FCustomizableInstanceDetails::CustomizeDetails(const TSharedPtr<IDetailLayo
 	check(CustomInstance);
 	
 	LayoutBuilder = DetailBuilder;
+
+	IDetailCategoryBuilder& ResourcesCategory = DetailBuilder->EditCategory("Generated Resources");
+	
+	TArray<UObject*> Private;
+	Private.Add(CustomInstance->GetPrivate(	));
+	
+	FAddPropertyParams PrivatePropertyParams;
+	PrivatePropertyParams.HideRootObjectNode(true);
 	
-	IDetailCategoryBuilder& MainCategory = DetailBuilder->EditCategory( "Customizable Instance" );
+	IDetailPropertyRow* PrivateDataRow = ResourcesCategory.AddExternalObjects(Private, EPropertyLocation::Default, PrivatePropertyParams);
+
+	IDetailCategoryBuilder& MainCategory = DetailBuilder->EditCategory( "InstanceParameters" );
 
 	MainCategory.AddCustomRow( LOCTEXT("CustomizableInstanceDetails", "Instance Parameters") )
 	[
@@ -37,6 +48,8 @@ void FCustomizableInstanceDetails::CustomizeDetails(const TSharedPtr<IDetailLayo
 			.CustomInstance(CustomInstance)
 			.InstanceDetails(SharedThis(this))
 	];
+
+	DetailBuilder->EditCategory("TextureParameter");
 }
 
 
diff --git a/Engine/Plugins/Experimental/Mutable/Source/CustomizableObjectEditor/Private/MuCOE/CustomizableObjectEditorViewportClient.cpp b/Engine/Plugins/Experimental/Mutable/Source/CustomizableObjectEditor/Private/MuCOE/CustomizableObjectEditorViewportClient.cpp
--- a/Engine/Plugins/Experimental/Mutable/Source/CustomizableObjectEditor/Private/MuCOE/CustomizableObjectEditorViewportClient.cpp
+++ b/Engine/Plugins/Experimental/Mutable/Source/CustomizableObjectEditor/Private/MuCOE/CustomizableObjectEditorViewportClient.cpp
@@ -1658,18 +1658,7 @@ void FCustomizableObjectEditorViewportClient::BakeInstance(UCustomizableObjectIn
 		return;
 	}
 
-	bool bHasSkeletalMesh = false;
-	int32 NumComponents = Instance->SkeletalMeshes.Num();
-
-	for (int32 ComponentIndex = 0; ComponentIndex < NumComponents; ++ComponentIndex)
-	{
-		if (Instance->SkeletalMeshes.IsValidIndex(ComponentIndex) && Instance->SkeletalMeshes[ComponentIndex])
-		{
-			bHasSkeletalMesh = true;
-		}
-	}
-		
-	if (!bHasSkeletalMesh)
+	if (!Instance->HasAnySkeletalMesh())
 	{
 		return;
 	}
@@ -1689,7 +1678,6 @@ void FCustomizableObjectEditorViewportClient::BakeInstance(UCustomizableObjectIn
 		System->SetImagePixelFormatOverride(UnrealPixelFormatFunc);
 
 		BakeTempInstance = Instance->Clone();
-		BakeTempInstance->SkeletalMeshes.Empty();
 		BakeTempInstance->UpdatedNativeDelegate.AddSP(this, &FCustomizableObjectEditorViewportClient::BakeInstance);
 		BakeTempInstance->UpdateSkeletalMeshAsync(true, true);
 
@@ -1772,11 +1760,11 @@ void FCustomizableObjectEditorViewportClient::BakeInstance(UCustomizableObjectIn
 		{
 			TArray<UPackage*> PackagesToSave;
 
+			const int32 NumComponents = Instance->GetNumComponents();
 			for (int32 ComponentIndex = 0; ComponentIndex < NumComponents; ++ComponentIndex)
 			{
-				USkeletalMesh* Mesh = Instance->SkeletalMeshes.IsValidIndex(ComponentIndex) && Instance->SkeletalMeshes[ComponentIndex] ?
-					Cast<USkeletalMesh>(Instance->SkeletalMeshes[ComponentIndex]) : nullptr;
-				
+				USkeletalMesh* Mesh = Instance->GetSkeletalMesh(ComponentIndex);
+
 				if (!Mesh)
 				{
 					continue;
diff --git a/Engine/Plugins/Experimental/Mutable/Source/CustomizableObjectEditor/Private/MuCOE/SCustomizableObjectEditorPerformanceReport.cpp b/Engine/Plugins/Experimental/Mutable/Source/CustomizableObjectEditor/Private/MuCOE/SCustomizableObjectEditorPerformanceReport.cpp
--- a/Engine/Plugins/Experimental/Mutable/Source/CustomizableObjectEditor/Private/MuCOE/SCustomizableObjectEditorPerformanceReport.cpp
+++ b/Engine/Plugins/Experimental/Mutable/Source/CustomizableObjectEditor/Private/MuCOE/SCustomizableObjectEditorPerformanceReport.cpp
@@ -964,7 +964,7 @@ void SCustomizableObjecEditorPerformanceReport::UpdateSelectedInstanceData()
 	if (CurrentReportInstance && CurrentReportInstance->IsValidLowLevelFast())
 	{
 		// If we have at least one skeletal mesh in the instance, we know it's generated (at least once)
-		if (CurrentReportInstance->SkeletalMeshes.Num() > 0 && CurrentReportInstance->SkeletalMeshes[0])
+		if (CurrentReportInstance->HasAnySkeletalMesh())
 		{
 			PerformanceReportTextureAnalyzer->RefreshTextureAnalyzerTable(CurrentReportInstance);
 		}
diff --git a/Engine/Plugins/Experimental/Mutable/Source/CustomizableObjectEditor/Private/MuCOE/SCustomizableObjectEditorTextureAnalyzer.cpp b/Engine/Plugins/Experimental/Mutable/Source/CustomizableObjectEditor/Private/MuCOE/SCustomizableObjectEditorTextureAnalyzer.cpp
--- a/Engine/Plugins/Experimental/Mutable/Source/CustomizableObjectEditor/Private/MuCOE/SCustomizableObjectEditorTextureAnalyzer.cpp
+++ b/Engine/Plugins/Experimental/Mutable/Source/CustomizableObjectEditor/Private/MuCOE/SCustomizableObjectEditorTextureAnalyzer.cpp
@@ -129,7 +129,7 @@ void SCustomizableObjecEditorTextureAnalyzer::FillTextureAnalyzerTable(UCustomiz
 
 	if (PreviewInstance && PreviewInstance->HasAnySkeletalMesh())
 	{
-		for (int32 ComponentIndex = 0; ComponentIndex < PreviewInstance->SkeletalMeshes.Num(); ++ComponentIndex)
+		for (int32 ComponentIndex = 0; ComponentIndex < PreviewInstance->GetNumComponents(); ++ComponentIndex)
 		{
 			if (PreviewInstance->GetSkeletalMesh(ComponentIndex))
 			{
diff --git a/Engine/Plugins/Experimental/Mutable/Source/CustomizableObjectEditor/Private/MuCOE/StressTest.cpp b/Engine/Plugins/Experimental/Mutable/Source/CustomizableObjectEditor/Private/MuCOE/StressTest.cpp
--- a/Engine/Plugins/Experimental/Mutable/Source/CustomizableObjectEditor/Private/MuCOE/StressTest.cpp
+++ b/Engine/Plugins/Experimental/Mutable/Source/CustomizableObjectEditor/Private/MuCOE/StressTest.cpp
@@ -26,14 +26,11 @@ void ULiveInstance::AddInstanceInformation()
 {
 	if(!Instance->HasAnySkeletalMesh())
 	{
-		UE_LOG(LogMutable, Warning, TEXT("ERROR: no skeletal mesh in ULiveInstance::AddInstanceInformataion"));
+		UE_LOG(LogMutable, Warning, TEXT("ERROR: No skeletal meshes found in ULiveInstance::AddInstanceInformation"));
 		return;
 	}
 
-	FSkeletalMeshRenderData* RenderData = Instance->GetSkeletalMesh()->GetResourceForRendering();
-	TIndirectArray<FSkeletalMeshLODRenderData>* ArrayLODRenderData = &RenderData->LODRenderData;
-	int32 MaxIndexLOD = RenderData->LODRenderData.Num();
-
+	const int32 MaxIndexLOD = Instance->GetNumLODsAvailable();
 	for (int32 i = 0; i < MaxIndexLOD; ++i)
 	{
 		TMap<FString, MaterialBriefInfo> MapMaterialInfo;
@@ -42,20 +39,29 @@ void ULiveInstance::AddInstanceInformation()
 
 		TMap<FString, MeasuredData> MapMeasuredData = ULiveInstance::InitInstanceInformation();
 
-		for (int32 c = 0; c < Instance->SkeletalMeshes.Num(); ++c)
+		const int32 NumComponents = Instance->GetNumComponents();
+		for (int32 c = 0; c < NumComponents; ++c)
 		{
-			RenderData = Instance->GetSkeletalMesh(c)->GetResourceForRendering();
+			USkeletalMesh* SkeletalMesh = Instance->GetSkeletalMesh(c);
+			if (!SkeletalMesh)
+			{
+				continue;
+			}
+
+			FSkeletalMeshRenderData* RenderData = SkeletalMesh->GetResourceForRendering();
+			TIndirectArray<FSkeletalMeshLODRenderData>* ArrayLODRenderData = &RenderData->LODRenderData;
+
 			ArrayLODRenderData = &RenderData->LODRenderData;
-			MaxIndexLOD = RenderData->LODRenderData.Num();
-			CurrentLODNumFaces = ((*ArrayLODRenderData)[i]).GetTotalFaces();
-			const TArray<struct FSkeletalMeshLODInfo>& LODInfoArray = Instance->GetSkeletalMesh(c)->GetLODInfoArray();
+			const TArray<struct FSkeletalMeshLODInfo>& LODInfoArray = SkeletalMesh->GetLODInfoArray();
 			if (LODInfoArray.IsValidIndex(i))
 			{
 				const FSkeletalMeshLODInfo& LODInfo = LODInfoArray[i];
 
+				CurrentLODNumFaces += RenderData->LODRenderData[i].GetTotalFaces();
+
 				MapMeasuredData["Number of faces"].AddMeasureData((*ArrayLODRenderData)[i].GetTotalFaces());
 				MapMeasuredData["Number of texture coordinates"].AddMeasureData((*ArrayLODRenderData)[i].GetNumTexCoords());
-				const TArray<FSkeletalMaterial>& Materials = Instance->GetSkeletalMesh(c)->GetMaterials();
+				const TArray<FSkeletalMaterial>& Materials = SkeletalMesh->GetMaterials();
 				MapMeasuredData["Number of materials"].AddMeasureData(LODInfo.LODMaterialMap.Num());
 
 				for (int32 j = 0; j < LODInfo.LODMaterialMap.Num(); ++j)
