[mutable] Added default value selector for group nodes.
+ By default, value is set to none (One or None groups) or the first option (At Least One groups).

#rnx

diff --git a/Engine/Plugins/Experimental/Mutable/Source/CustomizableObjectEditor/Classes/MuCOE/Nodes/CustomizableObjectNodeObjectGroup.h b/Engine/Plugins/Experimental/Mutable/Source/CustomizableObjectEditor/Classes/MuCOE/Nodes/CustomizableObjectNodeObjectGroup.h
--- a/Engine/Plugins/Experimental/Mutable/Source/CustomizableObjectEditor/Classes/MuCOE/Nodes/CustomizableObjectNodeObjectGroup.h
+++ b/Engine/Plugins/Experimental/Mutable/Source/CustomizableObjectEditor/Classes/MuCOE/Nodes/CustomizableObjectNodeObjectGroup.h
@@ -29,6 +29,9 @@ public:
 
 	UPROPERTY(EditAnywhere, Category=CustomizableObject)
 	ECustomizableObjectGroupType GroupType;
+	
+	UPROPERTY(EditAnywhere, Category = CustomizableObject)
+	FString DefaultValue;
 
 	UPROPERTY(EditAnywhere, Category = UI, meta = (DisplayName = "Parameter UI Metadata"))
 	FMutableParamUIMetadata ParamUIMetadata;
diff --git a/Engine/Plugins/Experimental/Mutable/Source/CustomizableObjectEditor/Private/MuCOE/CustomizableObjectNodeObjectGroupDetails.cpp b/Engine/Plugins/Experimental/Mutable/Source/CustomizableObjectEditor/Private/MuCOE/CustomizableObjectNodeObjectGroupDetails.cpp
--- a/Engine/Plugins/Experimental/Mutable/Source/CustomizableObjectEditor/Private/MuCOE/CustomizableObjectNodeObjectGroupDetails.cpp
+++ b/Engine/Plugins/Experimental/Mutable/Source/CustomizableObjectEditor/Private/MuCOE/CustomizableObjectNodeObjectGroupDetails.cpp
@@ -4,10 +4,13 @@
 
 #include "DetailLayoutBuilder.h"
 #include "IDetailsView.h"
+#include "Misc/Attribute.h"
 #include "MuCOE/CustomizableObjectEditor.h"
 #include "MuCOE/CustomizableObjectGraph.h"
+#include "MuCOE/GraphTraversal.h"
 #include "MuCOE/Nodes/CustomizableObjectNodeObjectGroup.h"
 #include "PropertyCustomizationHelpers.h"
+#include "Widgets/Input/STextComboBox.h"
 
 
 #define LOCTEXT_NAMESPACE "CustomizableObjectGroupDetails"
@@ -22,7 +25,7 @@ TSharedRef<IDetailCustomization> FCustomizableObjectNodeObjectGroupDetails::Make
 void FCustomizableObjectNodeObjectGroupDetails::CustomizeDetails(IDetailLayoutBuilder& DetailBuilder)
 {
 	const IDetailsView* DetailsView = DetailBuilder.GetDetailsView();
-	UCustomizableObjectNodeObjectGroup* NodeGroup = nullptr;
+
 	if (DetailsView->GetSelectedObjects().Num())
 	{
 		if (DetailsView->GetSelectedObjects()[0].Get()->IsA(UCustomizableObjectNodeObjectGroup::StaticClass()))
@@ -37,52 +40,160 @@ void FCustomizableObjectNodeObjectGroupDetails::CustomizeDetails(IDetailLayoutBu
 		{
 			if (UCustomizableObject* NodeGroupCO = CastChecked<UCustomizableObject>(NodeGroup->GetCustomizableObjectGraph()->GetOuter()))
 			{
-				TArray<UCustomizableObject*> ChildNodes;
 				IDetailCategoryBuilder& BlocksCategory = DetailBuilder.EditCategory("Group Info");
-				GraphEditor->GetExternalChildObjects(NodeGroupCO, ChildNodes, false, EObjectFlags::RF_NoFlags);
-				for (UCustomizableObject* ChildNode : ChildNodes)
+
+				TMultiMap<FGuid, UCustomizableObjectNodeObject*> ObjectNodesObjects = GetNodeGroupObjectNodeMapping(NodeGroupCO);
+				TArray<UCustomizableObjectNodeObject*> ChildNodes;
+
+				ObjectNodesObjects.MultiFind(NodeGroup->NodeGuid, ChildNodes);
+
+				for (const UCustomizableObjectNodeObject* ChildObjectNode : ChildNodes)
 				{
-					if (ChildNode)
+					if (ChildObjectNode)
 					{
-						TArray<UCustomizableObjectNodeObject*> ObjectNodes;
-						ChildNode->Source->GetNodesOfClass<UCustomizableObjectNodeObject>(ObjectNodes);
-						UCustomizableObjectNodeObject* groupParent = nullptr;
-
-						for (UCustomizableObjectNodeObject* ChildCONode : ObjectNodes)
-						{
-							if (ChildCONode->bIsBase)
-							{
-								groupParent = ChildCONode;
-								break;
-							}
-						}
-
-						if (!groupParent) continue;
-
-						FString GroupID = groupParent->Identifier.ToString();
-
-						FCustomizableObjectIdPair* DirectGroupChilds = NodeGroupCO->GroupNodeMap.Find(GroupID);
-						if (!DirectGroupChilds) continue;
-
-						if (NodeGroup->GroupName.Equals(DirectGroupChilds->CustomizableObjectGroupName))
-						{
-							BlocksCategory.AddCustomRow(LOCTEXT("FCustomizableObjectNodeObjectGroupDetails", "External Customizable Objects in this Group"))[
-								SNew(SObjectPropertyEntryBox)
-									.ObjectPath(ChildNode->GetPathName())
-									.AllowedClass(UCustomizableObject::StaticClass())
-									.AllowClear(false)
-									.DisplayUseSelected(false)
-									.DisplayBrowse(true)
-									.EnableContentPicker(false)
-									.DisplayThumbnail(true)
-							];
-						}
+						BlocksCategory.AddCustomRow(LOCTEXT("FCustomizableObjectNodeObjectGroupDetails", "External Customizable Objects in this Group"))
+						[
+							SNew(SObjectPropertyEntryBox)
+							.ObjectPath(ChildObjectNode->GetOutermostObject()->GetPathName())
+							.AllowedClass(UCustomizableObject::StaticClass())
+							.AllowClear(false)
+							.DisplayUseSelected(false)
+							.DisplayBrowse(true)
+							.EnableContentPicker(false)
+							.DisplayThumbnail(true)
+						];
 					}
 				}
 			}
 		}
+
+		DetailBuilder.HideProperty("DefaultValue");
+
+		IDetailCategoryBuilder& CustomizableObjectCategory = DetailBuilder.EditCategory("CustomizableObject");
+
+		// Forcing property order
+		CustomizableObjectCategory.AddProperty("GroupName");
+		CustomizableObjectCategory.AddProperty("GroupType");
+
+		// Getting group node children names
+		GenerateChildrenObjectNames();
+
+		CustomizableObjectCategory.AddCustomRow(LOCTEXT("NodeObjectGroupDetails_ComboBox", "Default Value Selector"))
+		.Visibility(TAttribute<EVisibility>(this, &FCustomizableObjectNodeObjectGroupDetails::DefaultValueSelectorVisibility))
+		.NameContent()
+		[
+			SNew(STextBlock)
+			.Text(LOCTEXT("NodeObjectGroupDetails_ComboBox_Text", "Default Value"))
+			.ToolTipText(LOCTEXT("NodeObjectGroupDetails_ComboBox_Tooltip", "Select the default value of the group."))
+			.Font(DetailBuilder.GetDetailFont())
+		]
+		.ValueContent()
+		.HAlign(EHorizontalAlignment::HAlign_Left)
+		[
+			SAssignNew(DefaultValueSelector, STextComboBox)
+			.InitiallySelectedItem(InitialNameOption)
+			.OptionsSource(&ChildrenNameOptions)
+			.OnComboBoxOpening(this, &FCustomizableObjectNodeObjectGroupDetails::GenerateChildrenObjectNames)
+			.OnSelectionChanged(this, &FCustomizableObjectNodeObjectGroupDetails::OnSetDefaultValue)
+			.Font(DetailBuilder.GetDetailFont())
+		];
 	}
 }
 
 
+void FCustomizableObjectNodeObjectGroupDetails::GenerateChildrenObjectNames()
+{
+	ChildrenNameOptions.Reset();
+	InitialNameOption = nullptr;
+
+	if (NodeGroup->GroupType == ECustomizableObjectGroupType::COGT_ONE_OR_NONE)
+	{
+		ChildrenNameOptions.Add(MakeShareable(new FString("None")));
+		InitialNameOption = ChildrenNameOptions.Last();
+	}
+
+	// Adding linked children names
+	const TArray<UEdGraphPin*> ConnectedChildrenPins = FollowInputPinArray(*NodeGroup->ObjectsPin());
+	for (const UEdGraphPin* ChildPin : ConnectedChildrenPins)
+	{
+		if (ChildPin)
+		{
+			if (const UCustomizableObjectNodeObject* ChildObjectNode = Cast<UCustomizableObjectNodeObject>(ChildPin->GetOwningNode()))
+			{
+				ChildrenNameOptions.Add(MakeShareable(new FString(ChildObjectNode->ObjectName)));
+
+				if (ChildObjectNode->ObjectName.Equals(NodeGroup->DefaultValue))
+				{
+					InitialNameOption = ChildrenNameOptions.Last();
+				}
+			}
+			else if (const UCustomizableObjectNodeObjectGroup* ChildGroupObjectNode = Cast<UCustomizableObjectNodeObjectGroup>(ChildPin->GetOwningNode()))
+			{
+				ChildrenNameOptions.Add(MakeShareable(new FString(ChildGroupObjectNode->GroupName)));
+
+				if (ChildGroupObjectNode->GroupName.Equals(NodeGroup->DefaultValue))
+				{
+					InitialNameOption = ChildrenNameOptions.Last();
+				}
+			}
+		}
+	}
+
+	// Adding external children names
+	if (TSharedPtr<FCustomizableObjectEditor> GraphEditor = StaticCastSharedPtr<FCustomizableObjectEditor>(NodeGroup->GetGraphEditor()))
+	{
+		if (UCustomizableObject* NodeGroupCO = CastChecked<UCustomizableObject>(NodeGroup->GetCustomizableObjectGraph()->GetOuter()))
+		{
+			TMultiMap<FGuid, UCustomizableObjectNodeObject*> ObjectNodesObjects = GetNodeGroupObjectNodeMapping(NodeGroupCO);
+			TArray<UCustomizableObjectNodeObject*> ChildNodes;
+
+			ObjectNodesObjects.MultiFind(NodeGroup->NodeGuid, ChildNodes);
+
+			for (const UCustomizableObjectNodeObject* ChildObjectNode : ChildNodes)
+			{
+				if (ChildObjectNode)
+				{
+					ChildrenNameOptions.Add(MakeShareable(new FString(ChildObjectNode->ObjectName)));
+
+					if (ChildObjectNode->ObjectName.Equals(NodeGroup->DefaultValue))
+					{
+						InitialNameOption = ChildrenNameOptions.Last();
+					}
+				}
+			}
+		}
+	}
+
+	if (ChildrenNameOptions.Num() && InitialNameOption == nullptr)
+	{
+		InitialNameOption = ChildrenNameOptions[0];
+	}
+
+	if (DefaultValueSelector.IsValid())
+	{
+		DefaultValueSelector->RefreshOptions();
+		DefaultValueSelector->SetSelectedItem(InitialNameOption);
+	}
+}
+
+
+void FCustomizableObjectNodeObjectGroupDetails::OnSetDefaultValue(TSharedPtr<FString> NewSelection, ESelectInfo::Type SelectInfo)
+{
+	if (NewSelection.IsValid() && (SelectInfo == ESelectInfo::OnMouseClick || SelectInfo == ESelectInfo::OnKeyPress))
+	{
+		NodeGroup->DefaultValue = *NewSelection;
+	}
+}
+
+
+EVisibility FCustomizableObjectNodeObjectGroupDetails::DefaultValueSelectorVisibility() const
+{
+	if (NodeGroup->GroupType == ECustomizableObjectGroupType::COGT_ONE || NodeGroup->GroupType == ECustomizableObjectGroupType::COGT_ONE_OR_NONE)
+	{
+		return EVisibility::Visible;
+	}
+
+	return EVisibility::Hidden;
+}
+
 #undef LOCTEXT_NAMESPACE
diff --git a/Engine/Plugins/Experimental/Mutable/Source/CustomizableObjectEditor/Private/MuCOE/CustomizableObjectNodeObjectGroupDetails.h b/Engine/Plugins/Experimental/Mutable/Source/CustomizableObjectEditor/Private/MuCOE/CustomizableObjectNodeObjectGroupDetails.h
--- a/Engine/Plugins/Experimental/Mutable/Source/CustomizableObjectEditor/Private/MuCOE/CustomizableObjectNodeObjectGroupDetails.h
+++ b/Engine/Plugins/Experimental/Mutable/Source/CustomizableObjectEditor/Private/MuCOE/CustomizableObjectNodeObjectGroupDetails.h
@@ -4,8 +4,11 @@
 
 #include "IDetailCustomization.h"
 
+struct EVisibility;
+namespace ESelectInfo { enum Type : int; }
 class IDetailLayoutBuilder;
-
+class STextComboBox;
+class UCustomizableObjectNodeObjectGroup;
 
 class FCustomizableObjectNodeObjectGroupDetails : public IDetailCustomization
 {
@@ -16,16 +19,27 @@ public:
 	/** ILayoutDetails interface */
 	virtual void CustomizeDetails( IDetailLayoutBuilder& DetailBuilder ) override;
 
+	/** Fills the name options array */
+	void GenerateChildrenObjectNames();
+
+	/** Sets the selected default value */
+	void OnSetDefaultValue(TSharedPtr<FString> NewSelection, ESelectInfo::Type SelectInfo);
+	
+	/** Determines the visibility of the default value selector */
+	EVisibility DefaultValueSelectorVisibility() const;
 
 private:
-	//TArray< TSharedPtr<FString> > BoneComboOptions;
 
-	//class UCustomizableObjectNodeProjectorConstant* NodeConstant;
-	//class UCustomizableObjectNodeProjectorParameter* NodeParameter;
-	//FReply OnProjectorCopyPressed();
-	//FReply OnProjectorPastePressed();
+	/** Pointer to the original group node */
+	UCustomizableObjectNodeObjectGroup* NodeGroup = nullptr;
+
+	/** Combobox to select the defaul value of the node group */
+	TSharedPtr<STextComboBox> DefaultValueSelector;
+
+	/** Combobox options */
+	TArray< TSharedPtr<FString> > ChildrenNameOptions;
 
-	//class USkeletalMesh* SkeletalMesh;
+	/** Initially selected option */
+	TSharedPtr<FString> InitialNameOption;
 
-	//void OnBoneComboBoxSelectionChanged(TSharedPtr<FString> Selection, ESelectInfo::Type SelectInfo, TSharedRef<IPropertyHandle> BoneProperty);
 };
diff --git a/Engine/Plugins/Experimental/Mutable/Source/CustomizableObjectEditor/Private/MuCOE/GenerateMutableSource/GenerateMutableSource.cpp b/Engine/Plugins/Experimental/Mutable/Source/CustomizableObjectEditor/Private/MuCOE/GenerateMutableSource/GenerateMutableSource.cpp
--- a/Engine/Plugins/Experimental/Mutable/Source/CustomizableObjectEditor/Private/MuCOE/GenerateMutableSource/GenerateMutableSource.cpp
+++ b/Engine/Plugins/Experimental/Mutable/Source/CustomizableObjectEditor/Private/MuCOE/GenerateMutableSource/GenerateMutableSource.cpp
@@ -1124,6 +1124,7 @@ mu::NodeObjectPtr GenerateMutableSource(const UEdGraphPin * Pin, FMutableGraphGe
 		const int32 TotalNumChildren = NumChildren + ExternalChildNodes.Num();
 
 		GroupNode->SetChildCount(TotalNumChildren);
+		GroupNode->SetDefaultValue(Type == mu::NodeObjectGroup::CS_ONE_OR_NONE ? -1 : 0);
 		int32 ChildIndex = 0;
 
 		// UI data
@@ -1145,12 +1146,13 @@ mu::NodeObjectPtr GenerateMutableSource(const UEdGraphPin * Pin, FMutableGraphGe
 			UCustomizableObjectNodeObject* CustomizableObjectNodeObject = Cast<UCustomizableObjectNodeObject>(ConnectedChildrenPins[ChildIndex]->GetOwningNode());
 
 			FString* SelectedOptionName = GenerationContext.ParamNamesToSelectedOptions.Find(TypedNodeGroup->GroupName); // If the param is in the map restrict to only the selected option
+			mu::NodeObjectPtr ChildNode = nullptr;
 
 			if (bConnectAtLeastTheLastChild || !SelectedOptionName || *SelectedOptionName == CustomizableObjectNodeObject->ObjectName)
 			{
 				bAtLeastOneConnected = true;
 
-				mu::NodeObjectPtr ChildNode = GenerateMutableSource(ConnectedChildrenPins[ChildIndex], GenerationContext, bPartialCompilation);
+				ChildNode = GenerateMutableSource(ConnectedChildrenPins[ChildIndex], GenerationContext, bPartialCompilation);
 				GroupNode->SetChild(ChildIndex, ChildNode.get());
 
 				if (CustomizableObjectNodeObject)
@@ -1171,10 +1173,17 @@ mu::NodeObjectPtr GenerateMutableSource(const UEdGraphPin * Pin, FMutableGraphGe
 			}
 			else
 			{
-				mu::NodeObjectPtr ChildNode = new mu::NodeObjectNew;
+				ChildNode = new mu::NodeObjectNew;
 				ChildNode->SetName(CustomizableObjectNodeObject->ObjectName);
 				GroupNode->SetChild(ChildIndex, ChildNode.get());
 			}
+
+			if ((TypedNodeGroup->GroupType == ECustomizableObjectGroupType::COGT_ONE ||
+				TypedNodeGroup->GroupType == ECustomizableObjectGroupType::COGT_ONE_OR_NONE)
+				&& TypedNodeGroup->DefaultValue == ChildNode->GetName())
+			{
+				GroupNode->SetDefaultValue(ChildIndex);
+			}
 		}
 
 		const bool bCollapseUnderParent = TypedNodeGroup->ParamUIMetadata.ExtraInformation.Find(FString("CollapseUnderParent")) != nullptr;
@@ -1228,12 +1237,13 @@ mu::NodeObjectPtr GenerateMutableSource(const UEdGraphPin * Pin, FMutableGraphGe
 			UCustomizableObjectNodeObject* CustomizableObjectNodeObject = Cast<UCustomizableObjectNodeObject>(ExternalChildNode->OutputPin()->GetOwningNode());
 
 			FString* SelectedOptionName = GenerationContext.ParamNamesToSelectedOptions.Find(TypedNodeGroup->GroupName); // If the param is in the map restrict to only the selected option
+			mu::NodeObjectPtr ChildNode = nullptr;
 
 			if (bConnectAtLeastTheLastChild || !SelectedOptionName || *SelectedOptionName == CustomizableObjectNodeObject->ObjectName)
 			{
 				bAtLeastOneConnected = true;
 
-				mu::NodeObjectPtr ChildNode = GenerateMutableSource(ExternalChildNode->OutputPin(), GenerationContext, bPartialCompilation);
+				ChildNode = GenerateMutableSource(ExternalChildNode->OutputPin(), GenerationContext, bPartialCompilation);
 				GroupNode->SetChild(ChildIndex, ChildNode.get());
 
 				if (CustomizableObjectNodeObject)
@@ -1259,11 +1269,18 @@ mu::NodeObjectPtr GenerateMutableSource(const UEdGraphPin * Pin, FMutableGraphGe
 			}
 			else
 			{
-				mu::NodeObjectPtr ChildNode = new mu::NodeObjectNew;
+				ChildNode = new mu::NodeObjectNew;
 				ChildNode->SetName(CustomizableObjectNodeObject->ObjectName);
 				GroupNode->SetChild(ChildIndex, ChildNode.get());
 			}
 
+			if ((TypedNodeGroup->GroupType == ECustomizableObjectGroupType::COGT_ONE ||
+				TypedNodeGroup->GroupType == ECustomizableObjectGroupType::COGT_ONE_OR_NONE)
+				&& TypedNodeGroup->DefaultValue == ChildNode->GetName())
+			{
+				GroupNode->SetDefaultValue(ChildIndex);
+			}
+
 			ChildIndex++;
 		}
 
diff --git a/Engine/Plugins/Experimental/Mutable/Source/CustomizableObjectEditor/Private/MuCOE/SCustomizableInstanceProperties.cpp b/Engine/Plugins/Experimental/Mutable/Source/CustomizableObjectEditor/Private/MuCOE/SCustomizableInstanceProperties.cpp
--- a/Engine/Plugins/Experimental/Mutable/Source/CustomizableObjectEditor/Private/MuCOE/SCustomizableInstanceProperties.cpp
+++ b/Engine/Plugins/Experimental/Mutable/Source/CustomizableObjectEditor/Private/MuCOE/SCustomizableInstanceProperties.cpp
@@ -2124,10 +2124,10 @@ bool SCustomizableInstanceProperties::SetParameterValueToDefault(int32 Parameter
 	}
 	case EMutableParameterType::Int:
 	{
-		int32 IndexValue = CustomObject->GetIntParameterDefaultValue(ParameterName);
-		FString DefaultValue = CustomObject->GetIntParameterAvailableOption(ParameterIndex, IndexValue);
+		int32 DefaultValue = CustomObject->GetIntParameterDefaultValue(ParameterName);
+		FString ValueName = CustomObject->FindIntParameterValueName(ParameterIndex, DefaultValue);
 
-		if (!DefaultValue.IsEmpty())
+		if (!ValueName.IsEmpty())
 		{
 			if (CustomObject->IsParameterMultidimensional(ParameterName))
 			{
@@ -2135,12 +2135,12 @@ bool SCustomizableInstanceProperties::SetParameterValueToDefault(int32 Parameter
 
 				for (int32 RangeIndex = 0; RangeIndex < NumRanges; ++RangeIndex)
 				{
-					CustomInstance->SetIntParameterSelectedOption(ParameterName, DefaultValue, RangeIndex);
+					CustomInstance->SetIntParameterSelectedOption(ParameterName, ValueName, RangeIndex);
 				}
 			}
 			else
 			{
-				CustomInstance->SetIntParameterSelectedOption(ParameterName, DefaultValue);
+				CustomInstance->SetIntParameterSelectedOption(ParameterName, ValueName);
 			}
 		}
 		break;
diff --git a/Engine/Plugins/Experimental/Mutable/Source/MutableTools/Private/MuT/CodeGenerator_FirstPass.cpp b/Engine/Plugins/Experimental/Mutable/Source/MutableTools/Private/MuT/CodeGenerator_FirstPass.cpp
--- a/Engine/Plugins/Experimental/Mutable/Source/MutableTools/Private/MuT/CodeGenerator_FirstPass.cpp
+++ b/Engine/Plugins/Experimental/Mutable/Source/MutableTools/Private/MuT/CodeGenerator_FirstPass.cpp
@@ -583,7 +583,7 @@ namespace mu
             op->parameter.m_name = Private->Name;
             op->parameter.m_uid = Private->Uid;
             op->parameter.m_type = PARAMETER_TYPE::T_INT;
-            op->parameter.m_defaultValue.Set<ParamIntType>(-1);
+            op->parameter.m_defaultValue.Set<ParamIntType>(Private->DefaultValue);
 
             if (Private->m_type==NodeObjectGroup::CS_ONE_OR_NONE )
             {
@@ -591,7 +591,6 @@ namespace mu
                 nullValue.m_value = -1;
                 nullValue.m_name = "None";
                 op->parameter.m_possibleValues.Add( nullValue );
-                op->parameter.m_defaultValue.Set<ParamIntType>(nullValue.m_value);
             }
 
             enumOp = op;
@@ -641,12 +640,6 @@ namespace mu
                         value.m_name = pChildNode->GetName();
                         enumOp->parameter.m_possibleValues.Add( value );
 
-                        // Set as default if none has been set.
-                        if ( enumOp->parameter.m_defaultValue.Get<ParamIntType>() == -1 )
-                        {
-                            enumOp->parameter.m_defaultValue.Set<ParamIntType>(t);
-                        }
-
                         check(enumOp);
 
                         // Create a comparison operation as the boolean parameter for the child
diff --git a/Engine/Plugins/Experimental/Mutable/Source/MutableTools/Private/MuT/NodeObjectGroup.cpp b/Engine/Plugins/Experimental/Mutable/Source/MutableTools/Private/MuT/NodeObjectGroup.cpp
--- a/Engine/Plugins/Experimental/Mutable/Source/MutableTools/Private/MuT/NodeObjectGroup.cpp
+++ b/Engine/Plugins/Experimental/Mutable/Source/MutableTools/Private/MuT/NodeObjectGroup.cpp
@@ -129,6 +129,10 @@ namespace mu
 		return pLayout;
 	}
 
+	void NodeObjectGroup::SetDefaultValue(int32 Value)
+	{
+		m_pD->DefaultValue = Value;
+	}
 
 }
 
diff --git a/Engine/Plugins/Experimental/Mutable/Source/MutableTools/Private/MuT/NodeObjectGroupPrivate.h b/Engine/Plugins/Experimental/Mutable/Source/MutableTools/Private/MuT/NodeObjectGroupPrivate.h
--- a/Engine/Plugins/Experimental/Mutable/Source/MutableTools/Private/MuT/NodeObjectGroupPrivate.h
+++ b/Engine/Plugins/Experimental/Mutable/Source/MutableTools/Private/MuT/NodeObjectGroupPrivate.h
@@ -28,17 +28,19 @@ namespace mu
 		void SetSelectionType( CHILD_SELECTION );
 
 		TArray<NodeObjectPtr> m_children;
+		int32 DefaultValue;
 
 		//!
 		void Serialise( OutputArchive& arch ) const
 		{
-            uint32_t ver = 2;
+            uint32_t ver = 3;
 			arch << ver;
 
 			arch << m_type;
 			arch << Name;
 			arch << Uid;
 			arch << m_children;
+			arch << DefaultValue;
 		}
 
 		//!
@@ -46,7 +48,7 @@ namespace mu
 		{
             uint32_t ver;
 			arch >> ver;
-            check(ver>=1 && ver<=2);
+            check(ver>=1 && ver<=3);
 
 			arch >> m_type;
 			if (ver <= 1)
@@ -64,6 +66,11 @@ namespace mu
 			}
 
 			arch >> m_children;
+
+			if (ver >= 3)
+			{
+				arch >> DefaultValue;
+			}
 		}
 
 
diff --git a/Engine/Plugins/Experimental/Mutable/Source/MutableTools/Public/MuT/NodeObjectGroup.h b/Engine/Plugins/Experimental/Mutable/Source/MutableTools/Public/MuT/NodeObjectGroup.h
--- a/Engine/Plugins/Experimental/Mutable/Source/MutableTools/Public/MuT/NodeObjectGroup.h
+++ b/Engine/Plugins/Experimental/Mutable/Source/MutableTools/Public/MuT/NodeObjectGroup.h
@@ -90,6 +90,12 @@ namespace mu
 		//! Set a child object node
 		void SetChild( int index, NodeObjectPtr );
 
+		//! Set default value for CS_ONE_OR_NONE or CS_ALWAYS_ONE groups
+		//! the value is the index of the child option in the group
+		//! -1 is the value for the None option.
+		//! 0 is the first child whether or not the NONE option is present.
+		void SetDefaultValue(int32 Value);
+
 		//-----------------------------------------------------------------------------------------
 		// Interface pattern
 		//-----------------------------------------------------------------------------------------
