[Mutable] Fix losing parameter values after a recompilation.

#rnx

diff --git a/Engine/Plugins/Experimental/Mutable/Source/CustomizableObjectEditor/Private/MuCOE/CustomizableObjectEditor.cpp b/Engine/Plugins/Experimental/Mutable/Source/CustomizableObjectEditor/Private/MuCOE/CustomizableObjectEditor.cpp
--- a/Engine/Plugins/Experimental/Mutable/Source/CustomizableObjectEditor/Private/MuCOE/CustomizableObjectEditor.cpp
+++ b/Engine/Plugins/Experimental/Mutable/Source/CustomizableObjectEditor/Private/MuCOE/CustomizableObjectEditor.cpp
@@ -2729,8 +2729,11 @@ void FCustomizableObjectEditor::OnCustomizableObjectStatusChanged(FCustomizableO
 {
 	switch (NextState)
 	{
-	case FCustomizableObjectStatus::EState::ModelLoaded: 
-		CreatePreviewInstance();
+	case FCustomizableObjectStatus::EState::ModelLoaded:
+		if (!PreviewInstance)
+		{
+			CreatePreviewInstance();
+		}
 		break;
 		
 	case FCustomizableObjectStatus::EState::NoModel:
