[Mutable] Fix CO is already being compiled or updated warning on commandlets.

#rnx

diff --git a/Engine/Plugins/Experimental/Mutable/Source/CustomizableObject/Classes/MuCO/CustomizableObjectSystem.h b/Engine/Plugins/Experimental/Mutable/Source/CustomizableObject/Classes/MuCO/CustomizableObjectSystem.h
--- a/Engine/Plugins/Experimental/Mutable/Source/CustomizableObject/Classes/MuCO/CustomizableObjectSystem.h
+++ b/Engine/Plugins/Experimental/Mutable/Source/CustomizableObject/Classes/MuCO/CustomizableObjectSystem.h
@@ -238,6 +238,12 @@ public:
 	// If true, uncompiled Customizable Objects will be compiled whenever an instance update is required
 	bool IsAutoCompileEnabled() const;
 
+	/** Return true if inside commandlets uncompiled Customizable Objects will be compiled whenever an instance update is required. */
+	bool IsAutoCompileCommandletEnabled() const;
+
+	/** Set if inside commandlets uncompiled Customizable Objects will be compiled whenever an instance update is required. */
+	void SetAutoCompileCommandletEnabled(bool bValue);
+	
 	// If true, uncompiled Customizable Objects will be compiled synchronously
 	bool IsAutoCompilationSync() const;
 
diff --git a/Engine/Plugins/Experimental/Mutable/Source/CustomizableObject/Private/MuCO/CustomizableObject.cpp b/Engine/Plugins/Experimental/Mutable/Source/CustomizableObject/Private/MuCO/CustomizableObject.cpp
--- a/Engine/Plugins/Experimental/Mutable/Source/CustomizableObject/Private/MuCO/CustomizableObject.cpp
+++ b/Engine/Plugins/Experimental/Mutable/Source/CustomizableObject/Private/MuCO/CustomizableObject.cpp
@@ -880,9 +880,9 @@ bool UCustomizableObject::ConditionalAutoCompile()
 		return false;
 	}
 
-	// Don't compile if the cook commandlet is running. Do not add a warning/error. Otherwise, we could end up
-	// invalidating the cook for no reason.
-	if (IsRunningCookCommandlet())
+	// By default, don't compile in a commandlet.
+	// Notice that the cook is also a commandlet. Do not add a warning/error, otherwise we could end up invalidating the cook for no reason.
+	if (IsRunningCookCommandlet() || (IsRunningCommandlet() && !System->IsAutoCompileCommandletEnabled()))
 	{
 		return false;
 	}
diff --git a/Engine/Plugins/Experimental/Mutable/Source/CustomizableObject/Private/MuCO/CustomizableObjectSystem.cpp b/Engine/Plugins/Experimental/Mutable/Source/CustomizableObject/Private/MuCO/CustomizableObjectSystem.cpp
--- a/Engine/Plugins/Experimental/Mutable/Source/CustomizableObject/Private/MuCO/CustomizableObjectSystem.cpp
+++ b/Engine/Plugins/Experimental/Mutable/Source/CustomizableObject/Private/MuCO/CustomizableObjectSystem.cpp
@@ -1425,6 +1425,18 @@ bool UCustomizableObjectSystem::IsAutoCompileEnabled() const
 }
 
 
+bool UCustomizableObjectSystem::IsAutoCompileCommandletEnabled() const
+{
+	return GetPrivate()->bAutoCompileCommandletEnabled;
+}
+
+
+void UCustomizableObjectSystem::SetAutoCompileCommandletEnabled(bool bValue)
+{
+	GetPrivate()->bAutoCompileCommandletEnabled = bValue;
+}
+
+
 bool UCustomizableObjectSystem::IsAutoCompilationSync() const
 {
 	return EditorSettings.bCompileObjectsSynchronously;
diff --git a/Engine/Plugins/Experimental/Mutable/Source/CustomizableObject/Private/MuCO/CustomizableObjectSystemPrivate.h b/Engine/Plugins/Experimental/Mutable/Source/CustomizableObject/Private/MuCO/CustomizableObjectSystemPrivate.h
--- a/Engine/Plugins/Experimental/Mutable/Source/CustomizableObject/Private/MuCO/CustomizableObjectSystemPrivate.h
+++ b/Engine/Plugins/Experimental/Mutable/Source/CustomizableObject/Private/MuCO/CustomizableObjectSystemPrivate.h
@@ -662,6 +662,8 @@ public:
 	FLogBenchmarkUtil LogBenchmarkUtil;
 
 	int32 NumSkeletalMeshes = 0;
+
+	bool bAutoCompileCommandletEnabled = false;
 };
 
 
