[Mutable] Initialize CO system OnPostEngineInit.
Remove system initialization from the CO/COI PostLoad.

#rnx

diff --git a/Engine/Plugins/Experimental/Mutable/Source/CustomizableObject/Private/MuCO/CustomizableObject.cpp b/Engine/Plugins/Experimental/Mutable/Source/CustomizableObject/Private/MuCO/CustomizableObject.cpp
--- a/Engine/Plugins/Experimental/Mutable/Source/CustomizableObject/Private/MuCO/CustomizableObject.cpp
+++ b/Engine/Plugins/Experimental/Mutable/Source/CustomizableObject/Private/MuCO/CustomizableObject.cpp
@@ -198,9 +198,6 @@ void UCustomizableObject::PostLoad()
 {
 	Super::PostLoad();
 
-	// Make sure mutable has been initialised.
-	UCustomizableObjectSystem::GetInstance();
-
 #if WITH_EDITOR
 	if (ReferenceSkeletalMesh_DEPRECATED)
 	{
diff --git a/Engine/Plugins/Experimental/Mutable/Source/CustomizableObject/Private/MuCO/CustomizableObjectInstance.cpp b/Engine/Plugins/Experimental/Mutable/Source/CustomizableObject/Private/MuCO/CustomizableObjectInstance.cpp
--- a/Engine/Plugins/Experimental/Mutable/Source/CustomizableObject/Private/MuCO/CustomizableObjectInstance.cpp
+++ b/Engine/Plugins/Experimental/Mutable/Source/CustomizableObject/Private/MuCO/CustomizableObjectInstance.cpp
@@ -556,9 +556,6 @@ void UCustomizableObjectInstance::Serialize(FArchive& Ar)
 
 void UCustomizableObjectInstance::PostLoad()
 {
-	// Make sure mutable has been initialised.
-	UCustomizableObjectSystem::GetInstance(); 
-
 	Super::PostLoad();
 
 #if WITH_EDITOR
diff --git a/Engine/Plugins/Experimental/Mutable/Source/CustomizableObject/Private/MuCO/CustomizableObjectModule.cpp b/Engine/Plugins/Experimental/Mutable/Source/CustomizableObject/Private/MuCO/CustomizableObjectModule.cpp
--- a/Engine/Plugins/Experimental/Mutable/Source/CustomizableObject/Private/MuCO/CustomizableObjectModule.cpp
+++ b/Engine/Plugins/Experimental/Mutable/Source/CustomizableObject/Private/MuCO/CustomizableObjectModule.cpp
@@ -10,6 +10,7 @@
 #include "MuCO/CustomizableObjectDGGUI.h"
 #include "MuCO/CustomizableObjectExtension.h"
 #include "MuCO/CustomizableObjectInstanceUsage.h"
+#include "MuCO/CustomizableObjectSystem.h"
 #include "MuCO/ICustomizableObjectModule.h"
 #include "UObject/StrongObjectPtr.h"
 #include "GPUSkinPublicDefs.h"
@@ -38,6 +39,8 @@ public:
 private:
 	void RefreshExtensionData();
 
+	static void InitializeSystem();
+
 	// Command to look for Customizable Object Instance in the player pawn of the current world and open a DGGUI to edit its parameters
 	IConsoleCommand* LaunchDGGUICommand;
 	static void ToggleDGGUI(const TArray<FString>& Arguments);
@@ -60,6 +63,8 @@ void FCustomizableObjectModule::StartupModule()
 		TEXT("mutable.ToggleDGGUI"),
 		TEXT("Looks for a Customizable Object Instance within the player pawn and opens a UI to modify its parameters, or closes it if it's open. Specify slot ID to control which component is modified."),
 		FConsoleCommandWithArgsDelegate::CreateStatic(&FCustomizableObjectModule::ToggleDGGUI));
+
+	FCoreDelegates::OnPostEngineInit.AddStatic(&FCustomizableObjectModule::InitializeSystem);
 }
 
 
@@ -217,6 +222,12 @@ void FCustomizableObjectModule::RefreshExtensionData()
 }
 
 
+void FCustomizableObjectModule::InitializeSystem()
+{
+	UCustomizableObjectSystem::GetInstance();
+}
+
+
 UCustomizableObjectInstanceUsage* GetPlayerCustomizableObjectInstanceUsage(const int32 SlotID, const UWorld* CurrentWorld, const int32 PlayerIndex)
 {
 	// Get customizable skeletal component attached to player pawn
