[Mutable] Change Parameter Uid type to reduce its size.

Only changed the type in MuR and CO. Core serialized nodes still are a string.

#rnx

diff --git a/Engine/Plugins/Experimental/Mutable/Source/CustomizableObject/Classes/MuCO/CustomizableObject.h b/Engine/Plugins/Experimental/Mutable/Source/CustomizableObject/Classes/MuCO/CustomizableObject.h
--- a/Engine/Plugins/Experimental/Mutable/Source/CustomizableObject/Classes/MuCO/CustomizableObject.h
+++ b/Engine/Plugins/Experimental/Mutable/Source/CustomizableObject/Classes/MuCO/CustomizableObject.h
@@ -1370,7 +1370,7 @@ private:
 	// This is a manual version number for the binary blobs in this asset.
 	// Increasing it invalidates all the previously compiled models.
 	// Warning: If while merging code both versions have changed, take the highest+1.
-	static const int32 CurrentSupportedVersion = 415;
+	static const int32 CurrentSupportedVersion = 416;
 
 	// Compile the object for a specific platform - Compile for Cook Customizable Object
 	void CompileForTargetPlatform(const ITargetPlatform* TargetPlatform);
diff --git a/Engine/Plugins/Experimental/Mutable/Source/CustomizableObject/Private/MuCO/CustomizableObjectInstanceDescriptor.cpp b/Engine/Plugins/Experimental/Mutable/Source/CustomizableObject/Private/MuCO/CustomizableObjectInstanceDescriptor.cpp
--- a/Engine/Plugins/Experimental/Mutable/Source/CustomizableObject/Private/MuCO/CustomizableObjectInstanceDescriptor.cpp
+++ b/Engine/Plugins/Experimental/Mutable/Source/CustomizableObject/Private/MuCO/CustomizableObjectInstanceDescriptor.cpp
@@ -466,7 +466,7 @@ mu::ParametersPtr FCustomizableObjectInstanceDescriptor::GetParameters() const
 	for (int32 ParamIndex = 0; ParamIndex < ParamCount; ++ParamIndex)
 	{
 		const FString& Name = MutableParameters->GetName(ParamIndex);
-		const FString& Uid = MutableParameters->GetUid(ParamIndex);
+		const FGuid& Uid = MutableParameters->GetUid(ParamIndex);
 		const mu::PARAMETER_TYPE MutableType = MutableParameters->GetType(ParamIndex);
 
 		switch (MutableType)
@@ -475,7 +475,7 @@ mu::ParametersPtr FCustomizableObjectInstanceDescriptor::GetParameters() const
 		{
 			for (const FCustomizableObjectBoolParameterValue& BoolParameter : BoolParameters)
 			{
-				if (BoolParameter.ParameterName == Name || (!Uid.IsEmpty() && BoolParameter.Uid == Uid))
+				if (BoolParameter.ParameterName == Name || (Uid.IsValid() && BoolParameter.Uid == Uid))
 				{
 					MutableParameters->SetBoolValue(ParamIndex,  BoolParameter.ParameterValue);
 					break;
@@ -489,7 +489,7 @@ mu::ParametersPtr FCustomizableObjectInstanceDescriptor::GetParameters() const
 		{
 			for (const FCustomizableObjectIntParameterValue& IntParameter : IntParameters)
 			{
-				if (IntParameter.ParameterName.Equals(Name, ESearchCase::CaseSensitive) || (!Uid.IsEmpty() && IntParameter.Uid == Uid))
+				if (IntParameter.ParameterName.Equals(Name, ESearchCase::CaseSensitive) || (Uid.IsValid() && IntParameter.Uid == Uid))
 				{
 					if (mu::RangeIndexPtr RangeIdxPtr = MutableParameters->NewRangeIndex(ParamIndex))
 					{
@@ -519,7 +519,7 @@ mu::ParametersPtr FCustomizableObjectInstanceDescriptor::GetParameters() const
 		{
 			for (const FCustomizableObjectFloatParameterValue& FloatParameter : FloatParameters)
 			{
-				if (FloatParameter.ParameterName == Name || (!Uid.IsEmpty() && FloatParameter.Uid == Uid))
+				if (FloatParameter.ParameterName == Name || (Uid.IsValid() && FloatParameter.Uid == Uid))
 				{
 					if (mu::RangeIndexPtr RangeIdxPtr = MutableParameters->NewRangeIndex(ParamIndex))
 					{
@@ -545,7 +545,7 @@ mu::ParametersPtr FCustomizableObjectInstanceDescriptor::GetParameters() const
 		{
 			for (const FCustomizableObjectVectorParameterValue& VectorParameter : VectorParameters)
 			{
-				if (VectorParameter.ParameterName == Name || (!Uid.IsEmpty() && VectorParameter.Uid == Uid))
+				if (VectorParameter.ParameterName == Name || (Uid.IsValid() && VectorParameter.Uid == Uid))
 				{
 					MutableParameters->SetColourValue(ParamIndex, VectorParameter.ParameterValue.R, VectorParameter.ParameterValue.G, VectorParameter.ParameterValue.B);
 
@@ -560,7 +560,7 @@ mu::ParametersPtr FCustomizableObjectInstanceDescriptor::GetParameters() const
 		{
 			for (const auto& ProjectorParameter : ProjectorParameters)
 			{
-				if (ProjectorParameter.ParameterName == Name || (!Uid.IsEmpty() && ProjectorParameter.Uid == Uid))
+				if (ProjectorParameter.ParameterName == Name || (Uid.IsValid() && ProjectorParameter.Uid == Uid))
 				{
 					auto CopyProjector = [&MutableParameters, ParamIndex](const FCustomizableObjectProjector& Value, const mu::RangeIndexPtr RangeIdxPtr = nullptr)
 					{
@@ -621,7 +621,7 @@ mu::ParametersPtr FCustomizableObjectInstanceDescriptor::GetParameters() const
 		{
 			for (const FCustomizableObjectTextureParameterValue& TextureParameter : TextureParameters)
 			{
-				if (TextureParameter.ParameterName == Name || (!Uid.IsEmpty() && TextureParameter.Uid == Uid))
+				if (TextureParameter.ParameterName == Name || (Uid.IsValid() && TextureParameter.Uid == Uid))
 				{
 					if (mu::RangeIndexPtr RangeIdxPtr = MutableParameters->NewRangeIndex(ParamIndex))
 					{
@@ -779,7 +779,7 @@ void FCustomizableObjectInstanceDescriptor::ReloadParameters()
 	for (int32 ParamIndex = 0; ParamIndex < ParamCount; ++ParamIndex)
 	{
 		const FString& Name = MutableParameters->GetName(ParamIndex);
-		const FString& Uid = MutableParameters->GetUid(ParamIndex);
+		const FGuid& Uid = MutableParameters->GetUid(ParamIndex);
 		const mu::PARAMETER_TYPE MutableType = MutableParameters->GetType(ParamIndex);
 
 		switch (MutableType)
@@ -792,7 +792,7 @@ void FCustomizableObjectInstanceDescriptor::ReloadParameters()
 
 			auto FindByNameAndUid = [&](const FCustomizableObjectBoolParameterValue& P)
 			{
-				return P.ParameterName == Name || (!Uid.IsEmpty() && P.Uid == Uid);
+				return P.ParameterName == Name || (Uid.IsValid() && P.Uid == Uid);
 			};
 
 			if (FCustomizableObjectBoolParameterValue* Result = OldBoolParameters.FindByPredicate(FindByNameAndUid))
@@ -816,7 +816,7 @@ void FCustomizableObjectInstanceDescriptor::ReloadParameters()
 
 			auto FindByNameAndUid = [&](const FCustomizableObjectIntParameterValue& P)
 			{
-				return P.ParameterName == Name || (!Uid.IsEmpty() && P.Uid == Uid);
+				return P.ParameterName == Name || (Uid.IsValid() && P.Uid == Uid);
 			};
 			
 			if (FCustomizableObjectIntParameterValue* Result = OldIntParameters.FindByPredicate(FindByNameAndUid))
@@ -920,7 +920,7 @@ void FCustomizableObjectInstanceDescriptor::ReloadParameters()
 
 			auto FindByNameAndUid = [&](const FCustomizableObjectFloatParameterValue& P)
 			{
-				return P.ParameterName == Name || (!Uid.IsEmpty() && P.Uid == Uid);
+				return P.ParameterName == Name || (Uid.IsValid() && P.Uid == Uid);
 			};
 			
 			if (FCustomizableObjectFloatParameterValue* Result = OldFloatParameters.FindByPredicate(FindByNameAndUid))
@@ -971,7 +971,7 @@ void FCustomizableObjectInstanceDescriptor::ReloadParameters()
 
 			auto FindByNameAndUid = [&](const FCustomizableObjectVectorParameterValue& P)
 			{
-				return P.ParameterName == Name || (!Uid.IsEmpty() && P.Uid == Uid);
+				return P.ParameterName == Name || (Uid.IsValid() && P.Uid == Uid);
 			};
 
 			if (FCustomizableObjectVectorParameterValue* Result = OldVectorParameters.FindByPredicate(FindByNameAndUid))
@@ -996,7 +996,7 @@ void FCustomizableObjectInstanceDescriptor::ReloadParameters()
 
 			auto FindByNameAndUid = [&](const FCustomizableObjectProjectorParameterValue& P)
 			{
-				return P.ParameterName == Name || (!Uid.IsEmpty() && P.Uid == Uid);
+				return P.ParameterName == Name || (Uid.IsValid() && P.Uid == Uid);
 			};
 			
 			if (FCustomizableObjectProjectorParameterValue* Result = OldProjectorParameters.FindByPredicate(FindByNameAndUid))
@@ -1071,7 +1071,7 @@ void FCustomizableObjectInstanceDescriptor::ReloadParameters()
 
 			auto FindByNameAndUid = [&](const FCustomizableObjectTextureParameterValue& P)
 			{
-				return P.ParameterName == Name || (!Uid.IsEmpty() && P.Uid == Uid);
+				return P.ParameterName == Name || (Uid.IsValid() && P.Uid == Uid);
 			};
 			
 			if (FCustomizableObjectTextureParameterValue* Result = OldTextureParameters.FindByPredicate(FindByNameAndUid))
diff --git a/Engine/Plugins/Experimental/Mutable/Source/CustomizableObject/Public/MuCO/CustomizableObjectParameterTypeDefinitions.h b/Engine/Plugins/Experimental/Mutable/Source/CustomizableObject/Public/MuCO/CustomizableObjectParameterTypeDefinitions.h
--- a/Engine/Plugins/Experimental/Mutable/Source/CustomizableObject/Public/MuCO/CustomizableObjectParameterTypeDefinitions.h
+++ b/Engine/Plugins/Experimental/Mutable/Source/CustomizableObject/Public/MuCO/CustomizableObjectParameterTypeDefinitions.h
@@ -59,7 +59,7 @@ struct FCustomizableObjectBoolParameterValue
 	bool ParameterValue = true;
 
 	UPROPERTY(Category = CustomizableObjectBoolParameterValue, VisibleAnywhere)
-	FString Uid;
+	FGuid Uid;
 };
 
 
@@ -87,7 +87,7 @@ struct FCustomizableObjectIntParameterValue
 	FString ParameterValueName;
 
 	UPROPERTY(Category = CustomizableObjectIntParameterValue, VisibleAnywhere)
-	FString Uid;
+	FGuid Uid;
 
 	// Same as ParameterValueName but for multidimensional params
 	UPROPERTY()
@@ -97,7 +97,7 @@ struct FCustomizableObjectIntParameterValue
 	{
 	}
 
-	FCustomizableObjectIntParameterValue(const FString & InParamName, const FString & InParameterValueName, const FString & InUid, const TArray<FString>& InParameterRangeValueNames)
+	FCustomizableObjectIntParameterValue(const FString & InParamName, const FString & InParameterValueName, const FGuid & InUid, const TArray<FString>& InParameterRangeValueNames)
 		: ParameterName(InParamName), ParameterValueName(InParameterValueName), Uid(InUid)
 	{
 		ParameterRangeValueNames = InParameterRangeValueNames;
@@ -133,7 +133,7 @@ struct FCustomizableObjectFloatParameterValue
 	float ParameterValue = 0.0f;
 
 	UPROPERTY(Category = CustomizableObjectFloatParameterValue, VisibleAnywhere)
-	FString Uid;
+	FGuid Uid;
 
 	UPROPERTY(Category = CustomizableObjectFloatParameterValue, VisibleAnywhere)
 	TArray<float> ParameterRangeValues;
@@ -168,7 +168,7 @@ struct FCustomizableObjectTextureParameterValue
 	FName ParameterValue;
 
 	UPROPERTY(Category = CustomizableObjectTextureParameterValue, VisibleAnywhere)
-	FString Uid;
+	FGuid Uid;
 
 	UPROPERTY(Category = CustomizableObjectTextureParameterValue, VisibleAnywhere)
 	TArray<FName> ParameterRangeValues;
@@ -203,7 +203,7 @@ struct FCustomizableObjectVectorParameterValue
 	FLinearColor ParameterValue = FLinearColor(ForceInit);
 
 	UPROPERTY(Category = CustomizableObjectVectorParameterValue, VisibleAnywhere)
-	FString Uid;
+	FGuid Uid;
 };
 
 
@@ -290,7 +290,7 @@ struct FCustomizableObjectProjectorParameterValue
 	FCustomizableObjectProjector Value;
 
 	UPROPERTY(Category = CustomizableObjectVectorParameterValue, VisibleAnywhere)
-	FString Uid;
+	FGuid Uid;
 
 	UPROPERTY(Category = CustomizableObjectVectorParameterValue, VisibleAnywhere)
 	TArray<FCustomizableObjectProjector> RangeValues;
diff --git a/Engine/Plugins/Experimental/Mutable/Source/MutableRuntime/Internal/MuR/ParametersPrivate.h b/Engine/Plugins/Experimental/Mutable/Source/MutableRuntime/Internal/MuR/ParametersPrivate.h
--- a/Engine/Plugins/Experimental/Mutable/Source/MutableRuntime/Internal/MuR/ParametersPrivate.h
+++ b/Engine/Plugins/Experimental/Mutable/Source/MutableRuntime/Internal/MuR/ParametersPrivate.h
@@ -408,7 +408,7 @@ namespace Private
         FString m_name;
 
         //! Unique id (provided externally, so no actual guarantee that it is unique.)
-		FString m_uid;
+		FGuid m_uid;
 
         PARAMETER_TYPE m_type = PARAMETER_TYPE::T_NONE;
 
@@ -462,7 +462,7 @@ namespace Private
         //!
         void Serialise( OutputArchive& arch ) const
         {
-            const int32 ver = 8;
+            const int32 ver = 9;
             arch << ver;
 
 			arch << m_name;
@@ -478,7 +478,7 @@ namespace Private
         {
             int32 ver;
             arch >> ver;
-            check( ver==8 );
+            check( ver==9 );
 
 			arch >> m_name;
 			arch >> m_uid;
diff --git a/Engine/Plugins/Experimental/Mutable/Source/MutableRuntime/Private/MuR/Parameters.cpp b/Engine/Plugins/Experimental/Mutable/Source/MutableRuntime/Private/MuR/Parameters.cpp
--- a/Engine/Plugins/Experimental/Mutable/Source/MutableRuntime/Private/MuR/Parameters.cpp
+++ b/Engine/Plugins/Experimental/Mutable/Source/MutableRuntime/Private/MuR/Parameters.cpp
@@ -200,7 +200,7 @@ namespace mu
 
 
 	//---------------------------------------------------------------------------------------------
-	const FString& Parameters::GetUid( int index ) const
+	const FGuid& Parameters::GetUid( int index ) const
 	{
 		const FProgram& program = m_pD->m_pModel->GetPrivate()->m_program;
 		check( index>=0 && index<(int)program.m_parameters.Num() );
diff --git a/Engine/Plugins/Experimental/Mutable/Source/MutableRuntime/Private/MuR/Serialisation.cpp b/Engine/Plugins/Experimental/Mutable/Source/MutableRuntime/Private/MuR/Serialisation.cpp
--- a/Engine/Plugins/Experimental/Mutable/Source/MutableRuntime/Private/MuR/Serialisation.cpp
+++ b/Engine/Plugins/Experimental/Mutable/Source/MutableRuntime/Private/MuR/Serialisation.cpp
@@ -27,8 +27,9 @@ namespace mu
     MUTABLE_IMPLEMENT_POD_SERIALISABLE( uint16 ); 
     MUTABLE_IMPLEMENT_POD_SERIALISABLE( uint32 ); 
     MUTABLE_IMPLEMENT_POD_SERIALISABLE( uint64 )
-	
-                                                                                      
+
+	MUTABLE_IMPLEMENT_POD_SERIALISABLE(FGuid);
+
 	MUTABLE_IMPLEMENT_POD_SERIALISABLE(vec2f);                                        
 	MUTABLE_IMPLEMENT_POD_SERIALISABLE(vec3f);                                        
 	MUTABLE_IMPLEMENT_POD_SERIALISABLE(mat3f);                                        
diff --git a/Engine/Plugins/Experimental/Mutable/Source/MutableRuntime/Public/MuR/Parameters.h b/Engine/Plugins/Experimental/Mutable/Source/MutableRuntime/Public/MuR/Parameters.h
--- a/Engine/Plugins/Experimental/Mutable/Source/MutableRuntime/Public/MuR/Parameters.h
+++ b/Engine/Plugins/Experimental/Mutable/Source/MutableRuntime/Public/MuR/Parameters.h
@@ -192,7 +192,7 @@ namespace mu
 
 		//! Return the Guid of the parameter, resistant to parameter name changes.
 		//! \param index Index of the parameter from 0 to GetCount()-1
-		const FString& GetUid( int index ) const;
+		const FGuid& GetUid( int index ) const;
 
 		//! Find the parameter index by name.
 		//! It returns -1 if the parameter doesn't exist.
diff --git a/Engine/Plugins/Experimental/Mutable/Source/MutableRuntime/Public/MuR/Serialisation.h b/Engine/Plugins/Experimental/Mutable/Source/MutableRuntime/Public/MuR/Serialisation.h
--- a/Engine/Plugins/Experimental/Mutable/Source/MutableRuntime/Public/MuR/Serialisation.h
+++ b/Engine/Plugins/Experimental/Mutable/Source/MutableRuntime/Public/MuR/Serialisation.h
@@ -348,6 +348,8 @@ namespace mu
     MUTABLE_DEFINE_POD_SERIALISABLE(uint16);
     MUTABLE_DEFINE_POD_SERIALISABLE(uint32);
     MUTABLE_DEFINE_POD_SERIALISABLE(uint64);
+	
+	MUTABLE_DEFINE_POD_SERIALISABLE(FGuid);
 
 	MUTABLE_DEFINE_POD_VECTOR_SERIALISABLE(float);
 	MUTABLE_DEFINE_POD_VECTOR_SERIALISABLE(double);
diff --git a/Engine/Plugins/Experimental/Mutable/Source/MutableTools/Private/MuT/CodeGenerator_Bool.cpp b/Engine/Plugins/Experimental/Mutable/Source/MutableTools/Private/MuT/CodeGenerator_Bool.cpp
--- a/Engine/Plugins/Experimental/Mutable/Source/MutableTools/Private/MuT/CodeGenerator_Bool.cpp
+++ b/Engine/Plugins/Experimental/Mutable/Source/MutableTools/Private/MuT/CodeGenerator_Bool.cpp
@@ -88,7 +88,8 @@ namespace mu
 		{
 			FParameterDesc param;
 			param.m_name = node.m_name;
-			param.m_uid = node.m_uid;
+			const TCHAR* CStr = ToCStr(node.m_uid);
+			param.m_uid.ImportTextItem(CStr, 0, nullptr, nullptr);
 			param.m_type = PARAMETER_TYPE::T_BOOL;
 			param.m_defaultValue.Set<ParamBoolType>(node.m_defaultValue);
 
diff --git a/Engine/Plugins/Experimental/Mutable/Source/MutableTools/Private/MuT/CodeGenerator_Color.cpp b/Engine/Plugins/Experimental/Mutable/Source/MutableTools/Private/MuT/CodeGenerator_Color.cpp
--- a/Engine/Plugins/Experimental/Mutable/Source/MutableTools/Private/MuT/CodeGenerator_Color.cpp
+++ b/Engine/Plugins/Experimental/Mutable/Source/MutableTools/Private/MuT/CodeGenerator_Color.cpp
@@ -148,7 +148,8 @@ namespace mu
 		{
 			FParameterDesc param;
 			param.m_name = node.m_name;
-			param.m_uid = node.m_uid;
+			const TCHAR* CStr = ToCStr(node.m_uid);
+			param.m_uid.ImportTextItem(CStr, 0, nullptr, nullptr);
 			param.m_type = PARAMETER_TYPE::T_COLOUR;
 
 			ParamColorType Value;
diff --git a/Engine/Plugins/Experimental/Mutable/Source/MutableTools/Private/MuT/CodeGenerator_FirstPass.cpp b/Engine/Plugins/Experimental/Mutable/Source/MutableTools/Private/MuT/CodeGenerator_FirstPass.cpp
--- a/Engine/Plugins/Experimental/Mutable/Source/MutableTools/Private/MuT/CodeGenerator_FirstPass.cpp
+++ b/Engine/Plugins/Experimental/Mutable/Source/MutableTools/Private/MuT/CodeGenerator_FirstPass.cpp
@@ -581,7 +581,8 @@ namespace mu
             op->type = OP_TYPE::NU_PARAMETER;
 
             op->parameter.m_name = Private->Name;
-            op->parameter.m_uid = Private->Uid;
+			const TCHAR* CStr = ToCStr(Private->Uid);
+			op->parameter.m_uid.ImportTextItem(CStr, 0, nullptr, nullptr);
             op->parameter.m_type = PARAMETER_TYPE::T_INT;
             op->parameter.m_defaultValue.Set<ParamIntType>(Private->DefaultValue);
 
@@ -613,7 +614,8 @@ namespace mu
                         op->type = OP_TYPE::BO_PARAMETER;
 
                         op->parameter.m_name = pChildNode->GetName();
-                        op->parameter.m_uid = pChildNode->GetUid();
+                   		const TCHAR* CStr = ToCStr(Private->Uid);
+                   		op->parameter.m_uid.ImportTextItem(CStr, 0, nullptr, nullptr);
                         op->parameter.m_type = PARAMETER_TYPE::T_BOOL;
                         op->parameter.m_defaultValue.Set<ParamBoolType>(false);
 
diff --git a/Engine/Plugins/Experimental/Mutable/Source/MutableTools/Private/MuT/CodeGenerator_Image.cpp b/Engine/Plugins/Experimental/Mutable/Source/MutableTools/Private/MuT/CodeGenerator_Image.cpp
--- a/Engine/Plugins/Experimental/Mutable/Source/MutableTools/Private/MuT/CodeGenerator_Image.cpp
+++ b/Engine/Plugins/Experimental/Mutable/Source/MutableTools/Private/MuT/CodeGenerator_Image.cpp
@@ -273,8 +273,9 @@ namespace mu
             op->type = OP_TYPE::IM_PARAMETER;
 
 			op->parameter.m_name = node.m_name;
-			op->parameter.m_uid = node.m_uid;
-            op->parameter.m_type = PARAMETER_TYPE::T_IMAGE;
+        	const TCHAR* CStr = ToCStr(node.m_uid);
+        	op->parameter.m_uid.ImportTextItem(CStr, 0, nullptr, nullptr);
+        	op->parameter.m_type = PARAMETER_TYPE::T_IMAGE;
         	op->parameter.m_defaultValue.Set<ParamImageType>(node.m_defaultValue);
 
 			// Generate the code for the ranges
diff --git a/Engine/Plugins/Experimental/Mutable/Source/MutableTools/Private/MuT/CodeGenerator_Projector.cpp b/Engine/Plugins/Experimental/Mutable/Source/MutableTools/Private/MuT/CodeGenerator_Projector.cpp
--- a/Engine/Plugins/Experimental/Mutable/Source/MutableTools/Private/MuT/CodeGenerator_Projector.cpp
+++ b/Engine/Plugins/Experimental/Mutable/Source/MutableTools/Private/MuT/CodeGenerator_Projector.cpp
@@ -117,7 +117,8 @@ namespace mu
 		{
 			FParameterDesc param;
 			param.m_name = node.m_name;
-			param.m_uid = node.m_uid;
+			const TCHAR* CStr = ToCStr(node.m_uid);
+			param.m_uid.ImportTextItem(CStr, 0, nullptr, nullptr);
 			param.m_type = PARAMETER_TYPE::T_PROJECTOR;
 
 			FProjector p = ProjectorFromNode(node);
diff --git a/Engine/Plugins/Experimental/Mutable/Source/MutableTools/Private/MuT/CodeGenerator_Scalar.cpp b/Engine/Plugins/Experimental/Mutable/Source/MutableTools/Private/MuT/CodeGenerator_Scalar.cpp
--- a/Engine/Plugins/Experimental/Mutable/Source/MutableTools/Private/MuT/CodeGenerator_Scalar.cpp
+++ b/Engine/Plugins/Experimental/Mutable/Source/MutableTools/Private/MuT/CodeGenerator_Scalar.cpp
@@ -137,7 +137,8 @@ namespace mu
 		{
 			FParameterDesc param;
 			param.m_name = node.m_name;
-			param.m_uid = node.m_uid;
+			const TCHAR* CStr = ToCStr(node.m_uid);
+			param.m_uid.ImportTextItem(CStr, 0, nullptr, nullptr);
 			param.m_type = PARAMETER_TYPE::T_FLOAT;
 			param.m_defaultValue.Set<ParamFloatType>(node.m_defaultValue);
 
@@ -176,7 +177,8 @@ namespace mu
 		{
 			FParameterDesc param;
 			param.m_name = node.m_name;
-			param.m_uid = node.m_uid;
+			const TCHAR* CStr = ToCStr(node.m_uid);
+			param.m_uid.ImportTextItem(CStr, 0, nullptr, nullptr);
 			param.m_type = PARAMETER_TYPE::T_INT;
 			param.m_defaultValue.Set<ParamIntType>(node.m_defaultValue);
 
diff --git a/Engine/Plugins/Experimental/Mutable/Source/MutableTools/Private/MuT/CodeGenerator_String.cpp b/Engine/Plugins/Experimental/Mutable/Source/MutableTools/Private/MuT/CodeGenerator_String.cpp
--- a/Engine/Plugins/Experimental/Mutable/Source/MutableTools/Private/MuT/CodeGenerator_String.cpp
+++ b/Engine/Plugins/Experimental/Mutable/Source/MutableTools/Private/MuT/CodeGenerator_String.cpp
@@ -87,7 +87,8 @@ namespace mu
 		{
 			FParameterDesc param;
 			param.m_name = node.m_name;
-			param.m_uid = node.m_uid;
+			const TCHAR* CStr = ToCStr(node.m_uid);
+			param.m_uid.ImportTextItem(CStr, 0, nullptr, nullptr);
 			param.m_type = PARAMETER_TYPE::T_FLOAT;
 			param.m_defaultValue.Set<ParamStringType>(node.m_defaultValue);
 
