[Mutable] Fix missing passthrough textures.

#rnx

diff --git a/Engine/Plugins/Experimental/Mutable/Source/CustomizableObject/Classes/MuCO/CustomizableObject.h b/Engine/Plugins/Experimental/Mutable/Source/CustomizableObject/Classes/MuCO/CustomizableObject.h
--- a/Engine/Plugins/Experimental/Mutable/Source/CustomizableObject/Classes/MuCO/CustomizableObject.h
+++ b/Engine/Plugins/Experimental/Mutable/Source/CustomizableObject/Classes/MuCO/CustomizableObject.h
@@ -1370,7 +1370,7 @@ private:
 	// This is a manual version number for the binary blobs in this asset.
 	// Increasing it invalidates all the previously compiled models.
 	// Warning: If while merging code both versions have changed, take the highest+1.
-	static const int32 CurrentSupportedVersion = 416;
+	static const int32 CurrentSupportedVersion = 417;
 
 	// Compile the object for a specific platform - Compile for Cook Customizable Object
 	void CompileForTargetPlatform(const ITargetPlatform* TargetPlatform);
diff --git a/Engine/Plugins/Experimental/Mutable/Source/CustomizableObject/Private/MuCO/CustomizableObject.cpp b/Engine/Plugins/Experimental/Mutable/Source/CustomizableObject/Private/MuCO/CustomizableObject.cpp
--- a/Engine/Plugins/Experimental/Mutable/Source/CustomizableObject/Private/MuCO/CustomizableObject.cpp
+++ b/Engine/Plugins/Experimental/Mutable/Source/CustomizableObject/Private/MuCO/CustomizableObject.cpp
@@ -337,6 +337,7 @@ void UCustomizableObject::ClearCompiledData()
 	ReferenceSkeletalMeshesData.Empty();
 	ReferencedMaterials.Empty();
 	ReferencedMaterialSlotNames.Empty();
+	ReferencedPassThroughTextures.Empty();
 	ReferencedSkeletons.Empty();
 	ImageProperties.Empty();
 	ParameterUIDataMap.Empty();
@@ -509,6 +510,15 @@ void UCustomizableObject::SaveCompiledData(FArchive& MemoryWriter, bool bIsCooki
 		MemoryWriter << StringRef;
 	}
 
+	int32 NumPassthroughTextures = ReferencedPassThroughTextures.Num();
+	MemoryWriter << NumPassthroughTextures;
+
+	for (const TSoftObjectPtr<UTexture>& PassthroughTexture : ReferencedPassThroughTextures)
+	{
+		FString StringRef = PassthroughTexture.ToSoftObjectPath().ToString();
+		MemoryWriter << StringRef;
+	}
+
 	MemoryWriter << ImageProperties;
 	MemoryWriter << ParameterUIDataMap;
 	MemoryWriter << StateUIDataMap;
@@ -610,6 +620,17 @@ void UCustomizableObject::LoadCompiledData(FArchive& MemoryReader, const ITarget
 			ReferencedSkeletons.Add(TSoftObjectPtr<USkeleton>(FSoftObjectPath(StringRef)));
 		}
 
+		int32 NumPassthroughTextures = 0;
+		MemoryReader << NumPassthroughTextures;
+
+		for (int32 Index = 0; Index < NumPassthroughTextures; ++Index)
+		{
+			FString StringRef;
+			MemoryReader << StringRef;
+
+			ReferencedPassThroughTextures.Add(TSoftObjectPtr<UTexture>(FSoftObjectPath(StringRef)));
+		}
+
 		MemoryReader << ImageProperties;
 		MemoryReader << ParameterUIDataMap;
 		MemoryReader << StateUIDataMap;
