[Mutable] Added new UI to show the Instance Parameters.
- Added CVar to use the old UI.
- Fixed viewport scale and rotation buttons. It seems that something from the viewport class changed and broke the viewport buttons.
- Now UI metadata of Object Nodes will be only added if its group node is set to toggle.
- Now Table Nodes also generate UI metadata.

#rnx

diff --git a/Engine/Plugins/Experimental/Mutable/Source/CustomizableObject/Classes/MuCO/CustomizableObjectInstance.h b/Engine/Plugins/Experimental/Mutable/Source/CustomizableObject/Classes/MuCO/CustomizableObjectInstance.h
--- a/Engine/Plugins/Experimental/Mutable/Source/CustomizableObject/Classes/MuCO/CustomizableObjectInstance.h
+++ b/Engine/Plugins/Experimental/Mutable/Source/CustomizableObject/Classes/MuCO/CustomizableObjectInstance.h
@@ -334,6 +334,7 @@ public:
 	
 	bool bShowOnlyRuntimeParameters = true;
 	bool bShowOnlyRelevantParameters = true;
+	bool bShowUISections = false;
 
 	// DEPRECATED: Use the method in the CustomizableObject instead which takes an index among all parameters
 	// Returns how many possible options an int parameter has
diff --git a/Engine/Plugins/Experimental/Mutable/Source/CustomizableObjectEditor/Private/MuCOE/CustomizableInstanceDetails.cpp b/Engine/Plugins/Experimental/Mutable/Source/CustomizableObjectEditor/Private/MuCOE/CustomizableInstanceDetails.cpp
--- a/Engine/Plugins/Experimental/Mutable/Source/CustomizableObjectEditor/Private/MuCOE/CustomizableInstanceDetails.cpp
+++ b/Engine/Plugins/Experimental/Mutable/Source/CustomizableObjectEditor/Private/MuCOE/CustomizableInstanceDetails.cpp
@@ -2,17 +2,44 @@
 
 #include "MuCOE/CustomizableInstanceDetails.h"
 
+#include "CustomizableObjectInstanceEditor.h"
+
+#include "ContentBrowserModule.h"
 #include "DetailLayoutBuilder.h"
 #include "DetailWidgetRow.h"
+#include "HAL/PlatformApplicationMisc.h"
 #include "IDetailsView.h"
-#include "MuCO/CustomizableObjectInstance.h"
+#include "IDetailGroup.h"
+#include "ScopedTransaction.h"
+#include "Serialization/BufferArchive.h"
+#include "Toolkits/ToolkitManager.h"
+
 #include "MuCO/CustomizableInstancePrivateData.h"
+#include "MuCO/CustomizableObject.h"
+#include "MuCO/CustomizableObjectInstance.h"
+#include "MuCO/CustomizableObjectSystem.h"
 #include "MuCOE/SCustomizableInstanceProperties.h"
+#include "MuCOE/SMutableTextSearchBox.h"
+#include "MuCOE/UnrealEditorPortabilityHelpers.h"
+
+#include "Widgets/Colors/SColorBlock.h"
+#include "Widgets/Colors/SColorPicker.h"
+#include "Widgets/Input/SButton.h"
+#include "Widgets/Input/SCheckBox.h"
+#include "Widgets/Input/SSearchBox.h"
+#include "Widgets/Input/STextComboBox.h"
+#include "Widgets/Layout/SExpandableArea.h"
+#include "Widgets/Layout/SUniformGridPanel.h"
 
 class UObject;
 
 #define LOCTEXT_NAMESPACE "CustomizableInstanceDetails"
 
+TAutoConsoleVariable<bool> CVarUseOldInstanceUI(
+	TEXT("mutable.UseOldInstanceUI"),
+	false,
+	TEXT("Enables the old Parameters UI for the Instance Properties tab."));
+
 
 TSharedRef<IDetailCustomization> FCustomizableInstanceDetails::MakeInstance()
 {
@@ -22,34 +49,189 @@ TSharedRef<IDetailCustomization> FCustomizableInstanceDetails::MakeInstance()
 
 void FCustomizableInstanceDetails::CustomizeDetails(const TSharedPtr<IDetailLayoutBuilder>& DetailBuilder)
 {
+	LayoutBuilder = DetailBuilder;
+
 	const IDetailsView* DetailsView = DetailBuilder->GetDetailsView();
 	check(DetailsView->GetSelectedObjects().Num());
 
-	UCustomizableObjectInstance* CustomInstance = Cast<UCustomizableObjectInstance>(DetailsView->GetSelectedObjects()[0].Get());
-	check(CustomInstance);
-	
-	LayoutBuilder = DetailBuilder;
+	CustomInstance = Cast<UCustomizableObjectInstance>(DetailsView->GetSelectedObjects()[0].Get());
+	check(CustomInstance.IsValid());
+
+	TSharedPtr<IToolkit> FoundAssetEditor = FToolkitManager::Get().FindEditorForAsset(CustomInstance.Get()); // Tab spawned in a COEInstanceEditor
+	if (!FoundAssetEditor)
+	{
+		FoundAssetEditor = FToolkitManager::Get().FindEditorForAsset(CustomInstance->GetCustomizableObject()); // Tab spawned in a COEditor
+	}
+	check(FoundAssetEditor);
+
+	WeakEditor = StaticCastSharedPtr<ICustomizableObjectInstanceEditor>(FoundAssetEditor).ToWeakPtr();
+
+	if (CustomInstance->IsSelectedParameterProfileDirty())
+	{
+		CustomInstance->SaveParametersToProfile(CustomInstance->SelectedProfileIndex);
+	}
+
+	// Delegate to refresh the detils when the instance has finished the Update
+	CustomInstance->UpdatedNativeDelegate.AddSP(this, &FCustomizableInstanceDetails::InstanceUpdated);
 
+	// New Category that will store all properties widgets
 	IDetailCategoryBuilder& ResourcesCategory = DetailBuilder->EditCategory("Generated Resources");
-	
+	IDetailCategoryBuilder& VisibilitySettingsCategory = DetailBuilder->EditCategory("ParametersVisibility");
+	IDetailCategoryBuilder& ParametersCategory = DetailBuilder->EditCategory("Instance Parameters");
+	IDetailCategoryBuilder& OldParametersCategory = DetailBuilder->EditCategory("Old Instance Parameters");
+	IDetailCategoryBuilder& SkeletalMeshCategory = DetailBuilder->EditCategory("CustomizableSkeletalMesh");
+	IDetailCategoryBuilder& TextureParametersCategory = DetailBuilder->EditCategory("TextureParameter");
+
+	// Showing warning message in case that the instance has not been compiled
+	UCustomizableObject* CustomizableObject = CustomInstance->GetCustomizableObject();
+	if (!CustomizableObject)
+	{
+		VisibilitySettingsCategory.AddCustomRow(LOCTEXT("CustomizableInstanceDetails", "Instance Parameters"))
+		[
+			SNew(STextBlock).Text(LOCTEXT("Model not compiled", "Model not compiled"))
+		];
+
+		return;
+	}
+
 	TArray<UObject*> Private;
-	Private.Add(CustomInstance->GetPrivate(	));
-	
+	Private.Add(CustomInstance->GetPrivate());
+
 	FAddPropertyParams PrivatePropertyParams;
 	PrivatePropertyParams.HideRootObjectNode(true);
-	
+
 	IDetailPropertyRow* PrivateDataRow = ResourcesCategory.AddExternalObjects(Private, EPropertyLocation::Default, PrivatePropertyParams);
 
-	IDetailCategoryBuilder& MainCategory = DetailBuilder->EditCategory( "InstanceParameters" );
+	// Store which Texture Parameter values can be selected.
+	GenerateTextureParameterOptions();
 
-	MainCategory.AddCustomRow( LOCTEXT("CustomizableInstanceDetails", "Instance Parameters") )
-	[
-		SNew(SCustomizableInstanceProperties)
-			.CustomInstance(CustomInstance)
+	// In case that something of the new UI doesn't work as expected
+	if(CVarUseOldInstanceUI.GetValueOnGameThread())
+	{
+		OldParametersCategory.AddCustomRow(LOCTEXT("CustomizableInstanceDetails", "Old Instance Parameters"))
+		[
+			SNew(SCustomizableInstanceProperties)
+			.CustomInstance(CustomInstance.Get())
 			.InstanceDetails(SharedThis(this))
+		];
+
+		return;
+	}
+
+	// State Selector Widget
+	VisibilitySettingsCategory.AddCustomRow(LOCTEXT("CustomizableInstanceDetails_StateSelector", "State"))
+	.NameContent()
+	[
+		SNew(STextBlock)
+		.Text(LOCTEXT("StateSelector_Text","State"))
+		.ToolTipText(LOCTEXT("StateSelector_Tooltip","Select a state."))
+	]
+	.ValueContent()
+	.HAlign(EHorizontalAlignment::HAlign_Fill)
+	[
+		GenerateStateSelector()
+	];
+
+	// Profile Selector Widget
+	VisibilitySettingsCategory.AddCustomRow(LOCTEXT("CustomizableInstanceDetails_InstanceProfileSelector", "Preview Instance Parameter Profiles"))
+	.NameContent()
+	[
+		SNew(STextBlock)
+		.Text(LOCTEXT("ProfileSelector_Text","Parameter Profile"))
+		.ToolTipText(LOCTEXT("ProfileSelector_Tooltip", "Select a profile to save the parameter options selected."))
+	]
+	.ValueContent()
+	.HAlign(EHorizontalAlignment::HAlign_Fill)
+	[
+		GenerateInstanceProfileSelector()
+	];
+
+	// Only Runtime Parameters Option
+	VisibilitySettingsCategory.AddCustomRow( LOCTEXT("CustomizableInstanceDetails_RuntimeParm", "Only Runtime") )
+	.NameContent()
+	[
+		SNew(STextBlock)
+		.Text(FText::FromString("Only Runtime"))
+	]
+	.ValueContent()
+	.HAlign(EHorizontalAlignment::HAlign_Fill)
+	[
+		SNew(SCheckBox)
+		.IsChecked(CustomInstance->bShowOnlyRuntimeParameters ? ECheckBoxState::Checked : ECheckBoxState::Unchecked)
+		.OnCheckStateChanged(this, &FCustomizableInstanceDetails::OnShowOnlyRuntimeSelectionChanged)
+	];
+	
+	// Only Relevant Parameters Option
+	VisibilitySettingsCategory.AddCustomRow( LOCTEXT("CustomizableInstanceDetails_RelevantParam", "Only Relevant") )
+	.NameContent()
+	[
+		SNew(STextBlock)
+		.Text(FText::FromString("Only Relevant"))
+	]
+	.ValueContent()
+	.HAlign(EHorizontalAlignment::HAlign_Fill)
+	[
+		SNew(SCheckBox)
+		.IsChecked(CustomInstance->bShowOnlyRelevantParameters ? ECheckBoxState::Checked : ECheckBoxState::Unchecked)
+		.OnCheckStateChanged(this, &FCustomizableInstanceDetails::OnShowOnlyRelevantSelectionChanged)
+	];
+	
+	// Show UI sections Option
+	VisibilitySettingsCategory.AddCustomRow( LOCTEXT("CustomizableInstanceDetails_UISections", "UI Sections") )
+	.NameContent()
+	[
+		SNew(STextBlock)
+		.Text(FText::FromString("UI Sections"))
+	]
+	.ValueContent()
+	.HAlign(EHorizontalAlignment::HAlign_Fill)
+	[
+		SNew(SCheckBox)
+		.IsChecked(CustomInstance->bShowUISections ? ECheckBoxState::Checked : ECheckBoxState::Unchecked)
+		.OnCheckStateChanged(this, &FCustomizableInstanceDetails::OnUseUISectionsSelectionChanged)
+	];
+
+	// Copy, Paste and Reset Parameters
+	ParametersCategory.AddCustomRow(LOCTEXT("CustomizableInstanceDetails_CopyPasteResetButtons", "Copy Paste Reset"))
+	[
+		SNew(SHorizontalBox)
+		+ SHorizontalBox::Slot()
+		.FillWidth(2.0f)
+		.Padding(0.0f, 5.0f, 0.0f, 5.0f)
+		[
+			SNew(SButton)
+			.HAlign(HAlign_Center)
+			.Text(LOCTEXT("Copy_Parameters", "Copy Parameters"))
+			.OnClicked(this, &FCustomizableInstanceDetails::OnCopyAllParameters)
+			.IsEnabled(CustomInstance->HasAnyParameters())
+			.ToolTipText(FText::FromString(FString("Copy the preview Instance parameters")))
+		]
+		+ SHorizontalBox::Slot()
+		.FillWidth(2.0f)
+		.Padding(0.0f, 5.0f, 0.0f, 5.0f)
+		[
+			SNew(SButton)
+			.HAlign(HAlign_Center)
+			.Text(LOCTEXT("Paste_Parameters", "Paste Parameters"))
+			.OnClicked(this, &FCustomizableInstanceDetails::OnPasteAllParameters)
+			.IsEnabled(CustomInstance->HasAnyParameters())
+			.ToolTipText(FText::FromString(FString("Paste the preview Instance parameters")))
+		]
+		+ SHorizontalBox::Slot()
+		.FillWidth(2.0f)
+		.Padding(0.0f, 5.0f, 0.0f, 5.0f)
+		[
+			SNew(SButton)
+			.HAlign(HAlign_Center)
+			.Text(LOCTEXT("Reset_Integer_Paramaters", "Reset parameters"))
+			.OnClicked(this, &FCustomizableInstanceDetails::OnResetAllParameters)
+			.IsEnabled(CustomInstance->HasAnyParameters())
+			.ToolTipText(FText::FromString(FString("Clear the preview Instance parameters")))
+		]
 	];
 
-	DetailBuilder->EditCategory("TextureParameter");
+	// Parameters Widgets
+	GenerateParametersView(ParametersCategory);
 }
 
 
@@ -61,5 +243,1898 @@ void FCustomizableInstanceDetails::Refresh() const
 	}
 }
 
+void FCustomizableInstanceDetails::InstanceUpdated(UCustomizableObjectInstance* Instance) const
+{
+	Refresh();
+}
+
+
+// STATE SELECTOR -----------------------------------------------------------------------------------------------------------------
+
+TSharedRef<SWidget> FCustomizableInstanceDetails::GenerateStateSelector()
+{
+	UCustomizableObject* CustomizableObject = CustomInstance->GetCustomizableObject();
+
+	// States selector options
+	const int NumStates = CustomizableObject->GetStateCount();
+	const int CurrentState = CustomInstance->GetState();
+	TSharedPtr<FString> CurrentStateName = nullptr;
+	
+	//I think that this is not necessary. There is always a "Default" state
+	if (NumStates == 0)
+	{
+		StateNames.Add(MakeShareable(new FString("Default")));
+		CurrentStateName = StateNames.Last();
+	}
+
+	for (int StateIndex = 0; StateIndex < NumStates; ++StateIndex)
+	{
+		if (StateIndex == CurrentState)
+		{
+			CurrentStateName = MakeShareable(new FString(CustomizableObject->GetStateName(StateIndex)));
+			StateNames.Add(CurrentStateName);
+		}
+		else
+		{
+			StateNames.Add(MakeShareable(new FString(CustomizableObject->GetStateName(StateIndex))));
+		}
+	}
+
+	StateNames.Sort([](const TSharedPtr<FString>& LHS, const TSharedPtr<FString>& RHS)
+	{
+		return *LHS < *RHS;
+	});
+
+	return SNew(STextComboBox)
+		.OptionsSource(&StateNames)
+		.InitiallySelectedItem((CustomInstance->GetState() != -1) ? CurrentStateName : StateNames[0])
+		.OnSelectionChanged(this, &FCustomizableInstanceDetails::OnStateComboBoxSelectionChanged);
+}
+
+
+void FCustomizableInstanceDetails::OnStateComboBoxSelectionChanged(TSharedPtr<FString> Selection, ESelectInfo::Type SelectInfo)
+{
+	CustomInstance->PreEditChange(nullptr);
+	CustomInstance->SetCurrentState(*Selection);
+	CustomInstance->UpdateSkeletalMeshAsync(true, true);
+	CustomInstance->PostEditChange();
+
+	// Non-continuous change: collect garbage.
+	GEngine->ForceGarbageCollection();
+}
+
+
+// INSTANCE PROFILE SELECTOR -----------------------------------------------------------------------------------------------------------------
+
+TSharedRef<SWidget> FCustomizableInstanceDetails::GenerateInstanceProfileSelector()
+{
+	UCustomizableObject* CustomizableObject = CustomInstance->GetCustomizableObject();
+	const int32 ProfileIdx = CustomInstance->SelectedProfileIndex;
+
+	ParameterProfileNames.Emplace(MakeShared<FString>("None"));
+	TSharedPtr<FString> CurrentProfileName = ParameterProfileNames.Last();
+
+	for (FProfileParameterDat& Profile : CustomInstance->GetCustomizableObject()->InstancePropertiesProfiles)
+	{
+		ParameterProfileNames.Emplace(MakeShared<FString>(Profile.ProfileName));
+
+		if (ProfileIdx != INDEX_NONE)
+		{
+			const FProfileParameterDat CurrentInstanceProfile = CustomizableObject->InstancePropertiesProfiles[ProfileIdx];
+
+			if (Profile.ProfileName == CurrentInstanceProfile.ProfileName)
+			{
+				CurrentProfileName = ParameterProfileNames.Last();
+			}
+		}
+	}
+
+	ParameterProfileNames.Sort([](const TSharedPtr<FString>& LHS, const TSharedPtr<FString>& RHS)
+	{
+		return *LHS < *RHS;
+	});
+
+	return SNew(SHorizontalBox)
+		+ SHorizontalBox::Slot()
+		.FillWidth(1.0f)
+		[
+			SNew(STextComboBox)
+			.OptionsSource(&ParameterProfileNames)
+			.OnSelectionChanged(this, &FCustomizableInstanceDetails::OnProfileSelectedChanged)
+			.ToolTipText(FText::FromString(ParameterProfileNames.Num() > 1 ? FString("Select an existing profile") : FString("No profiles are available")))
+			.InitiallySelectedItem(CurrentProfileName)
+			.IsEnabled(ParameterProfileNames.Num() > 1)
+		]
+
+		+ SHorizontalBox::Slot()
+		.AutoWidth()
+		[
+			SNew(SButton)
+			.Text(LOCTEXT("AddButtonLabel", " + "))
+			.ToolTipText(FText::FromString(CustomInstance->HasAnyParameters() ? FString("Add new profile") : FString("Create a profile functionality is not available, no parameters were found.")))
+			.IsEnabled(CustomInstance->HasAnyParameters())
+			.IsFocusable(false)
+			.OnClicked(this, &FCustomizableInstanceDetails::CreateParameterProfileWindow)
+		]
+
+		+ SHorizontalBox::Slot()
+		.AutoWidth()
+		[
+			SNew(SButton)
+			.Text(LOCTEXT("RemoveButtonLabel", " - "))
+			.ToolTipText(FText::FromString(CustomInstance->HasAnyParameters() ? FString("Delete selected profile") : FString("Delete selected profile functionality is not available, no profile is selected.")))
+			.IsEnabled(CustomInstance->SelectedProfileIndex != INDEX_NONE)
+			.IsFocusable(false)
+			.OnClicked(this, &FCustomizableInstanceDetails::RemoveParameterProfile)
+		];
+}
+
+
+class SProfileParametersWindow : public SWindow
+{
+public:
+	SLATE_BEGIN_ARGS(SProfileParametersWindow) {}
+		SLATE_ARGUMENT(FText, DefaultAssetPath)
+		SLATE_ARGUMENT(FText, DefaultFileName)
+	SLATE_END_ARGS()
+
+	SProfileParametersWindow() : UserResponse(EAppReturnType::Cancel) {}
+	void Construct(const FArguments& InArgs);
+
+public:
+	/** Displays the dialog in a blocking fashion */
+	EAppReturnType::Type ShowModal();
+
+	/** FileName getter */
+	FString GetFileName() const { return FileName.ToString(); }
+
+protected:
+
+	//FReply OnKeyDown(const FGeometry& MyGeometry, const FKeyEvent& InKeyEvent);
+	FReply OnButtonClick(EAppReturnType::Type ButtonID);
+	void OnNameChange(const FText& NewName, ETextCommit::Type CommitInfo);
+	
+	EAppReturnType::Type UserResponse;
+	FText AssetPath;
+	FText FileName;
+
+public:
+	TWeakObjectPtr<UCustomizableObjectInstance> CustomInstance;
+
+	FCustomizableInstanceDetails* InstanceDetails = nullptr;
+};
+
+
+FReply FCustomizableInstanceDetails::CreateParameterProfileWindow()
+{
+	const TSharedRef<SProfileParametersWindow> FolderDlg =
+		SNew(SProfileParametersWindow)
+		.DefaultAssetPath(LOCTEXT("DefaultAssethPath", "/Game"))
+		.DefaultFileName(FText::FromString("ProfileParameterData"));
+
+	FolderDlg->CustomInstance = CustomInstance;
+	FolderDlg->InstanceDetails = this;
+	FolderDlg->Construct(SProfileParametersWindow::FArguments());
+	FolderDlg->ShowModal();
+
+	return FReply::Handled();
+}
+
+
+FReply FCustomizableInstanceDetails::RemoveParameterProfile()
+{
+	UCustomizableObject* CustomizableObject = CustomInstance->GetCustomizableObject();
+
+	const int32 ProfileIdx = CustomInstance->SelectedProfileIndex;
+
+	if (ProfileIdx == INDEX_NONE)
+	{
+		return FReply::Handled();
+	}
+
+	TArray<FProfileParameterDat>& Profiles = CustomizableObject->InstancePropertiesProfiles;
+
+	Profiles.RemoveAt(ProfileIdx);
+	CustomInstance->SelectedProfileIndex = INDEX_NONE;
+	CustomizableObject->Modify();
+
+	CustomInstance->PreEditChange(nullptr);
+	CustomInstance->UpdateSkeletalMeshAsync(true, true);
+	CustomInstance->PostEditChange();
+
+	// Non-continuous change: collect garbage.
+	GEngine->ForceGarbageCollection();
+
+	return FReply::Handled();
+}
+
+
+void FCustomizableInstanceDetails::OnProfileSelectedChanged(TSharedPtr<FString> Selection, ESelectInfo::Type SelectInfo)
+{
+	const int32 ProfileIdx = CustomInstance->SelectedProfileIndex;
+	if (CustomInstance->IsSelectedParameterProfileDirty())
+	{
+		CustomInstance->SaveParametersToProfile(ProfileIdx);
+	}
+
+	if (*Selection == "None")
+	{
+		CustomInstance->SelectedProfileIndex = INDEX_NONE;
+	}
+	else
+	{
+		//Set selected profile
+		TArray<FProfileParameterDat>& Profiles = CustomInstance->GetCustomizableObject()->InstancePropertiesProfiles;
+		for (int32 Idx = 0; Idx < Profiles.Num(); ++Idx)
+		{
+			if (Profiles[Idx].ProfileName == *Selection)
+			{
+				CustomInstance->SelectedProfileIndex = Idx;
+				break;
+			}
+		}
+	}	
+
+	CustomInstance->LoadParametersFromProfile(CustomInstance->SelectedProfileIndex);
+	CustomInstance->PreEditChange(nullptr);
+	CustomInstance->UpdateSkeletalMeshAsync(true, true);
+	CustomInstance->PostEditChange();
+
+	// Non-continuous change: collect garbage.
+	GEngine->ForceGarbageCollection();
+}
+
+
+void FCustomizableInstanceDetails::OnShowOnlyRuntimeSelectionChanged(ECheckBoxState InCheckboxState)
+{
+	CustomInstance->bShowOnlyRuntimeParameters = InCheckboxState == ECheckBoxState::Checked;
+	Refresh();
+}
+
+
+void FCustomizableInstanceDetails::OnShowOnlyRelevantSelectionChanged(ECheckBoxState InCheckboxState)
+{
+	CustomInstance->bShowOnlyRelevantParameters = InCheckboxState == ECheckBoxState::Checked;
+	Refresh();
+}
+
+
+void FCustomizableInstanceDetails::OnUseUISectionsSelectionChanged(ECheckBoxState InCheckboxState)
+{
+	CustomInstance->bShowUISections = InCheckboxState == ECheckBoxState::Checked;
+	Refresh();
+}
+
+
+// PARAMETER INFO STRUCT -----------------------------------------------------------------------------------------------------------------
+
+struct FParameterInfo
+{
+	int32 ParameterIndex = -1;
+	int32 ParameterUIOrder = 0;
+	FString ParameterName = "";
+
+	bool operator==(const FParameterInfo& Other)
+	{
+		return ParameterName == Other.ParameterName;
+	}
+
+	friend bool operator<(const FParameterInfo& A, const FParameterInfo& B)
+	{
+		if (A.ParameterUIOrder != B.ParameterUIOrder)
+		{
+			return A.ParameterUIOrder < B.ParameterUIOrder;
+		}
+		else
+		{
+			return A.ParameterName < B.ParameterName;
+		}
+	}
+};
+
+
+// PARAMETERS WIDGET GENERATION -----------------------------------------------------------------------------------------------------------------
+
+void FCustomizableInstanceDetails::GenerateParametersView(IDetailCategoryBuilder& DetailsCategory)
+{
+	// Make this arrays locals
+	ParamChildren.Empty();
+	ParamHasParent.Empty();
+	GeneratedSections.Empty();
+
+	TArray<FParameterInfo> ParametersTree;
+	UCustomizableObject* CustomizableObject = CustomInstance->GetCustomizableObject();
+
+	if (!CustomizableObject)
+	{
+		return;
+	}
+	
+	if (CustomInstance->bShowOnlyRuntimeParameters)
+	{
+		const int32 NumStateParameters = CustomizableObject->GetStateParameterCount(CustomInstance->GetState());
+
+		for (int32 ParamIndexInState = 0; ParamIndexInState < NumStateParameters; ++ParamIndexInState)
+		{
+			FParameterInfo ParameterSortInfo;
+			ParameterSortInfo.ParameterIndex = CustomizableObject->GetStateParameterIndex(CustomInstance->GetState(), ParamIndexInState);
+
+			if (CustomInstance->IsParameterRelevant(ParameterSortInfo.ParameterIndex) && IsVisible(ParameterSortInfo.ParameterIndex))
+			{
+				ParameterSortInfo.ParameterUIOrder = CustomizableObject->GetParameterUIMetadataFromIndex(ParameterSortInfo.ParameterIndex).ParamUIMetadata.UIOrder;
+				ParameterSortInfo.ParameterName = CustomizableObject->GetParameterName(ParameterSortInfo.ParameterIndex);
+				ParametersTree.Add(ParameterSortInfo);
+			}
+		}
+
+		ParametersTree.Sort();
+
+		for (int32 ParamIndexInObject = 0; ParamIndexInObject < ParametersTree.Num(); ++ParamIndexInObject)
+		{
+			const FParameterInfo& ParamInfo = ParametersTree[ParamIndexInObject];
+			
+			if (CustomInstance->bShowUISections)
+			{
+				IDetailGroup* CurrentSection = GenerateParameterSection(ParamIndexInObject, DetailsCategory);
+				check(CurrentSection);
+
+				if (!IsMultidimensionalProjector(ParamIndexInObject))
+				{
+					GenerateWidgetRow(CurrentSection->AddWidgetRow(), ParamInfo.ParameterName, ParamInfo.ParameterIndex);
+				}
+				else
+				{
+					IDetailGroup* ProjectorGroup = &CurrentSection->AddGroup(FName(*ParamInfo.ParameterName), FText::GetEmpty());
+
+					// Call Order between the following lines maters.
+					ParentsGroups.Add(ParamInfo.ParameterName, ProjectorGroup);
+					GenerateWidgetRow(ProjectorGroup->HeaderRow(), ParamInfo.ParameterName, ParamIndexInObject);
+				}
+			}
+			else
+			{
+				if (!IsMultidimensionalProjector(ParamIndexInObject))
+				{
+					GenerateWidgetRow(DetailsCategory.AddCustomRow(FText::FromString(ParamInfo.ParameterName)), ParamInfo.ParameterName, ParamInfo.ParameterIndex);
+				}
+				else
+				{
+					IDetailGroup* ProjectorGroup = &DetailsCategory.AddGroup(FName(*ParamInfo.ParameterName), FText::GetEmpty());
+
+					// Call Order between the following lines maters.
+					ParentsGroups.Add(ParamInfo.ParameterName, ProjectorGroup);
+					GenerateWidgetRow(ProjectorGroup->HeaderRow(), ParamInfo.ParameterName, ParamIndexInObject);
+				}
+			}
+		}
+	}
+	else
+	{
+		const int32 NumObjectParameter = CustomizableObject->GetParameterCount();
+	
+		//TODO: get all parameters and sort, then make the next "for" use that sorted list as source of indexes
+		for (int32 ParamIndexInObject = 0; ParamIndexInObject < NumObjectParameter; ++ParamIndexInObject)
+		{
+			FParameterInfo ParameterSortInfo;
+			ParameterSortInfo.ParameterIndex = ParamIndexInObject;
+			if ((!CustomInstance->bShowOnlyRelevantParameters || CustomInstance->IsParameterRelevant(ParameterSortInfo.ParameterIndex)) && IsVisible(ParameterSortInfo.ParameterIndex))
+			{
+				ParameterSortInfo.ParameterUIOrder = CustomizableObject->GetParameterUIMetadataFromIndex(ParameterSortInfo.ParameterIndex).ParamUIMetadata.UIOrder;
+				ParameterSortInfo.ParameterName = CustomizableObject->GetParameterName(ParameterSortInfo.ParameterIndex);
+				ParametersTree.Add(ParameterSortInfo);
+			}
+		}
+	
+		ParametersTree.Sort();
+	
+		for (int32 ParamIndexInObject = 0; ParamIndexInObject < ParametersTree.Num(); ++ParamIndexInObject)
+		{
+			FillChildrenMap(ParametersTree[ParamIndexInObject].ParameterIndex);
+		}
+	
+		for (int32 ParamIndexInObject = 0; ParamIndexInObject < ParametersTree.Num(); ++ParamIndexInObject)
+		{
+			if (!ParamHasParent.Find(ParametersTree[ParamIndexInObject].ParameterIndex))
+			{
+				RecursivelyAddParamAndChildren(ParametersTree[ParamIndexInObject].ParameterIndex, "", DetailsCategory);
+			}
+		}
+	}
+}
+
+
+void FCustomizableInstanceDetails::RecursivelyAddParamAndChildren(const int32 ParamIndexInObject, const FString ParentName, IDetailCategoryBuilder& DetailsCategory)
+{
+	const UCustomizableObject* CustomizableObject = CustomInstance->GetCustomizableObject();
+	const FString ParamName = CustomizableObject->GetParameterName(ParamIndexInObject);	
+	TArray<int32> Children;
+	
+	ParamChildren.MultiFind(ParamName, Children, true);
+
+	if (ParentName.IsEmpty())
+	{
+		if (CustomInstance->bShowUISections)
+		{
+			IDetailGroup* CurrentSection = GenerateParameterSection(ParamIndexInObject, DetailsCategory);
+			check(CurrentSection);
+
+			if (Children.Num() == 0 && !IsMultidimensionalProjector(ParamIndexInObject))
+			{
+				GenerateWidgetRow(CurrentSection->AddWidgetRow(), ParamName, ParamIndexInObject);
+			}
+			else
+			{
+				IDetailGroup* ParentGroup = &CurrentSection->AddGroup(FName(*ParamName), FText::GetEmpty());
+				
+				// Call Order between the following lines maters.
+				ParentsGroups.Add(ParamName, ParentGroup);
+				GenerateWidgetRow(ParentGroup->HeaderRow(), ParamName, ParamIndexInObject);
+			}
+		}
+		else
+		{
+			if (Children.Num() == 0 && !IsMultidimensionalProjector(ParamIndexInObject))
+			{
+				GenerateWidgetRow(DetailsCategory.AddCustomRow(FText::FromString(ParamName)), ParamName, ParamIndexInObject);
+			}
+			else
+			{
+				IDetailGroup* ParentGroup = &DetailsCategory.AddGroup(FName(*ParamName), FText::GetEmpty());
+				
+				// Call Order between the following lines maters.
+				ParentsGroups.Add(ParamName, ParentGroup);
+				GenerateWidgetRow(ParentGroup->HeaderRow(), ParamName, ParamIndexInObject);
+			}
+		}
+	}
+	else
+	{
+		IDetailGroup* ParentGroup = *ParentsGroups.Find(ParentName);
+
+		if (Children.Num() == 0 && !IsMultidimensionalProjector(ParamIndexInObject))
+		{
+			GenerateWidgetRow(ParentGroup->AddWidgetRow(), ParamName, ParamIndexInObject);
+		}
+		else
+		{
+			IDetailGroup* ChildGroup = &ParentGroup->AddGroup(FName(*ParamName), FText::GetEmpty());
+
+			// Call Order between the following lines maters
+			ParentsGroups.Add(ParamName, ChildGroup);
+			GenerateWidgetRow(ParentGroup->HeaderRow(), ParamName, ParamIndexInObject);
+		}
+	}
+
+	for (const int32 ChildIndexInObject : Children)
+	{
+		RecursivelyAddParamAndChildren(ChildIndexInObject, ParamName, DetailsCategory);
+	}
+}
+
+
+void FCustomizableInstanceDetails::FillChildrenMap(int32 ParamIndexInObject)
+{
+	const UCustomizableObject* CustomizableObject = CustomInstance->GetCustomizableObject();
+	FString ParamName = CustomizableObject->GetParameterName(ParamIndexInObject);
+	FParameterUIData UIData = CustomizableObject->GetParameterUIMetadataFromIndex(ParamIndexInObject);
+
+	const FString* ParentName = UIData.ParamUIMetadata.ExtraInformation.Find(FString("__ParentParamName"));
+
+	if (ParentName)
+	{
+		ParamChildren.Add(*ParentName, ParamIndexInObject);
+		ParamHasParent.Add(ParamIndexInObject, true);
+	}
+}
+
+
+bool FCustomizableInstanceDetails::IsVisible(int32 ParamIndexInObject)
+{
+	const UCustomizableObject* CustomizableObject = CustomInstance->GetCustomizableObject();
+	FString ParamName = CustomizableObject->GetParameterName(ParamIndexInObject);
+	FParameterUIData UIData = CustomizableObject->GetParameterUIMetadataFromIndex(ParamIndexInObject);
+	FString* ParentName = UIData.ParamUIMetadata.ExtraInformation.Find(FString("__ParentParamName"));
+
+	bool IsAProjectorParam = ParamName.EndsWith(FMultilayerProjector::NUM_LAYERS_PARAMETER_POSTFIX)
+		|| (ParamName.EndsWith(FMultilayerProjector::IMAGE_PARAMETER_POSTFIX) && CustomizableObject->IsParameterMultidimensional(ParamIndexInObject))
+		|| (ParamName.EndsWith(FMultilayerProjector::OPACITY_PARAMETER_POSTFIX) && CustomizableObject->IsParameterMultidimensional(ParamIndexInObject))
+		|| (ParamName.EndsWith(FMultilayerProjector::POSE_PARAMETER_POSTFIX));
+
+	if (!IsAProjectorParam && ParentName && CustomInstance->bShowOnlyRelevantParameters)
+	{
+		FString* Value = UIData.ParamUIMetadata.ExtraInformation.Find(FString("__DisplayWhenParentValueEquals"));
+
+		if (Value && CustomInstance->GetIntParameterSelectedOption(*ParentName) != *Value)
+		{
+			return false;
+		}
+	}
+
+	return !IsAProjectorParam;
+}
+
+
+bool FCustomizableInstanceDetails::IsMultidimensionalProjector(int32 ParamIndexInObject)
+{
+	const UCustomizableObject* CustomizableObject = CustomInstance->GetCustomizableObject();
+
+	return CustomizableObject->GetParameterType(ParamIndexInObject) == EMutableParameterType::Projector && CustomInstance->GetCustomizableObject()->IsParameterMultidimensional(ParamIndexInObject);
+}
+
+
+IDetailGroup* FCustomizableInstanceDetails::GenerateParameterSection(const int32 ParamIndexInObject, IDetailCategoryBuilder& DetailsCategory)
+{
+	const UCustomizableObject* CustomizableObject = CustomInstance->GetCustomizableObject();
+
+	FMutableParamUIMetadata UIData = CustomizableObject->GetParameterUIMetadataFromIndex(ParamIndexInObject).ParamUIMetadata;
+	FString SectionName = UIData.UISectionName.IsEmpty() ? "Miscellaneous" : UIData.UISectionName;
+	IDetailGroup* CurrentSection = nullptr;
+
+	if (GeneratedSections.Contains(SectionName))
+	{
+		CurrentSection = *GeneratedSections.Find(SectionName);
+	}
+
+	if (!CurrentSection)
+	{
+		CurrentSection = &DetailsCategory.AddGroup(FName(SectionName), FText::FromString(SectionName));
+
+		GeneratedSections.Add(SectionName, CurrentSection);
+	}
+
+	return CurrentSection;
+}
+
+
+void FCustomizableInstanceDetails::GenerateWidgetRow(FDetailWidgetRow& WidgetRow, const FString& ParamName, const int32 ParamIndexInObject)
+{
+	WidgetRow.NameContent()
+	[
+		SNew(STextBlock)
+		.Text(FText::FromString(ParamName))
+	]
+	.ValueContent()
+	.HAlign(EHorizontalAlignment::HAlign_Fill)
+	.VAlign(EVerticalAlignment::VAlign_Fill)
+	[
+		SNew(SHorizontalBox)
+		+SHorizontalBox::Slot()
+		.HAlign(HAlign_Fill)
+		.Padding(0.0f, 5.0f, 0.0f, 5.0f)
+		[
+			GenerateParameterWidget(ParamIndexInObject)
+		]
+	]
+	.OverrideResetToDefault(FResetToDefaultOverride::Create(FSimpleDelegate::CreateSP(this, &FCustomizableInstanceDetails::OnResetParameterButtonClicked, ParamIndexInObject)))
+	.FilterString(FText::FromString(ParamName));
+}
+
+
+TSharedRef<SWidget> FCustomizableInstanceDetails::GenerateParameterWidget(const int32 ParamIndexInObject)
+{
+	UCustomizableObject* CustomizableObject = CustomInstance->GetCustomizableObject();
+
+	switch (CustomizableObject->GetParameterType(ParamIndexInObject))
+	{
+		case EMutableParameterType::Bool:
+		{
+			return GenerateBoolWidget(ParamIndexInObject);
+		}
+		case EMutableParameterType::Float:
+		{
+			return GenerateFloatWidget(ParamIndexInObject);
+		}
+		case EMutableParameterType::Color:
+		{
+			return GenerateColorWidget(ParamIndexInObject);
+		}
+		case EMutableParameterType::Texture:
+		{
+			return GenerateTextureWidget(ParamIndexInObject);
+		}
+		case EMutableParameterType::Projector:
+		{
+			bool bIsParamMultidimensional = CustomInstance->GetCustomizableObject()->IsParameterMultidimensional(ParamIndexInObject);
+
+			if (!bIsParamMultidimensional)
+			{
+				return GenerateSimpleProjector(ParamIndexInObject);
+			}
+			else 
+			{
+				return GenerateMultidimensionalProjector(ParamIndexInObject);
+			}
+		}
+		case EMutableParameterType::Int:
+		{
+			return GenerateIntWidget(ParamIndexInObject);
+		}
+		case EMutableParameterType::None:
+		{
+			return SNew(STextBlock).Text(LOCTEXT("ParameterTypeNotSupported_Text", "Parameter Type not supported"));
+		}
+	}
+
+	return SNew(STextBlock);
+}
+
+
+// INT PARAMETERS -----------------------------------------------------------------------------------------------------------------
+
+TSharedRef<SWidget> FCustomizableInstanceDetails::GenerateIntWidget(const int32 ParamIndexInObject)
+{
+	UCustomizableObject* CustomizableObject = CustomInstance->GetCustomizableObject();
+
+	int numValues = CustomizableObject->GetIntParameterNumOptions(ParamIndexInObject);
+	bool bIsParamMultidimensional = CustomInstance->GetCustomizableObject()->IsParameterMultidimensional((ParamIndexInObject));
+
+	if (!bIsParamMultidimensional && numValues)
+	{
+		FString ToolTipText = FString("None");
+		FString ParamName = CustomizableObject->GetParameterName(ParamIndexInObject);
+		FString Value = CustomInstance->GetIntParameterSelectedOption(ParamName, -1);
+		TArray<FString> OptionNamesAttribute;
+		int ValueIndex = 0;
+
+		for (int i = 0; i < numValues; ++i)
+		{
+			FString PossibleValue = CustomizableObject->GetIntParameterAvailableOption(ParamIndexInObject, i);
+
+			if (PossibleValue == Value)
+			{
+				ValueIndex = i;
+
+				const FString* Identifier = CustomizableObject->GroupNodeMap.FindKey(FCustomizableObjectIdPair(ParamName, PossibleValue));
+				if (Identifier)
+				{
+					if (FString* CustomizableObjectPath = CustomizableObject->CustomizableObjectPathMap.Find(*Identifier))
+					{
+						ToolTipText = *CustomizableObjectPath;
+					}
+				}
+			}
+
+			OptionNamesAttribute.Add(CustomizableObject->GetIntParameterAvailableOption(ParamIndexInObject, i));
+		}
+
+		return SNew(SMutableTextSearchBox)
+			.ToolTipText(FText::FromString(ToolTipText))
+			.PossibleSuggestions(OptionNamesAttribute)
+			.InitialText(FText::FromString(OptionNamesAttribute[ValueIndex]))
+			.MustMatchPossibleSuggestions(TAttribute<bool>(true))
+			.SuggestionListPlacement(EMenuPlacement::MenuPlacement_ComboBox)
+			.OnTextCommitted(this, &FCustomizableInstanceDetails::OnIntParameterComboBoxChanged, ParamName);
+	}
+
+	return SNew(STextBlock).Text(LOCTEXT("MultidimensionalINTParameter_Text", "Multidimensional INT Parameter not supported"));
+}
+
+
+void FCustomizableInstanceDetails::OnIntParameterComboBoxChanged(TSharedPtr<FString> Selection, ESelectInfo::Type SelectInfo, FString ParamName)
+{
+	TArray<FCustomizableObjectIntParameterValue>& IntParameters = CustomInstance->GetIntParameters();
+	const UCustomizableObject* CustomizableObject = CustomInstance->GetCustomizableObject();
+
+	for (int32 i = 0; i < IntParameters.Num(); ++i)
+	{
+		if (IntParameters[i].ParameterName == ParamName)
+		{
+			const int32 ParamIndexInObject = CustomizableObject->FindParameter(ParamName);
+			if (ParamIndexInObject != INDEX_NONE)
+			{
+				for (int32 v = 0; v < CustomizableObject->GetIntParameterNumOptions(ParamIndexInObject); ++v)
+				{
+					FString ValueName = CustomizableObject->GetIntParameterAvailableOption(ParamIndexInObject, v);
+
+					if (ValueName == *Selection)
+					{
+						IntParameters[i].ParameterValueName = ValueName;
+						break;
+					}
+				}
+			}
+			else
+			{
+				UE_LOG(LogMutable, Error, TEXT("Failed to find parameter."));
+			}
+		}
+	}
+
+	CustomInstance->SetSelectedParameterProfileDirty();
+	CustomInstance->PreEditChange(nullptr);
+	CustomInstance->UpdateSkeletalMeshAsync(true, true);
+	CustomInstance->PostEditChange();
+
+	// Non-continuous change: collect garbage.
+	GEngine->ForceGarbageCollection();
+}
+
+
+// FLOAT PARAMETERS -----------------------------------------------------------------------------------------------------------------
+
+TSharedRef<SWidget> FCustomizableInstanceDetails::GenerateFloatWidget(const int32 ParamIndexInObject)
+{
+	const UCustomizableObject* CustomizableObject = CustomInstance->GetCustomizableObject();
+	FString ParamName = CustomizableObject->GetParameterName(ParamIndexInObject);
+
+	return SNew(SSpinBox<float>)
+		.Value(this, &FCustomizableInstanceDetails::GetFloatParameterValue, ParamName, -1)
+		.MinValue(CustomizableObject->GetParameterUIMetadataFromIndex(ParamIndexInObject).ParamUIMetadata.MinimumValue)
+		.MaxValue(CustomizableObject->GetParameterUIMetadataFromIndex(ParamIndexInObject).ParamUIMetadata.MaximumValue)
+		.OnValueChanged(this, &FCustomizableInstanceDetails::OnFloatParameterChanged, ParamName, -1)
+		.OnEndSliderMovement(this, &FCustomizableInstanceDetails::OnFloatParameterSliderEnd, ParamName, -1);
+}
+
+
+float FCustomizableInstanceDetails::GetFloatParameterValue(FString ParamName, int32 RangeIndex) const
+{
+	return CustomInstance->GetFloatParameterSelectedOption(ParamName, RangeIndex);
+}
+
+
+void FCustomizableInstanceDetails::OnFloatParameterChanged(float Value, FString ParamName, int32 RangeIndex)
+{
+	CustomInstance->SetFloatParameterSelectedOption(ParamName, Value, RangeIndex);
+}
+
+
+void FCustomizableInstanceDetails::OnFloatParameterSliderEnd(float Value, FString ParamName, int32 RangeIndex)
+{
+	CustomInstance->SetFloatParameterSelectedOption(ParamName, Value, RangeIndex);
+
+	CustomInstance->SetSelectedParameterProfileDirty();
+	CustomInstance->PreEditChange(nullptr);
+	CustomInstance->UpdateSkeletalMeshAsync(true, true);
+	CustomInstance->PostEditChange();
+
+	// Non-continuous change: collect garbage.
+	GEngine->ForceGarbageCollection();
+}
+
+
+// TEXTURE PARAMETERS -----------------------------------------------------------------------------------------------------------------
+
+TSharedRef<SWidget> FCustomizableInstanceDetails::GenerateTextureWidget(const int32 ParamIndexInObject)
+{
+	UCustomizableObject* CustomizableObject = CustomInstance->GetCustomizableObject();
+	FString ParamName = CustomizableObject->GetParameterName(ParamIndexInObject);
+	const FName ParameterValue = CustomInstance->GetTextureParameterSelectedOption(ParamName);
+	TSharedPtr<FString> InitiallySelected;
+
+	// Look for the value index
+	for (int32 ValueIndex = 0; ValueIndex < TextureParameterValueNames.Num(); ++ValueIndex)
+	{
+		if (ParameterValue == TextureParameterValues[ValueIndex])
+		{
+			InitiallySelected = TextureParameterValueNames[ValueIndex];
+			break;
+		}
+	}
+
+	return SNew(STextComboBox)
+		.OptionsSource(&TextureParameterValueNames)
+		.InitiallySelectedItem(InitiallySelected)
+		.OnSelectionChanged(this, &FCustomizableInstanceDetails::OnTextureParameterComboBoxSelectionChanged, ParamName);
+}
+
+
+void FCustomizableInstanceDetails::GenerateTextureParameterOptions()
+{
+	UCustomizableObject* CustomizableObject = CustomInstance->GetCustomizableObject();
+
+	for (const FCustomizableObjectTextureParameterValue& TextureParameter : CustomInstance->GetTextureParameters())
+	{
+		// Get default values.
+		const FName& DefaultValue = CustomizableObject->GetTextureParameterDefaultValue(TextureParameter.ParameterName);
+
+		if (!TextureParameterValues.Contains(DefaultValue))
+		{
+			TextureParameterValueNames.Add(MakeShareable(new FString(DefaultValue.ToString())));
+			TextureParameterValues.Add(DefaultValue);
+		}
+
+		// Selected parameter value
+		const FName& SelectedValue = TextureParameter.ParameterValue;
+
+		if (!TextureParameterValues.Contains(SelectedValue))
+		{
+			TextureParameterValueNames.Add(MakeShareable(new FString(SelectedValue.ToString())));
+			TextureParameterValues.Add(SelectedValue);
+		}
+	}
+
+	// Get values from registered providers providers.
+	TArray<FCustomizableObjectExternalTexture> Textures = UCustomizableObjectSystem::GetInstance()->GetTextureParameterValues();
+	for (int i = 0; i < Textures.Num(); ++i)
+	{
+		const FName& Value = FName(Textures[i].Value);
+
+		if (!TextureParameterValues.Contains(Value))
+		{
+			TextureParameterValueNames.Add(MakeShareable(new FString(Textures[i].Name)));
+			TextureParameterValues.Add(Value);
+		}
+	}
+
+	// Get values from TextureParameterDeclarations.
+	for (TObjectPtr<UTexture2D> Declaration : CustomInstance->GetTextureParameterDeclarations())
+	{
+		if (!Declaration)
+		{
+			continue;
+		}
+
+		const FName Value = FName(Declaration.GetPathName());
+
+		if (!TextureParameterValues.Contains(Value))
+		{
+			TextureParameterValueNames.Add(MakeShareable(new FString(Declaration.GetName())));
+			TextureParameterValues.Add(Value);
+		}
+	}
+}
+
+
+void FCustomizableInstanceDetails::OnTextureParameterComboBoxSelectionChanged(TSharedPtr<FString> Selection, ESelectInfo::Type SelectInfo, FString ParamName)
+{
+	CustomInstance->PreEditChange(nullptr);
+
+	auto& TextureParameters = CustomInstance->GetTextureParameters();
+
+	for (int i = 0; i < TextureParameters.Num(); ++i)
+	{
+		if (TextureParameters[i].ParameterName == ParamName)
+		{
+			for (int o = 0; o < TextureParameterValueNames.Num(); ++o)
+			{
+				if (Selection.IsValid() && *TextureParameterValueNames[o] == *Selection)
+				{
+					TextureParameters[i].ParameterValue = TextureParameterValues[o];
+				}
+			}
+		}
+	}
+
+	CustomInstance->SetSelectedParameterProfileDirty();
+	CustomInstance->UpdateSkeletalMeshAsync(true, true);
+	CustomInstance->PostEditChange();
+
+	// Non-continuous change: collect garbage.
+	GEngine->ForceGarbageCollection();
+}
+
+
+// COLOR PARAMETERS -----------------------------------------------------------------------------------------------------------------
+
+TSharedRef<SWidget> FCustomizableInstanceDetails::GenerateColorWidget(const int32 ParamIndexInObject)
+{
+	UCustomizableObject* CustomizableObject = CustomInstance->GetCustomizableObject();
+	FString ParamName = CustomizableObject->GetParameterName(ParamIndexInObject);
+
+	return SNew(SColorBlock)
+		.Color(this, &FCustomizableInstanceDetails::GetColorParameterValue, ParamName)
+		.ShowBackgroundForAlpha(false)
+		.AlphaDisplayMode(EColorBlockAlphaDisplayMode::Ignore)
+		.UseSRGB(false)
+		.OnMouseButtonDown(this, &FCustomizableInstanceDetails::OnColorBlockMouseButtonDown, ParamName)
+		.CornerRadius(FVector4(4.0f, 4.0f, 4.0f, 4.0f));
+}
+
+
+FLinearColor FCustomizableInstanceDetails::GetColorParameterValue(FString ParamName) const
+{
+	return CustomInstance->GetColorParameterSelectedOption(ParamName);
+}
+
+
+FReply FCustomizableInstanceDetails::OnColorBlockMouseButtonDown(const FGeometry& MyGeometry, const FPointerEvent& MouseEvent, FString ParamName)
+{
+	if (MouseEvent.GetEffectingButton() != EKeys::LeftMouseButton)
+	{
+		return FReply::Unhandled();
+	}
+
+	FLinearColor col = GetColorParameterValue(ParamName);
+
+	TArray<FLinearColor*> LinearColorArray;
+	LinearColorArray.Add(&col);
+
+	FColorPickerArgs args;
+	args.bIsModal = true;
+	args.bUseAlpha = false;
+	args.bOnlyRefreshOnMouseUp = false;
+	args.InitialColor = col;
+	args.OnColorCommitted = FOnLinearColorValueChanged::CreateSP(this, &FCustomizableInstanceDetails::OnSetColorFromColorPicker, ParamName);
+	OpenColorPicker(args);
+
+	return FReply::Handled();
+}
+
+
+void FCustomizableInstanceDetails::OnSetColorFromColorPicker(FLinearColor NewColor, FString PickerParamName) const
+{
+	CustomInstance->PreEditChange(nullptr);
+	CustomInstance->SetColorParameterSelectedOption(PickerParamName, NewColor);
+	CustomInstance->SetSelectedParameterProfileDirty();
+
+	PickerParamName = FString();
+
+	CustomInstance->UpdateSkeletalMeshAsync(true, true);
+	CustomInstance->PostEditChange();
+}
+
+
+// PROJECTOR PARAMETERS -----------------------------------------------------------------------------------------------------------------
+
+TSharedRef<SWidget> FCustomizableInstanceDetails::GenerateSimpleProjector(const int32 ParamIndexInObject)
+{
+	UCustomizableObject* CustomizableObject = CustomInstance->GetCustomizableObject();
+	FString ParamName = CustomizableObject->GetParameterName(ParamIndexInObject);
+
+	const TSharedPtr<ICustomizableObjectInstanceEditor> Editor = GetEditorChecked();
+	const UProjectorParameter* ProjectorParameter = Editor->GetProjectorParameter();
+	const bool bSelectedProjector = ProjectorParameter->IsProjectorSelected(ParamName);
+
+	TSharedPtr<SButton> Button;
+	TSharedPtr<SHorizontalBox> SimpleProjectorBox;
+
+	SAssignNew(SimpleProjectorBox,SHorizontalBox)
+	+ SHorizontalBox::Slot()
+	.FillWidth(0.25f)
+	[
+		SAssignNew(Button, SButton)
+		.OnClicked(this, &FCustomizableInstanceDetails::OnProjectorSelectChanged, ParamName, -1)
+		.HAlign(HAlign_Center)
+		.Content()
+		[
+			SNew(STextBlock)
+			.Text(bSelectedProjector ? LOCTEXT("Unselect Projector", "Unselect Projector") : LOCTEXT("Select Projector", "Select Projector"))
+			.ToolTipText(bSelectedProjector ? LOCTEXT("Unselect Projector", "Unselect Projector") : LOCTEXT("Select Projector", "Select Projector"))
+			.Font(LayoutBuilder.Pin()->GetDetailFont())
+		]
+	]
+
+	+ SHorizontalBox::Slot()
+	.FillWidth(0.25f)
+	[
+		SNew(SButton)
+		.HAlign(HAlign_Center)
+		.OnClicked(this, &FCustomizableInstanceDetails::OnProjectorCopyTransform, ParamName, -1)
+		.Content()
+		[
+			SNew(STextBlock)
+			.Text(LOCTEXT("CopyTransform_Text", "Copy Transform"))
+			.ToolTipText(LOCTEXT("CopyTransform_Tooltip", "Copy Transform"))
+			.Font(LayoutBuilder.Pin()->GetDetailFont())
+		]
+	]
+
+	+ SHorizontalBox::Slot()
+	.FillWidth(0.25f)
+	[
+		SNew(SButton)
+		.HAlign(HAlign_Center)
+		.OnClicked(this, &FCustomizableInstanceDetails::OnProjectorPasteTransform, ParamName, -1)
+		.Content()
+		[
+			SNew(STextBlock)
+			.Text(LOCTEXT("PasteTransform_Text", "Paste Transform"))
+			.ToolTipText(LOCTEXT("PasteTransform_Tooltip", "Paste Transform"))
+			.Font(LayoutBuilder.Pin()->GetDetailFont())
+		]
+	]
+
+	+ SHorizontalBox::Slot()
+	.FillWidth(0.25f)
+	[
+		SNew(SButton)
+		.HAlign(HAlign_Center)
+		.OnClicked(this, &FCustomizableInstanceDetails::OnProjectorResetTransform, ParamName, -1)
+		.Content()
+		[
+			SNew(STextBlock)
+			.Text(LOCTEXT("ResetTransform_Text", "Reset Transform"))
+			.ToolTipText(LOCTEXT("ResetTransform_Tooltip", "Reset Transform"))
+			.Font(LayoutBuilder.Pin()->GetDetailFont())
+		]
+	];
+
+	Button->SetBorderBackgroundColor(bSelectedProjector ? FLinearColor::Green : FLinearColor::White);
+
+	return SimpleProjectorBox.ToSharedRef();
+}
+
+
+FReply FCustomizableInstanceDetails::OnProjectorSelectChanged(const FString ParamName, const int32 RangeIndex) const
+{
+	const TSharedPtr<ICustomizableObjectInstanceEditor> Editor = GetEditorChecked();
+
+	const UProjectorParameter* ProjectorParameter = Editor->GetProjectorParameter();
+	if (ProjectorParameter->IsProjectorSelected(ParamName, RangeIndex))
+	{
+		Editor->HideGizmoProjectorParameter();
+	}
+	else
+	{
+		Editor->ShowGizmoProjectorParameter(ParamName, RangeIndex);
+	}
+
+	return FReply::Handled();
+}
+
+
+FReply FCustomizableInstanceDetails::OnProjectorCopyTransform(const FString ParamName, const int32 RangeIndex) const
+{
+	const int32 ParameterIndexInObject = CustomInstance->GetCustomizableObject()->FindParameter(ParamName);
+	const int32 ProjectorParamIndex = CustomInstance->FindProjectorParameterNameIndex(ParamName);
+
+	if ((ParameterIndexInObject >= 0) && (ProjectorParamIndex >= 0))
+	{
+		TArray<FCustomizableObjectProjectorParameterValue>& ProjectorParameters = CustomInstance->GetProjectorParameters();
+		FCustomizableObjectProjector Value;
+
+		if (RangeIndex == -1)
+		{
+			Value = ProjectorParameters[ProjectorParamIndex].Value;
+		}
+		else if (ProjectorParameters[ProjectorParamIndex].RangeValues.IsValidIndex(RangeIndex))
+		{
+			Value = ProjectorParameters[ProjectorParamIndex].RangeValues[RangeIndex];
+		}
+		else
+		{
+			check(false);
+		}
+
+		const UScriptStruct* Struct = Value.StaticStruct();
+		FString Output = TEXT("");
+		Struct->ExportText(Output, &Value, nullptr, nullptr, (PPF_ExportsNotFullyQualified | PPF_Copy | PPF_Delimited | PPF_IncludeTransient), nullptr);
+
+		FPlatformApplicationMisc::ClipboardCopy(*Output);
+	}
+
+	return FReply::Handled();
+}
+
+
+FReply FCustomizableInstanceDetails::OnProjectorPasteTransform(const FString ParamName, const int32 RangeIndex)
+{
+	FScopedTransaction Transaction(LOCTEXT("PasteTransform", "Paste Transform"));
+	CustomInstance->Modify();
+
+	FString ClipboardText;
+	FPlatformApplicationMisc::ClipboardPaste(ClipboardText);
+	FCustomizableObjectProjector DefaultValue;
+	UScriptStruct* Struct = DefaultValue.StaticStruct();
+	Struct->ImportText(*ClipboardText, &DefaultValue, nullptr, 0, GLog, GetPathNameSafe(Struct));
+
+	CustomInstance->SetProjectorValue(ParamName,
+		static_cast<FVector>(DefaultValue.Position),
+		static_cast<FVector>(DefaultValue.Direction),
+		static_cast<FVector>(DefaultValue.Up),
+		static_cast<FVector>(DefaultValue.Scale),
+		DefaultValue.Angle,
+		RangeIndex);
+
+	const TSharedPtr<ICustomizableObjectInstanceEditor> Editor = GetEditorChecked();
+
+	Editor->ShowGizmoProjectorParameter(ParamName, RangeIndex);
+	CustomInstance->UpdateSkeletalMeshAsync(true, true);
+
+	return FReply::Handled();
+}
+
+
+FReply FCustomizableInstanceDetails::OnProjectorResetTransform(const FString ParamName, const int32 RangeIndex)
+{
+	FScopedTransaction Transaction(LOCTEXT("ResetTransform", "Reset Transform"));
+	CustomInstance->Modify();
+
+	const FCustomizableObjectProjector DefaultValue = CustomInstance->GetCustomizableObject()->GetProjectorParameterDefaultValue(ParamName);
+
+	CustomInstance->SetProjectorValue(ParamName,
+		static_cast<FVector>(DefaultValue.Position),
+		static_cast<FVector>(DefaultValue.Direction),
+		static_cast<FVector>(DefaultValue.Up),
+		static_cast<FVector>(DefaultValue.Scale),
+		DefaultValue.Angle,
+		RangeIndex);
+
+	const TSharedPtr<ICustomizableObjectInstanceEditor> Editor = GetEditorChecked();
+
+	Editor->ShowGizmoProjectorParameter(ParamName, RangeIndex);
+	CustomInstance->UpdateSkeletalMeshAsync(true, true);
+
+	return FReply::Handled();
+}
+
+
+TSharedRef<SWidget> FCustomizableInstanceDetails::GenerateMultidimensionalProjector(const int32 ParamIndexInObject)
+{
+	UCustomizableObject* CustomizableObject = CustomInstance->GetCustomizableObject();
+	FString ParamName = CustomizableObject->GetParameterName(ParamIndexInObject);
+
+	const TSharedPtr<ICustomizableObjectInstanceEditor> Editor = GetEditorChecked();
+	TArray<FCustomizableObjectProjectorParameterValue>& ProjectorParameters = CustomInstance->GetProjectorParameters();
+	const int32 ProjectorParamIndex = CustomInstance->FindProjectorParameterNameIndex(ParamName);
+
+	check(ProjectorParamIndex < ProjectorParameters.Num());
+
+	// Selected Pose UI
+	const FString PoseSwitchEnumParamName = ParamName + FMultilayerProjector::POSE_PARAMETER_POSTFIX;
+	const int32 PoseSwitchEnumParamIndexInObject = CustomizableObject->FindParameter(PoseSwitchEnumParamName);
+	TSharedPtr<SVerticalBox> ProjectorBox = SNew(SVerticalBox);
+
+	if (PoseSwitchEnumParamIndexInObject != INDEX_NONE)
+	{
+		const int32 NumPoseValues = CustomizableObject->GetIntParameterNumOptions(PoseSwitchEnumParamIndexInObject);
+
+		TArray<FString> PoseOptionNamesAttribute;
+		FString PoseValue = CustomInstance->GetIntParameterSelectedOption(ParamName, -1);
+		int32 PoseValueIndex = 0;
+
+		for (int32 j = 0; j < NumPoseValues; ++j)
+		{
+			const FString PossibleValue = CustomizableObject->GetIntParameterAvailableOption(PoseSwitchEnumParamIndexInObject, j);
+			if (PossibleValue == PoseValue)
+			{
+				PoseValueIndex = j;
+			}
+
+			PoseOptionNamesAttribute.Add(PossibleValue);
+		}
+
+		ProjectorBox->AddSlot()
+		.HAlign(HAlign_Right)
+		.VAlign(VAlign_Fill)
+		.FillHeight(10.f)
+		[
+			SNew(SHorizontalBox)
+			+ SHorizontalBox::Slot()
+			.HAlign(HAlign_Fill)
+			.VAlign(VAlign_Center)
+			.FillWidth(0.45f)
+			[
+				SNew(SMutableTextSearchBox)
+				.PossibleSuggestions(PoseOptionNamesAttribute)
+				.InitialText(FText::FromString(PoseOptionNamesAttribute[PoseValueIndex]))
+				.MustMatchPossibleSuggestions(TAttribute<bool>(true))
+				.SuggestionListPlacement(EMenuPlacement::MenuPlacement_ComboBox)
+				.OnTextCommitted(this, &FCustomizableInstanceDetails::OnProjectorTextureParameterComboBoxChanged, PoseSwitchEnumParamName, -1)
+				.ToolTipText(LOCTEXT("Pose selector tooltip", "Select the skeletal mesh pose used for projection. This does not control the actual visual mesh pose in the viewport (or during gameplay for that matter). It has to be manually set. You can drag&drop a pose onto the preview viewport."))
+			]
+
+			+ SHorizontalBox::Slot()
+			.HAlign(HAlign_Fill)
+			.VAlign(VAlign_Center)
+			.FillWidth(0.3f)
+			[
+				SNew(SButton)
+				.ToolTipText(LOCTEXT("Add Layer", "Add Layer"))
+				.Text(LOCTEXT("Add Layer", "Add Layer"))
+				.OnClicked(this, &FCustomizableInstanceDetails::OnProjectorLayerAdded, ParamName)
+				.HAlign(HAlign_Fill)
+			]
+		];
+	}
+	else
+	{
+		ProjectorBox->AddSlot()
+		.HAlign(HAlign_Right)
+		.VAlign(VAlign_Fill)
+		//.FillHeight(10.f)
+		[
+			SNew(SButton)
+			.ToolTipText(LOCTEXT("Add Layer", "Add Layer"))
+			.Text(LOCTEXT("Add Layer", "Add Layer"))
+			.OnClicked(this, &FCustomizableInstanceDetails::OnProjectorLayerAdded, ParamName)
+			.HAlign(HAlign_Fill)
+			.VAlign(VAlign_Center)
+		];
+	}
+
+	const FString TextureSwitchEnumParamName = ParamName + FMultilayerProjector::IMAGE_PARAMETER_POSTFIX;
+
+	IDetailGroup* ProjectorGroup = *ParentsGroups.Find(ParamName);
+
+	for (int32 RangeIndex = 0; RangeIndex < ProjectorParameters[ProjectorParamIndex].RangeValues.Num(); ++RangeIndex)
+	{
+		const int32 TextureSwitchEnumParamIndexInObject = CustomizableObject->FindParameter(TextureSwitchEnumParamName);
+		check(TextureSwitchEnumParamIndexInObject >= 0);
+		int32 NumValues = CustomizableObject->GetIntParameterNumOptions(TextureSwitchEnumParamIndexInObject);
+
+		FString OpacitySliderParamName = ParamName + FMultilayerProjector::OPACITY_PARAMETER_POSTFIX;
+
+		TArray<FString> OptionNamesAttribute;
+		FString Value = CustomInstance->GetIntParameterSelectedOption(ParamName, RangeIndex);
+		int32 ValueIndex = 0;
+
+		const UProjectorParameter* ProjectorParameter = Editor->GetProjectorParameter();
+		const bool bSelectedProjector = ProjectorParameter->IsProjectorSelected(ParamName, RangeIndex);
+
+		for (int32 CandidateIndex = 0; CandidateIndex < NumValues; ++CandidateIndex)
+		{
+			FString PossibleValue = CustomizableObject->GetIntParameterAvailableOption(TextureSwitchEnumParamIndexInObject, CandidateIndex);
+			if (PossibleValue == Value)
+			{
+				ValueIndex = CandidateIndex;
+			}
+			OptionNamesAttribute.Add(CustomizableObject->GetIntParameterAvailableOption(TextureSwitchEnumParamIndexInObject, CandidateIndex));
+		}
+
+		TSharedPtr<SHorizontalBox> SliderBox;
+		TSharedPtr<SSpinBox<float>> Slider;
+		TSharedPtr<SHorizontalBox> PropertiesHB;
+		TSharedPtr<SButton> Button;
+		TSharedPtr<STextBlock> TextBlock;
+
+		//Horizontal box that owns all the hidden layer properties
+		SAssignNew(PropertiesHB, SHorizontalBox)
+		+ SHorizontalBox::Slot()
+		.Padding(1, 0)
+		.AutoWidth()
+		[
+			SNew(SHorizontalBox)
+			+ SHorizontalBox::Slot()
+			.Padding(1, 0)
+			[
+				SNew(SBox)
+				.MinDesiredWidth(115.f)
+				.MaxDesiredWidth(115.f)
+				[
+					SAssignNew(Button, SButton)
+					.OnClicked(this, &FCustomizableInstanceDetails::OnProjectorSelectChanged, ParamName, RangeIndex)
+					.VAlign(VAlign_Center)
+					.HAlign(HAlign_Center)
+					.Content()
+					[
+						SNew(STextBlock)
+						.Text(bSelectedProjector ? LOCTEXT("Unselect Projector", "Unselect Projector") : LOCTEXT("Select Projector", "Select Projector"))
+						.ToolTipText(bSelectedProjector ? LOCTEXT("Unselect Projector", "Unselect Projector") : LOCTEXT("Select Projector", "Select Projector"))
+						.Justification(ETextJustify::Center)
+						.Font(LayoutBuilder.Pin()->GetDetailFont())
+						//.AutoWrapText(true)
+					]
+				]
+			]
+
+			+ SHorizontalBox::Slot()
+			.Padding(1, 0)
+			[
+				SNew(SBox)
+				.MinDesiredWidth(120.f)
+				[
+					SNew(SMutableTextSearchBox)
+					.PossibleSuggestions(OptionNamesAttribute)
+					.InitialText(FText::FromString(OptionNamesAttribute[ValueIndex]))
+					.MustMatchPossibleSuggestions(TAttribute<bool>(true))
+					.SuggestionListPlacement(EMenuPlacement::MenuPlacement_ComboBox)
+					.OnTextCommitted(this, &FCustomizableInstanceDetails::OnProjectorTextureParameterComboBoxChanged, TextureSwitchEnumParamName, RangeIndex)
+				]
+			]
+		];
+
+		Button->SetBorderBackgroundColor(bSelectedProjector ? FLinearColor::Green : FLinearColor::White);
+
+		PropertiesHB->AddSlot()
+		.HAlign(HAlign_Fill)
+		[
+			SNew(SHorizontalBox)
+			+ SHorizontalBox::Slot()
+			.Padding(1, 0)
+			.HAlign(HAlign_Fill)
+			[
+				SNew(SHorizontalBox)
+				+ SHorizontalBox::Slot()
+				[
+					SNew(SSpinBox<float>)
+					.MinValue(0.0f)
+					.MaxValue(1.0f)
+					.Value(this, &FCustomizableInstanceDetails::GetFloatParameterValue, OpacitySliderParamName, RangeIndex)
+					.OnValueChanged(this, &FCustomizableInstanceDetails::OnFloatParameterChanged, OpacitySliderParamName, RangeIndex)
+					.OnEndSliderMovement(this, &FCustomizableInstanceDetails::OnFloatParameterSliderEnd, OpacitySliderParamName, RangeIndex)
+				]
+			]
+		];
+
+		PropertiesHB->AddSlot()
+		.HAlign(HAlign_Right)
+		.AutoWidth()
+		[
+			SNew(SHorizontalBox)
+			+ SHorizontalBox::Slot()
+			.Padding(1, 0)
+			[
+				SNew(SButton)
+				.OnClicked(this, &FCustomizableInstanceDetails::OnProjectorCopyTransform, ParamName, RangeIndex)
+				.VAlign(VAlign_Center)
+				.HAlign(HAlign_Center)
+				.Content()
+				[
+					SNew(STextBlock)
+					.ToolTipText(LOCTEXT("Copy Transform", "Copy Transform"))
+					.Text(LOCTEXT("Copy Transform", "Copy Transform"))
+					.AutoWrapText(true)
+					.Justification(ETextJustify::Center)
+					.Font(LayoutBuilder.Pin()->GetDetailFont())
+				]
+			]
+
+			+ SHorizontalBox::Slot()
+			.Padding(1, 0)
+			[
+				SNew(SButton)
+				.OnClicked(this, &FCustomizableInstanceDetails::OnProjectorPasteTransform, ParamName, RangeIndex)
+				.VAlign(VAlign_Center)
+				.HAlign(HAlign_Center)
+				.Content()
+				[
+					SNew(STextBlock)
+					.ToolTipText(LOCTEXT("Paste Transform", "Paste Transform"))
+					.Text(LOCTEXT("Paste Transform", "Paste Transform"))
+					.Justification(ETextJustify::Center)
+					.AutoWrapText(true)
+					.Font(LayoutBuilder.Pin()->GetDetailFont())
+				]
+			]
+
+			+ SHorizontalBox::Slot()
+			.Padding(1, 0)
+			[
+				SNew(SButton)
+				.OnClicked(this, &FCustomizableInstanceDetails::OnProjectorResetTransform, ParamName, RangeIndex)
+				.VAlign(VAlign_Center)
+				.HAlign(HAlign_Center)
+				.Content()
+				[
+					SNew(STextBlock)
+					.ToolTipText(LOCTEXT("Reset Transform", "Reset Transform"))
+					.Text(LOCTEXT("Reset Transform", "Reset Transform"))
+					.Justification(ETextJustify::Center)
+					.AutoWrapText(true)
+					.Font(LayoutBuilder.Pin()->GetDetailFont())
+				]
+			]
+
+			+ SHorizontalBox::Slot()
+			.Padding(1, 0)
+			[
+				SNew(SButton)
+				.OnClicked(this, &FCustomizableInstanceDetails::OnProjectorLayerRemoved, ParamName, RangeIndex)
+				.VAlign(VAlign_Center)
+				.HAlign(HAlign_Center)
+				.Content()
+				[
+					SNew(STextBlock)
+					.ToolTipText(LOCTEXT("Remove Layer", "Remove Layer"))
+					.Text(LOCTEXT("Remove Layer", "Remove Layer"))
+					.Justification(ETextJustify::Center)
+					.AutoWrapText(true)
+					.Font(LayoutBuilder.Pin()->GetDetailFont())
+				]
+			]
+		];
+
+		FString LayerName = "Layer " + FString::FromInt(RangeIndex);
+
+		ProjectorGroup->AddWidgetRow()
+		.NameContent()
+		[
+			SNew(STextBlock)
+			.Text(FText::FromString(LayerName))
+		]
+		.ValueContent()
+		.HAlign(EHorizontalAlignment::HAlign_Fill)
+		.VAlign(EVerticalAlignment::VAlign_Fill)
+		[
+			SNew(SHorizontalBox)
+			+ SHorizontalBox::Slot()
+			.HAlign(HAlign_Fill)
+			.Padding(0.0f, 5.0f, 0.0f, 5.0f)
+			[
+				PropertiesHB.ToSharedRef()
+			]
+		];
+	}
+	
+	return ProjectorBox.ToSharedRef();
+}
+
+
+FReply FCustomizableInstanceDetails::OnProjectorLayerAdded(FString ParamName) const
+{
+	const int32 NumLayers = CustomInstance->AddValueToProjectorRange(ParamName) + 1;
+	if (NumLayers == 0)
+	{
+		return FReply::Handled();
+	}
+
+	const FString TextureSwitchEnumParamName = ParamName + FMultilayerProjector::IMAGE_PARAMETER_POSTFIX;
+	CustomInstance->AddValueToIntRange(TextureSwitchEnumParamName);
+	check(CustomInstance->FindIntParameterNameIndex(TextureSwitchEnumParamName) != INDEX_NONE);
+	check(NumLayers == CustomInstance->GetIntParameters()[CustomInstance->FindIntParameterNameIndex(TextureSwitchEnumParamName)].ParameterRangeValueNames.Num());
+
+	TArray<FCustomizableObjectFloatParameterValue>& FloatParameters = CustomInstance->GetFloatParameters();
+
+	const FString NumLayersParamName = ParamName + FMultilayerProjector::NUM_LAYERS_PARAMETER_POSTFIX;
+	check(CustomInstance->FindFloatParameterNameIndex(NumLayersParamName) != INDEX_NONE);
+	FloatParameters[CustomInstance->FindFloatParameterNameIndex(NumLayersParamName)].ParameterValue = NumLayers;
+
+	const FString OpacitySliderParamName = ParamName + FMultilayerProjector::OPACITY_PARAMETER_POSTFIX;
+	CustomInstance->AddValueToFloatRange(OpacitySliderParamName);
+	check(CustomInstance->FindFloatParameterNameIndex(OpacitySliderParamName) != INDEX_NONE);
+	check(NumLayers == FloatParameters[CustomInstance->FindFloatParameterNameIndex(OpacitySliderParamName)].ParameterRangeValues.Num());
+
+	CustomInstance->SetSelectedParameterProfileDirty();
+	CustomInstance->PreEditChange(nullptr);
+	CustomInstance->UpdateSkeletalMeshAsync(true, true);
+	CustomInstance->PostEditChange();
+
+	FCoreUObjectDelegates::BroadcastOnObjectModified(CustomInstance.Get());
+
+	return FReply::Handled();
+}
+
+
+FReply FCustomizableInstanceDetails::OnProjectorLayerRemoved(const FString ParamName, const int32 RangeIndex) const
+{
+	const int32 projectorParameterIndex = CustomInstance->FindProjectorParameterNameIndex(ParamName);
+	if (projectorParameterIndex == INDEX_NONE
+		|| CustomInstance->GetProjectorParameters()[projectorParameterIndex].RangeValues.Num() <= 0)
+	{
+		return FReply::Handled();
+	}
+
+	const int32 NumLayers = CustomInstance->RemoveValueFromProjectorRange(ParamName, RangeIndex) + 1;
+
+	const FString TextureSwitchEnumParamName = ParamName + FMultilayerProjector::IMAGE_PARAMETER_POSTFIX;
+	CustomInstance->RemoveValueFromIntRange(TextureSwitchEnumParamName, RangeIndex);
+	check(CustomInstance->FindIntParameterNameIndex(TextureSwitchEnumParamName) != INDEX_NONE);
+	check(NumLayers == CustomInstance->GetIntParameters()[CustomInstance->FindIntParameterNameIndex(TextureSwitchEnumParamName)].ParameterRangeValueNames.Num());
+
+	TArray<FCustomizableObjectFloatParameterValue>& FloatParameters = CustomInstance->GetFloatParameters();
+
+	const FString NumLayersParamName = ParamName + FMultilayerProjector::NUM_LAYERS_PARAMETER_POSTFIX;
+	check(CustomInstance->FindFloatParameterNameIndex(NumLayersParamName) != INDEX_NONE);
+	FloatParameters[CustomInstance->FindFloatParameterNameIndex(NumLayersParamName)].ParameterValue = NumLayers;
+
+	const FString OpacitySliderParamName = ParamName + FMultilayerProjector::OPACITY_PARAMETER_POSTFIX;
+	CustomInstance->RemoveValueFromFloatRange(OpacitySliderParamName, RangeIndex);
+	check(CustomInstance->FindFloatParameterNameIndex(OpacitySliderParamName) != INDEX_NONE);
+	check(NumLayers == FloatParameters[CustomInstance->FindFloatParameterNameIndex(OpacitySliderParamName)].ParameterRangeValues.Num());
+
+	CustomInstance->SetSelectedParameterProfileDirty();
+	CustomInstance->PreEditChange(nullptr);
+	CustomInstance->UpdateSkeletalMeshAsync(true, true);
+	CustomInstance->PostEditChange();
+
+	FCoreUObjectDelegates::BroadcastOnObjectModified(CustomInstance.Get());
+
+	return FReply::Handled();
+}
+
+
+void FCustomizableInstanceDetails::OnProjectorTextureParameterComboBoxChanged(TSharedPtr<FString> Selection, ESelectInfo::Type SelectInfo, FString ParamName, int32 RangeIndex) const
+{
+	TArray<FCustomizableObjectIntParameterValue>& IntParameters = CustomInstance->GetIntParameters();
+
+	const UCustomizableObject* CustomObject = CustomInstance->GetCustomizableObject();
+	for (int32 i = 0; i < IntParameters.Num(); ++i)
+	{
+		if (IntParameters[i].ParameterName == ParamName)
+		{
+			const int32 ParamIndexInObject = CustomObject->FindParameter(ParamName);
+			if (ParamIndexInObject != INDEX_NONE)
+			{
+				for (int32 v = 0; v < CustomObject->GetIntParameterNumOptions(ParamIndexInObject); ++v)
+				{
+					FString ValueName = CustomObject->GetIntParameterAvailableOption(ParamIndexInObject, v);
+					if (ValueName == *Selection)
+					{
+						if (IntParameters[i].ParameterRangeValueNames.IsValidIndex(RangeIndex))
+						{
+							IntParameters[i].ParameterRangeValueNames[RangeIndex] = ValueName;
+						}
+						else
+						{
+							IntParameters[i].ParameterValueName = ValueName;
+						}
+
+						break;
+					}
+				}
+			}
+			else
+			{
+				UE_LOG(LogMutable, Error, TEXT("Failed to find parameter."));
+			}
+		}
+	}
+
+	CustomInstance->SetSelectedParameterProfileDirty();
+	CustomInstance->UpdateSkeletalMeshAsync(true, true);
+	CustomInstance->PostEditChange();
+
+	// Non-continuous change: collect garbage.
+	GEngine->ForceGarbageCollection();
+}
+
+
+TSharedPtr<ICustomizableObjectInstanceEditor> FCustomizableInstanceDetails::GetEditorChecked() const
+{
+	TSharedPtr<ICustomizableObjectInstanceEditor> Editor = WeakEditor.Pin();
+	check(Editor);
+
+	return Editor;
+}
+
+
+// BOOL PARAMETERS -----------------------------------------------------------------------------------------------------------------
+
+TSharedRef<SWidget> FCustomizableInstanceDetails::GenerateBoolWidget(const int32 ParamIndexInObject)
+{
+	UCustomizableObject* CustomizableObject = CustomInstance->GetCustomizableObject();
+	FString ParamName = CustomizableObject->GetParameterName(ParamIndexInObject);
+
+	return SNew(SHorizontalBox)
+		+ SHorizontalBox::Slot()
+		.HAlign(HAlign_Left)
+		[
+			SNew(SCheckBox)
+			.HAlign(HAlign_Left)
+			.IsChecked(this, &FCustomizableInstanceDetails::GetBoolParameterValue, ParamName)
+			.OnCheckStateChanged(this, &FCustomizableInstanceDetails::OnBoolParameterChanged, ParamName)
+		];
+}
+
+
+ECheckBoxState FCustomizableInstanceDetails::GetBoolParameterValue(FString ParamName) const
+{
+	const bool bResult = CustomInstance->GetBoolParameterSelectedOption(ParamName);
+	return bResult ? ECheckBoxState::Checked : ECheckBoxState::Unchecked;
+}
+
+
+void FCustomizableInstanceDetails::OnBoolParameterChanged(ECheckBoxState InCheckboxState, FString ParamName)
+{
+	CustomInstance->PreEditChange(nullptr);
+	CustomInstance->SetBoolParameterSelectedOption(ParamName, InCheckboxState == ECheckBoxState::Checked);
+	CustomInstance->SetSelectedParameterProfileDirty();
+	CustomInstance->UpdateSkeletalMeshAsync(true, true);
+	CustomInstance->PostEditChange();
+
+	// Non-continuous change: collect garbage.
+	GEngine->ForceGarbageCollection();
+}
+
+
+// PARAMETERS -----------------------------------------------------------------------------------------------------------------
+
+FReply FCustomizableInstanceDetails::OnCopyAllParameters()
+{
+	FBufferArchive ParametersBuffer;
+	CustomInstance->SaveDescriptor(ParametersBuffer, true);
+
+	FString StringifyedData = FString::FromHexBlob(ParametersBuffer.GetData(), ParametersBuffer.Num());
+	FPlatformApplicationMisc::ClipboardCopy(StringifyedData.GetCharArray().GetData());
+
+	return FReply::Handled();
+}
+
+
+FReply FCustomizableInstanceDetails::OnPasteAllParameters()
+{
+	FString StringifiedData;
+	FPlatformApplicationMisc::ClipboardPaste(StringifiedData);
+
+	if (!StringifiedData.IsEmpty())
+	{
+		TArray<uint8> DescriptorData;
+		// Each char represents a hex number i.e. half byte.
+		DescriptorData.SetNum(StringifiedData.Len() * 2);
+		FString::ToHexBlob(StringifiedData, DescriptorData.GetData(), DescriptorData.Num());
+		FMemoryReader FromBinary = FMemoryReader(DescriptorData, true); //true, free data after done;
+		FromBinary.Seek(0);
+
+		CustomInstance->LoadDescriptor(FromBinary);
+		CustomInstance->SetSelectedParameterProfileDirty();
+		CustomInstance->UpdateSkeletalMeshAsync(true, true);
+		CustomInstance->PostEditChange();
+
+		// Non-continuous change: collect garbage.
+		GEngine->ForceGarbageCollection();
+	}
+
+	return FReply::Handled();
+}
+
+
+FReply FCustomizableInstanceDetails::OnResetAllParameters()
+{
+	CustomInstance->PreEditChange(nullptr);
+
+	TArray<FCustomizableObjectIntParameterValue>& IntParameters = CustomInstance->GetIntParameters();
+
+	const UCustomizableObject* CustomObject = CustomInstance->GetCustomizableObject();
+	const int32 NumObjectParameter = CustomObject->GetParameterCount();
+
+	for (int32 ParameterIndex = 0; ParameterIndex < NumObjectParameter; ++ParameterIndex)
+	{
+		SetParameterValueToDefault(ParameterIndex);
+	}
+
+	CustomInstance->UpdateSkeletalMeshAsync(true, true);
+	CustomInstance->PostEditChange();
+	CustomInstance->SelectedProfileIndex = INDEX_NONE;
+
+	// Non-continuous change: collect garbage.
+	GEngine->ForceGarbageCollection();
+
+	return FReply::Handled();
+}
+
+
+void FCustomizableInstanceDetails::OnResetParameterButtonClicked(int32 ParameterIndex)
+{	
+	CustomInstance->PreEditChange(nullptr);
+	SetParameterValueToDefault(ParameterIndex);
+	CustomInstance->SetSelectedParameterProfileDirty();
+	CustomInstance->UpdateSkeletalMeshAsync(true, true);
+	CustomInstance->PostEditChange();
+}
+
+
+void FCustomizableInstanceDetails::SetParameterValueToDefault(int32 ParameterIndex)
+{
+	UCustomizableObject* CustomObject = CustomInstance->GetCustomizableObject();
+	FString ParameterName = CustomObject->GetParameterName(ParameterIndex);
+	EMutableParameterType ParameterType = CustomObject->GetParameterType(ParameterIndex);
+
+	// Checking if there are parameteres with the same name
+	if (ParameterType != CustomObject->GetParameterTypeByName(ParameterName))
+	{
+		return;
+	}
+
+	switch (ParameterType)
+	{
+	case EMutableParameterType::None:
+		break;
+	case EMutableParameterType::Bool:
+	{
+		CustomInstance->SetBoolParameterSelectedOption(ParameterName, CustomObject->GetBoolParameterDefaultValue(ParameterName));
+		break;
+	}
+	case EMutableParameterType::Int:
+	{
+		int32 DefaultValue = CustomObject->GetIntParameterDefaultValue(ParameterName);
+		FString ValueName = CustomObject->FindIntParameterValueName(ParameterIndex, DefaultValue);
+
+		if (!ValueName.IsEmpty())
+		{
+			if (CustomObject->IsParameterMultidimensional(ParameterName))
+			{
+				int32 NumRanges = CustomInstance->GetIntValueRange(ParameterName);
+
+				for (int32 RangeIndex = 0; RangeIndex < NumRanges; ++RangeIndex)
+				{
+					CustomInstance->SetIntParameterSelectedOption(ParameterName, ValueName, RangeIndex);
+				}
+			}
+			else
+			{
+				CustomInstance->SetIntParameterSelectedOption(ParameterName, ValueName);
+			}
+		}
+		break;
+	}
+	case EMutableParameterType::Float:
+	{
+		float DefaultValue = CustomObject->GetFloatParameterDefaultValue(ParameterName);
+
+		if (CustomObject->IsParameterMultidimensional(ParameterName))
+		{
+			int32 NumRanges = CustomInstance->GetFloatValueRange(ParameterName);
+
+			for (int32 RangeIndex = 0; RangeIndex < NumRanges; ++RangeIndex)
+			{
+				CustomInstance->SetFloatParameterSelectedOption(ParameterName, DefaultValue, RangeIndex);
+			}
+		}
+		else
+		{
+			CustomInstance->SetFloatParameterSelectedOption(ParameterName, DefaultValue);
+		}
+		break;
+	}
+	case EMutableParameterType::Color:
+	{
+		CustomInstance->SetColorParameterSelectedOption(ParameterName, CustomObject->GetColorParameterDefaultValue(ParameterName));
+		break;
+	}
+	case EMutableParameterType::Projector:
+	{
+		FCustomizableObjectProjector DefaultValue = CustomObject->GetProjectorParameterDefaultValue(ParameterName);
+
+		if (CustomObject->IsParameterMultidimensional(ParameterName))
+		{
+			int32 NumRanges = CustomInstance->GetProjectorValueRange(ParameterName);
+
+			for (int32 RangeIndex = 0; RangeIndex < NumRanges; ++RangeIndex)
+			{
+				CustomInstance->SetProjectorValue(ParameterName, FVector(DefaultValue.Position), FVector(DefaultValue.Direction), FVector(DefaultValue.Up), FVector(DefaultValue.Scale), DefaultValue.Angle, RangeIndex);
+			}
+		}
+		else
+		{
+			CustomInstance->SetProjectorValue(ParameterName, FVector(DefaultValue.Position), FVector(DefaultValue.Direction), FVector(DefaultValue.Up), FVector(DefaultValue.Scale), DefaultValue.Angle);
+		}
+
+		break;
+	}
+	case EMutableParameterType::Texture:
+	{
+		FName DefaultValue = CustomObject->GetTextureParameterDefaultValue(ParameterName);
+
+		if (CustomObject->IsParameterMultidimensional(ParameterName))
+		{
+			int32 NumRanges = CustomInstance->GetTextureValueRange(ParameterName);
+
+			for (int32 RangeIndex = 0; RangeIndex < NumRanges; ++RangeIndex)
+			{
+				CustomInstance->SetTextureParameterSelectedOption(ParameterName, DefaultValue.ToString(), RangeIndex);
+			}
+		}
+		else
+		{
+			CustomInstance->SetTextureParameterSelectedOption(ParameterName, DefaultValue.ToString());
+		}
+		break;
+	}
+	}
+
+	return;
+}
+
+
+// PROFILES WINDOW -----------------------------------------------------------------------------------------------------------------
+
+void SProfileParametersWindow::Construct(const FArguments& InArgs)
+{
+	if (AssetPath.IsEmpty())
+	{
+		AssetPath = FText::FromString(TEXT("/Game/"));
+	}
+
+	FModuleManager::LoadModuleChecked<FContentBrowserModule>("ContentBrowser");
+
+	SWindow::Construct(SWindow::FArguments()
+	.Title(LOCTEXT("SSelectFolderDlg_Title", "Add a name to the new profile"))
+	.SupportsMinimize(false)
+	.SupportsMaximize(false)
+	.ClientSize(FVector2D(450, 85))
+	[
+		SNew(SVerticalBox)
+
+		+ SVerticalBox::Slot() // Add user input block
+		.Padding(2)
+		[
+			SNew(SBorder)
+			.BorderImage(UE_MUTABLE_GET_BRUSH("ToolPanel.GroupBorder"))
+			[
+				SNew(SVerticalBox)
+				+ SVerticalBox::Slot()
+				.AutoHeight()
+				[
+					SNew(STextBlock)
+					.Text(LOCTEXT("CustomizableProfileName", "Customizable Profile Name"))
+					.Font(FSlateFontInfo(FPaths::EngineContentDir() / TEXT("Slate/Fonts/Roboto-Regular.ttf"), 14))
+				]
+
+				+ SVerticalBox::Slot()
+				.AutoHeight()
+				[
+					SNew(SEditableTextBox)
+					.Text(InArgs._DefaultFileName)
+					.OnTextCommitted(this, &SProfileParametersWindow::OnNameChange)
+					.MinDesiredWidth(250)
+				]
+			]
+		]
+
+		+ SVerticalBox::Slot()
+		.AutoHeight()
+		.HAlign(HAlign_Right)
+		.Padding(5)
+		[
+			SNew(SUniformGridPanel)
+			.SlotPadding(UE_MUTABLE_GET_MARGIN("StandardDialog.SlotPadding"))
+			.MinDesiredSlotWidth(UE_MUTABLE_GET_FLOAT("StandardDialog.MinDesiredSlotWidth"))
+			.MinDesiredSlotHeight(UE_MUTABLE_GET_FLOAT("StandardDialog.MinDesiredSlotHeight"))
+			+ SUniformGridPanel::Slot(0, 0)
+			[
+				SNew(SButton)
+				.HAlign(HAlign_Center)
+				.ContentPadding(UE_MUTABLE_GET_MARGIN("StandardDialog.ContentPadding"))
+				.Text(LOCTEXT("OK", "OK"))
+				.OnClicked(this, &SProfileParametersWindow::OnButtonClick, EAppReturnType::Ok)
+			]
+			+ SUniformGridPanel::Slot(1, 0)
+			[
+				SNew(SButton)
+				.HAlign(HAlign_Center)
+				.ContentPadding(UE_MUTABLE_GET_MARGIN("StandardDialog.ContentPadding"))
+				.Text(LOCTEXT("Cancel", "Cancel"))
+				.OnClicked(this, &SProfileParametersWindow::OnButtonClick, EAppReturnType::Cancel)
+			]
+		]
+	]);
+}
+
+
+EAppReturnType::Type SProfileParametersWindow::ShowModal()
+{
+	GEditor->EditorAddModalWindow(SharedThis(this));
+
+	return UserResponse;
+}
+
+
+void SProfileParametersWindow::OnNameChange(const FText& NewName, ETextCommit::Type CommitInfo)
+{
+	if (CommitInfo == ETextCommit::OnEnter)
+	{
+		FileName = NewName;
+		
+		RequestDestroyWindow();
+
+		UCustomizableObject* CustomizableObject = CustomInstance->GetCustomizableObject();
+		CustomizableObject->AddNewParameterProfile(GetFileName(), *CustomInstance.Get());
+
+		if (CustomInstance->bSelectedProfileDirty && CustomInstance->SelectedProfileIndex != INDEX_NONE)
+		{
+			CustomInstance->SaveParametersToProfile(CustomInstance->SelectedProfileIndex);
+		}
+		CustomInstance->SelectedProfileIndex = CustomizableObject->InstancePropertiesProfiles.Num() - 1;
+
+		if (InstanceDetails)
+		{
+			InstanceDetails->Refresh();
+		}
+	}
+}
+
+
+FReply SProfileParametersWindow::OnButtonClick(EAppReturnType::Type ButtonID)
+{
+	if (ButtonID == EAppReturnType::Ok)
+	{
+		UserResponse = ButtonID;
+
+		RequestDestroyWindow();
+
+		UCustomizableObject* CustomizableObject = CustomInstance->GetCustomizableObject();
+		CustomizableObject->AddNewParameterProfile(GetFileName(), *CustomInstance.Get());
+
+		if (CustomInstance->bSelectedProfileDirty && CustomInstance->SelectedProfileIndex != INDEX_NONE)
+		{
+			CustomInstance->SaveParametersToProfile(CustomInstance->SelectedProfileIndex);
+		}
+		CustomInstance->SelectedProfileIndex = CustomizableObject->InstancePropertiesProfiles.Num() - 1;
+
+		if (InstanceDetails)
+		{
+			InstanceDetails->Refresh();
+		}
+	}
+	else if (ButtonID == EAppReturnType::Cancel)
+	{
+		UserResponse = ButtonID;
+
+		RequestDestroyWindow();
+	}
+	return FReply::Handled();
+}
+
 
 #undef LOCTEXT_NAMESPACE
diff --git a/Engine/Plugins/Experimental/Mutable/Source/CustomizableObjectEditor/Private/MuCOE/CustomizableInstanceDetails.h b/Engine/Plugins/Experimental/Mutable/Source/CustomizableObjectEditor/Private/MuCOE/CustomizableInstanceDetails.h
--- a/Engine/Plugins/Experimental/Mutable/Source/CustomizableObjectEditor/Private/MuCOE/CustomizableInstanceDetails.h
+++ b/Engine/Plugins/Experimental/Mutable/Source/CustomizableObjectEditor/Private/MuCOE/CustomizableInstanceDetails.h
@@ -3,12 +3,23 @@
 #pragma once
 
 #include "IDetailCustomization.h"
+#include "UObject/ObjectPtr.h"
 #include "UObject/WeakObjectPtr.h"
 
-class IDetailLayoutBuilder;
-class SCustomizableInstanceProperties;
+class FDetailWidgetRow;
+class FReply;
+class ICustomizableObjectInstanceEditor;
+class IDetailCategoryBuilder;
+class IDetailGroup;
+class IDetailLayoutBuilder; 
+class SWidget;
 class UCustomizableObjectInstance;
 
+enum class ECheckBoxState : uint8;
+namespace ESelectInfo { enum Type : int; }
+
+struct FGeometry;
+struct FPointerEvent;
 
 class FCustomizableInstanceDetails : public IDetailCustomization
 {
@@ -28,6 +39,118 @@ public:
 	void Refresh() const;
 
 private:
+
+	// Callback to regenerate the details when the instance has finished an update
+	void InstanceUpdated(UCustomizableObjectInstance* Instance) const;
+
+	// State Selector 
+	// Generates The StateSelector Widget
+	TSharedRef<SWidget> GenerateStateSelector();
+	void OnStateComboBoxSelectionChanged(TSharedPtr<FString> Selection, ESelectInfo::Type SelectInfo);
+
+	// Instance Profiles Functions
+	TSharedRef<SWidget> GenerateInstanceProfileSelector();
+	FReply CreateParameterProfileWindow();
+	FReply RemoveParameterProfile();
+	void OnProfileSelectedChanged(TSharedPtr<FString> Selection, ESelectInfo::Type SelectInfo);
+
+	// Function to change the parameters view
+	void OnShowOnlyRuntimeSelectionChanged(ECheckBoxState InCheckboxState);
+	void OnShowOnlyRelevantSelectionChanged(ECheckBoxState InCheckboxState);
+	void OnUseUISectionsSelectionChanged(ECheckBoxState InCheckboxState);
+
+	// Main parameter generation functions
+	void GenerateParametersView(IDetailCategoryBuilder& MainCategory);
+	void RecursivelyAddParamAndChildren(const int32 ParamIndexInObject, const FString ParentName, IDetailCategoryBuilder& DetailsCategory);
+	void FillChildrenMap(int32 ParamIndexInObject);
+
+	// Function to determine if a parameter widget should be generated
+	bool IsVisible(int32 ParamIndexInObject);
+	bool IsMultidimensionalProjector(int32 ParamIndexInObject);
+
+	// Main widget generation functions
+	IDetailGroup* GenerateParameterSection(const int32 ParamIndexInObject, IDetailCategoryBuilder& DetailsCategory);
+	void GenerateWidgetRow(FDetailWidgetRow& WidgetRow, const FString& ParamName, const int32 ParamIndexInObject);
+	TSharedRef<SWidget> GenerateParameterWidget(const int32 ParamIndexInObject);
+
+	// Int Parameters Functions
+	TSharedRef<SWidget> GenerateIntWidget(const int32 ParamIndexInObject);
+	void OnIntParameterComboBoxChanged(TSharedPtr<FString> Selection, ESelectInfo::Type SelectInfo, const FString ParamName);
+	
+	// Float Parameters Functions
+	TSharedRef<SWidget> GenerateFloatWidget(const int32 ParamIndexInObject);
+	float GetFloatParameterValue(const FString ParamName, int32 RangeIndex) const;
+	void OnFloatParameterChanged(float Value, const FString ParamName, int32 RangeIndex);
+	void OnFloatParameterSliderEnd(float Value, const FString ParamName, int32 RangeIndex);
+ 
+	// Texture Parameters Functions
+	TSharedRef<SWidget> GenerateTextureWidget(const int32 ParamIndexInObject);
+	void GenerateTextureParameterOptions();
+	void OnTextureParameterComboBoxSelectionChanged(TSharedPtr<FString> Selection, ESelectInfo::Type SelectInfo, const FString ParamName);
+	
+	// Color Parameters Functions
+	TSharedRef<SWidget> GenerateColorWidget(const int32 ParamIndexInObject);
+	FLinearColor GetColorParameterValue(const FString ParamName) const;
+	FReply OnColorBlockMouseButtonDown(const FGeometry& MyGeometry, const FPointerEvent& MouseEvent, const FString ParamName);
+	void OnSetColorFromColorPicker(FLinearColor NewColor, const FString PickerParamName) const;
+
+	// Bool Parameters Functions
+	TSharedRef<SWidget> GenerateBoolWidget(const int32 ParamIndexInObject);
+	ECheckBoxState GetBoolParameterValue(const FString ParamName) const;
+	void OnBoolParameterChanged(ECheckBoxState InCheckboxState, const FString ParamName);
+
+	// Projector Parameters Functions
+	TSharedRef<SWidget> GenerateSimpleProjector(const int32 ParamIndexInObject);
+	TSharedRef<SWidget> GenerateMultidimensionalProjector(const int32 ParamIndexInObject);
+	TSharedPtr<ICustomizableObjectInstanceEditor> GetEditorChecked() const;
+	FReply OnProjectorSelectChanged(const FString ParamName, const int32 RangeIndex) const;
+	FReply OnProjectorCopyTransform(const FString ParamName, const int32 RangeIndex) const;
+	FReply OnProjectorPasteTransform(const FString ParamName, const int32 RangeIndex);
+	FReply OnProjectorResetTransform(const FString ParamName, const int32 RangeIndex);
+	FReply OnProjectorLayerAdded(const FString ParamName) const;
+	FReply OnProjectorLayerRemoved(const FString ParamName, const int32 RangeIndex) const;
+	void OnProjectorTextureParameterComboBoxChanged(TSharedPtr<FString> Selection, ESelectInfo::Type SelectInfo, const FString ParamName, int32 RangeIndex) const;
+
+	// Parameter Functions
+	FReply OnCopyAllParameters();
+	FReply OnPasteAllParameters();
+	FReply OnResetAllParameters();
+	void OnResetParameterButtonClicked(int32 ParameterIndex);
+	void SetParameterValueToDefault(int32 ParameterIndex);
+
+private:
+
+	/** Pointer to the Customizable Object Instance */
+	TWeakObjectPtr<UCustomizableObjectInstance> CustomInstance;
+
+	/** Details builder pointer */
 	TWeakPtr<IDetailLayoutBuilder> LayoutBuilder;
+
+	/** Map to keep track of the generated parameter sections */
+	TMap<FString, IDetailGroup*> GeneratedSections;
+
+	/** Used to insert child params to a parent's expandable area */
+	TMap<FString, IDetailGroup*> ParentsGroups;
+
+	/*Stores all the possible profiles for the current COI*/
+	TArray<TSharedPtr<FString>> ParameterProfileNames;
+
+	/** Array with all the possible states for the states combo box */
+	TArray< TSharedPtr<FString> > StateNames;
+
+	// These arrays store the textures available for texture parameters of the model.
+	// These come from the texture generators registered in the CustomizableObjectSystem
+	TArray<TSharedPtr<FString>> TextureParameterValueNames;
+	TArray<FName> TextureParameterValues;
+
+	/** Weak pointer of the open editor */
+	TWeakPtr<ICustomizableObjectInstanceEditor> WeakEditor;
+
+	/** Maps param name to children param indices, used to walk the params in order respecting parent/children relationships */
+	TMultiMap<FString, int32> ParamChildren;
+
+	/** Maps param index to bool telling if it has parent, same use as previous line */
+	TMap<int32, bool> ParamHasParent;
+
 };
 
diff --git a/Engine/Plugins/Experimental/Mutable/Source/CustomizableObjectEditor/Private/MuCOE/CustomizableObjectEditor.cpp b/Engine/Plugins/Experimental/Mutable/Source/CustomizableObjectEditor/Private/MuCOE/CustomizableObjectEditor.cpp
--- a/Engine/Plugins/Experimental/Mutable/Source/CustomizableObjectEditor/Private/MuCOE/CustomizableObjectEditor.cpp
+++ b/Engine/Plugins/Experimental/Mutable/Source/CustomizableObjectEditor/Private/MuCOE/CustomizableObjectEditor.cpp
@@ -228,7 +228,7 @@ void FCustomizableObjectEditor::InitCustomizableObjectEditor(const EToolkitMode:
 	FDetailsViewArgs DetailsViewArgs;
 	DetailsViewArgs.NotifyHook = this;
 	DetailsViewArgs.NameAreaSettings = FDetailsViewArgs::ENameAreaSettings::HideNameArea;
-	DetailsViewArgs.bAllowSearch = false;
+	DetailsViewArgs.bAllowSearch = true;
 	//DetailsViewArgs.bShowActorLabel = false;
 	DetailsViewArgs.bShowObjectLabel = false;
 	DetailsViewArgs.bShowScrollBar = false;
diff --git a/Engine/Plugins/Experimental/Mutable/Source/CustomizableObjectEditor/Private/MuCOE/CustomizableObjectEditorViewportClient.cpp b/Engine/Plugins/Experimental/Mutable/Source/CustomizableObjectEditor/Private/MuCOE/CustomizableObjectEditorViewportClient.cpp
--- a/Engine/Plugins/Experimental/Mutable/Source/CustomizableObjectEditor/Private/MuCOE/CustomizableObjectEditorViewportClient.cpp
+++ b/Engine/Plugins/Experimental/Mutable/Source/CustomizableObjectEditor/Private/MuCOE/CustomizableObjectEditorViewportClient.cpp
@@ -54,8 +54,8 @@ class UTextureMipDataProviderFactory;
 #define LOCTEXT_NAMESPACE "CustomizableObjectEditor" 
 
 
-FCustomizableObjectEditorViewportClient::FCustomizableObjectEditorViewportClient(TWeakPtr<ICustomizableObjectInstanceEditor> InCustomizableObjectEditor, FPreviewScene* InPreviewScene)
-	: FEditorViewportClient(&GLevelEditorModeTools(), InPreviewScene)
+FCustomizableObjectEditorViewportClient::FCustomizableObjectEditorViewportClient(TWeakPtr<ICustomizableObjectInstanceEditor> InCustomizableObjectEditor, FPreviewScene* InPreviewScene, const TSharedPtr<SEditorViewport>& EditorViewportWidget)
+	: FEditorViewportClient(&GLevelEditorModeTools(), InPreviewScene, EditorViewportWidget)
 	, CustomizableObjectEditorPtr(InCustomizableObjectEditor)
 	, CustomizableObject(nullptr)
 	, BakingOverwritePermission(false)
diff --git a/Engine/Plugins/Experimental/Mutable/Source/CustomizableObjectEditor/Private/MuCOE/CustomizableObjectInstanceEditor.cpp b/Engine/Plugins/Experimental/Mutable/Source/CustomizableObjectEditor/Private/MuCOE/CustomizableObjectInstanceEditor.cpp
--- a/Engine/Plugins/Experimental/Mutable/Source/CustomizableObjectEditor/Private/MuCOE/CustomizableObjectInstanceEditor.cpp
+++ b/Engine/Plugins/Experimental/Mutable/Source/CustomizableObjectEditor/Private/MuCOE/CustomizableObjectInstanceEditor.cpp
@@ -257,7 +257,7 @@ void FCustomizableObjectInstanceEditor::InitCustomizableObjectInstanceEditor( co
 
 	FDetailsViewArgs DetailsViewArgs;
 	DetailsViewArgs.NameAreaSettings = FDetailsViewArgs::ENameAreaSettings::HideNameArea;
-	DetailsViewArgs.bAllowSearch = false;
+	DetailsViewArgs.bAllowSearch = true;
 	DetailsViewArgs.bShowObjectLabel = false;
 	
 	CustomizableInstanceDetailsView = PropPlugin.CreateDetailView( DetailsViewArgs );
diff --git a/Engine/Plugins/Experimental/Mutable/Source/CustomizableObjectEditor/Private/MuCOE/GenerateMutableSource/GenerateMutableSource.cpp b/Engine/Plugins/Experimental/Mutable/Source/CustomizableObjectEditor/Private/MuCOE/GenerateMutableSource/GenerateMutableSource.cpp
--- a/Engine/Plugins/Experimental/Mutable/Source/CustomizableObjectEditor/Private/MuCOE/GenerateMutableSource/GenerateMutableSource.cpp
+++ b/Engine/Plugins/Experimental/Mutable/Source/CustomizableObjectEditor/Private/MuCOE/GenerateMutableSource/GenerateMutableSource.cpp
@@ -825,7 +825,6 @@ mu::NodeObjectPtr GenerateMutableSource(const UEdGraphPin * Pin, FMutableGraphGe
 			}
 		}
 
-
 		// States
 		int NumStates = TypedNodeObj->States.Num();
 		ObjectNode->SetStateCount(NumStates);
@@ -881,12 +880,6 @@ mu::NodeObjectPtr GenerateMutableSource(const UEdGraphPin * Pin, FMutableGraphGe
 			GenerationContext.CurrentAutoLODStrategy = TypedNodeObj->AutoLODStrategy;
 		}
 
-		// UI Data
-		GenerationContext.ParameterUIDataMap.Add(TypedNodeObj->ObjectName, FParameterUIData(
-			TypedNodeObj->ObjectName,
-			TypedNodeObj->ParamUIMetadata,
-			EMutableParameterType::Int));
-
 		// Mesh components per instance
 		if (!GenerationContext.NumMeshComponentsInRoot)
 		{
@@ -1168,6 +1161,12 @@ mu::NodeObjectPtr GenerateMutableSource(const UEdGraphPin * Pin, FMutableGraphGe
 					if (TypedNodeGroup->GroupType == ECustomizableObjectGroupType::COGT_TOGGLE)
 					{
 						GenerationContext.AddParameterNameUnique(CustomizableObjectNodeObject, CustomizableObjectNodeObject->ObjectName);
+						
+						// UI Data is only relevant when the group node is set to Toggle
+						GenerationContext.ParameterUIDataMap.Add(CustomizableObjectNodeObject->ObjectName, FParameterUIData(
+							CustomizableObjectNodeObject->ObjectName,
+							CustomizableObjectNodeObject->ParamUIMetadata,
+							EMutableParameterType::Int));
 					}
 				}
 			}
@@ -1264,6 +1263,12 @@ mu::NodeObjectPtr GenerateMutableSource(const UEdGraphPin * Pin, FMutableGraphGe
 					if (TypedNodeGroup->GroupType == ECustomizableObjectGroupType::COGT_TOGGLE)
 					{
 						GenerationContext.AddParameterNameUnique(CustomizableObjectNodeObject, CustomizableObjectNodeObject->ObjectName);
+
+						// UI Data is only relevant when the group node is set to Toggle
+						GenerationContext.ParameterUIDataMap.Add(CustomizableObjectNodeObject->ObjectName, FParameterUIData(
+							CustomizableObjectNodeObject->ObjectName,
+							CustomizableObjectNodeObject->ParamUIMetadata,
+							EMutableParameterType::Int));
 					}
 				}
 			}
diff --git a/Engine/Plugins/Experimental/Mutable/Source/CustomizableObjectEditor/Private/MuCOE/GenerateMutableSource/GenerateMutableSourceColor.cpp b/Engine/Plugins/Experimental/Mutable/Source/CustomizableObjectEditor/Private/MuCOE/GenerateMutableSource/GenerateMutableSourceColor.cpp
--- a/Engine/Plugins/Experimental/Mutable/Source/CustomizableObjectEditor/Private/MuCOE/GenerateMutableSource/GenerateMutableSourceColor.cpp
+++ b/Engine/Plugins/Experimental/Mutable/Source/CustomizableObjectEditor/Private/MuCOE/GenerateMutableSource/GenerateMutableSourceColor.cpp
@@ -47,6 +47,11 @@ mu::NodeColourPtr GenerateMutableSourceColor(const UEdGraphPin* Pin, FMutableGra
 		return static_cast<mu::NodeColour*>(Generated->Node.get());
 	}
 
+	if (Node->IsNodeOutDatedAndNeedsRefresh())
+	{
+		Node->SetRefreshNodeWarning();
+	}
+
 	mu::NodeColourPtr Result;
 
 	if (const UCustomizableObjectNodeColorConstant* TypedNodeColorConst = Cast<UCustomizableObjectNodeColorConstant>(Node))
@@ -292,7 +297,7 @@ mu::NodeColourPtr GenerateMutableSourceColor(const UEdGraphPin* Pin, FMutableGra
 			{
 				// Generating a new data table if not exists
 				mu::TablePtr Table;
-				Table = GenerateMutableSourceTable(DataTable->GetName(), Pin, GenerationContext);
+				Table = GenerateMutableSourceTable(DataTable, TypedNodeTable, GenerationContext);
 
 				if (Table)
 				{
@@ -320,8 +325,6 @@ mu::NodeColourPtr GenerateMutableSourceColor(const UEdGraphPin* Pin, FMutableGra
 						ColorTableNode->SetColumn(ColumnName);
 						ColorTableNode->SetParameterName(TypedNodeTable->ParameterName);
 						ColorTableNode->SetNoneOption(TypedNodeTable->bAddNoneOption);
-
-						GenerationContext.AddParameterNameUnique(Node, TypedNodeTable->ParameterName);
 					}
 				}
 				else
diff --git a/Engine/Plugins/Experimental/Mutable/Source/CustomizableObjectEditor/Private/MuCOE/GenerateMutableSource/GenerateMutableSourceFloat.cpp b/Engine/Plugins/Experimental/Mutable/Source/CustomizableObjectEditor/Private/MuCOE/GenerateMutableSource/GenerateMutableSourceFloat.cpp
--- a/Engine/Plugins/Experimental/Mutable/Source/CustomizableObjectEditor/Private/MuCOE/GenerateMutableSource/GenerateMutableSourceFloat.cpp
+++ b/Engine/Plugins/Experimental/Mutable/Source/CustomizableObjectEditor/Private/MuCOE/GenerateMutableSource/GenerateMutableSourceFloat.cpp
@@ -74,6 +74,11 @@ mu::NodeScalarPtr GenerateMutableSourceFloat(const UEdGraphPin* Pin, FMutableGra
 		return static_cast<mu::NodeScalar*>(Generated->Node.get());
 	}
 
+	if (Node->IsNodeOutDatedAndNeedsRefresh())
+	{
+		Node->SetRefreshNodeWarning();
+	}
+
 	bool bDoNotAddToGeneratedCache = false;
 
 	mu::NodeScalarPtr Result;
@@ -326,7 +331,7 @@ mu::NodeScalarPtr GenerateMutableSourceFloat(const UEdGraphPin* Pin, FMutableGra
 			{
 				// Generating a new data table if not exists
 				mu::TablePtr Table;
-				Table = GenerateMutableSourceTable(DataTable->GetName(), Pin, GenerationContext);
+				Table = GenerateMutableSourceTable(DataTable, TypedNodeTable, GenerationContext);
 
 				if (Table)
 				{
@@ -360,8 +365,6 @@ mu::NodeScalarPtr GenerateMutableSourceFloat(const UEdGraphPin* Pin, FMutableGra
 						ScalarTableNode->SetColumn(ColumnName);
 						ScalarTableNode->SetParameterName(TypedNodeTable->ParameterName);
 						ScalarTableNode->SetNoneOption(TypedNodeTable->bAddNoneOption);
-
-						GenerationContext.AddParameterNameUnique(Node, TypedNodeTable->ParameterName);
 					}
 				}
 				else
diff --git a/Engine/Plugins/Experimental/Mutable/Source/CustomizableObjectEditor/Private/MuCOE/GenerateMutableSource/GenerateMutableSourceImage.cpp b/Engine/Plugins/Experimental/Mutable/Source/CustomizableObjectEditor/Private/MuCOE/GenerateMutableSource/GenerateMutableSourceImage.cpp
--- a/Engine/Plugins/Experimental/Mutable/Source/CustomizableObjectEditor/Private/MuCOE/GenerateMutableSource/GenerateMutableSourceImage.cpp
+++ b/Engine/Plugins/Experimental/Mutable/Source/CustomizableObjectEditor/Private/MuCOE/GenerateMutableSource/GenerateMutableSourceImage.cpp
@@ -143,6 +143,11 @@ mu::NodeImagePtr GenerateMutableSourceImage(const UEdGraphPin* Pin, FMutableGrap
 		return static_cast<mu::NodeImage*>(Generated->Node.get());
 	}
 
+	if (Node->IsNodeOutDatedAndNeedsRefresh())
+	{
+		Node->SetRefreshNodeWarning();
+	}
+
 	bool bDoNotAddToGeneratedCache = false;
 
 	mu::NodeImagePtr Result;
@@ -1010,7 +1015,7 @@ mu::NodeImagePtr GenerateMutableSourceImage(const UEdGraphPin* Pin, FMutableGrap
 				{
 					// Generating a new data table if not exists
 					mu::TablePtr Table = nullptr;
-					Table = GenerateMutableSourceTable(DataTable->GetName(), Pin, GenerationContext);
+					Table = GenerateMutableSourceTable(DataTable, TypedNodeTable, GenerationContext);
 
 					if (Table)
 					{
@@ -1080,8 +1085,6 @@ mu::NodeImagePtr GenerateMutableSourceImage(const UEdGraphPin* Pin, FMutableGrap
 								ImageTableNode->SetMaxTextureSize(FMath::Max(DefaultTexture2D->Source.GetSizeX(), DefaultTexture2D->Source.GetSizeY()));
 								ImageTableNode->SetReferenceImageDescriptor(GenerateImageDescriptor(DefaultTexture2D));
 							}
-							
-							GenerationContext.AddParameterNameUnique(Node, TypedNodeTable->ParameterName);
 						}
 					}
 					else
diff --git a/Engine/Plugins/Experimental/Mutable/Source/CustomizableObjectEditor/Private/MuCOE/GenerateMutableSource/GenerateMutableSourceMesh.cpp b/Engine/Plugins/Experimental/Mutable/Source/CustomizableObjectEditor/Private/MuCOE/GenerateMutableSource/GenerateMutableSourceMesh.cpp
--- a/Engine/Plugins/Experimental/Mutable/Source/CustomizableObjectEditor/Private/MuCOE/GenerateMutableSource/GenerateMutableSourceMesh.cpp
+++ b/Engine/Plugins/Experimental/Mutable/Source/CustomizableObjectEditor/Private/MuCOE/GenerateMutableSource/GenerateMutableSourceMesh.cpp
@@ -3648,7 +3648,7 @@ mu::NodeMeshPtr GenerateMutableSourceMesh(const UEdGraphPin* Pin,
 			{
 				// Generating a new data table if not exists
 				mu::TablePtr Table = nullptr;
-				Table = GenerateMutableSourceTable(DataTable->GetName(), Pin, GenerationContext);
+				Table = GenerateMutableSourceTable(DataTable, TypedNodeTable, GenerationContext);
 
 				if (Table)
 				{
@@ -3691,8 +3691,6 @@ mu::NodeMeshPtr GenerateMutableSourceMesh(const UEdGraphPin* Pin,
 						MeshTableNode->SetParameterName(TypedNodeTable->ParameterName);
 						MeshTableNode->SetNoneOption(TypedNodeTable->bAddNoneOption);
 
-						GenerationContext.AddParameterNameUnique(Node, TypedNodeTable->ParameterName);
-
 						if (DefaultSkeletalMesh)
 						{
 							FSkeletalMeshModel* ImportedModel = DefaultSkeletalMesh->GetImportedModel();
diff --git a/Engine/Plugins/Experimental/Mutable/Source/CustomizableObjectEditor/Private/MuCOE/GenerateMutableSource/GenerateMutableSourceTable.cpp b/Engine/Plugins/Experimental/Mutable/Source/CustomizableObjectEditor/Private/MuCOE/GenerateMutableSource/GenerateMutableSourceTable.cpp
--- a/Engine/Plugins/Experimental/Mutable/Source/CustomizableObjectEditor/Private/MuCOE/GenerateMutableSource/GenerateMutableSourceTable.cpp
+++ b/Engine/Plugins/Experimental/Mutable/Source/CustomizableObjectEditor/Private/MuCOE/GenerateMutableSource/GenerateMutableSourceTable.cpp
@@ -666,7 +666,7 @@ bool GenerateTableColumn(const UCustomizableObjectNodeTable* TableNode, const UE
 			{
 				bool bCellGenerated = FillTableColumn(TableNode, MutableTable, DataTableColumnName, RowNames[RowIndex].ToString(), RowIndex, CellData, ColumnProperty,
 					LODIndexConnected, SectionIndexConnected, LODIndex, SectionIndex, bOnlyConnectedLOD, GenerationContext);
-				
+
 				if (!bCellGenerated)
 				{
 					return false;
@@ -679,75 +679,68 @@ bool GenerateTableColumn(const UCustomizableObjectNodeTable* TableNode, const UE
 }
 
 
-mu::TablePtr GenerateMutableSourceTable(const FString& TableName, const UEdGraphPin* Pin, FMutableGraphGenerationContext& GenerationContext)
-{	
-	if (mu::TablePtr* Result = GenerationContext.GeneratedTables.Find(TableName))
-	{
-		return *Result;
-	}
+void GenerateTableParameterUIData(const UDataTable* DataTable, const UCustomizableObjectNodeTable* TableNode, FMutableGraphGenerationContext& GenerationContext)
+{
+	// Checking if the parameter name already exists
+	GenerationContext.AddParameterNameUnique(TableNode, TableNode->ParameterName);
 
-	mu::TablePtr MutableTable = new mu::Table();
-	
-	UCustomizableObjectNode* Node = CastChecked<UCustomizableObjectNode>(Pin->GetOwningNode());
-	if (Node->IsNodeOutDatedAndNeedsRefresh())
+	// Generating Parameter UI MetaData if not exists
+	if (!GenerationContext.ParameterUIDataMap.Contains(TableNode->ParameterName))
 	{
-		Node->SetRefreshNodeWarning();
+		// Getting Table and row names to access the information
+		const TArray<FName>& RowNames = TableNode->GetRowNames(DataTable);
+		FParameterUIData ParameterUIData(TableNode->ParameterName, TableNode->ParamUIMetadata, EMutableParameterType::Int);
+
+		for (int32 NameIndex = 0; NameIndex < RowNames.Num(); ++NameIndex)
+		{
+			ParameterUIData.ArrayIntegerParameterOption.Add(FIntegerParameterUIData(RowNames[NameIndex].ToString(), FMutableParamUIMetadata()));
+		}
+
+		ParameterUIData.IntegerParameterGroupType = TableNode->bAddNoneOption ? ECustomizableObjectGroupType::COGT_ONE_OR_NONE : ECustomizableObjectGroupType::COGT_ONE;
+		GenerationContext.ParameterUIDataMap.Add(TableNode->ParameterName, ParameterUIData);
 	}
+}
 
-	if (const UCustomizableObjectNodeTable* TypedTable = Cast<UCustomizableObjectNodeTable>(Node))
-	{
-		UDataTable* DataTable = GetDataTable(TypedTable, GenerationContext);
 
-		if (!DataTable)
-		{
-			FString msg = "Couldn't find the Data Table asset in the Node.";
-			GenerationContext.Compiler->CompilerLog(FText::FromString(msg), Node);
+mu::TablePtr GenerateMutableSourceTable(const UDataTable* DataTable, const UCustomizableObjectNodeTable* TableNode, FMutableGraphGenerationContext& GenerationContext)
+{
+	check(DataTable && TableNode);
 
-			return nullptr;
-		}
+	// Checking if the table in the cache
+	const FString TableName = DataTable->GetName();
 
-		const UScriptStruct* TableStruct = DataTable->GetRowStruct();
+	if (mu::TablePtr* Result = GenerationContext.GeneratedTables.Find(TableName))
+	{
+		// Generating Parameter Metadata for parameters that reuse a Table
+		GenerateTableParameterUIData(DataTable, TableNode, GenerationContext);
 
-		if (TableStruct)
-		{
-			// Getting names of the rows to access the information
-			const TArray<FName>& RowNames = TypedTable->GetRowNames(DataTable);
+		return *Result;
+	}
 
-			// Adding and filling Name Column
-			MutableTable->AddColumn("Name", mu::ETableColumnType::String);
+	mu::TablePtr MutableTable = new mu::Table();
 
-			// Add metadata
-			FParameterUIData ParameterUIData(
-				TypedTable->ParameterName,
-				TypedTable->ParamUIMetadata,
-				EMutableParameterType::Int);
-			
-			ParameterUIData.IntegerParameterGroupType = TypedTable->bAddNoneOption ? ECustomizableObjectGroupType::COGT_ONE_OR_NONE : ECustomizableObjectGroupType::COGT_ONE;
+	if (const UScriptStruct* TableStruct = DataTable->GetRowStruct())
+	{
+		// Getting Table and row names to access the information
+		const TArray<FName>& RowNames = TableNode->GetRowNames(DataTable);
 
-			for (int32 i = 0; i < RowNames.Num(); ++i)
-			{
-				MutableTable->AddRow(i);
-				FString RowName= RowNames[i].ToString();
-				MutableTable->SetCell(0, i, RowName);
-				ParameterUIData.ArrayIntegerParameterOption.Add(FIntegerParameterUIData(
-					RowName,
-					FMutableParamUIMetadata()));
-			}
+		// Adding and filling Name Column
+		MutableTable->AddColumn("Name", mu::ETableColumnType::String);
 
-			GenerationContext.ParameterUIDataMap.Add(TypedTable->ParameterName, ParameterUIData);
-		}
-		else
+		for (int32 NameIndex = 0; NameIndex < RowNames.Num(); ++NameIndex)
 		{
-			FString msg = "Couldn't find the Data Table's Struct asset in the Node.";
-			GenerationContext.Compiler->CompilerLog(FText::FromString(msg), Node);
-			
-			return nullptr;
+			MutableTable->AddRow(NameIndex);
+			MutableTable->SetCell(0, NameIndex, RowNames[NameIndex].ToString());
 		}
+
+		// Generating Parameter Metadata for new table parameters
+		GenerateTableParameterUIData(DataTable, TableNode, GenerationContext);
 	}
 	else
 	{
-		GenerationContext.Compiler->CompilerLog(LOCTEXT("UnimplementedNode", "Node type not implemented yet."), Node);
-
+		FString msg = "Couldn't find the Data Table's Struct asset in the Node.";
+		GenerationContext.Compiler->CompilerLog(FText::FromString(msg), DataTable);
+		
 		return nullptr;
 	}
 
diff --git a/Engine/Plugins/Experimental/Mutable/Source/CustomizableObjectEditor/Private/MuCOE/GenerateMutableSource/GenerateMutableSourceTable.h b/Engine/Plugins/Experimental/Mutable/Source/CustomizableObjectEditor/Private/MuCOE/GenerateMutableSource/GenerateMutableSourceTable.h
--- a/Engine/Plugins/Experimental/Mutable/Source/CustomizableObjectEditor/Private/MuCOE/GenerateMutableSource/GenerateMutableSourceTable.h
+++ b/Engine/Plugins/Experimental/Mutable/Source/CustomizableObjectEditor/Private/MuCOE/GenerateMutableSource/GenerateMutableSourceTable.h
@@ -57,7 +57,10 @@ bool FillTableColumn(const UCustomizableObjectNodeTable* TableNode, mu::TablePtr
 	FMutableGraphGenerationContext& GenerationContext);
 
 
-mu::TablePtr GenerateMutableSourceTable(const FString& TableName, const UEdGraphPin* Pin, FMutableGraphGenerationContext& GenerationContext);
+mu::TablePtr GenerateMutableSourceTable(const UDataTable* DataTable, const UCustomizableObjectNodeTable* TableNode, FMutableGraphGenerationContext& GenerationContext);
+
+/** Generates a new FParameterUIData for a new table parameter. Also, adds the UIData to the parameters map. */
+void GenerateTableParameterUIData(const UDataTable* DataTable, const UCustomizableObjectNodeTable* TableNode, FMutableGraphGenerationContext& GenerationContext);
 
 /** Gets the data table needed during the compilation process */
 UDataTable* GetDataTable(const UCustomizableObjectNodeTable* TableNode, FMutableGraphGenerationContext& GenerationContext);
@@ -67,3 +70,4 @@ UDataTable* GenerateDataTableFromStruct(const UCustomizableObjectNodeTable* Tabl
 
 /** This method adds the original Data Table(s) of the processed row at the end of the log message when the data table is a Composite Data table. */
 void LogRowGenerationMessage(const UCustomizableObjectNodeTable* TableNode, const UDataTable* DataTable, FMutableGraphGenerationContext& GenerationContext, const FString& Message, const FString& RowName);
+
diff --git a/Engine/Plugins/Experimental/Mutable/Source/CustomizableObjectEditor/Private/MuCOE/SCustomizableObjectEditorViewport.cpp b/Engine/Plugins/Experimental/Mutable/Source/CustomizableObjectEditor/Private/MuCOE/SCustomizableObjectEditorViewport.cpp
--- a/Engine/Plugins/Experimental/Mutable/Source/CustomizableObjectEditor/Private/MuCOE/SCustomizableObjectEditorViewport.cpp
+++ b/Engine/Plugins/Experimental/Mutable/Source/CustomizableObjectEditor/Private/MuCOE/SCustomizableObjectEditorViewport.cpp
@@ -46,7 +46,7 @@ void SCustomizableObjectEditorViewport::Construct(const FArguments& InArgs, cons
 
 TSharedRef<FEditorViewportClient> SCustomizableObjectEditorViewport::MakeEditorViewportClient()
 {
-	LevelViewportClient = MakeShareable(new FCustomizableObjectEditorViewportClient(TabBodyPtr.Pin()->CustomizableObjectEditorPtr, PreviewScenePtr.Pin().Get()));
+	LevelViewportClient = MakeShareable(new FCustomizableObjectEditorViewportClient(TabBodyPtr.Pin()->CustomizableObjectEditorPtr, PreviewScenePtr.Pin().Get(), SharedThis(this)));
 
 	LevelViewportClient->ViewportType = LVT_Perspective;
 	LevelViewportClient->bSetListenerPosition = false;
diff --git a/Engine/Plugins/Experimental/Mutable/Source/CustomizableObjectEditor/Public/MuCOE/CustomizableObjectEditorViewportClient.h b/Engine/Plugins/Experimental/Mutable/Source/CustomizableObjectEditor/Public/MuCOE/CustomizableObjectEditorViewportClient.h
--- a/Engine/Plugins/Experimental/Mutable/Source/CustomizableObjectEditor/Public/MuCOE/CustomizableObjectEditorViewportClient.h
+++ b/Engine/Plugins/Experimental/Mutable/Source/CustomizableObjectEditor/Public/MuCOE/CustomizableObjectEditorViewportClient.h
@@ -68,7 +68,7 @@ enum class EWidgetType : uint8
 class FCustomizableObjectEditorViewportClient : public FEditorViewportClient, public TSharedFromThis<FCustomizableObjectEditorViewportClient>
 {
 public:
-	FCustomizableObjectEditorViewportClient(TWeakPtr<ICustomizableObjectInstanceEditor> InCustomizableObjectEditor, FPreviewScene* InPreviewScene);
+	FCustomizableObjectEditorViewportClient(TWeakPtr<ICustomizableObjectInstanceEditor> InCustomizableObjectEditor, FPreviewScene* InPreviewScene, const TSharedPtr<SEditorViewport>& EditorViewportWidget);
 	~FCustomizableObjectEditorViewportClient();
 
 	/** persona config options **/
