[Mutable] Fix Type mismatch in Uid of CustomizableObjectIntParameterValue warnings.

#rnx

diff --git a/Engine/Plugins/Experimental/Mutable/Source/CustomizableObject/Private/MuCO/CustomizableObjectInstance.cpp b/Engine/Plugins/Experimental/Mutable/Source/CustomizableObject/Private/MuCO/CustomizableObjectInstance.cpp
--- a/Engine/Plugins/Experimental/Mutable/Source/CustomizableObject/Private/MuCO/CustomizableObjectInstance.cpp
+++ b/Engine/Plugins/Experimental/Mutable/Source/CustomizableObject/Private/MuCO/CustomizableObjectInstance.cpp
@@ -525,7 +525,7 @@ void UCustomizableObjectInstance::Serialize(FArchive& Ar)
 			FloatParameters_DEPRECATED.AddDefaulted();
 			FloatParameters_DEPRECATED.Last().ParameterName = IntParameters_DEPRECATED[IntParametersToMove[i]].ParameterName;
 			FloatParameters_DEPRECATED.Last().ParameterValue = FCString::Atoi(*IntParameters_DEPRECATED[IntParametersToMove[i]].ParameterValueName);
-			FloatParameters_DEPRECATED.Last().Uid = IntParameters_DEPRECATED[IntParametersToMove[i]].Uid;
+			FloatParameters_DEPRECATED.Last().Id = IntParameters_DEPRECATED[IntParametersToMove[i]].Id;
 		}
 
 		// Remove them from the int params in reverse order
diff --git a/Engine/Plugins/Experimental/Mutable/Source/CustomizableObject/Private/MuCO/CustomizableObjectInstanceDescriptor.cpp b/Engine/Plugins/Experimental/Mutable/Source/CustomizableObject/Private/MuCO/CustomizableObjectInstanceDescriptor.cpp
--- a/Engine/Plugins/Experimental/Mutable/Source/CustomizableObject/Private/MuCO/CustomizableObjectInstanceDescriptor.cpp
+++ b/Engine/Plugins/Experimental/Mutable/Source/CustomizableObject/Private/MuCO/CustomizableObjectInstanceDescriptor.cpp
@@ -475,7 +475,7 @@ mu::ParametersPtr FCustomizableObjectInstanceDescriptor::GetParameters() const
 		{
 			for (const FCustomizableObjectBoolParameterValue& BoolParameter : BoolParameters)
 			{
-				if (BoolParameter.ParameterName == Name || (Uid.IsValid() && BoolParameter.Uid == Uid))
+				if (BoolParameter.ParameterName == Name || (Uid.IsValid() && BoolParameter.Id == Uid))
 				{
 					MutableParameters->SetBoolValue(ParamIndex,  BoolParameter.ParameterValue);
 					break;
@@ -489,7 +489,7 @@ mu::ParametersPtr FCustomizableObjectInstanceDescriptor::GetParameters() const
 		{
 			for (const FCustomizableObjectIntParameterValue& IntParameter : IntParameters)
 			{
-				if (IntParameter.ParameterName.Equals(Name, ESearchCase::CaseSensitive) || (Uid.IsValid() && IntParameter.Uid == Uid))
+				if (IntParameter.ParameterName.Equals(Name, ESearchCase::CaseSensitive) || (Uid.IsValid() && IntParameter.Id == Uid))
 				{
 					if (mu::RangeIndexPtr RangeIdxPtr = MutableParameters->NewRangeIndex(ParamIndex))
 					{
@@ -519,7 +519,7 @@ mu::ParametersPtr FCustomizableObjectInstanceDescriptor::GetParameters() const
 		{
 			for (const FCustomizableObjectFloatParameterValue& FloatParameter : FloatParameters)
 			{
-				if (FloatParameter.ParameterName == Name || (Uid.IsValid() && FloatParameter.Uid == Uid))
+				if (FloatParameter.ParameterName == Name || (Uid.IsValid() && FloatParameter.Id == Uid))
 				{
 					if (mu::RangeIndexPtr RangeIdxPtr = MutableParameters->NewRangeIndex(ParamIndex))
 					{
@@ -545,7 +545,7 @@ mu::ParametersPtr FCustomizableObjectInstanceDescriptor::GetParameters() const
 		{
 			for (const FCustomizableObjectVectorParameterValue& VectorParameter : VectorParameters)
 			{
-				if (VectorParameter.ParameterName == Name || (Uid.IsValid() && VectorParameter.Uid == Uid))
+				if (VectorParameter.ParameterName == Name || (Uid.IsValid() && VectorParameter.Id == Uid))
 				{
 					MutableParameters->SetColourValue(ParamIndex, VectorParameter.ParameterValue.R, VectorParameter.ParameterValue.G, VectorParameter.ParameterValue.B);
 
@@ -560,7 +560,7 @@ mu::ParametersPtr FCustomizableObjectInstanceDescriptor::GetParameters() const
 		{
 			for (const auto& ProjectorParameter : ProjectorParameters)
 			{
-				if (ProjectorParameter.ParameterName == Name || (Uid.IsValid() && ProjectorParameter.Uid == Uid))
+				if (ProjectorParameter.ParameterName == Name || (Uid.IsValid() && ProjectorParameter.Id == Uid))
 				{
 					auto CopyProjector = [&MutableParameters, ParamIndex](const FCustomizableObjectProjector& Value, const mu::RangeIndexPtr RangeIdxPtr = nullptr)
 					{
@@ -621,7 +621,7 @@ mu::ParametersPtr FCustomizableObjectInstanceDescriptor::GetParameters() const
 		{
 			for (const FCustomizableObjectTextureParameterValue& TextureParameter : TextureParameters)
 			{
-				if (TextureParameter.ParameterName == Name || (Uid.IsValid() && TextureParameter.Uid == Uid))
+				if (TextureParameter.ParameterName == Name || (Uid.IsValid() && TextureParameter.Id == Uid))
 				{
 					if (mu::RangeIndexPtr RangeIdxPtr = MutableParameters->NewRangeIndex(ParamIndex))
 					{
@@ -788,11 +788,11 @@ void FCustomizableObjectInstanceDescriptor::ReloadParameters()
 		{
 			FCustomizableObjectBoolParameterValue Param;
 			Param.ParameterName = Name;
-			Param.Uid = Uid;
+			Param.Id = Uid;
 
 			auto FindByNameAndUid = [&](const FCustomizableObjectBoolParameterValue& P)
 			{
-				return P.ParameterName == Name || (Uid.IsValid() && P.Uid == Uid);
+				return P.ParameterName == Name || (Uid.IsValid() && P.Id == Uid);
 			};
 
 			if (FCustomizableObjectBoolParameterValue* Result = OldBoolParameters.FindByPredicate(FindByNameAndUid))
@@ -812,11 +812,11 @@ void FCustomizableObjectInstanceDescriptor::ReloadParameters()
 		{
 			FCustomizableObjectIntParameterValue Param;
 			Param.ParameterName = Name;
-			Param.Uid = Uid;
+			Param.Id = Uid;
 
 			auto FindByNameAndUid = [&](const FCustomizableObjectIntParameterValue& P)
 			{
-				return P.ParameterName == Name || (Uid.IsValid() && P.Uid == Uid);
+				return P.ParameterName == Name || (Uid.IsValid() && P.Id == Uid);
 			};
 			
 			if (FCustomizableObjectIntParameterValue* Result = OldIntParameters.FindByPredicate(FindByNameAndUid))
@@ -916,11 +916,11 @@ void FCustomizableObjectInstanceDescriptor::ReloadParameters()
 		{
 			FCustomizableObjectFloatParameterValue Param;
 			Param.ParameterName = Name;
-			Param.Uid = Uid;
+			Param.Id = Uid;
 
 			auto FindByNameAndUid = [&](const FCustomizableObjectFloatParameterValue& P)
 			{
-				return P.ParameterName == Name || (Uid.IsValid() && P.Uid == Uid);
+				return P.ParameterName == Name || (Uid.IsValid() && P.Id == Uid);
 			};
 			
 			if (FCustomizableObjectFloatParameterValue* Result = OldFloatParameters.FindByPredicate(FindByNameAndUid))
@@ -967,11 +967,11 @@ void FCustomizableObjectInstanceDescriptor::ReloadParameters()
 		{
 			FCustomizableObjectVectorParameterValue Param;
 			Param.ParameterName = Name;
-			Param.Uid = Uid;
+			Param.Id = Uid;
 
 			auto FindByNameAndUid = [&](const FCustomizableObjectVectorParameterValue& P)
 			{
-				return P.ParameterName == Name || (Uid.IsValid() && P.Uid == Uid);
+				return P.ParameterName == Name || (Uid.IsValid() && P.Id == Uid);
 			};
 
 			if (FCustomizableObjectVectorParameterValue* Result = OldVectorParameters.FindByPredicate(FindByNameAndUid))
@@ -992,11 +992,11 @@ void FCustomizableObjectInstanceDescriptor::ReloadParameters()
 		{
 			FCustomizableObjectProjectorParameterValue Param;
 			Param.ParameterName = Name;
-			Param.Uid = Uid;
+			Param.Id = Uid;
 
 			auto FindByNameAndUid = [&](const FCustomizableObjectProjectorParameterValue& P)
 			{
-				return P.ParameterName == Name || (Uid.IsValid() && P.Uid == Uid);
+				return P.ParameterName == Name || (Uid.IsValid() && P.Id == Uid);
 			};
 			
 			if (FCustomizableObjectProjectorParameterValue* Result = OldProjectorParameters.FindByPredicate(FindByNameAndUid))
@@ -1067,11 +1067,11 @@ void FCustomizableObjectInstanceDescriptor::ReloadParameters()
 		{
 			FCustomizableObjectTextureParameterValue Param;
 			Param.ParameterName = Name;
-			Param.Uid = Uid;
+			Param.Id = Uid;
 
 			auto FindByNameAndUid = [&](const FCustomizableObjectTextureParameterValue& P)
 			{
-				return P.ParameterName == Name || (Uid.IsValid() && P.Uid == Uid);
+				return P.ParameterName == Name || (Uid.IsValid() && P.Id == Uid);
 			};
 			
 			if (FCustomizableObjectTextureParameterValue* Result = OldTextureParameters.FindByPredicate(FindByNameAndUid))
diff --git a/Engine/Plugins/Experimental/Mutable/Source/CustomizableObject/Public/MuCO/CustomizableObjectParameterTypeDefinitions.h b/Engine/Plugins/Experimental/Mutable/Source/CustomizableObject/Public/MuCO/CustomizableObjectParameterTypeDefinitions.h
--- a/Engine/Plugins/Experimental/Mutable/Source/CustomizableObject/Public/MuCO/CustomizableObjectParameterTypeDefinitions.h
+++ b/Engine/Plugins/Experimental/Mutable/Source/CustomizableObject/Public/MuCO/CustomizableObjectParameterTypeDefinitions.h
@@ -59,7 +59,7 @@ struct FCustomizableObjectBoolParameterValue
 	bool ParameterValue = true;
 
 	UPROPERTY(Category = CustomizableObjectBoolParameterValue, VisibleAnywhere)
-	FGuid Uid;
+	FGuid Id;
 };
 
 
@@ -87,7 +87,7 @@ struct FCustomizableObjectIntParameterValue
 	FString ParameterValueName;
 
 	UPROPERTY(Category = CustomizableObjectIntParameterValue, VisibleAnywhere)
-	FGuid Uid;
+	FGuid Id;
 
 	// Same as ParameterValueName but for multidimensional params
 	UPROPERTY()
@@ -98,7 +98,7 @@ struct FCustomizableObjectIntParameterValue
 	}
 
 	FCustomizableObjectIntParameterValue(const FString & InParamName, const FString & InParameterValueName, const FGuid & InUid, const TArray<FString>& InParameterRangeValueNames)
-		: ParameterName(InParamName), ParameterValueName(InParameterValueName), Uid(InUid)
+		: ParameterName(InParamName), ParameterValueName(InParameterValueName), Id(InUid)
 	{
 		ParameterRangeValueNames = InParameterRangeValueNames;
 	}
@@ -133,7 +133,7 @@ struct FCustomizableObjectFloatParameterValue
 	float ParameterValue = 0.0f;
 
 	UPROPERTY(Category = CustomizableObjectFloatParameterValue, VisibleAnywhere)
-	FGuid Uid;
+	FGuid Id;
 
 	UPROPERTY(Category = CustomizableObjectFloatParameterValue, VisibleAnywhere)
 	TArray<float> ParameterRangeValues;
@@ -168,7 +168,7 @@ struct FCustomizableObjectTextureParameterValue
 	FName ParameterValue;
 
 	UPROPERTY(Category = CustomizableObjectTextureParameterValue, VisibleAnywhere)
-	FGuid Uid;
+	FGuid Id;
 
 	UPROPERTY(Category = CustomizableObjectTextureParameterValue, VisibleAnywhere)
 	TArray<FName> ParameterRangeValues;
@@ -203,7 +203,7 @@ struct FCustomizableObjectVectorParameterValue
 	FLinearColor ParameterValue = FLinearColor(ForceInit);
 
 	UPROPERTY(Category = CustomizableObjectVectorParameterValue, VisibleAnywhere)
-	FGuid Uid;
+	FGuid Id;
 };
 
 
@@ -290,7 +290,7 @@ struct FCustomizableObjectProjectorParameterValue
 	FCustomizableObjectProjector Value;
 
 	UPROPERTY(Category = CustomizableObjectVectorParameterValue, VisibleAnywhere)
-	FGuid Uid;
+	FGuid Id;
 
 	UPROPERTY(Category = CustomizableObjectVectorParameterValue, VisibleAnywhere)
 	TArray<FCustomizableObjectProjector> RangeValues;
