Fixed unreachable code warning/error.

diff --git a/Engine/Plugins/Experimental/Mutable/Source/CustomizableObjectEditor/Private/MuCOE/CustomizableObjectEditor.cpp b/Engine/Plugins/Experimental/Mutable/Source/CustomizableObjectEditor/Private/MuCOE/CustomizableObjectEditor.cpp
--- a/Engine/Plugins/Experimental/Mutable/Source/CustomizableObjectEditor/Private/MuCOE/CustomizableObjectEditor.cpp
+++ b/Engine/Plugins/Experimental/Mutable/Source/CustomizableObjectEditor/Private/MuCOE/CustomizableObjectEditor.cpp
@@ -2559,18 +2559,16 @@ UCustomizableObject* FCustomizableObjectEditor::GetAbsoluteCOParent(const UCusto
 		//Get all the NodeObjects
 		TArray<UCustomizableObjectNodeObject*> ObjectNodes;
 		Root->ParentObject->Source->GetNodesOfClass<UCustomizableObjectNodeObject>(ObjectNodes);
-
-		for (int i = 0; i < ObjectNodes.Num(); ++i)
+		if (!ObjectNodes.IsEmpty())
 		{
 			//Getting the parent of the root
-			if (ObjectNodes[i]->ParentObject == nullptr)
+			UCustomizableObjectNodeObject* FirstObjectNode = ObjectNodes[0];
+			if (FirstObjectNode->ParentObject == nullptr)
 			{
 				return Root->ParentObject;
 			}
-			else
-			{
-				return GetAbsoluteCOParent(ObjectNodes[0]);
-			}
+
+			return GetAbsoluteCOParent(FirstObjectNode);
 		}
 	}
 
