[mutable] Fixed LOCTEXT coflict.

#rnx

diff --git a/Engine/Plugins/Experimental/Mutable/Source/CustomizableObjectEditor/Private/MuCOE/CustomizableInstanceDetails.cpp b/Engine/Plugins/Experimental/Mutable/Source/CustomizableObjectEditor/Private/MuCOE/CustomizableInstanceDetails.cpp
--- a/Engine/Plugins/Experimental/Mutable/Source/CustomizableObjectEditor/Private/MuCOE/CustomizableInstanceDetails.cpp
+++ b/Engine/Plugins/Experimental/Mutable/Source/CustomizableObjectEditor/Private/MuCOE/CustomizableInstanceDetails.cpp
@@ -86,7 +86,7 @@ void FCustomizableInstanceDetails::CustomizeDetails(const TSharedPtr<IDetailLayo
 	UCustomizableObject* CustomizableObject = CustomInstance->GetCustomizableObject();
 	if (!CustomizableObject)
 	{
-		VisibilitySettingsCategory.AddCustomRow(LOCTEXT("CustomizableInstanceDetails", "Instance Parameters"))
+		VisibilitySettingsCategory.AddCustomRow(LOCTEXT("CustomizableInstanceDetails_NoCOMessage", "Instance Parameters"))
 		[
 			SNew(STextBlock).Text(LOCTEXT("Model not compiled", "Model not compiled"))
 		];
@@ -108,7 +108,7 @@ void FCustomizableInstanceDetails::CustomizeDetails(const TSharedPtr<IDetailLayo
 	// In case that something of the new UI doesn't work as expected
 	if(CVarUseOldInstanceUI.GetValueOnGameThread())
 	{
-		OldParametersCategory.AddCustomRow(LOCTEXT("CustomizableInstanceDetails", "Old Instance Parameters"))
+		OldParametersCategory.AddCustomRow(LOCTEXT("CustomizableInstanceDetails_OldUI", "Old Instance Parameters"))
 		[
 			SNew(SCustomizableInstanceProperties)
 			.CustomInstance(CustomInstance.Get())
