[Mutable] Remove PostLoad auto-compilations.

#rnx

diff --git a/Engine/Plugins/Experimental/Mutable/Source/CustomizableObject/Private/MuCO/CustomizableObject.cpp b/Engine/Plugins/Experimental/Mutable/Source/CustomizableObject/Private/MuCO/CustomizableObject.cpp
--- a/Engine/Plugins/Experimental/Mutable/Source/CustomizableObject/Private/MuCO/CustomizableObject.cpp
+++ b/Engine/Plugins/Experimental/Mutable/Source/CustomizableObject/Private/MuCO/CustomizableObject.cpp
@@ -798,7 +798,6 @@ void UCustomizableObject::LoadCompiledDataFromDisk()
 	if (!CustomizableObjectPrivate->GetModel()) // Not failed to load the model
 	{
 		CustomizableObjectPrivate->Status.NextState(FCustomizableObjectStatusTypes::EState::NoModel);
-		ConditionalAutoCompile(); // Recompile if no model has been loaded
 	}
 }
 
diff --git a/Engine/Plugins/Experimental/Mutable/Source/CustomizableObject/Public/MuCO/StateMachine.h b/Engine/Plugins/Experimental/Mutable/Source/CustomizableObject/Public/MuCO/StateMachine.h
--- a/Engine/Plugins/Experimental/Mutable/Source/CustomizableObject/Public/MuCO/StateMachine.h
+++ b/Engine/Plugins/Experimental/Mutable/Source/CustomizableObject/Public/MuCO/StateMachine.h
@@ -23,9 +23,10 @@ public:
 	void NextState(EState State)
 	{
 		check(T::ValidTransitions[static_cast<int32>(CurrentState)][static_cast<int32>(State)])
-		
-		OnStateChangedDelegate.Broadcast(CurrentState, State);
+
+		EState PreviousState = CurrentState; 
 		CurrentState = State;
+		OnStateChangedDelegate.Broadcast(PreviousState, CurrentState);
 	}
 
 	EState Get() const
diff --git a/Engine/Plugins/Experimental/Mutable/Source/CustomizableObjectEditor/Private/MuCOE/CustomizableObjectEditor.cpp b/Engine/Plugins/Experimental/Mutable/Source/CustomizableObjectEditor/Private/MuCOE/CustomizableObjectEditor.cpp
--- a/Engine/Plugins/Experimental/Mutable/Source/CustomizableObjectEditor/Private/MuCOE/CustomizableObjectEditor.cpp
+++ b/Engine/Plugins/Experimental/Mutable/Source/CustomizableObjectEditor/Private/MuCOE/CustomizableObjectEditor.cpp
@@ -2723,9 +2723,9 @@ void FCustomizableObjectEditor::CreatePreviewComponents()
 }
 
 
-void FCustomizableObjectEditor::OnCustomizableObjectStatusChanged(FCustomizableObjectStatus::EState, const FCustomizableObjectStatus::EState NextState)
+void FCustomizableObjectEditor::OnCustomizableObjectStatusChanged(FCustomizableObjectStatus::EState, const FCustomizableObjectStatus::EState CurrentState)
 {
-	switch (NextState)
+	switch (CurrentState)
 	{
 	case FCustomizableObjectStatus::EState::ModelLoaded:
 		if (!PreviewInstance)
diff --git a/Engine/Plugins/Experimental/Mutable/Source/CustomizableObjectEditor/Private/MuCOE/CustomizableObjectEditor.h b/Engine/Plugins/Experimental/Mutable/Source/CustomizableObjectEditor/Private/MuCOE/CustomizableObjectEditor.h
--- a/Engine/Plugins/Experimental/Mutable/Source/CustomizableObjectEditor/Private/MuCOE/CustomizableObjectEditor.h
+++ b/Engine/Plugins/Experimental/Mutable/Source/CustomizableObjectEditor/Private/MuCOE/CustomizableObjectEditor.h
@@ -358,7 +358,7 @@ private:
 	void CreatePreviewComponents();
 
 public:
-	void OnCustomizableObjectStatusChanged(FCustomizableObjectStatus::EState PreviousState, FCustomizableObjectStatus::EState NextState);
+	void OnCustomizableObjectStatusChanged(FCustomizableObjectStatus::EState PreviousState, FCustomizableObjectStatus::EState CurrentState);
 
 	// Helpers to get the absolute parent of a Customizable Object
 	static UCustomizableObject* GetAbsoluteCOParent(const UCustomizableObjectNodeObject* const Root);
diff --git a/Engine/Plugins/Experimental/Mutable/Source/CustomizableObjectEditor/Private/MuCOE/CustomizableObjectInstanceEditor.cpp b/Engine/Plugins/Experimental/Mutable/Source/CustomizableObjectEditor/Private/MuCOE/CustomizableObjectInstanceEditor.cpp
--- a/Engine/Plugins/Experimental/Mutable/Source/CustomizableObjectEditor/Private/MuCOE/CustomizableObjectInstanceEditor.cpp
+++ b/Engine/Plugins/Experimental/Mutable/Source/CustomizableObjectEditor/Private/MuCOE/CustomizableObjectInstanceEditor.cpp
@@ -340,7 +340,7 @@ FText FCustomizableObjectInstanceEditor::GetBaseToolkitName() const
 void FCustomizableObjectInstanceEditor::CreatePreviewInstance()
 {
 	check(CustomizableObjectInstance);
-	if (!CustomizableObjectInstance->CanUpdateInstance())
+	if (!CustomizableObjectInstance->GetCustomizableObject())
 	{
 		return;
 	}
@@ -733,9 +733,9 @@ TStatId FCustomizableObjectInstanceEditor::GetStatId() const
 }
 
 
-void FCustomizableObjectInstanceEditor::OnCustomizableObjectStatusChanged(FCustomizableObjectStatus::EState, const FCustomizableObjectStatus::EState NextState)
+void FCustomizableObjectInstanceEditor::OnCustomizableObjectStatusChanged(FCustomizableObjectStatus::EState, const FCustomizableObjectStatus::EState CurrentState)
 {
-	switch (NextState)
+	switch (CurrentState)
 	{
 	case FCustomizableObjectStatus::EState::ModelLoaded:
 		{
diff --git a/Engine/Plugins/Experimental/Mutable/Source/CustomizableObjectEditor/Private/MuCOE/CustomizableObjectInstanceEditor.h b/Engine/Plugins/Experimental/Mutable/Source/CustomizableObjectEditor/Private/MuCOE/CustomizableObjectInstanceEditor.h
--- a/Engine/Plugins/Experimental/Mutable/Source/CustomizableObjectEditor/Private/MuCOE/CustomizableObjectInstanceEditor.h
+++ b/Engine/Plugins/Experimental/Mutable/Source/CustomizableObjectEditor/Private/MuCOE/CustomizableObjectInstanceEditor.h
@@ -193,7 +193,7 @@ public:
 	virtual void Tick(float InDeltaTime) override;
 	virtual TStatId GetStatId() const override;
 
-	void OnCustomizableObjectStatusChanged(FCustomizableObjectStatus::EState PreviousState, FCustomizableObjectStatus::EState NextState);
+	void OnCustomizableObjectStatusChanged(FCustomizableObjectStatus::EState PreviousState, FCustomizableObjectStatus::EState CurrentState);
 	
 private:
 
