[mutable] Added log to report CO embeded data size.
- This log will be used to report to MongoDB

diff --git a/Engine/Plugins/Experimental/Mutable/Source/MutableValidation/Private/MuV/CustomizableObjectValidationCommandlet.cpp b/Engine/Plugins/Experimental/Mutable/Source/MutableValidation/Private/MuV/CustomizableObjectValidationCommandlet.cpp
--- a/Engine/Plugins/Experimental/Mutable/Source/MutableValidation/Private/MuV/CustomizableObjectValidationCommandlet.cpp
+++ b/Engine/Plugins/Experimental/Mutable/Source/MutableValidation/Private/MuV/CustomizableObjectValidationCommandlet.cpp
@@ -56,7 +56,7 @@ int32 UCustomizableObjectValidationCommandlet::Main(const FString& Params)
 	}
 	
 	// Compile the Customizable Object ------------------------------------------------------------------------------ //
-	bool bWasCoCompilationSuccesfull = false;
+	bool bWasCoCompilationSuccessful = false;
 	{
 		UE_LOG(LogMutable,Display,TEXT("Compiling Customizable Object..."));
     	
@@ -93,11 +93,11 @@ int32 UCustomizableObjectValidationCommandlet::Main(const FString& Params)
     	check(CompilationEndResult != ECustomizableObjectCompilationState::None);
     	check(CompilationEndResult != ECustomizableObjectCompilationState::InProgress);
     	
-    	bWasCoCompilationSuccesfull = CompilationEndResult == ECustomizableObjectCompilationState::Completed;
+    	bWasCoCompilationSuccessful = CompilationEndResult == ECustomizableObjectCompilationState::Completed;
 	}
 	// -------------------------------------------------------------------------------------------------------------- //
 	
-	if (bWasCoCompilationSuccesfull)
+	if (bWasCoCompilationSuccessful)
 	{
 		UE_LOG(LogMutable,Display,TEXT("Customizable Object was compiled succesfully."));
 		
@@ -119,7 +119,18 @@ int32 UCustomizableObjectValidationCommandlet::Main(const FString& Params)
 			UE_LOG(LogMutable, Log,TEXT("(int) model_rom_count : %d "), RomCount);
 			UE_LOG(LogMutable, Log,TEXT("(int) model_roms_size : %lld "), TotalRomSizeBytes);
 		}
-		
+
+		// CO embedded data size ------ //
+		{
+			TArray<uint8> EmbeddedDataBytes;
+			FMemoryWriter SerializationTarget{EmbeddedDataBytes, false};
+			
+			ToTestCustomizableObject->SaveEmbeddedData(SerializationTarget);
+			const int64 COEmbeddedDataSizeBytes = EmbeddedDataBytes.Num();
+			
+			UE_LOG(LogMutable, Log,TEXT("(int) co_embedded_data_bytes : %lld "), COEmbeddedDataSizeBytes);
+		}
+
 		// Generate target random instances to be tested ------------------------------------------------------------ //
 		bool bWasInstancesCreationSuccessful = true;
 		{
