[mutable] Fixed NSLOCTEXT conflict.

#rnx

diff --git a/Engine/Plugins/Experimental/Mutable/Source/CustomizableObjectEditor/Private/MuCOE/CustomizableObjectEditorViewportClient.cpp b/Engine/Plugins/Experimental/Mutable/Source/CustomizableObjectEditor/Private/MuCOE/CustomizableObjectEditorViewportClient.cpp
--- a/Engine/Plugins/Experimental/Mutable/Source/CustomizableObjectEditor/Private/MuCOE/CustomizableObjectEditorViewportClient.cpp
+++ b/Engine/Plugins/Experimental/Mutable/Source/CustomizableObjectEditor/Private/MuCOE/CustomizableObjectEditorViewportClient.cpp
@@ -1634,7 +1634,7 @@ void FCustomizableObjectEditorViewportClient::BakeInstance(UCustomizableObjectIn
 {
 	if (CustomizableObject->GetPrivate()->Status.Get() == FCustomizableObjectStatus::EState::Loading)
 	{
-		FNotificationInfo Info(NSLOCTEXT("CustomizableObjectEditor", "CustomizableObjectCompileTryLater", "Please wait unitl Customizable Object is loaded"));
+		FNotificationInfo Info(NSLOCTEXT("CustomizableObjectEditor", "CustomizableObjectCompileTryLater_BakeInstance", "Please wait unitl Customizable Object is loaded"));
 		Info.bFireAndForget = true;
 		Info.bUseThrobber = true;
 		Info.FadeOutDuration = 1.0f;
