[Mutable] Fix crash when opening a CO Instance before loading the asset registry.

#rnx

diff --git a/Engine/Plugins/Experimental/Mutable/Source/CustomizableObject/Classes/MuCO/CustomizableObjectInstance.h b/Engine/Plugins/Experimental/Mutable/Source/CustomizableObject/Classes/MuCO/CustomizableObjectInstance.h
--- a/Engine/Plugins/Experimental/Mutable/Source/CustomizableObject/Classes/MuCO/CustomizableObjectInstance.h
+++ b/Engine/Plugins/Experimental/Mutable/Source/CustomizableObject/Classes/MuCO/CustomizableObjectInstance.h
@@ -680,12 +680,6 @@ public:
 	bool bSelectedProfileDirty = false;
 #endif 
 
-private:
-
-#if WITH_EDITOR
-	void BindPostCompileDelegate(UCustomizableObject* CustomizableObject);
-#endif
-
 private:
 	UPROPERTY()
 	FCustomizableObjectInstanceDescriptor Descriptor;
diff --git a/Engine/Plugins/Experimental/Mutable/Source/CustomizableObject/Internal/MuCO/CustomizableInstancePrivateData.h b/Engine/Plugins/Experimental/Mutable/Source/CustomizableObject/Internal/MuCO/CustomizableInstancePrivateData.h
--- a/Engine/Plugins/Experimental/Mutable/Source/CustomizableObject/Internal/MuCO/CustomizableInstancePrivateData.h
+++ b/Engine/Plugins/Experimental/Mutable/Source/CustomizableObject/Internal/MuCO/CustomizableInstancePrivateData.h
@@ -7,6 +7,7 @@
 #include "Rendering/SkeletalMeshLODModel.h"
 #include "Rendering/SkeletalMeshModel.h"
 #include "MuCO/CustomizableObjectInstance.h"
+#include "MuCO/CustomizableObjectPrivate.h"
 #include "MuR/Instance.h"
 #include "GameplayTagContainer.h"
 #include "MuCO/DescriptorHash.h"
@@ -197,8 +198,10 @@ public:
 
 	virtual void PostDuplicate(bool bDuplicateForPIE) override;
 
-	void OnPostCompile();
+	void BindObjectDelegates(UCustomizableObject* CurrentCustomizableObject, UCustomizableObject* NewCustomizableObject);
 
+	void OnPostCompile();
+	void OnObjectStatusChanged(FCustomizableObjectStatus::EState Previous, FCustomizableObjectStatus::EState Next);
 #endif
 
 	/** Invalidates the previously generated data and retrieves information from the CObject after specific actions.
diff --git a/Engine/Plugins/Experimental/Mutable/Source/CustomizableObject/Private/MuCO/CustomizableObjectInstance.cpp b/Engine/Plugins/Experimental/Mutable/Source/CustomizableObject/Private/MuCO/CustomizableObjectInstance.cpp
--- a/Engine/Plugins/Experimental/Mutable/Source/CustomizableObject/Private/MuCO/CustomizableObjectInstance.cpp
+++ b/Engine/Plugins/Experimental/Mutable/Source/CustomizableObject/Private/MuCO/CustomizableObjectInstance.cpp
@@ -265,7 +265,7 @@ void UCustomizableObjectInstance::SetDescriptor(const FCustomizableObjectInstanc
 
 #if WITH_EDITOR
 	// Bind a lambda to the PostCompileDelegate and unbind from the previous object if any.
-	BindPostCompileDelegate(InCustomizableObject);
+	PrivateData->BindObjectDelegates(GetCustomizableObject(), InCustomizableObject);
 #endif
 
 	Descriptor = InDescriptor;
@@ -327,19 +327,34 @@ void UCustomizableInstancePrivate::OnPostCompile()
 }
 
 
-void UCustomizableObjectInstance::BindPostCompileDelegate(UCustomizableObject* InCustomizableObject)
+void UCustomizableInstancePrivate::OnObjectStatusChanged(FCustomizableObjectStatus::EState Previous, FCustomizableObjectStatus::EState Next)
 {
+	if (Previous != Next && Next == FCustomizableObjectStatus::EState::ModelLoaded)
+	{
+		OnPostCompile();
+	}
+}
+
+
+void UCustomizableInstancePrivate::BindObjectDelegates(UCustomizableObject*  CurrentCustomizableObject, UCustomizableObject* NewCustomizableObject)
+{
+	if (CurrentCustomizableObject == NewCustomizableObject)
+	{
+		return;
+	}
+
 	// Unbind callback from the previous CO
-	if (UCustomizableObject* CurrentObject = GetCustomizableObject())
+	if (CurrentCustomizableObject)
 	{
-		CurrentObject->PostCompileDelegate.RemoveAll(GetPrivate());
+		CurrentCustomizableObject->PostCompileDelegate.RemoveAll(this);
+		CurrentCustomizableObject->GetPrivate()->Status.GetOnStateChangedDelegate().RemoveAll(this);
 	}
 
 	// Bind callback to the new CO
-	if (InCustomizableObject &&
-		!InCustomizableObject->PostCompileDelegate.IsBoundToObject(GetPrivate()))
+	if (NewCustomizableObject)
 	{
-		InCustomizableObject->PostCompileDelegate.AddUObject(GetPrivate(), &UCustomizableInstancePrivate::OnPostCompile);
+		NewCustomizableObject->PostCompileDelegate.AddUObject(this, &UCustomizableInstancePrivate::OnPostCompile);
+		NewCustomizableObject->GetPrivate()->Status.GetOnStateChangedDelegate().AddUObject(this, &UCustomizableInstancePrivate::OnObjectStatusChanged);
 	}
 }
 
@@ -393,11 +408,6 @@ bool UCustomizableObjectInstance::IsEditorOnly() const
 
 void UCustomizableObjectInstance::BeginDestroy()
 {
-#if WITH_EDITOR
-	// Unbind PostCompileDelegate
-	BindPostCompileDelegate(nullptr);
-#endif
-
 	BeginDestroyDelegate.Broadcast(this);
 	BeginDestroyNativeDelegate.Broadcast(this);
 
@@ -406,6 +416,11 @@ void UCustomizableObjectInstance::BeginDestroy()
 
 	if (PrivateData)
 	{
+#if WITH_EDITOR
+		// Unbind Object delegates
+		PrivateData->BindObjectDelegates(GetCustomizableObject(), nullptr);
+#endif
+
 		if (PrivateData->StreamingHandle.IsValid() && PrivateData->StreamingHandle->IsActive())
 		{
 			PrivateData->StreamingHandle->CancelHandle();
@@ -559,7 +574,7 @@ void UCustomizableObjectInstance::PostLoad()
 	Super::PostLoad();
 
 #if WITH_EDITOR
-	BindPostCompileDelegate(GetCustomizableObject());
+	PrivateData->BindObjectDelegates(nullptr, GetCustomizableObject());
 #endif
 
 	// Skip the cost of ReloadParameters in the cook commandlet; it will be reloaded during PreSave. For cooked runtime
@@ -684,7 +699,7 @@ void UCustomizableObjectInstance::SetObject(UCustomizableObject* InObject)
 {
 #if WITH_EDITOR
 	// Bind a lambda to the PostCompileDelegate and unbind from the previous object if any.
-	BindPostCompileDelegate(InObject);
+	PrivateData->BindObjectDelegates(GetCustomizableObject(), InObject);
 #endif
 
 	Descriptor.SetCustomizableObject(*InObject);
