[Mutable] Fix import data and render data out of sync.

#rnx

diff --git a/Engine/Plugins/Experimental/Mutable/Source/CustomizableObject/Private/MuCO/CustomizableObjectInstance.cpp b/Engine/Plugins/Experimental/Mutable/Source/CustomizableObject/Private/MuCO/CustomizableObjectInstance.cpp
--- a/Engine/Plugins/Experimental/Mutable/Source/CustomizableObject/Private/MuCO/CustomizableObjectInstance.cpp
+++ b/Engine/Plugins/Experimental/Mutable/Source/CustomizableObject/Private/MuCO/CustomizableObjectInstance.cpp
@@ -6850,17 +6850,17 @@ void UCustomizableInstancePrivate::RegenerateImportedModels()
 							ImportedSection.MaterialIndex = 0;
 						}
 					}
+				}
 
-					ImportedSection.MaxBoneInfluences = RenderSection.MaxBoneInfluences;
-					ImportedSection.OriginalDataSectionIndex = OriginalIndex++;
+				ImportedSection.MaxBoneInfluences = RenderSection.MaxBoneInfluences;
+				ImportedSection.OriginalDataSectionIndex = OriginalIndex++;
 
-					FSkelMeshSourceSectionUserData& SectionUserData = ImportedModel->LODModels[LODIndex].UserSectionsData.FindOrAdd(ImportedSection.OriginalDataSectionIndex);
+				FSkelMeshSourceSectionUserData& SectionUserData = ImportedModel->LODModels[LODIndex].UserSectionsData.FindOrAdd(ImportedSection.OriginalDataSectionIndex);
 
-					SectionUserData.CorrespondClothAssetIndex = RenderSection.CorrespondClothAssetIndex;
-					SectionUserData.ClothingData.AssetGuid = RenderSection.ClothingData.AssetGuid;
-					SectionUserData.ClothingData.AssetLodIndex = RenderSection.ClothingData.AssetLodIndex;
-					SectionUserData.bCastShadow = RenderSection.bCastShadow;
-				}
+				SectionUserData.CorrespondClothAssetIndex = RenderSection.CorrespondClothAssetIndex;
+				SectionUserData.ClothingData.AssetGuid = RenderSection.ClothingData.AssetGuid;
+				SectionUserData.ClothingData.AssetLodIndex = RenderSection.ClothingData.AssetLodIndex;
+				SectionUserData.bCastShadow = RenderSection.bCastShadow;
 
 				ImportedModel->LODModels[LODIndex].SyncronizeUserSectionsDataArray();
 
