[Mutable] Fix Group Parameters not working as expected.

Parameters had duplicate ids, showing incorrect selected options.

#reveiw [REDACTED_HANDLE]
#rnx

diff --git a/Engine/Plugins/Experimental/Mutable/Source/CustomizableObject/Classes/MuCO/CustomizableObject.h b/Engine/Plugins/Experimental/Mutable/Source/CustomizableObject/Classes/MuCO/CustomizableObject.h
--- a/Engine/Plugins/Experimental/Mutable/Source/CustomizableObject/Classes/MuCO/CustomizableObject.h
+++ b/Engine/Plugins/Experimental/Mutable/Source/CustomizableObject/Classes/MuCO/CustomizableObject.h
@@ -1370,7 +1370,7 @@ private:
 	// This is a manual version number for the binary blobs in this asset.
 	// Increasing it invalidates all the previously compiled models.
 	// Warning: If while merging code both versions have changed, take the highest+1.
-	static const int32 CurrentSupportedVersion = 417;
+	static const int32 CurrentSupportedVersion = 418;
 
 	// Compile the object for a specific platform - Compile for Cook Customizable Object
 	void CompileForTargetPlatform(const ITargetPlatform* TargetPlatform);
diff --git a/Engine/Plugins/Experimental/Mutable/Source/MutableTools/Private/MuT/CodeGenerator_FirstPass.cpp b/Engine/Plugins/Experimental/Mutable/Source/MutableTools/Private/MuT/CodeGenerator_FirstPass.cpp
--- a/Engine/Plugins/Experimental/Mutable/Source/MutableTools/Private/MuT/CodeGenerator_FirstPass.cpp
+++ b/Engine/Plugins/Experimental/Mutable/Source/MutableTools/Private/MuT/CodeGenerator_FirstPass.cpp
@@ -614,7 +614,7 @@ namespace mu
                         op->type = OP_TYPE::BO_PARAMETER;
 
                         op->parameter.m_name = pChildNode->GetName();
-                   		const TCHAR* CStr = ToCStr(Private->Uid);
+                   		const TCHAR* CStr = ToCStr(pChildNode->GetUid());
                    		op->parameter.m_uid.ImportTextItem(CStr, 0, nullptr, nullptr);
                         op->parameter.m_type = PARAMETER_TYPE::T_BOOL;
                         op->parameter.m_defaultValue.Set<ParamBoolType>(false);
