[mutable] Remove usage of RTTI.
Some coding standards cleanup.
#rnx

diff --git a/Engine/Plugins/Experimental/Mutable/Source/CustomizableObjectEditor/Classes/MuCOE/Nodes/CustomizableObjectNode.h b/Engine/Plugins/Experimental/Mutable/Source/CustomizableObjectEditor/Classes/MuCOE/Nodes/CustomizableObjectNode.h
--- a/Engine/Plugins/Experimental/Mutable/Source/CustomizableObjectEditor/Classes/MuCOE/Nodes/CustomizableObjectNode.h
+++ b/Engine/Plugins/Experimental/Mutable/Source/CustomizableObjectEditor/Classes/MuCOE/Nodes/CustomizableObjectNode.h
@@ -319,17 +319,17 @@ private:
  * 
  * Return nullptr if not found.
  */
-template<class NODE_TYPE>
-NODE_TYPE* GetCustomizableObjectExternalNode(UCustomizableObject* Object, const FGuid& NodeGuid)
+template<class FNodeType>
+FNodeType* GetCustomizableObjectExternalNode(UCustomizableObject* Object, const FGuid& NodeGuid)
 {
-	NODE_TYPE* Result = nullptr;
+	FNodeType* Result = nullptr;
 
 	if (Object && Object->Source && NodeGuid.IsValid())
 	{
-		TArray<NODE_TYPE*> GroupNodes;
-		Object->Source->GetNodesOfClass<NODE_TYPE>(GroupNodes);
+		TArray<FNodeType*> GroupNodes;
+		Object->Source->GetNodesOfClass<FNodeType>(GroupNodes);
 
-		for (NODE_TYPE* GroupNode : GroupNodes)
+		for (FNodeType* GroupNode : GroupNodes)
 		{
 			if (NodeGuid == GroupNode->NodeGuid)
 			{
diff --git a/Engine/Plugins/Experimental/Mutable/Source/CustomizableObjectEditor/CustomizableObjectEditor.Build.cs b/Engine/Plugins/Experimental/Mutable/Source/CustomizableObjectEditor/CustomizableObjectEditor.Build.cs
--- a/Engine/Plugins/Experimental/Mutable/Source/CustomizableObjectEditor/CustomizableObjectEditor.Build.cs
+++ b/Engine/Plugins/Experimental/Mutable/Source/CustomizableObjectEditor/CustomizableObjectEditor.Build.cs
@@ -12,15 +12,6 @@ public class CustomizableObjectEditor : ModuleRules
 
 		DefaultBuildSettings = BuildSettingsVersion.V2;
 
-		// Strangely enough:
-		// - this has to be enabled for the windows editor build, or it raises an exception at runtime (compiling Mutable).
-		// - this cannot be enabled in Linux, or it doesn't compile.
-		// No class in this module requires RTTI, but this module uses RTTI for classes in the MutableTools module.
-		if (Target.Platform == UnrealTargetPlatform.Win64)
-		{
-			bUseRTTI = true;
-		}
-
 		PrivateIncludePathModuleNames.AddRange(
 			new string[] { 
 				"AssetRegistry", 
diff --git a/Engine/Plugins/Experimental/Mutable/Source/MutableRuntime/Internal/MuR/SerialisationPrivate.h b/Engine/Plugins/Experimental/Mutable/Source/MutableRuntime/Internal/MuR/SerialisationPrivate.h
--- a/Engine/Plugins/Experimental/Mutable/Source/MutableRuntime/Internal/MuR/SerialisationPrivate.h
+++ b/Engine/Plugins/Experimental/Mutable/Source/MutableRuntime/Internal/MuR/SerialisationPrivate.h
@@ -259,7 +259,7 @@ namespace mu
 		template<>															\
         void DLLEXPORT operator<< <T>( OutputArchive& arch, const T& t )	\
 		{																	\
-            auto v = (uint32)t;                                         	\
+            uint32 v = (uint32)t;                                         	\
             arch.GetPrivate()->m_pStream->Write( &v, sizeof(uint32) );  	\
 		}																	\
 																			\
diff --git a/Engine/Plugins/Experimental/Mutable/Source/MutableRuntime/Internal/MuR/SystemPrivate.h b/Engine/Plugins/Experimental/Mutable/Source/MutableRuntime/Internal/MuR/SystemPrivate.h
--- a/Engine/Plugins/Experimental/Mutable/Source/MutableRuntime/Internal/MuR/SystemPrivate.h
+++ b/Engine/Plugins/Experimental/Mutable/Source/MutableRuntime/Internal/MuR/SystemPrivate.h
@@ -55,27 +55,27 @@ namespace mu
 	{
 	public:
 		//! Set or add a value to the index
-		void SetFromModelRangeIndex(uint16 rangeIndex, int rangeValue)
+		void SetFromModelRangeIndex(uint16 RangeIndex, int32 RangeValue)
 		{
-			auto Index = IndexOfByPredicate([=](const ElementType& v) { return v.Key >= rangeIndex; });
-			if (Index != INDEX_NONE && (*this)[Index].Key == rangeIndex)
+			SizeType Index = IndexOfByPredicate([=](const ElementType& v) { return v.Key >= RangeIndex; });
+			if (Index != INDEX_NONE && (*this)[Index].Key == RangeIndex)
 			{
 				// Update
-				(*this)[Index].Value = rangeValue;
+				(*this)[Index].Value = RangeValue;
 			}
 			else
 			{
 				// Add new
-				Push(ElementType(rangeIndex, rangeValue));
+				Push(ElementType(RangeIndex, RangeValue));
 			}
 		}
 
 		//! Get the value of the index from the range index in the model.
-		int GetFromModelRangeIndex(int modelRangeIndex) const
+		int32 GetFromModelRangeIndex(int32 ModelRangeIndex) const
 		{
-			for (const auto& e : *this)
+			for (const TPair<int32, int32>& e : *this)
 			{
-				if (e.Key == modelRangeIndex)
+				if (e.Key == ModelRangeIndex)
 				{
 					return e.Value;
 				}
@@ -506,7 +506,7 @@ namespace mu
 		}
 
 		//!
-		inline uint32_t GetRangeIndexIndex(const ExecutionIndex& rangeIndex)
+		inline uint32 GetRangeIndexIndex(const ExecutionIndex& rangeIndex)
 		{
 			if (rangeIndex.IsEmpty())
 			{
@@ -520,7 +520,7 @@ namespace mu
 			}
 
 			// Look for or add the new element.
-			auto ElemIndex = m_usedRangeIndices.Find(rangeIndex);
+			int32 ElemIndex = m_usedRangeIndices.Find(rangeIndex);
 			if (ElemIndex != INDEX_NONE)
 			{
 				return ElemIndex;
diff --git a/Engine/Plugins/Experimental/Mutable/Source/MutableRuntime/Private/MuR/CodeVisitor.cpp b/Engine/Plugins/Experimental/Mutable/Source/MutableRuntime/Private/MuR/CodeVisitor.cpp
--- a/Engine/Plugins/Experimental/Mutable/Source/MutableRuntime/Private/MuR/CodeVisitor.cpp
+++ b/Engine/Plugins/Experimental/Mutable/Source/MutableRuntime/Private/MuR/CodeVisitor.cpp
@@ -1,6 +1,5 @@
 // Copyright Epic Games, Inc. All Rights Reserved.
 
-
 #include "MuR/CodeVisitor.h"
 
 #include "HAL/PlatformCrt.h"
@@ -9,85 +8,83 @@
 namespace mu
 {
 
-
 	//-------------------------------------------------------------------------------------------------
 	//-------------------------------------------------------------------------------------------------
 	//-------------------------------------------------------------------------------------------------
-	void SubtreeParametersVisitor::Run(OP::ADDRESS root, FProgram& program)
+	void SubtreeParametersVisitor::Run(OP::ADDRESS Root, const FProgram& Program)
 	{
 		// Cached?
-		auto it = m_resultCache.Find(root);
+		TArray<int32>* it = ResultCache.Find(Root);
 		if (it)
 		{
-			m_params = *it;
+			RelevantParams = *it;
 			return;
 		}
 
 		// Not cached
+		MUTABLE_CPUPROFILER_SCOPE(SubtreeParametersVisitor);
+
+		Visited.SetNum(Program.m_opAddress.Num());
+		if (Program.m_opAddress.Num())
 		{
-			MUTABLE_CPUPROFILER_SCOPE(SubtreeParametersVisitor);
+			FMemory::Memzero(&Visited[0], Visited.Num());
+		}
 
-			m_visited.SetNum(program.m_opAddress.Num());
-			if (program.m_opAddress.Num())
-			{
-				FMemory::Memzero(&m_visited[0], m_visited.Num());
-			}
+		CurrentParams.SetNum(Program.m_parameters.Num());
+		if (CurrentParams.Num())
+		{
+			FMemory::Memzero(CurrentParams.GetData(), CurrentParams.GetAllocatedSize());
+		}
 
-			m_currentParams.SetNum(program.m_parameters.Num());
-			if (m_currentParams.Num())
-			{
-				FMemory::Memzero(&m_currentParams[0], sizeof(int) * m_currentParams.Num());
-			}
+		Pending.Reserve(Program.m_opAddress.Num() / 4);
+		Pending.Add(Root);
 
-			m_pending.Reserve(program.m_opAddress.Num() / 4);
-			m_pending.Add(root);
+		while (Pending.Num())
+		{
+			OP::ADDRESS at = Pending.Pop();
 
-			while (m_pending.Num())
+			if (!Visited[at])
 			{
-				OP::ADDRESS at = m_pending.Pop();
+				Visited[at] = true;
 
-				if (!m_visited[at])
+				switch (Program.GetOpType(at))
 				{
-					m_visited[at] = true;
+				case OP_TYPE::NU_PARAMETER:
+				case OP_TYPE::SC_PARAMETER:
+				case OP_TYPE::BO_PARAMETER:
+				case OP_TYPE::CO_PARAMETER:
+				case OP_TYPE::PR_PARAMETER:
+				case OP_TYPE::IM_PARAMETER:
+				case OP_TYPE::ST_PARAMETER:
+					CurrentParams[Program.GetOpArgs<OP::ParameterArgs>(at).variable]++;
+					break;
+
+				default:
+					break;
+				}
 
-					switch (program.GetOpType(at))
+				ForEachReference(Program, at, [&](OP::ADDRESS ref)
 					{
-					case OP_TYPE::NU_PARAMETER:
-					case OP_TYPE::SC_PARAMETER:
-					case OP_TYPE::BO_PARAMETER:
-					case OP_TYPE::CO_PARAMETER:
-					case OP_TYPE::PR_PARAMETER:
-					case OP_TYPE::IM_PARAMETER:
-						m_currentParams[program.GetOpArgs<OP::ParameterArgs>(at).variable]++;
-						break;
-
-					default:
-						break;
-					}
-
-					ForEachReference(program, at, [&](OP::ADDRESS ref)
+						if (ref)
 						{
-							if (ref)
-							{
-								m_pending.Add(ref);
-							}
-						});
-				}
+							Pending.Add(ref);
+						}
+					});
 			}
+		}
 
 
-			// Build result
-			m_params.Empty();
-			for (size_t i = 0; i < m_currentParams.Num(); ++i)
+		// Build result
+		RelevantParams.Empty();
+		for (int32 i = 0; i < CurrentParams.Num(); ++i)
+		{
+			if (CurrentParams[i])
 			{
-				if (m_currentParams[i])
-				{
-					m_params.Add(int(i));
-				}
+				RelevantParams.Add(i);
 			}
-
-			m_resultCache.Add(root, m_params);
 		}
+
+		ResultCache.Add(Root, RelevantParams);
 	}
 
 }
\ No newline at end of file
diff --git a/Engine/Plugins/Experimental/Mutable/Source/MutableRuntime/Private/MuR/CodeVisitor.h b/Engine/Plugins/Experimental/Mutable/Source/MutableRuntime/Private/MuR/CodeVisitor.h
--- a/Engine/Plugins/Experimental/Mutable/Source/MutableRuntime/Private/MuR/CodeVisitor.h
+++ b/Engine/Plugins/Experimental/Mutable/Source/MutableRuntime/Private/MuR/CodeVisitor.h
@@ -74,7 +74,7 @@ namespace mu
         //! For manual recursion that changes the state for a specific path.
         void RecurseWithState(OP::ADDRESS at, const STATE& newState)
         {
-            auto it = m_states.Find(newState);
+			int32 it = m_states.Find(newState);
             if (it==INDEX_NONE)
             {
                 m_states.Add(newState);
@@ -93,7 +93,7 @@ namespace mu
         //! Can be called from visit to set the state to visit all children ops
         void SetCurrentState(const STATE& newState)
         {
-            auto it = m_states.Find(newState);
+            int32 it = m_states.Find(newState);
             if (it==INDEX_NONE)
             {
                 m_states.Add(newState);
@@ -114,7 +114,7 @@ namespace mu
         void FullTraverse( FProgram& program )
         {
             // Visit all the state roots
-            for ( std::size_t p=0; p<program.m_states.Num(); ++p )
+            for ( int32 p=0; p<program.m_states.Num(); ++p )
             {
                 m_pending.Add(PENDING(program.m_states[p].m_root,0) );
                 Recurse( program );
@@ -245,18 +245,16 @@ namespace mu
         }
 
         //! For manual recursion that changes the state for a specific path.
-        void RecurseWithState(OP::ADDRESS at, const STATE& newState)
+        void RecurseWithState(OP::ADDRESS At, const STATE& NewState)
         {
-            auto it = std::find(m_states.begin(),m_states.end(),newState);
-            if (it==m_states.end())
+            int32 StateIndex = m_states.Find(NewState);
+            if (StateIndex ==INDEX_NONE)
             {
-                m_states.Add(newState);
-                it = m_states.end() - 1;
+				StateIndex = m_states.Add(NewState);
             }
-            int stateIndex = (int)(it - m_states.begin());
 
             //check(at<1000000);
-            m_pending.Add( PENDING(at,stateIndex) );
+            m_pending.Add( PENDING(At,StateIndex) );
         }
 
         //! For manual recursion that doesn't change the state for a specific path.
@@ -267,34 +265,34 @@ namespace mu
         }
 
         //! Can be called from visit to set the state to visit all children ops
-        void SetCurrentState(const STATE& newState)
+        void SetCurrentState(const STATE& NewState)
         {
-            auto it = std::find(m_states.begin(),m_states.end(),newState);
-            if (it==m_states.end())
-            {
-                m_states.Add(newState);
-                it = m_states.end() - 1;
-            }
-            m_currentState = (int)(it - m_states.begin());
+			int32 StateIndex = m_states.Find(NewState);
+			if (StateIndex == INDEX_NONE)
+			{
+				StateIndex = m_states.Add(NewState);
+			}
+
+			m_currentState = StateIndex;
         }
 
 
-        void Traverse( OP::ADDRESS root, FProgram& program )
+        void Traverse( OP::ADDRESS Root, FProgram& Program )
         {
-            m_pending.reserve( program.m_opAddress.Num() );
+            m_pending.reserve(Program.m_opAddress.Num() );
 
             // Visit the given root
-            m_pending.Add( PENDING(root,0) );
-            Recurse( program );
+            m_pending.Add( PENDING(Root,0) );
+            Recurse(Program);
         }
 
-        void FullTraverse( FProgram& program )
+        void FullTraverse( FProgram& Program)
         {
             // Visit all the state roots
-            for ( std::size_t p=0; p<program.m_states.Num(); ++p )
+            for ( int32 p=0; p<Program.m_states.Num(); ++p )
             {
-                m_pending.Add(PENDING(program.m_states[p].m_root,0) );
-                Recurse( program );
+                m_pending.Add(PENDING(Program.m_states[p].m_root,0) );
+                Recurse(Program);
             }
         }
 
@@ -353,7 +351,7 @@ namespace mu
                 {
                     if (m_skipResources && VisitorIsAddResource(program.GetOpType(at)))
                     {
-                        auto args = program.GetOpArgs<OP::InstanceAddArgs>(at);
+						OP::InstanceAddArgs args = program.GetOpArgs<OP::InstanceAddArgs>(at);
 
                         // Recurse only the base
                         OP::ADDRESS base = args.instance;
@@ -449,7 +447,7 @@ namespace mu
             case OP_TYPE::IN_CONDITIONAL:
 			case OP_TYPE::ED_CONDITIONAL:
             {
-                auto args = program.GetOpArgs<OP::ConditionalArgs>(at);
+				OP::ConditionalArgs args = program.GetOpArgs<OP::ConditionalArgs>(at);
 
                 recurse = false;
 
@@ -485,7 +483,7 @@ namespace mu
             {
                 recurse = false;
 
-				const uint8_t* data = program.GetOpArgsPointer(at);
+				const uint8* data = program.GetOpArgsPointer(at);
 				
 				OP::ADDRESS VarAddress;
 				FMemory::Memcpy( &VarAddress, data, sizeof(OP::ADDRESS));
@@ -497,20 +495,20 @@ namespace mu
 					FMemory::Memcpy( &DefAddress, data, sizeof(OP::ADDRESS));
 					data += sizeof(OP::ADDRESS);
 
-					uint32_t CaseCount;
-					FMemory::Memcpy( &CaseCount, data, sizeof(uint32_t));
-					data += sizeof(uint32_t);
+					uint32 CaseCount;
+					FMemory::Memcpy( &CaseCount, data, sizeof(uint32));
+					data += sizeof(uint32);
 
                     PARENT::RecurseWithCurrentState( VarAddress );
 
                     int var = m_pSystem->BuildInt( m_pModel, m_pParams, VarAddress );
 
 					OP::ADDRESS valueAt = DefAddress;
-					for (uint32_t C = 0; C < CaseCount; ++C)
+					for (uint32 C = 0; C < CaseCount; ++C)
 					{
-						int32_t Condition;
-						FMemory::Memcpy( &Condition, data, sizeof(int32_t));		
-						data += sizeof(int32_t);
+						int32 Condition;
+						FMemory::Memcpy( &Condition, data, sizeof(int32));		
+						data += sizeof(int32);
 
 						OP::ADDRESS At;
 						FMemory::Memcpy( &At, data, sizeof(OP::ADDRESS));
@@ -532,12 +530,12 @@ namespace mu
 
             case OP_TYPE::IN_ADDLOD:
             {
-                auto args = program.GetOpArgs<OP::InstanceAddLODArgs>(at);
+				OP::InstanceAddLODArgs args = program.GetOpArgs<OP::InstanceAddLODArgs>(at);
 
                 recurse = false;
 
                 STATE newState = PARENT::GetCurrentState();
-                for (int t=0;t<MUTABLE_OP_MAX_ADD_COUNT;++t)
+                for (int32 t=0;t<MUTABLE_OP_MAX_ADD_COUNT;++t)
                 {
                     OP::ADDRESS lodAt = args.lod[t];
                     if (lodAt)
@@ -555,7 +553,7 @@ namespace mu
 
             case OP_TYPE::IN_ADDMESH:
             {
-                auto args = program.GetOpArgs<OP::InstanceAddArgs>(at);
+				OP::InstanceAddArgs args = program.GetOpArgs<OP::InstanceAddArgs>(at);
 
                 recurse = false;
 
@@ -575,7 +573,7 @@ namespace mu
 
             case OP_TYPE::IN_ADDIMAGE:
             {
-                auto args = program.GetOpArgs<OP::InstanceAddArgs>(at);
+				OP::InstanceAddArgs args = program.GetOpArgs<OP::InstanceAddArgs>(at);
 
                 recurse = false;
 
@@ -626,20 +624,18 @@ namespace mu
             COVERED_STATE
             >
     {
-        using PARENT=DiscreteCoveredCodeVisitorBase<
-        UniqueConstCodeVisitorIterative<COVERED_STATE>,
-        COVERED_STATE
-        >;
+        using PARENT=DiscreteCoveredCodeVisitorBase<UniqueConstCodeVisitorIterative<COVERED_STATE>, COVERED_STATE>;
+
     public:
 
         UniqueDiscreteCoveredCodeVisitor
             (
-                System::Private* pSystem,
-				const TSharedPtr<const Model>& pModel,
-                const ParametersPtrConst& pParams,
-                unsigned lodMask
+                System::Private* InSystem,
+				const TSharedPtr<const Model>& InModel,
+                const ParametersPtrConst& InParams,
+                uint32 InLodMask
             )
-            : PARENT( pSystem, pModel, pParams, lodMask )
+            : PARENT(InSystem, InModel, InParams, InLodMask )
         {
         }
 
@@ -655,19 +651,19 @@ namespace mu
     {
     public:
 
-        void Run( OP::ADDRESS root, FProgram& program );
+        void Run( OP::ADDRESS Root, const FProgram& );
 
         //! After Run, list of relevant parameters.
-		TArray<int> m_params;
+		TArray<int32> RelevantParams;
 
     private:
 
-		TArray<int> m_currentParams;
-		TArray<uint8_t> m_visited;
-		TArray<OP::ADDRESS> m_pending;
+		TArray<int32> CurrentParams;
+		TArray<uint8> Visited;
+		TArray<OP::ADDRESS> Pending;
 
         // Result cache
-        TMap< OP::ADDRESS, TArray<int> > m_resultCache;
+        TMap< OP::ADDRESS, TArray<int32> > ResultCache;
     };
 
 
diff --git a/Engine/Plugins/Experimental/Mutable/Source/MutableRuntime/Private/MuR/ConvertData.h b/Engine/Plugins/Experimental/Mutable/Source/MutableRuntime/Private/MuR/ConvertData.h
--- a/Engine/Plugins/Experimental/Mutable/Source/MutableRuntime/Private/MuR/ConvertData.h
+++ b/Engine/Plugins/Experimental/Mutable/Source/MutableRuntime/Private/MuR/ConvertData.h
@@ -950,7 +950,7 @@ namespace mu
 			{
 			case MBF_NUINT8:
 			{
-				auto pTypedSource = reinterpret_cast<const uint8*>(pSource);
+				const uint8* pTypedSource = reinterpret_cast<const uint8*>(pSource);
 				pTypedResult[channel] = pTypedSource[channel];
 				break;
 			}
@@ -1015,7 +1015,7 @@ namespace mu
 			{
 			case MBF_NUINT16:
 			{
-				auto pTypedSource = reinterpret_cast<const uint16*>(pSource);
+				const uint16* pTypedSource = reinterpret_cast<const uint16*>(pSource);
 				pTypedResult[channel] = pTypedSource[channel];
 				break;
 			}
diff --git a/Engine/Plugins/Experimental/Mutable/Source/MutableRuntime/Private/MuR/Mesh.cpp b/Engine/Plugins/Experimental/Mutable/Source/MutableRuntime/Private/MuR/Mesh.cpp
--- a/Engine/Plugins/Experimental/Mutable/Source/MutableRuntime/Private/MuR/Mesh.cpp
+++ b/Engine/Plugins/Experimental/Mutable/Source/MutableRuntime/Private/MuR/Mesh.cpp
@@ -1764,20 +1764,20 @@ namespace
         (void)out;
         (void)bufset;
 
-        size_t elemCount = bufset.m_elementCount;
+		uint32 elemCount = bufset.m_elementCount;
         out += "  Set with "
                 + FString::Printf(TEXT("%d"), bufset.m_buffers.Num())
                 + " buffers and "
                 + FString::Printf(TEXT("%d"), elemCount)
                 + " elements.\n";
 
-        for( const auto& buf : bufset.m_buffers )
+        for( const MESH_BUFFER& buf : bufset.m_buffers )
         {
-            const uint8_t* pData = buf.m_data.GetData();
+            const uint8* pData = buf.m_data.GetData();
 
             out += "    Buffer with "+ FString::Printf(TEXT("%d"), buf.m_channels.Num())
                     + " channels and "+ FString::Printf(TEXT("%d"), buf.m_elementSize)+" elementsize\n";
-            for( const auto& chan : buf.m_channels )
+            for( const MESH_BUFFER_CHANNEL& chan : buf.m_channels )
             {
                 out += "      Channel with format: "+ FString::Printf(TEXT("%d"), chan.m_format)
                         + " semantic: "+ FString::Printf(TEXT("%d"), chan.m_semantic)
@@ -1786,8 +1786,8 @@ namespace
                         + " offset: " + FString::Printf(TEXT("%d"), chan.m_offset)+"\n";
                 for( size_t e=0; e<elemCount && e<BufferElementLimit; ++e )
                 {
-                    const uint8_t* pElementData = pData+buf.m_elementSize*e;
-                    const uint8_t* pChanData = pElementData+chan.m_offset;
+                    const uint8* pElementData = pData+buf.m_elementSize*e;
+                    const uint8* pChanData = pElementData+chan.m_offset;
                     out += "        ";
                     for (int c=0; c<chan.m_componentCount; ++c)
                     {
diff --git a/Engine/Plugins/Experimental/Mutable/Source/MutableRuntime/Private/MuR/Model.cpp b/Engine/Plugins/Experimental/Mutable/Source/MutableRuntime/Private/MuR/Model.cpp
--- a/Engine/Plugins/Experimental/Mutable/Source/MutableRuntime/Private/MuR/Model.cpp
+++ b/Engine/Plugins/Experimental/Mutable/Source/MutableRuntime/Private/MuR/Model.cpp
@@ -39,17 +39,17 @@ namespace mu
     void FProgram::LogHistogram() const
     {
 #if 0
-        uint64 countPerType[(int)OP_TYPE::COUNT];
+        uint64 countPerType[(int32)OP_TYPE::COUNT];
         mutable_memset(countPerType,0,sizeof(countPerType));
 
-        for ( const auto& o: m_opAddress )
+        for ( const uint32& o: m_opAddress )
         {
-            auto type = GetOpType(o);
-            countPerType[(int)type]++;
+            OP_TYPE type = GetOpType(o);
+            countPerType[(int32)type]++;
         }
 
-		TArray< TPair<uint64,OP_TYPE> > sorted((int)OP_TYPE::COUNT);
-        for (int i=0; i<(int)OP_TYPE::COUNT; ++i)
+		TArray< TPair<uint64,OP_TYPE> > sorted((int32)OP_TYPE::COUNT);
+        for (int32 i=0; i<(int32)OP_TYPE::COUNT; ++i)
         {
             sorted[i].second = (OP_TYPE)i;
             sorted[i].first = countPerType[i];
@@ -61,10 +61,10 @@ namespace mu
         });
 
         UE_LOG(LogMutableCore,Log, TEXT("Op histogram (%llu ops):"), m_opAddress.Num());
-        for(int i=0; i<8; ++i)
+        for(int32 i=0; i<8; ++i)
         {
             float p = sorted[i].first/float(m_opAddress.Num())*100.0f;
-            UE_LOG(LogMutableCore,Log, TEXT("  %3.2f%% : %d"), p, (int)sorted[i].second );
+            UE_LOG(LogMutableCore,Log, TEXT("  %3.2f%% : %d"), p, (int32)sorted[i].second );
         }
 #endif
     }
diff --git a/Engine/Plugins/Experimental/Mutable/Source/MutableRuntime/Private/MuR/OpImageDisplace.h b/Engine/Plugins/Experimental/Mutable/Source/MutableRuntime/Private/MuR/OpImageDisplace.h
--- a/Engine/Plugins/Experimental/Mutable/Source/MutableRuntime/Private/MuR/OpImageDisplace.h
+++ b/Engine/Plugins/Experimental/Mutable/Source/MutableRuntime/Private/MuR/OpImageDisplace.h
@@ -158,7 +158,7 @@ namespace mu
 
         const uint8* pSourceData = pSource->GetData();
 
-        auto mapFormat = pMap->GetFormat();
+		EImageFormat mapFormat = pMap->GetFormat();
         bool isUncompressed = ( mapFormat == EImageFormat::IF_L_UBYTE );
 
         if ( isUncompressed )
diff --git a/Engine/Plugins/Experimental/Mutable/Source/MutableRuntime/Private/MuR/OpLayoutPack.h b/Engine/Plugins/Experimental/Mutable/Source/MutableRuntime/Private/MuR/OpLayoutPack.h
--- a/Engine/Plugins/Experimental/Mutable/Source/MutableRuntime/Private/MuR/OpLayoutPack.h
+++ b/Engine/Plugins/Experimental/Mutable/Source/MutableRuntime/Private/MuR/OpLayoutPack.h
@@ -101,7 +101,7 @@ namespace mu
         {
             if (packedFlag[b])
             {
-                auto i = scratch.sorted[b].index;
+				int32 i = scratch.sorted[b].index;
                 if ( x>=scratch.positions[i][0] &&
                      x<scratch.positions[i][0]+scratch.sorted[b].size[0] &&
                      y>=scratch.positions[i][1] &&
diff --git a/Engine/Plugins/Experimental/Mutable/Source/MutableRuntime/Private/MuR/OpMeshMerge.h b/Engine/Plugins/Experimental/Mutable/Source/MutableRuntime/Private/MuR/OpMeshMerge.h
--- a/Engine/Plugins/Experimental/Mutable/Source/MutableRuntime/Private/MuR/OpMeshMerge.h
+++ b/Engine/Plugins/Experimental/Mutable/Source/MutableRuntime/Private/MuR/OpMeshMerge.h
@@ -873,11 +873,11 @@ namespace mu
 
                 // See if we need to add additional buffers from the second mesh (like vertex colours or additional UV Channels)
                 // This is a bit ad-hoc: we only add buffers containing all new channels
-                for ( const auto& buf : pSecond->GetVertexBuffers().m_buffers )
+                for ( const MESH_BUFFER& buf : pSecond->GetVertexBuffers().m_buffers )
                 {
                     bool someChannel = false;
                     bool allNewChannels = true;
-					for (const auto& chan : buf.m_channels)
+					for (const MESH_BUFFER_CHANNEL& chan : buf.m_channels)
 					{
 						// Skip system buffers
 						if (chan.m_semantic == MBS_VERTEXINDEX
diff --git a/Engine/Plugins/Experimental/Mutable/Source/MutableRuntime/Private/MuR/OpMeshRemove.cpp b/Engine/Plugins/Experimental/Mutable/Source/MutableRuntime/Private/MuR/OpMeshRemove.cpp
--- a/Engine/Plugins/Experimental/Mutable/Source/MutableRuntime/Private/MuR/OpMeshRemove.cpp
+++ b/Engine/Plugins/Experimental/Mutable/Source/MutableRuntime/Private/MuR/OpMeshRemove.cpp
@@ -62,9 +62,9 @@ namespace mu
 
     int FindPositionInIntervals( const TArray<ID_INTERVAL>& intervals, int32 id )
     {
-        for( const auto& interval: intervals )
+        for( const ID_INTERVAL& interval: intervals )
         {
-            int deltaId = id - interval.idStart;
+            int32 deltaId = id - interval.idStart;
             if (deltaId>=0 && deltaId<interval.size)
             {
                 return interval.idPosition+deltaId;
diff --git a/Engine/Plugins/Experimental/Mutable/Source/MutableTools/MutableTools.Build.cs b/Engine/Plugins/Experimental/Mutable/Source/MutableTools/MutableTools.Build.cs
--- a/Engine/Plugins/Experimental/Mutable/Source/MutableTools/MutableTools.Build.cs
+++ b/Engine/Plugins/Experimental/Mutable/Source/MutableTools/MutableTools.Build.cs
@@ -14,10 +14,6 @@ namespace UnrealBuildTool.Rules
 			IWYUSupport = IWYUSupport.KeepAsIsForNow;
 			//bUseUnity = false;
 
-			// 
-			bUseRTTI = true;
-			bEnableExceptions = true;
-
 			PrivateIncludePaths.AddRange(new string[] {
 				System.IO.Path.Combine(GetModuleDirectory("MutableRuntime"), "Private"),
 			});
diff --git a/Engine/Plugins/Experimental/Mutable/Source/MutableTools/Private/MuT/AST.cpp b/Engine/Plugins/Experimental/Mutable/Source/MutableTools/Private/MuT/AST.cpp
--- a/Engine/Plugins/Experimental/Mutable/Source/MutableTools/Private/MuT/AST.cpp
+++ b/Engine/Plugins/Experimental/Mutable/Source/MutableTools/Private/MuT/AST.cpp
@@ -91,7 +91,7 @@ ASTChild& ASTChild::operator=( ASTChild&& rhs )
 //-------------------------------------------------------------------------------------------------
 void ASTOp::ForEachParent(const TFunctionRef<void(ASTOp*)> f) const
 {
-	for (auto& p : m_parents)
+	for (ASTOp* p : m_parents)
 	{
 		if (p)
 		{
@@ -224,7 +224,7 @@ mu::Ptr<ASTOp> ASTOp::DeepClone( const Ptr<ASTOp>& root )
     MapChildFunc m = [&](const Ptr<ASTOp>&n)
     {
         if (!n) return Ptr<ASTOp>();
-        auto it = visited.find(n);
+		std::unordered_map<Ptr<const ASTOp>, Ptr<ASTOp>>::const_iterator it = visited.find(n);
         check(it!=visited.end());
         return it->second;
     };
@@ -236,7 +236,7 @@ mu::Ptr<ASTOp> ASTOp::DeepClone( const Ptr<ASTOp>& root )
         visited[n] = cloned;
     });
 
-    auto it = visited.find(r);
+	std::unordered_map<Ptr<const ASTOp>, Ptr<ASTOp>>::const_iterator it = visited.find(r);
     check(it!=visited.end());
     return it->second;
 }
@@ -310,9 +310,9 @@ void ASTOp::LogHistogram( ASTOpList& roots )
 	//	{
 	//		if (n->GetOpType() == OP_TYPE::IM_MULTILAYER)
 	//		{
-	//			Ptr<const ASTOpImageMultiLayer> Typed = dynamic_cast<const ASTOpImageMultiLayer*>(n.get());
+	//			Ptr<const ASTOpImageMultiLayer> Typed = static_cast<const ASTOpImageMultiLayer*>(n.get());
 	//			Ptr<const ASTOp> Base = Typed->base.child();
-	//			Ptr<const ASTOpConstantResource> Constant = dynamic_cast<const ASTOpConstantResource*>(Base.get());
+	//			Ptr<const ASTOpConstantResource> Constant = static_cast<const ASTOpConstantResource*>(Base.get());
 
 	//			// Log the op
 	//			UE_LOG(LogMutableCore, Log, TEXT("Multilayer at %x:"), n.get());
@@ -340,7 +340,7 @@ void ASTOp::Traverse_TopDown_Unique( const TArray<Ptr<ASTOp>>& roots,
                                      TFunctionRef<bool(Ptr<ASTOp>&)> f )
 {
     TQueue<Ptr<ASTOp>> pending;
-	for (auto& r : roots)
+	for (const Ptr<ASTOp>& r : roots)
 	{
 		pending.Enqueue(r);
 	}
@@ -348,7 +348,7 @@ void ASTOp::Traverse_TopDown_Unique( const TArray<Ptr<ASTOp>>& roots,
     TSet<Ptr<const ASTOp>> traversed;
 
     // We record the parents of all roots as traversed
-    for ( const auto& r: roots )
+    for (const Ptr<ASTOp>& r: roots)
     {
         r->ForEachParent( [&]( const ASTOp* parent )
         {
@@ -415,7 +415,7 @@ void ASTOp::Traverse_TopDown_Unique_Imprecise( const TArray<Ptr<ASTOp>>& roots,
                                      TFunctionRef<bool(Ptr<ASTOp>&)> f )
 {
     TQueue<Ptr<ASTOp>> pending;
-	for (auto& r : roots)
+	for (const Ptr<ASTOp>& r : roots)
 	{
 		pending.Enqueue(r);
 	}
@@ -459,7 +459,7 @@ void ASTOp::Traverse_TopRandom_Unique_NonReentrant( const TArray<Ptr<ASTOp>>& ro
 
     uint32 traverseIndex = s_lastTraverseIndex++;
 
-    for (auto& r:roots)
+    for (const Ptr<ASTOp>& r:roots)
     {
         if (r && r->m_traverseIndex!=traverseIndex )
         {
@@ -467,7 +467,7 @@ void ASTOp::Traverse_TopRandom_Unique_NonReentrant( const TArray<Ptr<ASTOp>>& ro
             pending.Add( r );
         }
     }
-    for( auto& p : pending)
+    for(const Ptr<ASTOp>& p : pending)
     {
         p->m_traverseIndex = traverseIndex-1;
     }
@@ -532,13 +532,13 @@ void Visitor_TopDown_Unique_Cloning::Process()
 		if (item.Key)
         {
             // Item indicating we finished with all the children of this instruction
-			auto cop = at->Clone(Identity);
+			Ptr<ASTOp> cop = at->Clone(Identity);
 
             // Fix the references to the children
             bool childChanged = false;
             cop->ForEachChild( [&](ASTChild& ref)
             {
-                auto it = m_oldToNew.find(ref.m_child);
+				std::unordered_map<Ptr<ASTOp>, Ptr<ASTOp>>::iterator it = m_oldToNew.find(ref.m_child);
                 if ( ref && it!=m_oldToNew.end() && it->second.get()!=nullptr )
                 {
                     auto oldRef = ref.m_child;
@@ -559,21 +559,21 @@ void Visitor_TopDown_Unique_Cloning::Process()
         }
         else
         {
-            auto it = m_oldToNew.find(at);
+			std::unordered_map<Ptr<ASTOp>, Ptr<ASTOp>>::iterator it = m_oldToNew.find(at);
             if (it==m_oldToNew.end())
             {
-                auto initialAt = at;
+				Ptr<ASTOp> initialAt = at;
 
                 // Fix the references to the children, possibly adding a new instruction
                 {
-                    auto cop = at->Clone(Identity);
+					Ptr<ASTOp> cop = at->Clone(Identity);
                     bool childChanged = false;
                     cop->ForEachChild( [&](ASTChild& ref)
                     {
-                        auto ito = m_oldToNew.find(ref.m_child);
+						std::unordered_map<Ptr<ASTOp>, Ptr<ASTOp>>::iterator ito = m_oldToNew.find(ref.m_child);
                         if ( ref && ito!=m_oldToNew.end() && ito->second.get()!=nullptr )
                         {
-                            auto oldRef = ref.m_child;
+							Ptr<ASTOp> oldRef = ref.m_child;
                             ref=GetOldToNew(ref.m_child);
                             if (ref.m_child!=oldRef)
                             {
@@ -590,11 +590,11 @@ void Visitor_TopDown_Unique_Cloning::Process()
                     }
                 }
 
-                //auto test1 = at->Clone();
+                //Ptr<ASTOp> test1 = at->Clone();
                 //check(*test1==*at);
 
                 bool processChildren = true;
-                auto newAt = Visit( at, processChildren );
+				Ptr<ASTOp> newAt = Visit( at, processChildren );
                 m_oldToNew[initialAt]=newAt;
 
                 //check(*test1==*at);
@@ -663,8 +663,8 @@ void ASTOp::Traverse_BottomUp_Unique_NonReentrant
 {
     uint32 traverseIndex = s_lastTraverseIndex++;
 
-    TArray< std::pair<Ptr<ASTOp>,int> > pending;
-    for (auto& r:roots)
+    TArray< std::pair<Ptr<ASTOp>,int32> > pending;
+    for (Ptr<ASTOp>& r:roots)
     {
         if (r && r->m_traverseIndex!=traverseIndex )
         {
@@ -672,14 +672,14 @@ void ASTOp::Traverse_BottomUp_Unique_NonReentrant
             pending.Add( std::make_pair<>(r,0) );
         }
     }
-    for(auto& p : pending)
+    for(std::pair<Ptr<ASTOp>, int32>& p : pending)
     {
         p.first->m_traverseIndex = traverseIndex-1;
     }
 
     while (pending.Num())
     {
-        int phase = pending.Last().second;
+        int32 phase = pending.Last().second;
         Ptr<ASTOp> pCurrent = pending.Last().first;
         pending.Pop();
 
@@ -696,7 +696,7 @@ void ASTOp::Traverse_BottomUp_Unique_NonReentrant
                 {
                     if (c && c.m_child->m_traverseIndex!=traverseIndex )
                     {
-                        auto e = std::make_pair<>(c.m_child,0);
+						std::pair<Ptr<ASTOp>, int32> e = std::make_pair<>(c.m_child,0);
                         pending.Add( e );
                     }
                 });
@@ -725,8 +725,8 @@ void ASTOp::Traverse_BottomUp_Unique_NonReentrant
 {
     uint32 traverseIndex = s_lastTraverseIndex++;
 
-    TArray< std::pair<Ptr<ASTOp>,int> > pending;
-    for (auto& r:roots)
+    TArray< std::pair<Ptr<ASTOp>,int32> > pending;
+    for (Ptr<ASTOp>& r:roots)
     {
         if (r && r->m_traverseIndex!=traverseIndex)
         {
@@ -734,14 +734,14 @@ void ASTOp::Traverse_BottomUp_Unique_NonReentrant
             pending.Add( std::make_pair<>(r,0) );
         }
     }
-    for(auto& p : pending)
+    for(std::pair<Ptr<ASTOp>, int32>& p : pending)
     {
         p.first->m_traverseIndex = traverseIndex-1;
     }
 
     while (pending.Num())
     {
-        int phase = pending.Last().second;
+        int32 phase = pending.Last().second;
         Ptr<ASTOp> pCurrent = pending.Last().first;
         pending.Pop();
 
@@ -785,12 +785,12 @@ void ASTOp::Traverse_BottomUp_Unique
 )
 {
     TSet<Ptr<ASTOp>> Traversed;
-    TArray< std::pair<Ptr<ASTOp>,int> > Pending;
-    for (auto& r:roots)
+    TArray< std::pair<Ptr<ASTOp>,int32> > Pending;
+    for (Ptr<ASTOp>& r:roots)
     {
         if (r)
         {
-			auto It = Pending.FindByPredicate([&](const std::pair<Ptr<ASTOp>, int>& p)
+			std::pair<Ptr<ASTOp>, int32>* It = Pending.FindByPredicate([&](const std::pair<Ptr<ASTOp>, int>& p)
 				{
 					return r == p.first;
 				});
@@ -877,9 +877,9 @@ void ASTOp::Replace( const Ptr<ASTOp>& node, const Ptr<ASTOp>& other )
         return;
     }
 
-    auto parentsCopy = node->m_parents;
+	TArray<ASTOp*, TInlineAllocator<4> > parentsCopy = node->m_parents;
 
-    for(auto& p:parentsCopy)
+    for(ASTOp* p:parentsCopy)
     {
         if(p)
         {
@@ -976,7 +976,7 @@ void ASTOp::LinkRange(FProgram& program,
 				||
 				range.rangeSize->GetOpType() == OP_TYPE::NU_PARAMETER )
 			{
-				const ASTOpParameter* ParamOp = dynamic_cast<const ASTOpParameter*>(range.rangeSize.child().get());
+				const ASTOpParameter* ParamOp = static_cast<const ASTOpParameter*>(range.rangeSize.child().get());
 
 				EstimatedSizeParameter = ParamOp->LinkedParameterIndex;
 			}
@@ -1064,8 +1064,9 @@ void ASTOpFixed::Link( FProgram& program, FLinkerOptions* )
 
 bool ASTOpFixed::IsEqual(const ASTOp& otherUntyped) const
 {
-    if ( auto other = dynamic_cast<const ASTOpFixed*>(&otherUntyped) )
+    if (otherUntyped.GetOpType()==GetOpType())
     {
+		const ASTOpFixed* other = static_cast<const ASTOpFixed*>(&otherUntyped);
         return op==other->op && children==other->children;
     }
     return false;
@@ -1089,7 +1090,7 @@ mu::Ptr<ASTOp> ASTOpFixed::Clone( MapChildFuncRef mapChild ) const
 uint64 ASTOpFixed::Hash() const
 {
 	uint64 res = std::hash<uint64>()(uint64(op.type));
-    for (const auto& c: children)
+    for (const ASTChild& c: children)
     {
         hash_combine( res, c.child().get() );
     }
@@ -1285,7 +1286,7 @@ ASTOp::FBoolEvalResult ASTOpFixed::EvaluateBool( ASTOpList& facts, FEvaluateBool
     else
     {
         // Is this in the cache?
-        auto it = cache->find(this);
+		FEvaluateBoolCache::iterator it = cache->find(this);
         if (it!=cache->end())
         {
             return it->second;
@@ -1314,8 +1315,8 @@ ASTOp::FBoolEvalResult ASTOpFixed::EvaluateBool( ASTOpList& facts, FEvaluateBool
 
     case OP_TYPE::BO_EQUAL_INT_CONST:
     {
-        int intValue = op.args.BoolEqualScalarConst.constant;
-        const auto& intExp = children[op.args.BoolEqualScalarConst.value].child();
+        int32 intValue = op.args.BoolEqualScalarConst.constant;
+        const Ptr<ASTOp>& intExp = children[op.args.BoolEqualScalarConst.value].child();
         bool intUnknown = true;
         int intResult = intExp->EvaluateInt( facts, intUnknown );
         if (intUnknown)
@@ -1335,8 +1336,8 @@ ASTOp::FBoolEvalResult ASTOpFixed::EvaluateBool( ASTOpList& facts, FEvaluateBool
 
     case OP_TYPE::BO_AND:
     {
-        const auto& a = children[op.args.BoolBinary.a].child();
-        const auto& b = children[op.args.BoolBinary.b].child();
+        const Ptr<ASTOp>& a = children[op.args.BoolBinary.a].child();
+        const Ptr<ASTOp>& b = children[op.args.BoolBinary.b].child();
         FBoolEvalResult resultA = BET_UNKNOWN;
         FBoolEvalResult resultB = BET_UNKNOWN;
         for ( size_t f=0; f<facts.Num(); ++f )
@@ -1377,8 +1378,8 @@ ASTOp::FBoolEvalResult ASTOpFixed::EvaluateBool( ASTOpList& facts, FEvaluateBool
 
     case OP_TYPE::BO_OR:
     {
-        const auto& a = children[op.args.BoolBinary.a].child();
-        const auto& b = children[op.args.BoolBinary.b].child();
+        const Ptr<ASTOp>& a = children[op.args.BoolBinary.a].child();
+        const Ptr<ASTOp>& b = children[op.args.BoolBinary.b].child();
         FBoolEvalResult resultA = BET_UNKNOWN;
         FBoolEvalResult resultB = BET_UNKNOWN;
         for ( size_t f=0; f<facts.Num(); ++f )
diff --git a/Engine/Plugins/Experimental/Mutable/Source/MutableTools/Private/MuT/AST.h b/Engine/Plugins/Experimental/Mutable/Source/MutableTools/Private/MuT/AST.h
--- a/Engine/Plugins/Experimental/Mutable/Source/MutableTools/Private/MuT/AST.h
+++ b/Engine/Plugins/Experimental/Mutable/Source/MutableTools/Private/MuT/AST.h
@@ -318,7 +318,7 @@ namespace mu
 	{
 	public:
 
-		Ptr<ASTOp> Apply(const ASTOp* root);
+		Ptr<ASTOp> Apply(const class ASTOpImagePixelFormat* root);
 
 	protected:
 
@@ -338,7 +338,7 @@ namespace mu
 	public:
 
 		// \TODO This is recursive and may cause stack overflows in big models.
-		Ptr<ASTOp> Apply(const ASTOp* Root);
+		Ptr<ASTOp> Apply(const class ASTOpMeshFormat* Root);
 
 	protected:
 
@@ -491,6 +491,10 @@ namespace mu
 		using MapChildFuncRef = TFunctionRef<Ptr<ASTOp>(const Ptr<ASTOp>&)>;
 		virtual Ptr<ASTOp> Clone(MapChildFuncRef mapChild) const = 0;
 
+		//
+		virtual bool IsConditional() const { return false; }
+		virtual bool IsSwitch() const { return false; }
+
     protected:
 
         void RemoveChildren();
@@ -798,8 +802,7 @@ namespace mu
     {
 		ASTOp::MapChildFunc Identity = [](const Ptr<ASTOp>& o) {return o; };
 		Ptr<ASTOp> c = s->Clone(Identity);
-		Ptr<DERIVED> t = dynamic_cast<DERIVED*>(c.get());
-        check(t);
+		Ptr<DERIVED> t = static_cast<DERIVED*>(c.get());
         return t;
     }
 
@@ -808,9 +811,8 @@ namespace mu
     {
 		ASTOp::MapChildFunc Identity = [](const Ptr<ASTOp>& o) {return o; };
 		Ptr<ASTOp> c = s->Clone(Identity);
-        Ptr<DERIVED> t = dynamic_cast<DERIVED*>(c.get());
-        check(t);
-        return t;
+        Ptr<DERIVED> t = static_cast<DERIVED*>(c.get());
+		return t;
     }
 
     //---------------------------------------------------------------------------------------------
diff --git a/Engine/Plugins/Experimental/Mutable/Source/MutableTools/Private/MuT/ASTOpAddExtensionData.cpp b/Engine/Plugins/Experimental/Mutable/Source/MutableTools/Private/MuT/ASTOpAddExtensionData.cpp
--- a/Engine/Plugins/Experimental/Mutable/Source/MutableTools/Private/MuT/ASTOpAddExtensionData.cpp
+++ b/Engine/Plugins/Experimental/Mutable/Source/MutableTools/Private/MuT/ASTOpAddExtensionData.cpp
@@ -7,7 +7,6 @@
 #include "MuR/RefCounted.h"
 #include "MuT/StreamsPrivate.h"
 
-#include <utility>
 
 namespace mu
 {
@@ -31,8 +30,9 @@ ASTOpAddExtensionData::~ASTOpAddExtensionData()
 
 bool ASTOpAddExtensionData::IsEqual(const ASTOp& OtherUntyped) const
 {
-	if (const ASTOpAddExtensionData* Other = dynamic_cast<const ASTOpAddExtensionData*>(&OtherUntyped))
+	if (OtherUntyped.GetOpType() == GetOpType())
 	{
+		const ASTOpAddExtensionData* Other = static_cast<const ASTOpAddExtensionData*>(&OtherUntyped);
 		return Instance == Other->Instance
 			&& ExtensionData == Other->ExtensionData
 			&& ExtensionDataName == Other->ExtensionDataName;
diff --git a/Engine/Plugins/Experimental/Mutable/Source/MutableTools/Private/MuT/ASTOpAddExtensionData.h b/Engine/Plugins/Experimental/Mutable/Source/MutableTools/Private/MuT/ASTOpAddExtensionData.h
--- a/Engine/Plugins/Experimental/Mutable/Source/MutableTools/Private/MuT/ASTOpAddExtensionData.h
+++ b/Engine/Plugins/Experimental/Mutable/Source/MutableTools/Private/MuT/ASTOpAddExtensionData.h
@@ -17,7 +17,7 @@ struct FProgram;
 //---------------------------------------------------------------------------------------------
 //! Adds a named ExtensionData to an Instance
 //---------------------------------------------------------------------------------------------
-class ASTOpAddExtensionData : public ASTOp
+class ASTOpAddExtensionData final : public ASTOp
 {
 public:
 
diff --git a/Engine/Plugins/Experimental/Mutable/Source/MutableTools/Private/MuT/ASTOpAddLOD.cpp b/Engine/Plugins/Experimental/Mutable/Source/MutableTools/Private/MuT/ASTOpAddLOD.cpp
--- a/Engine/Plugins/Experimental/Mutable/Source/MutableTools/Private/MuT/ASTOpAddLOD.cpp
+++ b/Engine/Plugins/Experimental/Mutable/Source/MutableTools/Private/MuT/ASTOpAddLOD.cpp
@@ -7,9 +7,6 @@
 #include "MuR/RefCounted.h"
 #include "MuT/StreamsPrivate.h"
 
-#include <memory>
-#include <utility>
-
 
 namespace mu
 {
@@ -32,8 +29,9 @@ namespace mu
 
 	bool ASTOpAddLOD::IsEqual(const ASTOp& otherUntyped) const
 	{
-		if (auto other = dynamic_cast<const ASTOpAddLOD*>(&otherUntyped))
+		if (otherUntyped.GetOpType() == GetOpType())
 		{
+			const ASTOpAddLOD* other = static_cast<const ASTOpAddLOD*>(&otherUntyped);
 			return lods == other->lods;
 		}
 		return false;
diff --git a/Engine/Plugins/Experimental/Mutable/Source/MutableTools/Private/MuT/ASTOpAddLOD.h b/Engine/Plugins/Experimental/Mutable/Source/MutableTools/Private/MuT/ASTOpAddLOD.h
--- a/Engine/Plugins/Experimental/Mutable/Source/MutableTools/Private/MuT/ASTOpAddLOD.h
+++ b/Engine/Plugins/Experimental/Mutable/Source/MutableTools/Private/MuT/ASTOpAddLOD.h
@@ -16,7 +16,7 @@ struct FProgram;
 	//---------------------------------------------------------------------------------------------
 	//! Operations to add elements to a LOD
 	//---------------------------------------------------------------------------------------------
-	class ASTOpAddLOD : public ASTOp
+	class ASTOpAddLOD final : public ASTOp
 	{
 	public:
 
diff --git a/Engine/Plugins/Experimental/Mutable/Source/MutableTools/Private/MuT/ASTOpConditional.cpp b/Engine/Plugins/Experimental/Mutable/Source/MutableTools/Private/MuT/ASTOpConditional.cpp
--- a/Engine/Plugins/Experimental/Mutable/Source/MutableTools/Private/MuT/ASTOpConditional.cpp
+++ b/Engine/Plugins/Experimental/Mutable/Source/MutableTools/Private/MuT/ASTOpConditional.cpp
@@ -39,8 +39,9 @@ namespace mu
 	//-------------------------------------------------------------------------------------------------
 	bool ASTOpConditional::IsEqual(const ASTOp& otherUntyped) const
 	{
-		if (auto other = dynamic_cast<const ASTOpConditional*>(&otherUntyped))
+		if (otherUntyped.GetOpType() == GetOpType())
 		{
+			const ASTOpConditional* other = static_cast<const ASTOpConditional*>(&otherUntyped);
 			return type == other->type && condition == other->condition &&
 				yes == other->yes && no == other->no;
 		}
@@ -299,7 +300,7 @@ namespace mu
 		// Constant condition?
 		if (condition->GetOpType() == OP_TYPE::BO_CONSTANT)
 		{
-			auto typedCondition = dynamic_cast<const ASTOpConstantBool*>(condition.child().get());
+			const ASTOpConstantBool* typedCondition = static_cast<const ASTOpConstantBool*>(condition.child().get());
 			if (typedCondition->value)
 			{
 				return yes.child();
@@ -313,7 +314,7 @@ namespace mu
 			// If the yes branch is a conditional with the same condition
 			if (yes && yes->GetOpType() == type)
 			{
-				auto typedYes = dynamic_cast<const ASTOpConditional*>(yes.child().get());
+				const ASTOpConditional* typedYes = static_cast<const ASTOpConditional*>(yes.child().get());
 				if (condition.child() == typedYes->condition.child()
 					||
 					*condition.child() == *typedYes->condition.child())
@@ -327,7 +328,7 @@ namespace mu
 			// If the no branch is a conditional with the same condition
 			else if (no && no->GetOpType() == type)
 			{
-				auto typedNo = dynamic_cast<const ASTOpConditional*>(no.child().get());
+				const ASTOpConditional* typedNo = static_cast<const ASTOpConditional*>(no.child().get());
 				if (condition.child() == typedNo->condition.child()
 					||
 					*condition.child() == *typedNo->condition.child())
diff --git a/Engine/Plugins/Experimental/Mutable/Source/MutableTools/Private/MuT/ASTOpConditional.h b/Engine/Plugins/Experimental/Mutable/Source/MutableTools/Private/MuT/ASTOpConditional.h
--- a/Engine/Plugins/Experimental/Mutable/Source/MutableTools/Private/MuT/ASTOpConditional.h
+++ b/Engine/Plugins/Experimental/Mutable/Source/MutableTools/Private/MuT/ASTOpConditional.h
@@ -16,7 +16,7 @@ struct FProgram;
 	//---------------------------------------------------------------------------------------------
 	//! Conditional operation.
 	//---------------------------------------------------------------------------------------------
-	class ASTOpConditional : public ASTOp
+	class ASTOpConditional final : public ASTOp
 	{
 	public:
 
@@ -50,6 +50,7 @@ struct FProgram;
 		Ptr<ASTOp> OptimiseSemantic(const FModelOptimizationOptions&, int32 Pass) const override;
 		bool GetNonBlackRect(FImageRect& maskUsage) const override;
 		Ptr<ImageSizeExpression> GetImageSizeExpression() const override;
+		virtual bool IsConditional() const override { return true; }
 	};
 
 
diff --git a/Engine/Plugins/Experimental/Mutable/Source/MutableTools/Private/MuT/ASTOpConstantBool.cpp b/Engine/Plugins/Experimental/Mutable/Source/MutableTools/Private/MuT/ASTOpConstantBool.cpp
--- a/Engine/Plugins/Experimental/Mutable/Source/MutableTools/Private/MuT/ASTOpConstantBool.cpp
+++ b/Engine/Plugins/Experimental/Mutable/Source/MutableTools/Private/MuT/ASTOpConstantBool.cpp
@@ -29,8 +29,9 @@ namespace mu
 
 	bool ASTOpConstantBool::IsEqual(const ASTOp& otherUntyped) const
 	{
-		if (auto other = dynamic_cast<const ASTOpConstantBool*>(&otherUntyped))
+		if (otherUntyped.GetOpType() == GetOpType())
 		{
+			const ASTOpConstantBool* other = static_cast<const ASTOpConstantBool*>(&otherUntyped);
 			return value == other->value;
 		}
 		return false;
diff --git a/Engine/Plugins/Experimental/Mutable/Source/MutableTools/Private/MuT/ASTOpConstantBool.h b/Engine/Plugins/Experimental/Mutable/Source/MutableTools/Private/MuT/ASTOpConstantBool.h
--- a/Engine/Plugins/Experimental/Mutable/Source/MutableTools/Private/MuT/ASTOpConstantBool.h
+++ b/Engine/Plugins/Experimental/Mutable/Source/MutableTools/Private/MuT/ASTOpConstantBool.h
@@ -16,7 +16,7 @@ struct FProgram;
 	//---------------------------------------------------------------------------------------------
 	//!
 	//---------------------------------------------------------------------------------------------
-	class ASTOpConstantBool : public ASTOp
+	class ASTOpConstantBool final : public ASTOp
 	{
 	public:
 
diff --git a/Engine/Plugins/Experimental/Mutable/Source/MutableTools/Private/MuT/ASTOpConstantExtensionData.cpp b/Engine/Plugins/Experimental/Mutable/Source/MutableTools/Private/MuT/ASTOpConstantExtensionData.cpp
--- a/Engine/Plugins/Experimental/Mutable/Source/MutableTools/Private/MuT/ASTOpConstantExtensionData.cpp
+++ b/Engine/Plugins/Experimental/Mutable/Source/MutableTools/Private/MuT/ASTOpConstantExtensionData.cpp
@@ -6,8 +6,6 @@
 #include "MuR/ModelPrivate.h"
 #include "MuT/StreamsPrivate.h"
 
-#include <utility>
-
 
 namespace mu
 {
@@ -16,8 +14,9 @@ namespace mu
 //-------------------------------------------------------------------------------------------------
 bool ASTOpConstantExtensionData::IsEqual(const ASTOp& OtherUntyped) const
 {
-	if (const ASTOpConstantExtensionData* Other = dynamic_cast<const ASTOpConstantExtensionData*>(&OtherUntyped))
+	if (OtherUntyped.GetOpType() == GetOpType())
 	{
+		const ASTOpConstantExtensionData* Other = static_cast<const ASTOpConstantExtensionData*>(&OtherUntyped);
 		return Other->Value == Value;
 	}
 
diff --git a/Engine/Plugins/Experimental/Mutable/Source/MutableTools/Private/MuT/ASTOpConstantExtensionData.h b/Engine/Plugins/Experimental/Mutable/Source/MutableTools/Private/MuT/ASTOpConstantExtensionData.h
--- a/Engine/Plugins/Experimental/Mutable/Source/MutableTools/Private/MuT/ASTOpConstantExtensionData.h
+++ b/Engine/Plugins/Experimental/Mutable/Source/MutableTools/Private/MuT/ASTOpConstantExtensionData.h
@@ -17,7 +17,7 @@ struct FProgram;
 //---------------------------------------------------------------------------------------------
 //! A constant ExtensionData
 //---------------------------------------------------------------------------------------------
-class ASTOpConstantExtensionData : public ASTOp
+class ASTOpConstantExtensionData final : public ASTOp
 {
 public:
 
diff --git a/Engine/Plugins/Experimental/Mutable/Source/MutableTools/Private/MuT/ASTOpConstantProjector.cpp b/Engine/Plugins/Experimental/Mutable/Source/MutableTools/Private/MuT/ASTOpConstantProjector.cpp
--- a/Engine/Plugins/Experimental/Mutable/Source/MutableTools/Private/MuT/ASTOpConstantProjector.cpp
+++ b/Engine/Plugins/Experimental/Mutable/Source/MutableTools/Private/MuT/ASTOpConstantProjector.cpp
@@ -30,8 +30,9 @@ namespace mu
 
 	bool ASTOpConstantProjector::IsEqual(const ASTOp& otherUntyped) const
 	{
-		if (auto other = dynamic_cast<const ASTOpConstantProjector*>(&otherUntyped))
+		if (otherUntyped.GetOpType() == GetOpType())
 		{
+			const ASTOpConstantProjector* other = static_cast<const ASTOpConstantProjector*>(&otherUntyped);
 			return value == other->value;
 		}
 		return false;
diff --git a/Engine/Plugins/Experimental/Mutable/Source/MutableTools/Private/MuT/ASTOpConstantProjector.h b/Engine/Plugins/Experimental/Mutable/Source/MutableTools/Private/MuT/ASTOpConstantProjector.h
--- a/Engine/Plugins/Experimental/Mutable/Source/MutableTools/Private/MuT/ASTOpConstantProjector.h
+++ b/Engine/Plugins/Experimental/Mutable/Source/MutableTools/Private/MuT/ASTOpConstantProjector.h
@@ -17,7 +17,7 @@ struct FProgram;
 	//---------------------------------------------------------------------------------------------
 	//!
 	//---------------------------------------------------------------------------------------------
-	class ASTOpConstantProjector : public ASTOp
+	class ASTOpConstantProjector final : public ASTOp
 	{
 	public:
 
diff --git a/Engine/Plugins/Experimental/Mutable/Source/MutableTools/Private/MuT/ASTOpConstantResource.cpp b/Engine/Plugins/Experimental/Mutable/Source/MutableTools/Private/MuT/ASTOpConstantResource.cpp
--- a/Engine/Plugins/Experimental/Mutable/Source/MutableTools/Private/MuT/ASTOpConstantResource.cpp
+++ b/Engine/Plugins/Experimental/Mutable/Source/MutableTools/Private/MuT/ASTOpConstantResource.cpp
@@ -16,9 +16,6 @@
 #include "MuR/Types.h"
 #include "MuT/StreamsPrivate.h"
 
-#include <memory>
-#include <utility>
-
 
 namespace mu
 {
@@ -34,8 +31,9 @@ namespace mu
 	//-------------------------------------------------------------------------------------------------
 	bool ASTOpConstantResource::IsEqual(const ASTOp& otherUntyped) const
 	{
-		if (const ASTOpConstantResource* other = dynamic_cast<const ASTOpConstantResource*>(&otherUntyped))
+		if (otherUntyped.GetOpType()==GetOpType())
 		{
+			const ASTOpConstantResource* other = static_cast<const ASTOpConstantResource*>(&otherUntyped);
 			return type == other->type && hash == other->hash &&
 				loadedValue == other->loadedValue && proxy == other->proxy;
 		}
@@ -435,7 +433,7 @@ namespace mu
 
 			case OP_TYPE::IM_CONSTANT:
 			{
-				Ptr<ResourceProxy<Image>> typedProxy = dynamic_cast<ResourceProxy<Image>*>(proxy.get());
+				Ptr<ResourceProxy<Image>> typedProxy = static_cast<ResourceProxy<Image>*>(proxy.get());
 				Ptr<const Image> r = typedProxy->Get();
 				return r;
 			}
diff --git a/Engine/Plugins/Experimental/Mutable/Source/MutableTools/Private/MuT/ASTOpConstantResource.h b/Engine/Plugins/Experimental/Mutable/Source/MutableTools/Private/MuT/ASTOpConstantResource.h
--- a/Engine/Plugins/Experimental/Mutable/Source/MutableTools/Private/MuT/ASTOpConstantResource.h
+++ b/Engine/Plugins/Experimental/Mutable/Source/MutableTools/Private/MuT/ASTOpConstantResource.h
@@ -19,7 +19,7 @@ namespace mu
 	//---------------------------------------------------------------------------------------------
 	//! A constant mesh, image, volume or layout
 	//---------------------------------------------------------------------------------------------
-	class ASTOpConstantResource : public ASTOp
+	class ASTOpConstantResource final : public ASTOp
 	{
 	private:
 
diff --git a/Engine/Plugins/Experimental/Mutable/Source/MutableTools/Private/MuT/ASTOpConstantString.cpp b/Engine/Plugins/Experimental/Mutable/Source/MutableTools/Private/MuT/ASTOpConstantString.cpp
--- a/Engine/Plugins/Experimental/Mutable/Source/MutableTools/Private/MuT/ASTOpConstantString.cpp
+++ b/Engine/Plugins/Experimental/Mutable/Source/MutableTools/Private/MuT/ASTOpConstantString.cpp
@@ -7,8 +7,6 @@
 #include "MuR/RefCounted.h"
 #include "MuR/Types.h"
 
-#include <string>
-
 
 namespace mu
 {
@@ -31,8 +29,9 @@ namespace mu
 
 	bool ASTOpConstantString::IsEqual(const ASTOp& otherUntyped) const
 	{
-		if (auto other = dynamic_cast<const ASTOpConstantString*>(&otherUntyped))
+		if (otherUntyped.GetOpType() == GetOpType())
 		{
+			const ASTOpConstantString* other = static_cast<const ASTOpConstantString*>(&otherUntyped);
 			return value == other->value;
 		}
 		return false;
diff --git a/Engine/Plugins/Experimental/Mutable/Source/MutableTools/Private/MuT/ASTOpConstantString.h b/Engine/Plugins/Experimental/Mutable/Source/MutableTools/Private/MuT/ASTOpConstantString.h
--- a/Engine/Plugins/Experimental/Mutable/Source/MutableTools/Private/MuT/ASTOpConstantString.h
+++ b/Engine/Plugins/Experimental/Mutable/Source/MutableTools/Private/MuT/ASTOpConstantString.h
@@ -16,7 +16,7 @@ struct FProgram;
 	//---------------------------------------------------------------------------------------------
 	//!
 	//---------------------------------------------------------------------------------------------
-	class ASTOpConstantString : public ASTOp
+	class ASTOpConstantString final : public ASTOp
 	{
 	public:
 		//!
diff --git a/Engine/Plugins/Experimental/Mutable/Source/MutableTools/Private/MuT/ASTOpImageCompose.cpp b/Engine/Plugins/Experimental/Mutable/Source/MutableTools/Private/MuT/ASTOpImageCompose.cpp
--- a/Engine/Plugins/Experimental/Mutable/Source/MutableTools/Private/MuT/ASTOpImageCompose.cpp
+++ b/Engine/Plugins/Experimental/Mutable/Source/MutableTools/Private/MuT/ASTOpImageCompose.cpp
@@ -44,9 +44,10 @@ ASTOpImageCompose::~ASTOpImageCompose()
 //-------------------------------------------------------------------------------------------------
 bool ASTOpImageCompose::IsEqual(const ASTOp& otherUntyped) const
 {
-    if ( const ASTOpImageCompose* other = dynamic_cast<const ASTOpImageCompose*>(&otherUntyped) )
+	if (otherUntyped.GetOpType() == GetOpType())
     {
-        return Layout ==other->Layout &&
+		const ASTOpImageCompose* other = static_cast<const ASTOpImageCompose*>(&otherUntyped);
+		return Layout == other->Layout &&
 			Base ==other->Base &&
 			BlockImage == other->BlockImage &&
 			Mask == other->Mask &&
@@ -234,7 +235,7 @@ mu::Ptr<ASTOp> ASTOpImageCompose::OptimiseSemantic(const FModelOptimizationOptio
 		&&
 		blockAt)
 	{
-		auto typedLayout = dynamic_cast<const ASTOpConstantResource*>(layoutAt.get());
+		const ASTOpConstantResource* typedLayout = static_cast<const ASTOpConstantResource*>(layoutAt.get());
 		mu::Ptr<const mu::Layout> pLayout = static_cast<const mu::Layout*>(typedLayout->GetValue().get());
 
 		// Constant single-block full layout?
diff --git a/Engine/Plugins/Experimental/Mutable/Source/MutableTools/Private/MuT/ASTOpImageCompose.h b/Engine/Plugins/Experimental/Mutable/Source/MutableTools/Private/MuT/ASTOpImageCompose.h
--- a/Engine/Plugins/Experimental/Mutable/Source/MutableTools/Private/MuT/ASTOpImageCompose.h
+++ b/Engine/Plugins/Experimental/Mutable/Source/MutableTools/Private/MuT/ASTOpImageCompose.h
@@ -14,7 +14,7 @@ namespace mu
 struct FProgram;
 template <class SCALAR> class vec4;
 
-	class ASTOpImageCompose : public ASTOp
+	class ASTOpImageCompose final : public ASTOp
 	{
 	public:
 
diff --git a/Engine/Plugins/Experimental/Mutable/Source/MutableTools/Private/MuT/ASTOpImageCrop.cpp b/Engine/Plugins/Experimental/Mutable/Source/MutableTools/Private/MuT/ASTOpImageCrop.cpp
--- a/Engine/Plugins/Experimental/Mutable/Source/MutableTools/Private/MuT/ASTOpImageCrop.cpp
+++ b/Engine/Plugins/Experimental/Mutable/Source/MutableTools/Private/MuT/ASTOpImageCrop.cpp
@@ -34,8 +34,9 @@ namespace mu
 
 	bool ASTOpImageCrop::IsEqual(const ASTOp& InOther) const
 	{
-		if (const ASTOpImageCrop* other = dynamic_cast<const ASTOpImageCrop*>(&InOther))
+		if (InOther.GetOpType()==GetOpType())
 		{
+			const ASTOpImageCrop* other = static_cast<const ASTOpImageCrop*>(&InOther);
 			return Source == other->Source &&
 				Min == other->Min &&
 				Size == other->Size;
@@ -283,7 +284,7 @@ namespace mu
 
 		case OP_TYPE::IM_PATCH:
 		{
-			const ASTOpImagePatch* typedPatch = dynamic_cast<const ASTOpImagePatch*>(at.get());
+			const ASTOpImagePatch* typedPatch = static_cast<const ASTOpImagePatch*>(at.get());
 
 			Ptr<ASTOp> rectOp = typedPatch->patch.child();
 			ASTOp::FGetImageDescContext context;
@@ -333,7 +334,7 @@ namespace mu
 		case OP_TYPE::IM_CROP:
 		{
 			// We can combine the two crops into a possibly smaller crop
-			const ASTOpImageCrop* childCrop = dynamic_cast<const ASTOpImageCrop*>(at.get());
+			const ASTOpImageCrop* childCrop = static_cast<const ASTOpImageCrop*>(at.get());
 
 			box<vec2<int16>> childCropBox;
 			childCropBox.min[0] = childCrop->Min[0];
diff --git a/Engine/Plugins/Experimental/Mutable/Source/MutableTools/Private/MuT/ASTOpImageCrop.h b/Engine/Plugins/Experimental/Mutable/Source/MutableTools/Private/MuT/ASTOpImageCrop.h
--- a/Engine/Plugins/Experimental/Mutable/Source/MutableTools/Private/MuT/ASTOpImageCrop.h
+++ b/Engine/Plugins/Experimental/Mutable/Source/MutableTools/Private/MuT/ASTOpImageCrop.h
@@ -14,7 +14,7 @@ namespace mu
 {
 struct FProgram;
 
-	class ASTOpImageCrop : public ASTOp
+	class ASTOpImageCrop final : public ASTOp
 	{
 	public:
 
diff --git a/Engine/Plugins/Experimental/Mutable/Source/MutableTools/Private/MuT/ASTOpImageLayer.cpp b/Engine/Plugins/Experimental/Mutable/Source/MutableTools/Private/MuT/ASTOpImageLayer.cpp
--- a/Engine/Plugins/Experimental/Mutable/Source/MutableTools/Private/MuT/ASTOpImageLayer.cpp
+++ b/Engine/Plugins/Experimental/Mutable/Source/MutableTools/Private/MuT/ASTOpImageLayer.cpp
@@ -40,8 +40,9 @@ namespace mu
 	//-------------------------------------------------------------------------------------------------
 	bool ASTOpImageLayer::IsEqual(const ASTOp& InOtherUntyped) const
 	{
-		if (const ASTOpImageLayer* Other = dynamic_cast<const ASTOpImageLayer*>(&InOtherUntyped))
+		if (InOtherUntyped.GetOpType() == GetOpType())
 		{
+			const ASTOpImageLayer* Other = static_cast<const ASTOpImageLayer*>(&InOtherUntyped);
 			return base == Other->base &&
 				blend == Other->blend &&
 				mask == Other->mask &&
@@ -191,7 +192,7 @@ namespace mu
 			// TODO: May some flags be supported?
 			if (Flags == 0)
 			{
-				const ASTOpFixed* BlendPlainColor = dynamic_cast<const ASTOpFixed*>(blendAt.get());
+				const ASTOpFixed* BlendPlainColor = static_cast<const ASTOpFixed*>(blendAt.get());
 
 				Ptr<ASTOpImageLayerColor> NewLayerColor = new ASTOpImageLayerColor;
 				NewLayerColor->base = baseAt;
@@ -257,7 +258,7 @@ namespace mu
 					{
 					case OP_TYPE::IM_LAYERCOLOUR:
 					{
-						const ASTOpImageLayerColor* BlendLayer = dynamic_cast<const ASTOpImageLayerColor*>(CurrentBlend.get());
+						const ASTOpImageLayerColor* BlendLayer = static_cast<const ASTOpImageLayerColor*>(CurrentBlend.get());
 						if (BlendLayer->blendTypeAlpha == EBlendType::BT_NONE)
 						{
 							CurrentBlend = BlendLayer->base.child();
@@ -268,7 +269,7 @@ namespace mu
 
 					case OP_TYPE::IM_LAYER:
 					{
-						const ASTOpImageLayer* BlendLayer = dynamic_cast<const ASTOpImageLayer*>(CurrentBlend.get());
+						const ASTOpImageLayer* BlendLayer = static_cast<const ASTOpImageLayer*>(CurrentBlend.get());
 						if (BlendLayer->blendTypeAlpha == EBlendType::BT_NONE)
 						{
 							CurrentBlend = BlendLayer->base.child();
@@ -293,8 +294,8 @@ namespace mu
 
 				case OP_TYPE::IM_DISPLACE:
 				{
-					const ASTOpFixed* MaskDisplace = dynamic_cast<const ASTOpFixed*>(CurrentMask.get());
-					const ASTOpFixed* BlendDisplace = dynamic_cast<const ASTOpFixed*>(CurrentBlend.get());
+					const ASTOpFixed* MaskDisplace = static_cast<const ASTOpFixed*>(CurrentMask.get());
+					const ASTOpFixed* BlendDisplace = static_cast<const ASTOpFixed*>(CurrentBlend.get());
 					if (MaskDisplace && BlendDisplace 
 						&&
 						MaskDisplace->children[MaskDisplace->op.args.ImageDisplace.displacementMap].child()
@@ -310,8 +311,8 @@ namespace mu
 
 				case OP_TYPE::IM_RASTERMESH:
 				{
-					const ASTOpImageRasterMesh* MaskRaster = dynamic_cast<const ASTOpImageRasterMesh*>(CurrentMask.get());
-					const ASTOpImageRasterMesh* BlendRaster = dynamic_cast<const ASTOpImageRasterMesh*>(CurrentBlend.get());
+					const ASTOpImageRasterMesh* MaskRaster = static_cast<const ASTOpImageRasterMesh*>(CurrentMask.get());
+					const ASTOpImageRasterMesh* BlendRaster = static_cast<const ASTOpImageRasterMesh*>(CurrentBlend.get());
 					if (MaskRaster && BlendRaster
 						&&
 						MaskRaster->mesh.child() == BlendRaster->mesh.child()
@@ -336,8 +337,8 @@ namespace mu
 
 				case OP_TYPE::IM_RESIZE:
 				{
-					const ASTOpFixed* MaskResize = dynamic_cast<const ASTOpFixed*>(CurrentMask.get());
-					const ASTOpFixed* BlendResize = dynamic_cast<const ASTOpFixed*>(CurrentBlend.get());
+					const ASTOpFixed* MaskResize = static_cast<const ASTOpFixed*>(CurrentMask.get());
+					const ASTOpFixed* BlendResize = static_cast<const ASTOpFixed*>(CurrentBlend.get());
 					if (MaskResize && BlendResize 
 						&&
 						MaskResize->op.args.ImageResize.size[0] == BlendResize->op.args.ImageResize.size[0]
@@ -360,7 +361,7 @@ namespace mu
 			if (CurrentMask && CurrentMask->GetOpType() == OP_TYPE::IM_SWIZZLE)
 			{
 				// End of the possible mask expression chain match should have a swizzle selecting the alpha.
-				const ASTOpImageSwizzle* MaskSwizzle = dynamic_cast<const ASTOpImageSwizzle*>(CurrentMask.get());
+				const ASTOpImageSwizzle* MaskSwizzle = static_cast<const ASTOpImageSwizzle*>(CurrentMask.get());
 				if (MaskSwizzle->SourceChannels[0] == 3
 					&&
 					!MaskSwizzle->SourceChannels[1]
@@ -387,7 +388,7 @@ namespace mu
 			// Is the base a swizzle expanding alpha from a texture?
 			if (baseAt->GetOpType() == OP_TYPE::IM_SWIZZLE)
 			{
-				const ASTOpImageSwizzle* TypedBase = dynamic_cast<const ASTOpImageSwizzle*>(baseAt.get());
+				const ASTOpImageSwizzle* TypedBase = static_cast<const ASTOpImageSwizzle*>(baseAt.get());
 				bool bAreAllAlpha = true;
 				for (int32 c=0; c<MUTABLE_OP_MAX_SWIZZLE_CHANNELS; ++c)
 				{
@@ -407,7 +408,7 @@ namespace mu
 			// Is the mask a swizzle expanding alpha from a texture?
 			if (!at && maskAt && maskAt->GetOpType() == OP_TYPE::IM_SWIZZLE)
 			{
-				const ASTOpImageSwizzle* TypedBase = dynamic_cast<const ASTOpImageSwizzle*>(maskAt.get());
+				const ASTOpImageSwizzle* TypedBase = static_cast<const ASTOpImageSwizzle*>(maskAt.get());
 				bool bAreAllAlpha = true;
 				for (int32 c = 0; c < MUTABLE_OP_MAX_SWIZZLE_CHANNELS; ++c)
 				{
@@ -427,7 +428,7 @@ namespace mu
 			// Is the blend a swizzle selecting alpha?
 			if (Pass>0 && !at && blendAt && blendAt->GetOpType() == OP_TYPE::IM_SWIZZLE && Flags==0)
 			{
-				const ASTOpImageSwizzle* TypedBlend = dynamic_cast<const ASTOpImageSwizzle*>(blendAt.get());
+				const ASTOpImageSwizzle* TypedBlend = static_cast<const ASTOpImageSwizzle*>(blendAt.get());
 				if (TypedBlend->Format==EImageFormat::IF_L_UBYTE
 					&&
 					TypedBlend->SourceChannels[0]==3)
@@ -654,7 +655,7 @@ namespace mu
 
 			case OP_TYPE::IM_SWITCH:
 			{
-				const ASTOpSwitch* BlendSwitch = dynamic_cast<const ASTOpSwitch*>(blendAt.get());
+				const ASTOpSwitch* BlendSwitch = static_cast<const ASTOpSwitch*>(blendAt.get());
 
 				// If at least a switch option is a plain colour, sink the layer into the switch
 				bool bWorthSinking = false;
@@ -674,9 +675,10 @@ namespace mu
 				if (bWorthSinking)
 				{
 					bool bMaskIsCompatibleSwitch = false;
-					const ASTOpSwitch* MaskSwitch = dynamic_cast<const ASTOpSwitch*>(maskAt.get());
+					const ASTOpSwitch* MaskSwitch = nullptr;
 					if (maskAt && maskAt->GetOpType()== OP_TYPE::IM_SWITCH)
 					{
+						MaskSwitch = static_cast<const ASTOpSwitch*>(maskAt.get());
 						bMaskIsCompatibleSwitch = MaskSwitch->IsCompatibleWith(BlendSwitch);
 					}
 
@@ -728,7 +730,7 @@ namespace mu
 
 			case OP_TYPE::IM_SWITCH:
 			{
-				const ASTOpSwitch* MaskSwitch = dynamic_cast<const ASTOpSwitch*>(maskAt.get());
+				const ASTOpSwitch* MaskSwitch = static_cast<const ASTOpSwitch*>(maskAt.get());
 
 				// If at least a switch option is a plain colour, sink the layer into the switch
 				bool bWorthSinking = false;
diff --git a/Engine/Plugins/Experimental/Mutable/Source/MutableTools/Private/MuT/ASTOpImageLayer.h b/Engine/Plugins/Experimental/Mutable/Source/MutableTools/Private/MuT/ASTOpImageLayer.h
--- a/Engine/Plugins/Experimental/Mutable/Source/MutableTools/Private/MuT/ASTOpImageLayer.h
+++ b/Engine/Plugins/Experimental/Mutable/Source/MutableTools/Private/MuT/ASTOpImageLayer.h
@@ -17,7 +17,7 @@ namespace mu
 	//---------------------------------------------------------------------------------------------
 	//!
 	//---------------------------------------------------------------------------------------------
-	class ASTOpImageLayer : public ASTOp
+	class ASTOpImageLayer final : public ASTOp
 	{
 	public:
 
diff --git a/Engine/Plugins/Experimental/Mutable/Source/MutableTools/Private/MuT/ASTOpImageLayerColor.cpp b/Engine/Plugins/Experimental/Mutable/Source/MutableTools/Private/MuT/ASTOpImageLayerColor.cpp
--- a/Engine/Plugins/Experimental/Mutable/Source/MutableTools/Private/MuT/ASTOpImageLayerColor.cpp
+++ b/Engine/Plugins/Experimental/Mutable/Source/MutableTools/Private/MuT/ASTOpImageLayerColor.cpp
@@ -35,8 +35,9 @@ namespace mu
 	//-------------------------------------------------------------------------------------------------
 	bool ASTOpImageLayerColor::IsEqual(const ASTOp& InOtherUntyped) const
 	{
-		if (const ASTOpImageLayerColor* Other = dynamic_cast<const ASTOpImageLayerColor*>(&InOtherUntyped))
+		if (InOtherUntyped.GetOpType()==GetOpType())
 		{
+			const ASTOpImageLayerColor* Other = static_cast<const ASTOpImageLayerColor*>(&InOtherUntyped);
 			return base == Other->base &&
 				color == Other->color &&
 				mask == Other->mask &&
@@ -209,7 +210,7 @@ namespace mu
 			{
 				if (color)
 				{
-					const ASTOpFixed* TypedColor = dynamic_cast<const ASTOpFixed*>(color.child().get());
+					const ASTOpFixed* TypedColor = static_cast<const ASTOpFixed*>(color.child().get());
 					const float* Value = TypedColor->op.args.ColourConstant.value;
 					ColorConst.Set(Value[0], Value[1], Value[2], Value[3]);
 				}
@@ -279,7 +280,7 @@ namespace mu
 			// Is the base a swizzle getting expanding alpha from the same texture?
 			if (base.child()->GetOpType() == OP_TYPE::IM_SWIZZLE)
 			{
-				const ASTOpImageSwizzle* TypedBase = dynamic_cast<const ASTOpImageSwizzle*>(base.child().get());
+				const ASTOpImageSwizzle* TypedBase = static_cast<const ASTOpImageSwizzle*>(base.child().get());
 				bool bAreAllSameAlpha = true;
 				Ptr<ASTOp> Source = nullptr;
 				for (int32 c = 0; c < MUTABLE_OP_MAX_SWIZZLE_CHANNELS; ++c)
@@ -358,7 +359,7 @@ namespace mu
 			// all switch branches become unique constants
 
 			// See if the blended has an identical switch, to optimise it too
-			const ASTOpSwitch* baseSwitch = dynamic_cast<const ASTOpSwitch*>(baseAt.get());
+			const ASTOpSwitch* baseSwitch = static_cast<const ASTOpSwitch*>(baseAt.get());
 
 			// Mask not supported yet
 			if (maskAt)
diff --git a/Engine/Plugins/Experimental/Mutable/Source/MutableTools/Private/MuT/ASTOpImageLayerColor.h b/Engine/Plugins/Experimental/Mutable/Source/MutableTools/Private/MuT/ASTOpImageLayerColor.h
--- a/Engine/Plugins/Experimental/Mutable/Source/MutableTools/Private/MuT/ASTOpImageLayerColor.h
+++ b/Engine/Plugins/Experimental/Mutable/Source/MutableTools/Private/MuT/ASTOpImageLayerColor.h
@@ -17,7 +17,7 @@ namespace mu
 	//---------------------------------------------------------------------------------------------
 	//!
 	//---------------------------------------------------------------------------------------------
-	class ASTOpImageLayerColor : public ASTOp
+	class ASTOpImageLayerColor final : public ASTOp
 	{
 	public:
 
diff --git a/Engine/Plugins/Experimental/Mutable/Source/MutableTools/Private/MuT/ASTOpImageMakeGrowMap.cpp b/Engine/Plugins/Experimental/Mutable/Source/MutableTools/Private/MuT/ASTOpImageMakeGrowMap.cpp
--- a/Engine/Plugins/Experimental/Mutable/Source/MutableTools/Private/MuT/ASTOpImageMakeGrowMap.cpp
+++ b/Engine/Plugins/Experimental/Mutable/Source/MutableTools/Private/MuT/ASTOpImageMakeGrowMap.cpp
@@ -18,9 +18,6 @@
 #include "MuT/ASTOpSwitch.h"
 #include "MuT/StreamsPrivate.h"
 
-#include <memory>
-#include <utility>
-
 
 namespace mu
 {
@@ -40,8 +37,9 @@ namespace mu
 
 	bool ASTOpImageMakeGrowMap::IsEqual(const ASTOp& otherUntyped) const
 	{
-		if (const ASTOpImageMakeGrowMap* other = dynamic_cast<const ASTOpImageMakeGrowMap*>(&otherUntyped))
+		if (otherUntyped.GetOpType()==GetOpType())
 		{
+			const ASTOpImageMakeGrowMap* other = static_cast<const ASTOpImageMakeGrowMap*>(&otherUntyped);
 			return Mask == other->Mask &&
 				Border == other->Border;
 		}
diff --git a/Engine/Plugins/Experimental/Mutable/Source/MutableTools/Private/MuT/ASTOpImageMakeGrowMap.h b/Engine/Plugins/Experimental/Mutable/Source/MutableTools/Private/MuT/ASTOpImageMakeGrowMap.h
--- a/Engine/Plugins/Experimental/Mutable/Source/MutableTools/Private/MuT/ASTOpImageMakeGrowMap.h
+++ b/Engine/Plugins/Experimental/Mutable/Source/MutableTools/Private/MuT/ASTOpImageMakeGrowMap.h
@@ -13,7 +13,7 @@ namespace mu
 {
 	struct FProgram;
 
-	class ASTOpImageMakeGrowMap : public ASTOp
+	class ASTOpImageMakeGrowMap final : public ASTOp
 	{
 	public:
 
diff --git a/Engine/Plugins/Experimental/Mutable/Source/MutableTools/Private/MuT/ASTOpImageMipmap.cpp b/Engine/Plugins/Experimental/Mutable/Source/MutableTools/Private/MuT/ASTOpImageMipmap.cpp
--- a/Engine/Plugins/Experimental/Mutable/Source/MutableTools/Private/MuT/ASTOpImageMipmap.cpp
+++ b/Engine/Plugins/Experimental/Mutable/Source/MutableTools/Private/MuT/ASTOpImageMipmap.cpp
@@ -18,9 +18,6 @@
 #include "MuT/ASTOpSwitch.h"
 #include "MuT/StreamsPrivate.h"
 
-#include <memory>
-#include <utility>
-
 
 namespace mu
 {
@@ -40,8 +37,9 @@ namespace mu
 
 	bool ASTOpImageMipmap::IsEqual(const ASTOp& otherUntyped) const
 	{
-		if (const ASTOpImageMipmap* other = dynamic_cast<const ASTOpImageMipmap*>(&otherUntyped))
+		if (otherUntyped.GetOpType()==GetOpType())
 		{
+			const ASTOpImageMipmap* other = static_cast<const ASTOpImageMipmap*>(&otherUntyped);
 			return Source == other->Source &&
 				Levels == other->Levels &&
 				BlockLevels == other->BlockLevels &&
@@ -379,7 +377,7 @@ namespace mu
 
 		case OP_TYPE::IM_COMPOSE:
 		{
-			const ASTOpImageCompose* typedAt = dynamic_cast<const ASTOpImageCompose*>(at.get());
+			const ASTOpImageCompose* typedAt = static_cast<const ASTOpImageCompose*>(at.get());
 			if (!currentMipmapOp->bOnlyTail
 				&&
 				// Don't move the mipmapping if we are composing with a mask.
@@ -410,7 +408,7 @@ namespace mu
 				// a top-level  mipmapping operation to generate the smallest
 				// mipmaps after the patch is done.
 
-				const ASTOpImagePatch* typedSource = dynamic_cast<const ASTOpImagePatch*>(at.get());
+				const ASTOpImagePatch* typedSource = static_cast<const ASTOpImagePatch*>(at.get());
 				Ptr<ASTOp> rectOp = typedSource->patch.child();
 				ASTOp::FGetImageDescContext context;
 				FImageDesc patchDesc = rectOp->GetImageDesc(false, &context);
diff --git a/Engine/Plugins/Experimental/Mutable/Source/MutableTools/Private/MuT/ASTOpImageMipmap.h b/Engine/Plugins/Experimental/Mutable/Source/MutableTools/Private/MuT/ASTOpImageMipmap.h
--- a/Engine/Plugins/Experimental/Mutable/Source/MutableTools/Private/MuT/ASTOpImageMipmap.h
+++ b/Engine/Plugins/Experimental/Mutable/Source/MutableTools/Private/MuT/ASTOpImageMipmap.h
@@ -13,7 +13,7 @@ namespace mu
 {
 	struct FProgram;
 
-	class ASTOpImageMipmap : public ASTOp
+	class ASTOpImageMipmap final : public ASTOp
 	{
 	public:
 
diff --git a/Engine/Plugins/Experimental/Mutable/Source/MutableTools/Private/MuT/ASTOpImageMultiLayer.cpp b/Engine/Plugins/Experimental/Mutable/Source/MutableTools/Private/MuT/ASTOpImageMultiLayer.cpp
--- a/Engine/Plugins/Experimental/Mutable/Source/MutableTools/Private/MuT/ASTOpImageMultiLayer.cpp
+++ b/Engine/Plugins/Experimental/Mutable/Source/MutableTools/Private/MuT/ASTOpImageMultiLayer.cpp
@@ -37,8 +37,9 @@ namespace mu
 	//-------------------------------------------------------------------------------------------------
 	bool ASTOpImageMultiLayer::IsEqual(const ASTOp& InOtherUntyped) const
 	{
-		if (const ASTOpImageMultiLayer* Other = dynamic_cast<const ASTOpImageMultiLayer*>(&InOtherUntyped))
+		if (InOtherUntyped.GetOpType()==GetOpType())
 		{
+			const ASTOpImageMultiLayer* Other = static_cast<const ASTOpImageMultiLayer*>(&InOtherUntyped);
 			return base == Other->base &&
 				blend == Other->blend &&
 				mask == Other->mask &&
diff --git a/Engine/Plugins/Experimental/Mutable/Source/MutableTools/Private/MuT/ASTOpImageMultiLayer.h b/Engine/Plugins/Experimental/Mutable/Source/MutableTools/Private/MuT/ASTOpImageMultiLayer.h
--- a/Engine/Plugins/Experimental/Mutable/Source/MutableTools/Private/MuT/ASTOpImageMultiLayer.h
+++ b/Engine/Plugins/Experimental/Mutable/Source/MutableTools/Private/MuT/ASTOpImageMultiLayer.h
@@ -17,7 +17,7 @@ namespace mu
 	//---------------------------------------------------------------------------------------------
 	//!
 	//---------------------------------------------------------------------------------------------
-	class ASTOpImageMultiLayer : public ASTOp
+	class ASTOpImageMultiLayer final : public ASTOp
 	{
 	public:
 
diff --git a/Engine/Plugins/Experimental/Mutable/Source/MutableTools/Private/MuT/ASTOpImageNormalComposite.cpp b/Engine/Plugins/Experimental/Mutable/Source/MutableTools/Private/MuT/ASTOpImageNormalComposite.cpp
--- a/Engine/Plugins/Experimental/Mutable/Source/MutableTools/Private/MuT/ASTOpImageNormalComposite.cpp
+++ b/Engine/Plugins/Experimental/Mutable/Source/MutableTools/Private/MuT/ASTOpImageNormalComposite.cpp
@@ -33,9 +33,13 @@ namespace mu
 
 	bool ASTOpImageNormalComposite::IsEqual(const ASTOp& OtherUntyped) const
 	{
-		if (const ASTOpImageNormalComposite* Other = dynamic_cast<const ASTOpImageNormalComposite*>(&OtherUntyped))
+		if (OtherUntyped.GetOpType()==GetOpType())
 		{
-			return Base == Other->Base && Normal == Other->Normal && Power == Other->Power && Mode == Other->Mode;
+			const ASTOpImageNormalComposite* Other = static_cast<const ASTOpImageNormalComposite*>(&OtherUntyped);
+			return Base == Other->Base 
+				&& Normal == Other->Normal 
+				&& Power == Other->Power 
+				&& Mode == Other->Mode;
 		}
 
 		return false;
diff --git a/Engine/Plugins/Experimental/Mutable/Source/MutableTools/Private/MuT/ASTOpImageNormalComposite.h b/Engine/Plugins/Experimental/Mutable/Source/MutableTools/Private/MuT/ASTOpImageNormalComposite.h
--- a/Engine/Plugins/Experimental/Mutable/Source/MutableTools/Private/MuT/ASTOpImageNormalComposite.h
+++ b/Engine/Plugins/Experimental/Mutable/Source/MutableTools/Private/MuT/ASTOpImageNormalComposite.h
@@ -13,7 +13,7 @@ namespace mu
 {
 struct FProgram;
 
-	class ASTOpImageNormalComposite : public ASTOp
+	class ASTOpImageNormalComposite final : public ASTOp
 	{
 	public:
 		ASTChild Base;
diff --git a/Engine/Plugins/Experimental/Mutable/Source/MutableTools/Private/MuT/ASTOpImagePatch.cpp b/Engine/Plugins/Experimental/Mutable/Source/MutableTools/Private/MuT/ASTOpImagePatch.cpp
--- a/Engine/Plugins/Experimental/Mutable/Source/MutableTools/Private/MuT/ASTOpImagePatch.cpp
+++ b/Engine/Plugins/Experimental/Mutable/Source/MutableTools/Private/MuT/ASTOpImagePatch.cpp
@@ -30,8 +30,9 @@ namespace mu
 	//-------------------------------------------------------------------------------------------------
 	bool ASTOpImagePatch::IsEqual(const ASTOp& otherUntyped) const
 	{
-		if (auto other = dynamic_cast<const ASTOpImagePatch*>(&otherUntyped))
+		if (otherUntyped.GetOpType()==GetOpType())
 		{
+			const ASTOpImagePatch* other = static_cast<const ASTOpImagePatch*>(&otherUntyped);
 			return base == other->base &&
 				patch == other->patch &&
 				location == other->location;
diff --git a/Engine/Plugins/Experimental/Mutable/Source/MutableTools/Private/MuT/ASTOpImagePatch.h b/Engine/Plugins/Experimental/Mutable/Source/MutableTools/Private/MuT/ASTOpImagePatch.h
--- a/Engine/Plugins/Experimental/Mutable/Source/MutableTools/Private/MuT/ASTOpImagePatch.h
+++ b/Engine/Plugins/Experimental/Mutable/Source/MutableTools/Private/MuT/ASTOpImagePatch.h
@@ -14,7 +14,7 @@ namespace mu
 {
 struct FProgram;
 
-	class ASTOpImagePatch : public ASTOp
+	class ASTOpImagePatch final : public ASTOp
 	{
 	public:
 
diff --git a/Engine/Plugins/Experimental/Mutable/Source/MutableTools/Private/MuT/ASTOpImagePixelFormat.cpp b/Engine/Plugins/Experimental/Mutable/Source/MutableTools/Private/MuT/ASTOpImagePixelFormat.cpp
--- a/Engine/Plugins/Experimental/Mutable/Source/MutableTools/Private/MuT/ASTOpImagePixelFormat.cpp
+++ b/Engine/Plugins/Experimental/Mutable/Source/MutableTools/Private/MuT/ASTOpImagePixelFormat.cpp
@@ -43,8 +43,9 @@ namespace mu
 
 	bool ASTOpImagePixelFormat::IsEqual(const ASTOp& otherUntyped) const
 	{
-		if (const ASTOpImagePixelFormat* other = dynamic_cast<const ASTOpImagePixelFormat*>(&otherUntyped))
+		if (otherUntyped.GetOpType()==GetOpType())
 		{
+			const ASTOpImagePixelFormat* other = static_cast<const ASTOpImagePixelFormat*>(&otherUntyped);
 			return Source == other->Source && Format == other->Format && FormatIfAlpha == other->FormatIfAlpha;
 		}
 		return false;
@@ -132,7 +133,7 @@ namespace mu
 		case OP_TYPE::IM_PIXELFORMAT:
 		{
 			// Keep only the top pixel format
-			const ASTOpImagePixelFormat* typedSource = dynamic_cast<const ASTOpImagePixelFormat*>(sourceAt.get());
+			const ASTOpImagePixelFormat* typedSource = static_cast<const ASTOpImagePixelFormat*>(sourceAt.get());
 			mu::Ptr<ASTOpImagePixelFormat> formatOp = mu::Clone<ASTOpImagePixelFormat>(this);
 			formatOp->Source = typedSource->Source.child();
 			at = formatOp;
@@ -176,7 +177,7 @@ namespace mu
 			// This op doesn't support compressed formats
 			if (!bIsCompressedFormat
 				&&
-				dynamic_cast<const ASTOpImageRasterMesh*>(sourceAt.get())->image)
+				static_cast<const ASTOpImageRasterMesh*>(sourceAt.get())->image)
 			{
 				Ptr<ASTOpImageRasterMesh> newOp = mu::Clone<ASTOpImageRasterMesh>(sourceAt);
 
@@ -325,9 +326,9 @@ namespace mu
 	//---------------------------------------------------------------------------------------------
 	//---------------------------------------------------------------------------------------------
 	//---------------------------------------------------------------------------------------------
-	mu::Ptr<ASTOp> Sink_ImagePixelFormatAST::Apply(const ASTOp* root)
+	mu::Ptr<ASTOp> Sink_ImagePixelFormatAST::Apply(const ASTOpImagePixelFormat* root)
 	{
-		m_root = dynamic_cast<const ASTOpImagePixelFormat*>(root);
+		m_root = root;
 		OldToNew.Empty();
 
 		check(root->GetOpType() == OP_TYPE::IM_PIXELFORMAT);
@@ -402,7 +403,7 @@ namespace mu
 
 				if (!acceptable)
 				{
-					const ASTOpImageCompose* typedAt = dynamic_cast<const ASTOpImageCompose*>(at.get());
+					const ASTOpImageCompose* typedAt = static_cast<const ASTOpImageCompose*>(at.get());
 					auto originalBaseOp = typedAt->Base.child();
 
 					int layoutBlockPixelsX = 0;
@@ -456,7 +457,7 @@ namespace mu
 
 		case OP_TYPE::IM_MIPMAP:
 		{
-			const ASTOpImageMipmap* typedSource = dynamic_cast<const ASTOpImageMipmap*>(at.get());
+			const ASTOpImageMipmap* typedSource = static_cast<const ASTOpImageMipmap*>(at.get());
 
 			// If its a compressed format, only sink formats on mipmap operations that
 			// generate the tail. To avoid optimization loop.
diff --git a/Engine/Plugins/Experimental/Mutable/Source/MutableTools/Private/MuT/ASTOpImagePixelFormat.h b/Engine/Plugins/Experimental/Mutable/Source/MutableTools/Private/MuT/ASTOpImagePixelFormat.h
--- a/Engine/Plugins/Experimental/Mutable/Source/MutableTools/Private/MuT/ASTOpImagePixelFormat.h
+++ b/Engine/Plugins/Experimental/Mutable/Source/MutableTools/Private/MuT/ASTOpImagePixelFormat.h
@@ -11,10 +11,9 @@
 
 namespace mu
 {
-struct FProgram;
-template <class SCALAR> class vec4;
+	struct FProgram;
 
-	class ASTOpImagePixelFormat : public ASTOp
+	class ASTOpImagePixelFormat final : public ASTOp
 	{
 	public:
 
diff --git a/Engine/Plugins/Experimental/Mutable/Source/MutableTools/Private/MuT/ASTOpImageRasterMesh.cpp b/Engine/Plugins/Experimental/Mutable/Source/MutableTools/Private/MuT/ASTOpImageRasterMesh.cpp
--- a/Engine/Plugins/Experimental/Mutable/Source/MutableTools/Private/MuT/ASTOpImageRasterMesh.cpp
+++ b/Engine/Plugins/Experimental/Mutable/Source/MutableTools/Private/MuT/ASTOpImageRasterMesh.cpp
@@ -52,8 +52,9 @@ namespace mu
 	//-------------------------------------------------------------------------------------------------
 	bool ASTOpImageRasterMesh::IsEqual(const ASTOp& InOtherUntyped) const
 	{
-		if (const ASTOpImageRasterMesh* Other = dynamic_cast<const ASTOpImageRasterMesh*>(&InOtherUntyped))
+		if (InOtherUntyped.GetOpType()==GetOpType())
 		{
+			const ASTOpImageRasterMesh* Other = static_cast<const ASTOpImageRasterMesh*>(&InOtherUntyped);
 			return mesh == Other->mesh &&
 				image == Other->image &&
 				angleFadeProperties == Other->angleFadeProperties &&
@@ -252,7 +253,7 @@ namespace mu
 			if (!imageAt)
 			{
 				// We remove the project from the raster children
-				const ASTOpFixed* MeshProjectOp = dynamic_cast<const ASTOpFixed*>(sourceAt.get());
+				const ASTOpFixed* MeshProjectOp = static_cast<const ASTOpFixed*>(sourceAt.get());
 				Ptr<ASTOpImageRasterMesh> nop = mu::Clone<ASTOpImageRasterMesh>(this);
 				nop->mesh = MeshProjectOp->children[MeshProjectOp->op.args.MeshProject.mesh].child();
 				at = nop;
@@ -263,7 +264,7 @@ namespace mu
 		case OP_TYPE::ME_INTERPOLATE:
 		{
 			// TODO: should be sink only if no imageAt?
-			auto typedSource = dynamic_cast<const ASTOpFixed*>(sourceAt.get());
+			auto typedSource = static_cast<const ASTOpFixed*>(sourceAt.get());
 			Ptr<ASTOpImageRasterMesh> rasterOp = mu::Clone<ASTOpImageRasterMesh>(this);
 			rasterOp->mesh = typedSource->children[typedSource->op.args.MeshInterpolate.base].child();
 			at = rasterOp;
@@ -273,7 +274,7 @@ namespace mu
 		case OP_TYPE::ME_MORPH:
 		{
 			// TODO: should be sink only if no imageAt?
-			const ASTOpMeshMorph* typedSource = dynamic_cast<const ASTOpMeshMorph*>(sourceAt.get());
+			const ASTOpMeshMorph* typedSource = static_cast<const ASTOpMeshMorph*>(sourceAt.get());
 			Ptr<ASTOpImageRasterMesh> rasterOp = mu::Clone<ASTOpImageRasterMesh>(this);
 			rasterOp->mesh = typedSource->Base.child();
 			at = rasterOp;
@@ -283,7 +284,7 @@ namespace mu
 		case OP_TYPE::ME_ADDTAGS:
 		{
 			// Ignore tags
-			const ASTOpMeshAddTags* typedSource = dynamic_cast<const ASTOpMeshAddTags*>(sourceAt.get());
+			const ASTOpMeshAddTags* typedSource = static_cast<const ASTOpMeshAddTags*>(sourceAt.get());
 			Ptr<ASTOpImageRasterMesh> rasterOp = mu::Clone<ASTOpImageRasterMesh>(this);
 			rasterOp->mesh = typedSource->Source.child();
 			at = rasterOp;
diff --git a/Engine/Plugins/Experimental/Mutable/Source/MutableTools/Private/MuT/ASTOpImageRasterMesh.h b/Engine/Plugins/Experimental/Mutable/Source/MutableTools/Private/MuT/ASTOpImageRasterMesh.h
--- a/Engine/Plugins/Experimental/Mutable/Source/MutableTools/Private/MuT/ASTOpImageRasterMesh.h
+++ b/Engine/Plugins/Experimental/Mutable/Source/MutableTools/Private/MuT/ASTOpImageRasterMesh.h
@@ -13,7 +13,7 @@ namespace mu
 {
 	struct FProgram;
 
-	class ASTOpImageRasterMesh : public ASTOp
+	class ASTOpImageRasterMesh final : public ASTOp
 	{
 	public:
 
