[mutable] Remove usage of RTTI.

diff --git a/Engine/Plugins/Experimental/Mutable/Source/MutableTools/Private/MuT/ASTOpImageSwizzle.cpp b/Engine/Plugins/Experimental/Mutable/Source/MutableTools/Private/MuT/ASTOpImageSwizzle.cpp
--- a/Engine/Plugins/Experimental/Mutable/Source/MutableTools/Private/MuT/ASTOpImageSwizzle.cpp
+++ b/Engine/Plugins/Experimental/Mutable/Source/MutableTools/Private/MuT/ASTOpImageSwizzle.cpp
@@ -44,8 +44,9 @@ namespace mu
 
 	bool ASTOpImageSwizzle::IsEqual(const ASTOp& otherUntyped) const
 	{
-		if (const ASTOpImageSwizzle* Other = dynamic_cast<const ASTOpImageSwizzle*>(&otherUntyped))
+		if (otherUntyped.GetOpType()==GetOpType())
 		{
+			const ASTOpImageSwizzle* Other = static_cast<const ASTOpImageSwizzle*>(&otherUntyped);
 			for (int32 i = 0; i<MUTABLE_OP_MAX_SWIZZLE_CHANNELS; ++i)
 			{
 				if (!(Sources[i] == Other->Sources[i] && SourceChannels[i] == Other->SourceChannels[i]))
@@ -140,7 +141,7 @@ namespace mu
 				{
 					sat = mu::Clone<ASTOpImageSwizzle>(this);
 				}
-				const ASTOpImageSwizzle* typedCandidate = dynamic_cast<const ASTOpImageSwizzle*>(candidate.get());
+				const ASTOpImageSwizzle* typedCandidate = static_cast<const ASTOpImageSwizzle*>(candidate.get());
 				int candidateChannel = SourceChannels[c];
 
 				sat->Sources[c] = typedCandidate->Sources[candidateChannel].child();
@@ -153,7 +154,7 @@ namespace mu
 			case OP_TYPE::IM_PIXELFORMAT:
 			{
 				// We can remove the format if its source is already an uncompressed format
-				ASTOpImagePixelFormat* typedCandidate = dynamic_cast<ASTOpImagePixelFormat*>(candidate.get());
+				ASTOpImagePixelFormat* typedCandidate = static_cast<ASTOpImagePixelFormat*>(candidate.get());
 				Ptr<ASTOp> formatSource = typedCandidate->Source.child();
 
 				if (formatSource)
@@ -489,9 +490,9 @@ namespace mu
 				SourceChannels[0] == 0 && SourceChannels[1] == 1 && SourceChannels[2] == 2 && SourceChannels[3] == 0
 				)
 			{
-				const ASTOpImageMultiLayer* ColorMultiLayer = dynamic_cast<const ASTOpImageMultiLayer*>(Sources[0].child().get());
+				const ASTOpImageMultiLayer* ColorMultiLayer = static_cast<const ASTOpImageMultiLayer*>(Sources[0].child().get());
 				check(ColorMultiLayer);
-				const ASTOpImageMultiLayer* AlphaMultiLayer = dynamic_cast<const ASTOpImageMultiLayer*>(Sources[3].child().get());
+				const ASTOpImageMultiLayer* AlphaMultiLayer = static_cast<const ASTOpImageMultiLayer*>(Sources[3].child().get());
 				check(AlphaMultiLayer);
 
 				bool bIsSpecialMultiLayer = !AlphaMultiLayer->mask
@@ -548,9 +549,9 @@ namespace mu
 				SourceChannels[0] == 0 && SourceChannels[1] == 1 && (SourceChannels[2] == 2 || !Sources[2] ) && SourceChannels[3] == 0
 				)
 			{
-				const ASTOpImageLayer* ColorLayer = dynamic_cast<const ASTOpImageLayer*>(Sources[0].child().get());
+				const ASTOpImageLayer* ColorLayer = static_cast<const ASTOpImageLayer*>(Sources[0].child().get());
 				check(ColorLayer);
-				const ASTOpImageLayer* AlphaLayer = dynamic_cast<const ASTOpImageLayer*>(Sources[3].child().get());
+				const ASTOpImageLayer* AlphaLayer = static_cast<const ASTOpImageLayer*>(Sources[3].child().get());
 				check(AlphaLayer);
 
 				bool bIsSpecialMultiLayer = !AlphaLayer->mask && !ColorLayer->Flags && !AlphaLayer->Flags;
@@ -595,7 +596,7 @@ namespace mu
 			// If the channels are compatible switches, we can still sink the swizzle.
 			if (!at && sourceType == OP_TYPE::IM_SWITCH)
 			{
-				const ASTOpSwitch* FirstSwitch = dynamic_cast<const ASTOpSwitch*>(Sources[0].child().get());
+				const ASTOpSwitch* FirstSwitch = static_cast<const ASTOpSwitch*>(Sources[0].child().get());
 				check(FirstSwitch);
 
 				bool bAreAllSwitchesCompatible = true;
@@ -603,7 +604,7 @@ namespace mu
 				{
 					if (Sources[c])
 					{
-						const ASTOpSwitch* Typed = dynamic_cast<const ASTOpSwitch*>(Sources[c].child().get());
+						const ASTOpSwitch* Typed = static_cast<const ASTOpSwitch*>(Sources[c].child().get());
 						check(Typed);
 						if (!Typed->IsCompatibleWith(FirstSwitch))
 						{
@@ -623,7 +624,7 @@ namespace mu
 						Ptr<ASTOpImageSwizzle> defOp = mu::Clone<ASTOpImageSwizzle>(this);
 						for (int c = 0; c < MUTABLE_OP_MAX_SWIZZLE_CHANNELS; ++c)
 						{
-							const ASTOpSwitch* ChannelSwitch = dynamic_cast<const ASTOpSwitch*>(Sources[c].child().get());
+							const ASTOpSwitch* ChannelSwitch = static_cast<const ASTOpSwitch*>(Sources[c].child().get());
 							if (ChannelSwitch)
 							{
 								defOp->Sources[c] = ChannelSwitch->def.child();
@@ -639,7 +640,7 @@ namespace mu
 							Ptr<ASTOpImageSwizzle> branchOp = mu::Clone<ASTOpImageSwizzle>(this);
 							for (int c = 0; c < MUTABLE_OP_MAX_SWIZZLE_CHANNELS; ++c)
 							{
-								const ASTOpSwitch* ChannelSwitch = dynamic_cast<const ASTOpSwitch*>(Sources[c].child().get());
+								const ASTOpSwitch* ChannelSwitch = static_cast<const ASTOpSwitch*>(Sources[c].child().get());
 								if (ChannelSwitch)
 								{
 									branchOp->Sources[c] = ChannelSwitch->cases[v].branch.child();
@@ -656,7 +657,7 @@ namespace mu
 			// Swizzle down compatible displaces.
 			if (!at && sourceType == OP_TYPE::IM_DISPLACE)
 			{
-				const ASTOpFixed* FirstDisplace = dynamic_cast<const ASTOpFixed*>(Sources[0].child().get());
+				const ASTOpFixed* FirstDisplace = static_cast<const ASTOpFixed*>(Sources[0].child().get());
 				check(FirstDisplace);
 
 				bool bAreAllDisplacesCompatible = true;
@@ -664,7 +665,7 @@ namespace mu
 				{
 					if (Sources[c])
 					{
-						const ASTOpFixed* Typed = dynamic_cast<const ASTOpFixed*>(Sources[c].child().get());
+						const ASTOpFixed* Typed = static_cast<const ASTOpFixed*>(Sources[c].child().get());
 						check(Typed);
 						if (FirstDisplace->op.args.ImageDisplace.displacementMap != Typed->op.args.ImageDisplace.displacementMap)
 						{
@@ -682,7 +683,7 @@ namespace mu
 					Ptr<ASTOpImageSwizzle> SourceOp = mu::Clone<ASTOpImageSwizzle>(this);
 					for (int c = 0; c < MUTABLE_OP_MAX_SWIZZLE_CHANNELS; ++c)
 					{
-						const ASTOpFixed* ChannelDisplace = dynamic_cast<const ASTOpFixed*>(Sources[c].child().get());
+						const ASTOpFixed* ChannelDisplace = static_cast<const ASTOpFixed*>(Sources[c].child().get());
 						if (ChannelDisplace)
 						{
 							SourceOp->Sources[c] = ChannelDisplace->children[ChannelDisplace->op.args.ImageDisplace.source].child();
@@ -699,7 +700,7 @@ namespace mu
 			// Swizzle down compatible raster meshes.
 			if (!at && sourceType == OP_TYPE::IM_RASTERMESH)
 			{
-				const ASTOpImageRasterMesh* FirstRasterMesh = dynamic_cast<const ASTOpImageRasterMesh*>(Sources[0].child().get());
+				const ASTOpImageRasterMesh* FirstRasterMesh = static_cast<const ASTOpImageRasterMesh*>(Sources[0].child().get());
 				check(FirstRasterMesh);
 
 				bool bAreAllRasterMeshesCompatible = true;
@@ -707,7 +708,7 @@ namespace mu
 				{
 					if (Sources[c])
 					{
-						const ASTOpImageRasterMesh* Typed = dynamic_cast<const ASTOpImageRasterMesh*>(Sources[c].child().get());
+						const ASTOpImageRasterMesh* Typed = static_cast<const ASTOpImageRasterMesh*>(Sources[c].child().get());
 						check(Typed);
 
 						// Compare all args but the source image
@@ -742,7 +743,7 @@ namespace mu
 					Ptr<ASTOpImageSwizzle> NewSwizzle = mu::Clone<ASTOpImageSwizzle>(this);
 					for (int c = 0; c < MUTABLE_OP_MAX_SWIZZLE_CHANNELS; ++c)
 					{
-						const ASTOpImageRasterMesh* ChannelRaster = dynamic_cast<const ASTOpImageRasterMesh*>(Sources[c].child().get());
+						const ASTOpImageRasterMesh* ChannelRaster = static_cast<const ASTOpImageRasterMesh*>(Sources[c].child().get());
 						if (ChannelRaster)
 						{
 							NewSwizzle->Sources[c] = ChannelRaster->image.child();
@@ -773,7 +774,7 @@ namespace mu
 			// Swizzle down compatible image transforms.
 			if (!at && sourceType == OP_TYPE::IM_TRANSFORM)
 			{
-				const ASTOpImageTransform* FirstTransform = dynamic_cast<const ASTOpImageTransform*>(Sources[0].child().get());
+				const ASTOpImageTransform* FirstTransform = static_cast<const ASTOpImageTransform*>(Sources[0].child().get());
 				check(FirstTransform);
 
 				bool bAreAllTransformsCompatible = true;
@@ -781,7 +782,7 @@ namespace mu
 				{
 					if (Sources[c])
 					{
-						const ASTOpImageTransform* Typed = dynamic_cast<const ASTOpImageTransform*>(Sources[c].child().get());
+						const ASTOpImageTransform* Typed = static_cast<const ASTOpImageTransform*>(Sources[c].child().get());
 						check(Typed);
 
 						// Compare all args but the source image
@@ -812,7 +813,7 @@ namespace mu
 					Ptr<ASTOpImageSwizzle> NewSwizzle = mu::Clone<ASTOpImageSwizzle>(this);
 					for (int c = 0; c < MUTABLE_OP_MAX_SWIZZLE_CHANNELS; ++c)
 					{
-						const ASTOpImageTransform* ChannelTransform = dynamic_cast<const ASTOpImageTransform*>(Sources[c].child().get());
+						const ASTOpImageTransform* ChannelTransform = static_cast<const ASTOpImageTransform*>(Sources[c].child().get());
 						if (ChannelTransform)
 						{
 							NewSwizzle->Sources[c] = ChannelTransform->Base.child();
@@ -828,7 +829,7 @@ namespace mu
 			// Swizzle down compatible resizes.
 			//if (!at && sourceType == OP_TYPE::IM_RESIZE)
 			//{
-			//	const ASTOpFixed* FirstResize = dynamic_cast<const ASTOpFixed*>(Sources[0].child().get());
+			//	const ASTOpFixed* FirstResize = static_cast<const ASTOpFixed*>(Sources[0].child().get());
 			//	check(FirstResize);
 
 			//	bool bAreAllResizesCompatible = true;
@@ -836,7 +837,7 @@ namespace mu
 			//	{
 			//		if (Sources[c])
 			//		{
-			//			const ASTOpFixed* Typed = dynamic_cast<const ASTOpFixed*>(Sources[c].child().get());
+			//			const ASTOpFixed* Typed = static_cast<const ASTOpFixed*>(Sources[c].child().get());
 			//			check(Typed);
 			//			// Compare all args but the source image
 			//			OP::ImageResizeArgs ArgCopy = FirstResize->op.args.ImageResize;
@@ -858,7 +859,7 @@ namespace mu
 			//		Ptr<ASTOpImageSwizzle> NewSwizzle = mu::Clone<ASTOpImageSwizzle>(this);
 			//		for (int c = 0; c < MUTABLE_OP_MAX_SWIZZLE_CHANNELS; ++c)
 			//		{
-			//			const ASTOpFixed* ChannelResize = dynamic_cast<const ASTOpFixed*>(Sources[c].child().get());
+			//			const ASTOpFixed* ChannelResize = static_cast<const ASTOpFixed*>(Sources[c].child().get());
 			//			if (ChannelResize)
 			//			{
 			//				NewSwizzle->Sources[c] = ChannelResize->children[ChannelResize->op.args.ImageResize.source].child();
@@ -875,7 +876,7 @@ namespace mu
 			// Swizzle down compatible pixelformats.
 			if (!at && sourceType == OP_TYPE::IM_PIXELFORMAT && bSameChannelOrder)
 			{
-				const ASTOpImagePixelFormat* FirstFormat = dynamic_cast<const ASTOpImagePixelFormat*>(Sources[0].child().get());
+				const ASTOpImagePixelFormat* FirstFormat = static_cast<const ASTOpImagePixelFormat*>(Sources[0].child().get());
 				check(FirstFormat);
 
 				bool bAreAllFormatsCompatible = true;
@@ -883,7 +884,7 @@ namespace mu
 				{
 					if (Sources[c])
 					{
-						const ASTOpImagePixelFormat* Typed = dynamic_cast<const ASTOpImagePixelFormat*>(Sources[c].child().get());
+						const ASTOpImagePixelFormat* Typed = static_cast<const ASTOpImagePixelFormat*>(Sources[c].child().get());
 						check(Typed);
 
 						if (Typed->Source.child() != FirstFormat->Source.child())
@@ -913,7 +914,7 @@ namespace mu
 				{
 					if (Sources[c])
 					{
-						const ASTOpFixed* TypedPlain = dynamic_cast<const ASTOpFixed*>(Sources[c].child().get());
+						const ASTOpFixed* TypedPlain = static_cast<const ASTOpFixed*>(Sources[c].child().get());
 
 						NewSwizzle->SetChild(NewSwizzle->op.args.ColourSwizzle.sources[c], TypedPlain->children[TypedPlain->op.args.ImagePlainColour.colour]);
 					}
@@ -1020,7 +1021,7 @@ namespace mu
 						NewSwizzle = mu::Clone<ASTOpImageSwizzle>(this);
 					}
 
-					const ASTOpFixed* OldSaturate = dynamic_cast<const ASTOpFixed*>(Sources[Channel].child().get());
+					const ASTOpFixed* OldSaturate = static_cast<const ASTOpFixed*>(Sources[Channel].child().get());
 					Ptr<ASTOp> OldSaturateBase = OldSaturate->children[OldSaturate->op.args.ImageSaturate.base].child();
 
 					NewSwizzle->Sources[Channel] = OldSaturateBase;
@@ -1091,7 +1092,7 @@ namespace mu
 			Sources[3] && Sources[3]->GetOpType() == OP_TYPE::IM_LAYER
 			)
 		{
-			const ASTOpImageLayer* OldLayer = dynamic_cast<const ASTOpImageLayer*>(Sources[3].child().get());
+			const ASTOpImageLayer* OldLayer = static_cast<const ASTOpImageLayer*>(Sources[3].child().get());
 
 			Ptr<ASTOp> SwizzleRGBOp = Sources[0].child();
 			Ptr<ASTOp> OldLayerBlendOp = OldLayer->blend.child();
@@ -1106,7 +1107,7 @@ namespace mu
 						{
 						case OP_TYPE::IM_PIXELFORMAT:
 						{
-							const ASTOpImagePixelFormat* Typed = dynamic_cast<const ASTOpImagePixelFormat*>(Op.get());
+							const ASTOpImagePixelFormat* Typed = static_cast<const ASTOpImagePixelFormat*>(Op.get());
 							Op = Typed->Source.child();
 							bUpdated = true;
 							break;
@@ -1163,7 +1164,7 @@ namespace mu
 			Sources[3] && Sources[3]->GetOpType() == OP_TYPE::IM_LAYER
 			)
 		{
-			ASTOpImageLayer* OldLayer = dynamic_cast<ASTOpImageLayer*>(Sources[3].child().get());
+			ASTOpImageLayer* OldLayer = static_cast<ASTOpImageLayer*>(Sources[3].child().get());
 
 			// For now just check the case that we are observing in the working data: 
 			if (OldLayer->Flags == 0
@@ -1172,7 +1173,7 @@ namespace mu
 				&&
 				OldLayer->blend->GetOpType()==OP_TYPE::IM_PIXELFORMAT )
 			{
-				const ASTOpImagePixelFormat* OldFormat = dynamic_cast<const ASTOpImagePixelFormat*>(OldLayer->blend.child().get());
+				const ASTOpImagePixelFormat* OldFormat = static_cast<const ASTOpImagePixelFormat*>(OldLayer->blend.child().get());
 				if (OldFormat->Source->GetOpType() == OP_TYPE::IM_SWIZZLE
 					&&
 					(
@@ -1182,7 +1183,7 @@ namespace mu
 						) 
 					)
 				{
-					const ASTOpImageSwizzle* OldChildSwizzle = dynamic_cast<const ASTOpImageSwizzle*>(OldFormat->Source.child().get());
+					const ASTOpImageSwizzle* OldChildSwizzle = static_cast<const ASTOpImageSwizzle*>(OldFormat->Source.child().get());
 					if (OldChildSwizzle->Format == EImageFormat::IF_L_UBYTE)
 					{
 						Ptr<ASTOpImageSwizzle> NewBaseSwizzle = new ASTOpImageSwizzle;
diff --git a/Engine/Plugins/Experimental/Mutable/Source/MutableTools/Private/MuT/ASTOpImageSwizzle.h b/Engine/Plugins/Experimental/Mutable/Source/MutableTools/Private/MuT/ASTOpImageSwizzle.h
--- a/Engine/Plugins/Experimental/Mutable/Source/MutableTools/Private/MuT/ASTOpImageSwizzle.h
+++ b/Engine/Plugins/Experimental/Mutable/Source/MutableTools/Private/MuT/ASTOpImageSwizzle.h
@@ -13,7 +13,7 @@ namespace mu
 {
 	struct FProgram;
 
-	class ASTOpImageSwizzle : public ASTOp
+	class ASTOpImageSwizzle final : public ASTOp
 	{
 	public:
 
diff --git a/Engine/Plugins/Experimental/Mutable/Source/MutableTools/Private/MuT/ASTOpImageTransform.cpp b/Engine/Plugins/Experimental/Mutable/Source/MutableTools/Private/MuT/ASTOpImageTransform.cpp
--- a/Engine/Plugins/Experimental/Mutable/Source/MutableTools/Private/MuT/ASTOpImageTransform.cpp
+++ b/Engine/Plugins/Experimental/Mutable/Source/MutableTools/Private/MuT/ASTOpImageTransform.cpp
@@ -39,8 +39,9 @@ namespace mu
 	//-------------------------------------------------------------------------------------------------
 	bool ASTOpImageTransform::IsEqual(const ASTOp& OtherUntyped) const
 	{
-		if (auto Other = dynamic_cast<const ASTOpImageTransform*>(&OtherUntyped))
+		if (OtherUntyped.GetOpType()==GetOpType())
 		{
+			auto Other = static_cast<const ASTOpImageTransform*>(&OtherUntyped);
 			return 
 				Base 	 == Other->Base &&
 				OffsetX  == Other->OffsetX &&
diff --git a/Engine/Plugins/Experimental/Mutable/Source/MutableTools/Private/MuT/ASTOpImageTransform.h b/Engine/Plugins/Experimental/Mutable/Source/MutableTools/Private/MuT/ASTOpImageTransform.h
--- a/Engine/Plugins/Experimental/Mutable/Source/MutableTools/Private/MuT/ASTOpImageTransform.h
+++ b/Engine/Plugins/Experimental/Mutable/Source/MutableTools/Private/MuT/ASTOpImageTransform.h
@@ -17,7 +17,7 @@ struct FProgram;
 	//---------------------------------------------------------------------------------------------
 	//!
 	//---------------------------------------------------------------------------------------------
-	class ASTOpImageTransform : public ASTOp
+	class ASTOpImageTransform final : public ASTOp
 	{
 	public:
 
diff --git a/Engine/Plugins/Experimental/Mutable/Source/MutableTools/Private/MuT/ASTOpInstanceAdd.cpp b/Engine/Plugins/Experimental/Mutable/Source/MutableTools/Private/MuT/ASTOpInstanceAdd.cpp
--- a/Engine/Plugins/Experimental/Mutable/Source/MutableTools/Private/MuT/ASTOpInstanceAdd.cpp
+++ b/Engine/Plugins/Experimental/Mutable/Source/MutableTools/Private/MuT/ASTOpInstanceAdd.cpp
@@ -38,8 +38,9 @@ namespace mu
 	//-------------------------------------------------------------------------------------------------
 	bool ASTOpInstanceAdd::IsEqual(const ASTOp& otherUntyped) const
 	{
-		if (auto other = dynamic_cast<const ASTOpInstanceAdd*>(&otherUntyped))
+		if (otherUntyped.GetOpType()==GetOpType())
 		{
+			const ASTOpInstanceAdd* other = static_cast<const ASTOpInstanceAdd*>(&otherUntyped);
 			return type == other->type &&
 				instance == other->instance &&
 				value == other->value &&
diff --git a/Engine/Plugins/Experimental/Mutable/Source/MutableTools/Private/MuT/ASTOpInstanceAdd.h b/Engine/Plugins/Experimental/Mutable/Source/MutableTools/Private/MuT/ASTOpInstanceAdd.h
--- a/Engine/Plugins/Experimental/Mutable/Source/MutableTools/Private/MuT/ASTOpInstanceAdd.h
+++ b/Engine/Plugins/Experimental/Mutable/Source/MutableTools/Private/MuT/ASTOpInstanceAdd.h
@@ -16,7 +16,7 @@ struct FProgram;
 	//---------------------------------------------------------------------------------------------
 	//! Operations to add elements to an instance
 	//---------------------------------------------------------------------------------------------
-	class ASTOpInstanceAdd : public ASTOp
+	class ASTOpInstanceAdd final : public ASTOp
 	{
 	public:
 
diff --git a/Engine/Plugins/Experimental/Mutable/Source/MutableTools/Private/MuT/ASTOpLayoutFromMesh.cpp b/Engine/Plugins/Experimental/Mutable/Source/MutableTools/Private/MuT/ASTOpLayoutFromMesh.cpp
--- a/Engine/Plugins/Experimental/Mutable/Source/MutableTools/Private/MuT/ASTOpLayoutFromMesh.cpp
+++ b/Engine/Plugins/Experimental/Mutable/Source/MutableTools/Private/MuT/ASTOpLayoutFromMesh.cpp
@@ -47,8 +47,9 @@ namespace mu
 
 	bool ASTOpLayoutFromMesh::IsEqual(const ASTOp& otherUntyped) const
 	{
-		if (const ASTOpLayoutFromMesh* other = dynamic_cast<const ASTOpLayoutFromMesh*>(&otherUntyped))
+		if (otherUntyped.GetOpType()==GetOpType())
 		{
+			const ASTOpLayoutFromMesh* other = static_cast<const ASTOpLayoutFromMesh*>(&otherUntyped);
 			return Mesh == other->Mesh && LayoutIndex == other->LayoutIndex;
 		}
 		return false;
@@ -159,8 +160,7 @@ namespace mu
 				case OP_TYPE::ME_MORPH:
 				{
 					// Sink, ignoring the op
-					const ASTOpMeshMorph* Typed = dynamic_cast<const ASTOpMeshMorph*>(at.get());
-					check(Typed);
+					const ASTOpMeshMorph* Typed = static_cast<const ASTOpMeshMorph*>(at.get());
 					newAt = Visit(Typed->Base.child(), CurrentSinkingOp);
 					break;
 				}
@@ -168,8 +168,7 @@ namespace mu
 				case OP_TYPE::ME_FORMAT:
 				{
 					// Sink, ignoring the op
-					const ASTOpMeshFormat* Typed = dynamic_cast<const ASTOpMeshFormat*>(at.get());
-					check(Typed);
+					const ASTOpMeshFormat* Typed = static_cast<const ASTOpMeshFormat*>(at.get());
 					newAt = Visit(Typed->Source.child(), CurrentSinkingOp);
 					break;
 				}
@@ -177,8 +176,7 @@ namespace mu
 				case OP_TYPE::ME_APPLYSHAPE:
 				{
 					// Sink, ignoring the op
-					const ASTOpMeshApplyShape* Typed = dynamic_cast<const ASTOpMeshApplyShape*>(at.get());
-					check(Typed);
+					const ASTOpMeshApplyShape* Typed = static_cast<const ASTOpMeshApplyShape*>(at.get());
 					newAt = Visit(Typed->Mesh.child(), CurrentSinkingOp);
 					break;
 				}
@@ -186,8 +184,7 @@ namespace mu
 				case OP_TYPE::ME_BINDSHAPE:
 				{
 					// Sink, ignoring the op
-					const ASTOpMeshBindShape* Typed = dynamic_cast<const ASTOpMeshBindShape*>(at.get());
-					check(Typed);
+					const ASTOpMeshBindShape* Typed = static_cast<const ASTOpMeshBindShape*>(at.get());
 					newAt = Visit(Typed->Mesh.child(), CurrentSinkingOp);
 					break;
 				}
@@ -195,8 +192,7 @@ namespace mu
 				case OP_TYPE::ME_ADDTAGS:
 				{
 					// Sink, ignoring the op
-					const ASTOpMeshAddTags* Typed = dynamic_cast<const ASTOpMeshAddTags*>(at.get());
-					check(Typed);
+					const ASTOpMeshAddTags* Typed = static_cast<const ASTOpMeshAddTags*>(at.get());
 					newAt = Visit(Typed->Source.child(), CurrentSinkingOp);
 					break;
 				}
@@ -204,8 +200,7 @@ namespace mu
 				case OP_TYPE::ME_INTERPOLATE:
 				{
 					// Sink, ignoring the op
-					const ASTOpFixed* Typed = dynamic_cast<const ASTOpFixed*>(at.get());
-					check(Typed);
+					const ASTOpFixed* Typed = static_cast<const ASTOpFixed*>(at.get());
 					newAt = Visit(Typed->children[Typed->op.args.MeshInterpolate.base].child(), CurrentSinkingOp);
 					break;
 				}
@@ -235,8 +230,7 @@ namespace mu
 
 				case OP_TYPE::ME_MERGE:
 				{
-					const ASTOpFixed* Typed = dynamic_cast<const ASTOpFixed*>(at.get());
-					check(Typed);
+					const ASTOpFixed* Typed = static_cast<const ASTOpFixed*>(at.get());
 
 					Ptr<ASTOpLayoutMerge> NewMerge = new ASTOpLayoutMerge;
 					NewMerge->Base = Visit(Typed->children[Typed->op.args.MeshMerge.base].child(), CurrentSinkingOp);
diff --git a/Engine/Plugins/Experimental/Mutable/Source/MutableTools/Private/MuT/ASTOpLayoutFromMesh.h b/Engine/Plugins/Experimental/Mutable/Source/MutableTools/Private/MuT/ASTOpLayoutFromMesh.h
--- a/Engine/Plugins/Experimental/Mutable/Source/MutableTools/Private/MuT/ASTOpLayoutFromMesh.h
+++ b/Engine/Plugins/Experimental/Mutable/Source/MutableTools/Private/MuT/ASTOpLayoutFromMesh.h
@@ -12,7 +12,7 @@ namespace mu
 {
 	struct FProgram;
 	
-	class ASTOpLayoutFromMesh : public ASTOp
+	class ASTOpLayoutFromMesh final : public ASTOp
 	{
 	public:
 
diff --git a/Engine/Plugins/Experimental/Mutable/Source/MutableTools/Private/MuT/ASTOpLayoutMerge.cpp b/Engine/Plugins/Experimental/Mutable/Source/MutableTools/Private/MuT/ASTOpLayoutMerge.cpp
--- a/Engine/Plugins/Experimental/Mutable/Source/MutableTools/Private/MuT/ASTOpLayoutMerge.cpp
+++ b/Engine/Plugins/Experimental/Mutable/Source/MutableTools/Private/MuT/ASTOpLayoutMerge.cpp
@@ -42,8 +42,9 @@ namespace mu
 
 	bool ASTOpLayoutMerge::IsEqual(const ASTOp& otherUntyped) const
 	{
-		if (const ASTOpLayoutMerge* other = dynamic_cast<const ASTOpLayoutMerge*>(&otherUntyped))
+		if (otherUntyped.GetOpType()==GetOpType())
 		{
+			const ASTOpLayoutMerge* other = static_cast<const ASTOpLayoutMerge*>(&otherUntyped);
 			return Base == other->Base && Added == other->Added;
 		}
 		return false;
diff --git a/Engine/Plugins/Experimental/Mutable/Source/MutableTools/Private/MuT/ASTOpLayoutMerge.h b/Engine/Plugins/Experimental/Mutable/Source/MutableTools/Private/MuT/ASTOpLayoutMerge.h
--- a/Engine/Plugins/Experimental/Mutable/Source/MutableTools/Private/MuT/ASTOpLayoutMerge.h
+++ b/Engine/Plugins/Experimental/Mutable/Source/MutableTools/Private/MuT/ASTOpLayoutMerge.h
@@ -12,7 +12,7 @@ namespace mu
 {
 	struct FProgram;
 	
-	class ASTOpLayoutMerge : public ASTOp
+	class ASTOpLayoutMerge final : public ASTOp
 	{
 	public:
 
diff --git a/Engine/Plugins/Experimental/Mutable/Source/MutableTools/Private/MuT/ASTOpLayoutPack.cpp b/Engine/Plugins/Experimental/Mutable/Source/MutableTools/Private/MuT/ASTOpLayoutPack.cpp
--- a/Engine/Plugins/Experimental/Mutable/Source/MutableTools/Private/MuT/ASTOpLayoutPack.cpp
+++ b/Engine/Plugins/Experimental/Mutable/Source/MutableTools/Private/MuT/ASTOpLayoutPack.cpp
@@ -41,8 +41,9 @@ namespace mu
 
 	bool ASTOpLayoutPack::IsEqual(const ASTOp& otherUntyped) const
 	{
-		if (const ASTOpLayoutPack* other = dynamic_cast<const ASTOpLayoutPack*>(&otherUntyped))
+		if (otherUntyped.GetOpType()==GetOpType())
 		{
+			const ASTOpLayoutPack* other = static_cast<const ASTOpLayoutPack*>(&otherUntyped);
 			return Source == other->Source;
 		}
 		return false;
diff --git a/Engine/Plugins/Experimental/Mutable/Source/MutableTools/Private/MuT/ASTOpLayoutPack.h b/Engine/Plugins/Experimental/Mutable/Source/MutableTools/Private/MuT/ASTOpLayoutPack.h
--- a/Engine/Plugins/Experimental/Mutable/Source/MutableTools/Private/MuT/ASTOpLayoutPack.h
+++ b/Engine/Plugins/Experimental/Mutable/Source/MutableTools/Private/MuT/ASTOpLayoutPack.h
@@ -12,7 +12,7 @@ namespace mu
 {
 	struct FProgram;
 	
-	class ASTOpLayoutPack : public ASTOp
+	class ASTOpLayoutPack final : public ASTOp
 	{
 	public:
 
diff --git a/Engine/Plugins/Experimental/Mutable/Source/MutableTools/Private/MuT/ASTOpLayoutRemoveBlocks.cpp b/Engine/Plugins/Experimental/Mutable/Source/MutableTools/Private/MuT/ASTOpLayoutRemoveBlocks.cpp
--- a/Engine/Plugins/Experimental/Mutable/Source/MutableTools/Private/MuT/ASTOpLayoutRemoveBlocks.cpp
+++ b/Engine/Plugins/Experimental/Mutable/Source/MutableTools/Private/MuT/ASTOpLayoutRemoveBlocks.cpp
@@ -42,8 +42,9 @@ namespace mu
 
 	bool ASTOpLayoutRemoveBlocks::IsEqual(const ASTOp& otherUntyped) const
 	{
-		if (const ASTOpLayoutRemoveBlocks* other = dynamic_cast<const ASTOpLayoutRemoveBlocks*>(&otherUntyped))
+		if (otherUntyped.GetOpType()==GetOpType())
 		{
+			const ASTOpLayoutRemoveBlocks* other = static_cast<const ASTOpLayoutRemoveBlocks*>(&otherUntyped);
 			return Source == other->Source && ReferenceLayout == other->ReferenceLayout;
 		}
 		return false;
diff --git a/Engine/Plugins/Experimental/Mutable/Source/MutableTools/Private/MuT/ASTOpLayoutRemoveBlocks.h b/Engine/Plugins/Experimental/Mutable/Source/MutableTools/Private/MuT/ASTOpLayoutRemoveBlocks.h
--- a/Engine/Plugins/Experimental/Mutable/Source/MutableTools/Private/MuT/ASTOpLayoutRemoveBlocks.h
+++ b/Engine/Plugins/Experimental/Mutable/Source/MutableTools/Private/MuT/ASTOpLayoutRemoveBlocks.h
@@ -12,7 +12,7 @@ namespace mu
 {
 	struct FProgram;
 	
-	class ASTOpLayoutRemoveBlocks : public ASTOp
+	class ASTOpLayoutRemoveBlocks final : public ASTOp
 	{
 	public:
 
diff --git a/Engine/Plugins/Experimental/Mutable/Source/MutableTools/Private/MuT/ASTOpMeshAddTags.cpp b/Engine/Plugins/Experimental/Mutable/Source/MutableTools/Private/MuT/ASTOpMeshAddTags.cpp
--- a/Engine/Plugins/Experimental/Mutable/Source/MutableTools/Private/MuT/ASTOpMeshAddTags.cpp
+++ b/Engine/Plugins/Experimental/Mutable/Source/MutableTools/Private/MuT/ASTOpMeshAddTags.cpp
@@ -28,8 +28,9 @@ namespace mu
 
 	bool ASTOpMeshAddTags::IsEqual(const ASTOp& otherUntyped) const
 	{
-		if (const ASTOpMeshAddTags* other = dynamic_cast<const ASTOpMeshAddTags*>(&otherUntyped))
+		if (otherUntyped.GetOpType()==GetOpType())
 		{
+			const ASTOpMeshAddTags* other = static_cast<const ASTOpMeshAddTags*>(&otherUntyped);
 			return Source == other->Source && Tags == other->Tags;
 		}
 		return false;
diff --git a/Engine/Plugins/Experimental/Mutable/Source/MutableTools/Private/MuT/ASTOpMeshAddTags.h b/Engine/Plugins/Experimental/Mutable/Source/MutableTools/Private/MuT/ASTOpMeshAddTags.h
--- a/Engine/Plugins/Experimental/Mutable/Source/MutableTools/Private/MuT/ASTOpMeshAddTags.h
+++ b/Engine/Plugins/Experimental/Mutable/Source/MutableTools/Private/MuT/ASTOpMeshAddTags.h
@@ -13,7 +13,7 @@ namespace mu
 	struct FProgram;
 
 	/** Add tags to a mesh. */
-	class ASTOpMeshAddTags : public ASTOp
+	class ASTOpMeshAddTags final : public ASTOp
 	{
 	public:
 
diff --git a/Engine/Plugins/Experimental/Mutable/Source/MutableTools/Private/MuT/ASTOpMeshApplyPose.cpp b/Engine/Plugins/Experimental/Mutable/Source/MutableTools/Private/MuT/ASTOpMeshApplyPose.cpp
--- a/Engine/Plugins/Experimental/Mutable/Source/MutableTools/Private/MuT/ASTOpMeshApplyPose.cpp
+++ b/Engine/Plugins/Experimental/Mutable/Source/MutableTools/Private/MuT/ASTOpMeshApplyPose.cpp
@@ -33,8 +33,9 @@ namespace mu
 
 	bool ASTOpMeshApplyPose::IsEqual(const ASTOp& otherUntyped) const
 	{
-		if (auto other = dynamic_cast<const ASTOpMeshApplyPose*>(&otherUntyped))
+		if (otherUntyped.GetOpType()==GetOpType())
 		{
+			const ASTOpMeshApplyPose* other = static_cast<const ASTOpMeshApplyPose*>(&otherUntyped);
 			return base == other->base &&
 				pose == other->pose;
 		}
diff --git a/Engine/Plugins/Experimental/Mutable/Source/MutableTools/Private/MuT/ASTOpMeshApplyPose.h b/Engine/Plugins/Experimental/Mutable/Source/MutableTools/Private/MuT/ASTOpMeshApplyPose.h
--- a/Engine/Plugins/Experimental/Mutable/Source/MutableTools/Private/MuT/ASTOpMeshApplyPose.h
+++ b/Engine/Plugins/Experimental/Mutable/Source/MutableTools/Private/MuT/ASTOpMeshApplyPose.h
@@ -16,7 +16,7 @@ struct FProgram;
 	//---------------------------------------------------------------------------------------------
 	//!
 	//---------------------------------------------------------------------------------------------
-	class ASTOpMeshApplyPose : public ASTOp
+	class ASTOpMeshApplyPose final : public ASTOp
 	{
 	public:
 
diff --git a/Engine/Plugins/Experimental/Mutable/Source/MutableTools/Private/MuT/ASTOpMeshApplyShape.cpp b/Engine/Plugins/Experimental/Mutable/Source/MutableTools/Private/MuT/ASTOpMeshApplyShape.cpp
--- a/Engine/Plugins/Experimental/Mutable/Source/MutableTools/Private/MuT/ASTOpMeshApplyShape.cpp
+++ b/Engine/Plugins/Experimental/Mutable/Source/MutableTools/Private/MuT/ASTOpMeshApplyShape.cpp
@@ -36,8 +36,10 @@ namespace mu
 
 	bool ASTOpMeshApplyShape::IsEqual(const ASTOp& OtherUntyped) const
 	{
-		if (const ASTOpMeshApplyShape* Other = dynamic_cast<const ASTOpMeshApplyShape*>(&OtherUntyped))
+		if (OtherUntyped.GetOpType()==GetOpType())
 		{
+			const ASTOpMeshApplyShape* Other = static_cast<const ASTOpMeshApplyShape*>(&OtherUntyped);
+
 			const bool bSameFlags =
 				bReshapePhysicsVolumes == Other->bReshapePhysicsVolumes &&
 				bReshapeSkeleton == Other->bReshapeSkeleton &&
diff --git a/Engine/Plugins/Experimental/Mutable/Source/MutableTools/Private/MuT/ASTOpMeshApplyShape.h b/Engine/Plugins/Experimental/Mutable/Source/MutableTools/Private/MuT/ASTOpMeshApplyShape.h
--- a/Engine/Plugins/Experimental/Mutable/Source/MutableTools/Private/MuT/ASTOpMeshApplyShape.h
+++ b/Engine/Plugins/Experimental/Mutable/Source/MutableTools/Private/MuT/ASTOpMeshApplyShape.h
@@ -15,7 +15,7 @@ struct FProgram;
 	//---------------------------------------------------------------------------------------------
 	//!
 	//---------------------------------------------------------------------------------------------
-	class ASTOpMeshApplyShape : public ASTOp
+	class ASTOpMeshApplyShape final : public ASTOp
 	{
 	public:
 
diff --git a/Engine/Plugins/Experimental/Mutable/Source/MutableTools/Private/MuT/ASTOpMeshBindShape.cpp b/Engine/Plugins/Experimental/Mutable/Source/MutableTools/Private/MuT/ASTOpMeshBindShape.cpp
--- a/Engine/Plugins/Experimental/Mutable/Source/MutableTools/Private/MuT/ASTOpMeshBindShape.cpp
+++ b/Engine/Plugins/Experimental/Mutable/Source/MutableTools/Private/MuT/ASTOpMeshBindShape.cpp
@@ -41,8 +41,10 @@ namespace mu
 
 	bool ASTOpMeshBindShape::IsEqual(const ASTOp& OtherUntyped) const
 	{
-		if (const ASTOpMeshBindShape* Other = dynamic_cast<const ASTOpMeshBindShape*>(&OtherUntyped))
+		if (OtherUntyped.GetOpType()==GetOpType())
 		{
+			const ASTOpMeshBindShape* Other = static_cast<const ASTOpMeshBindShape*>(&OtherUntyped);
+
 			const bool bSameFlags =
 				bReshapeSkeleton == Other->bReshapeSkeleton	&&
 				bReshapePhysicsVolumes == Other->bReshapePhysicsVolumes &&
@@ -451,7 +453,7 @@ namespace mu
 			{
 				// Ignore the tags in the shape
 				Ptr<ASTOpMeshBindShape> NewBind = mu::Clone<ASTOpMeshBindShape>(this);
-				const ASTOpMeshAddTags* New = dynamic_cast<const ASTOpMeshAddTags*>(ShapeAt.get());
+				const ASTOpMeshAddTags* New = static_cast<const ASTOpMeshAddTags*>(ShapeAt.get());
 				NewBind->Shape = New->Source.child();
 				NewOp = NewBind;
 				break;
diff --git a/Engine/Plugins/Experimental/Mutable/Source/MutableTools/Private/MuT/ASTOpMeshBindShape.h b/Engine/Plugins/Experimental/Mutable/Source/MutableTools/Private/MuT/ASTOpMeshBindShape.h
--- a/Engine/Plugins/Experimental/Mutable/Source/MutableTools/Private/MuT/ASTOpMeshBindShape.h
+++ b/Engine/Plugins/Experimental/Mutable/Source/MutableTools/Private/MuT/ASTOpMeshBindShape.h
@@ -17,7 +17,7 @@ struct FProgram;
 	//---------------------------------------------------------------------------------------------
 	//!
 	//---------------------------------------------------------------------------------------------
-	class ASTOpMeshBindShape : public ASTOp
+	class ASTOpMeshBindShape final : public ASTOp
 	{
 	public:
 
diff --git a/Engine/Plugins/Experimental/Mutable/Source/MutableTools/Private/MuT/ASTOpMeshClipDeform.cpp b/Engine/Plugins/Experimental/Mutable/Source/MutableTools/Private/MuT/ASTOpMeshClipDeform.cpp
--- a/Engine/Plugins/Experimental/Mutable/Source/MutableTools/Private/MuT/ASTOpMeshClipDeform.cpp
+++ b/Engine/Plugins/Experimental/Mutable/Source/MutableTools/Private/MuT/ASTOpMeshClipDeform.cpp
@@ -32,8 +32,9 @@ namespace mu
 
 	bool ASTOpMeshClipDeform::IsEqual(const ASTOp& OtherUntyped) const
 	{
-		if (const ASTOpMeshClipDeform* Other = dynamic_cast<const ASTOpMeshClipDeform*>(&OtherUntyped))
+		if (OtherUntyped.GetOpType()==GetOpType())
 		{
+			const ASTOpMeshClipDeform* Other = static_cast<const ASTOpMeshClipDeform*>(&OtherUntyped);
 			return Mesh == Other->Mesh && ClipShape == Other->ClipShape;
 		}
 
diff --git a/Engine/Plugins/Experimental/Mutable/Source/MutableTools/Private/MuT/ASTOpMeshClipDeform.h b/Engine/Plugins/Experimental/Mutable/Source/MutableTools/Private/MuT/ASTOpMeshClipDeform.h
--- a/Engine/Plugins/Experimental/Mutable/Source/MutableTools/Private/MuT/ASTOpMeshClipDeform.h
+++ b/Engine/Plugins/Experimental/Mutable/Source/MutableTools/Private/MuT/ASTOpMeshClipDeform.h
@@ -16,7 +16,7 @@ struct FProgram;
 	//---------------------------------------------------------------------------------------------
 	//!
 	//---------------------------------------------------------------------------------------------
-	class ASTOpMeshClipDeform : public ASTOp
+	class ASTOpMeshClipDeform final : public ASTOp
 	{
 	public:
 
diff --git a/Engine/Plugins/Experimental/Mutable/Source/MutableTools/Private/MuT/ASTOpMeshClipMorphPlane.cpp b/Engine/Plugins/Experimental/Mutable/Source/MutableTools/Private/MuT/ASTOpMeshClipMorphPlane.cpp
--- a/Engine/Plugins/Experimental/Mutable/Source/MutableTools/Private/MuT/ASTOpMeshClipMorphPlane.cpp
+++ b/Engine/Plugins/Experimental/Mutable/Source/MutableTools/Private/MuT/ASTOpMeshClipMorphPlane.cpp
@@ -32,8 +32,10 @@ namespace mu
 
 	bool ASTOpMeshClipMorphPlane::IsEqual(const ASTOp& otherUntyped) const
 	{
-		if (auto other = dynamic_cast<const ASTOpMeshClipMorphPlane*>(&otherUntyped))
+		if (otherUntyped.GetOpType()==GetOpType())
 		{
+			const ASTOpMeshClipMorphPlane* other = static_cast<const ASTOpMeshClipMorphPlane*>(&otherUntyped);
+
 			return source == other->source &&
 				morphShape == other->morphShape &&
 				selectionShape == other->selectionShape &&
diff --git a/Engine/Plugins/Experimental/Mutable/Source/MutableTools/Private/MuT/ASTOpMeshClipMorphPlane.h b/Engine/Plugins/Experimental/Mutable/Source/MutableTools/Private/MuT/ASTOpMeshClipMorphPlane.h
--- a/Engine/Plugins/Experimental/Mutable/Source/MutableTools/Private/MuT/ASTOpMeshClipMorphPlane.h
+++ b/Engine/Plugins/Experimental/Mutable/Source/MutableTools/Private/MuT/ASTOpMeshClipMorphPlane.h
@@ -16,7 +16,7 @@ struct FProgram;
 	//---------------------------------------------------------------------------------------------
 	//!
 	//---------------------------------------------------------------------------------------------
-	class ASTOpMeshClipMorphPlane : public ASTOp
+	class ASTOpMeshClipMorphPlane final : public ASTOp
 	{
 	public:
 
diff --git a/Engine/Plugins/Experimental/Mutable/Source/MutableTools/Private/MuT/ASTOpMeshDifference.cpp b/Engine/Plugins/Experimental/Mutable/Source/MutableTools/Private/MuT/ASTOpMeshDifference.cpp
--- a/Engine/Plugins/Experimental/Mutable/Source/MutableTools/Private/MuT/ASTOpMeshDifference.cpp
+++ b/Engine/Plugins/Experimental/Mutable/Source/MutableTools/Private/MuT/ASTOpMeshDifference.cpp
@@ -36,8 +36,10 @@ namespace mu
 	//---------------------------------------------------------------------------------------------
 	bool ASTOpMeshDifference::IsEqual(const ASTOp& otherUntyped) const
 	{
-		if (const ASTOpMeshDifference* Other = dynamic_cast<const ASTOpMeshDifference*>(&otherUntyped))
+		if (otherUntyped.GetOpType() == GetOpType())
 		{
+			const ASTOpMeshDifference* Other = static_cast<const ASTOpMeshDifference*>(&otherUntyped);
+
 			return Base == Other->Base && Target == Other->Target
 				&& bIgnoreTextureCoords == Other->bIgnoreTextureCoords
 				&& Channels == Other->Channels;
@@ -344,7 +346,7 @@ namespace mu
 			case OP_TYPE::ME_ADDTAGS:
 			{
 				// Ignore tags in this branch
-				const ASTOpMeshAddTags* Add = dynamic_cast<const ASTOpMeshAddTags*>(TargetAt.get());
+				const ASTOpMeshAddTags* Add = static_cast<const ASTOpMeshAddTags*>(TargetAt.get());
 
 				Ptr<ASTOpMeshDifference> NewDiff = mu::Clone<ASTOpMeshDifference>(this);
 				NewDiff->Target = Add->Source.child();
diff --git a/Engine/Plugins/Experimental/Mutable/Source/MutableTools/Private/MuT/ASTOpMeshDifference.h b/Engine/Plugins/Experimental/Mutable/Source/MutableTools/Private/MuT/ASTOpMeshDifference.h
--- a/Engine/Plugins/Experimental/Mutable/Source/MutableTools/Private/MuT/ASTOpMeshDifference.h
+++ b/Engine/Plugins/Experimental/Mutable/Source/MutableTools/Private/MuT/ASTOpMeshDifference.h
@@ -13,7 +13,7 @@ namespace mu
 	struct FProgram;
 
 	/** Calculate the difference between two meshes (usually to calculate a morph delta). */
-	class ASTOpMeshDifference : public ASTOp
+	class ASTOpMeshDifference final : public ASTOp
 	{
 	public:
 
diff --git a/Engine/Plugins/Experimental/Mutable/Source/MutableTools/Private/MuT/ASTOpMeshExtractLayoutBlocks.cpp b/Engine/Plugins/Experimental/Mutable/Source/MutableTools/Private/MuT/ASTOpMeshExtractLayoutBlocks.cpp
--- a/Engine/Plugins/Experimental/Mutable/Source/MutableTools/Private/MuT/ASTOpMeshExtractLayoutBlocks.cpp
+++ b/Engine/Plugins/Experimental/Mutable/Source/MutableTools/Private/MuT/ASTOpMeshExtractLayoutBlocks.cpp
@@ -34,8 +34,9 @@ namespace mu
 
 	bool ASTOpMeshExtractLayoutBlocks::IsEqual(const ASTOp& otherUntyped) const
 	{
-		if (auto other = dynamic_cast<const ASTOpMeshExtractLayoutBlocks*>(&otherUntyped))
+		if (otherUntyped.GetOpType() == GetOpType())
 		{
+			const ASTOpMeshExtractLayoutBlocks* other = static_cast<const ASTOpMeshExtractLayoutBlocks*>(&otherUntyped);
 			return Source == other->Source && Layout == other->Layout && Blocks == other->Blocks;
 		}
 		return false;
diff --git a/Engine/Plugins/Experimental/Mutable/Source/MutableTools/Private/MuT/ASTOpMeshExtractLayoutBlocks.h b/Engine/Plugins/Experimental/Mutable/Source/MutableTools/Private/MuT/ASTOpMeshExtractLayoutBlocks.h
--- a/Engine/Plugins/Experimental/Mutable/Source/MutableTools/Private/MuT/ASTOpMeshExtractLayoutBlocks.h
+++ b/Engine/Plugins/Experimental/Mutable/Source/MutableTools/Private/MuT/ASTOpMeshExtractLayoutBlocks.h
@@ -16,7 +16,7 @@ namespace mu
 	//---------------------------------------------------------------------------------------------
 	//! Variable sized mesh block extract operation.
 	//---------------------------------------------------------------------------------------------
-	class ASTOpMeshExtractLayoutBlocks : public ASTOp
+	class ASTOpMeshExtractLayoutBlocks final : public ASTOp
 	{
 	public:
 
diff --git a/Engine/Plugins/Experimental/Mutable/Source/MutableTools/Private/MuT/ASTOpMeshFormat.cpp b/Engine/Plugins/Experimental/Mutable/Source/MutableTools/Private/MuT/ASTOpMeshFormat.cpp
--- a/Engine/Plugins/Experimental/Mutable/Source/MutableTools/Private/MuT/ASTOpMeshFormat.cpp
+++ b/Engine/Plugins/Experimental/Mutable/Source/MutableTools/Private/MuT/ASTOpMeshFormat.cpp
@@ -46,8 +46,9 @@ ASTOpMeshFormat::~ASTOpMeshFormat()
 
 bool ASTOpMeshFormat::IsEqual(const ASTOp& otherUntyped) const
 {
-    if (const ASTOpMeshFormat* other = dynamic_cast<const ASTOpMeshFormat*>(&otherUntyped) )
-    {
+	if (otherUntyped.GetOpType() == GetOpType())
+	{
+		const ASTOpMeshFormat* other = static_cast<const ASTOpMeshFormat*>(&otherUntyped);
         return Source==other->Source && Format==other->Format && Buffers==other->Buffers;
     }
     return false;
@@ -111,10 +112,9 @@ mu::Ptr<ASTOp> ASTOpMeshFormat::OptimiseSink(const FModelOptimizationOptions& op
 //---------------------------------------------------------------------------------------------
 //---------------------------------------------------------------------------------------------
 //---------------------------------------------------------------------------------------------
-mu::Ptr<ASTOp> Sink_MeshFormatAST::Apply(const ASTOp* root)
+mu::Ptr<ASTOp> Sink_MeshFormatAST::Apply(const ASTOpMeshFormat* root)
 {
-	m_root = dynamic_cast<const ASTOpMeshFormat*>(root);
-	check(m_root);
+	m_root = root;
 
 	OldToNew.Empty();
 
@@ -141,7 +141,7 @@ namespace
 		{
 		case OP_TYPE::ME_CONSTANT:
 		{
-			const ASTOpConstantResource* typed = dynamic_cast<const ASTOpConstantResource*>(at.get());
+			const ASTOpConstantResource* typed = static_cast<const ASTOpConstantResource*>(at.get());
 			res = static_cast<const Mesh*>(typed->GetValue().get());
 			break;
 		}
diff --git a/Engine/Plugins/Experimental/Mutable/Source/MutableTools/Private/MuT/ASTOpMeshFormat.h b/Engine/Plugins/Experimental/Mutable/Source/MutableTools/Private/MuT/ASTOpMeshFormat.h
--- a/Engine/Plugins/Experimental/Mutable/Source/MutableTools/Private/MuT/ASTOpMeshFormat.h
+++ b/Engine/Plugins/Experimental/Mutable/Source/MutableTools/Private/MuT/ASTOpMeshFormat.h
@@ -12,7 +12,7 @@ namespace mu
 {
 	struct FProgram;
 	
-	class ASTOpMeshFormat : public ASTOp
+	class ASTOpMeshFormat final : public ASTOp
 	{
 	public:
 
diff --git a/Engine/Plugins/Experimental/Mutable/Source/MutableTools/Private/MuT/ASTOpMeshGeometryOperation.cpp b/Engine/Plugins/Experimental/Mutable/Source/MutableTools/Private/MuT/ASTOpMeshGeometryOperation.cpp
--- a/Engine/Plugins/Experimental/Mutable/Source/MutableTools/Private/MuT/ASTOpMeshGeometryOperation.cpp
+++ b/Engine/Plugins/Experimental/Mutable/Source/MutableTools/Private/MuT/ASTOpMeshGeometryOperation.cpp
@@ -31,8 +31,9 @@ namespace mu
 
 	bool ASTOpMeshGeometryOperation::IsEqual(const ASTOp& otherUntyped) const
 	{
-		if (auto other = dynamic_cast<const ASTOpMeshGeometryOperation*>(&otherUntyped))
+		if (otherUntyped.GetOpType() == GetOpType())
 		{
+			const ASTOpMeshGeometryOperation* other = static_cast<const ASTOpMeshGeometryOperation*>(&otherUntyped);
 			return meshA == other->meshA && meshB == other->meshB &&
 				scalarA == other->scalarA && scalarB == other->scalarB;
 		}
diff --git a/Engine/Plugins/Experimental/Mutable/Source/MutableTools/Private/MuT/ASTOpMeshGeometryOperation.h b/Engine/Plugins/Experimental/Mutable/Source/MutableTools/Private/MuT/ASTOpMeshGeometryOperation.h
--- a/Engine/Plugins/Experimental/Mutable/Source/MutableTools/Private/MuT/ASTOpMeshGeometryOperation.h
+++ b/Engine/Plugins/Experimental/Mutable/Source/MutableTools/Private/MuT/ASTOpMeshGeometryOperation.h
@@ -10,13 +10,13 @@
 
 namespace mu
 {
-struct FProgram;
+	struct FProgram;
 
 
 	//---------------------------------------------------------------------------------------------
 	//!
 	//---------------------------------------------------------------------------------------------
-	class ASTOpMeshGeometryOperation : public ASTOp
+	class ASTOpMeshGeometryOperation final : public ASTOp
 	{
 	public:
 
diff --git a/Engine/Plugins/Experimental/Mutable/Source/MutableTools/Private/MuT/ASTOpMeshMaskClipMesh.cpp b/Engine/Plugins/Experimental/Mutable/Source/MutableTools/Private/MuT/ASTOpMeshMaskClipMesh.cpp
--- a/Engine/Plugins/Experimental/Mutable/Source/MutableTools/Private/MuT/ASTOpMeshMaskClipMesh.cpp
+++ b/Engine/Plugins/Experimental/Mutable/Source/MutableTools/Private/MuT/ASTOpMeshMaskClipMesh.cpp
@@ -38,8 +38,9 @@ namespace mu
 	//-------------------------------------------------------------------------------------------------
 	bool ASTOpMeshMaskClipMesh::IsEqual(const ASTOp& otherUntyped) const
 	{
-		if (const ASTOpMeshMaskClipMesh* other = dynamic_cast<const ASTOpMeshMaskClipMesh*>(&otherUntyped))
+		if (otherUntyped.GetOpType() == GetOpType())
 		{
+			const ASTOpMeshMaskClipMesh* other = static_cast<const ASTOpMeshMaskClipMesh*>(&otherUntyped);
 			return source == other->source && clip == other->clip;
 		}
 		return false;
@@ -165,7 +166,7 @@ namespace mu
 					// faces removed by the ignored removemask, but it is ok
 
 					// TODO: Swap instead of ignore, and implement removemask on a mask?
-					const ASTOpMeshRemoveMask* typedAt = dynamic_cast<const ASTOpMeshRemoveMask*>(at.get());
+					const ASTOpMeshRemoveMask* typedAt = static_cast<const ASTOpMeshRemoveMask*>(at.get());
 					newAt = Visit(typedAt->source.child());
 					break;
 				}
@@ -316,7 +317,7 @@ namespace mu
 				case OP_TYPE::ME_ADDTAGS:
 				{
 					// Ignore tags in the clip mesh branch
-					const ASTOpMeshAddTags* AddOp = dynamic_cast<const ASTOpMeshAddTags*>(at.get());
+					const ASTOpMeshAddTags* AddOp = static_cast<const ASTOpMeshAddTags*>(at.get());
 					newAt = Visit(AddOp->Source.child());
 					break;
 				}
diff --git a/Engine/Plugins/Experimental/Mutable/Source/MutableTools/Private/MuT/ASTOpMeshMaskClipMesh.h b/Engine/Plugins/Experimental/Mutable/Source/MutableTools/Private/MuT/ASTOpMeshMaskClipMesh.h
--- a/Engine/Plugins/Experimental/Mutable/Source/MutableTools/Private/MuT/ASTOpMeshMaskClipMesh.h
+++ b/Engine/Plugins/Experimental/Mutable/Source/MutableTools/Private/MuT/ASTOpMeshMaskClipMesh.h
@@ -16,7 +16,7 @@ struct FProgram;
 	//---------------------------------------------------------------------------------------------
 	//!
 	//---------------------------------------------------------------------------------------------
-	class ASTOpMeshMaskClipMesh : public ASTOp
+	class ASTOpMeshMaskClipMesh final : public ASTOp
 	{
 	public:
 
diff --git a/Engine/Plugins/Experimental/Mutable/Source/MutableTools/Private/MuT/ASTOpMeshMaskClipUVMask.cpp b/Engine/Plugins/Experimental/Mutable/Source/MutableTools/Private/MuT/ASTOpMeshMaskClipUVMask.cpp
--- a/Engine/Plugins/Experimental/Mutable/Source/MutableTools/Private/MuT/ASTOpMeshMaskClipUVMask.cpp
+++ b/Engine/Plugins/Experimental/Mutable/Source/MutableTools/Private/MuT/ASTOpMeshMaskClipUVMask.cpp
@@ -38,8 +38,9 @@ namespace mu
 	//-------------------------------------------------------------------------------------------------
 	bool ASTOpMeshMaskClipUVMask::IsEqual(const ASTOp& otherUntyped) const
 	{
-		if (const ASTOpMeshMaskClipUVMask* other = dynamic_cast<const ASTOpMeshMaskClipUVMask*>(&otherUntyped))
+		if (otherUntyped.GetOpType() == GetOpType())
 		{
+			const ASTOpMeshMaskClipUVMask* other = static_cast<const ASTOpMeshMaskClipUVMask*>(&otherUntyped);
 			return Source == other->Source && Mask == other->Mask && LayoutIndex==other->LayoutIndex;
 		}
 		return false;
@@ -166,7 +167,7 @@ namespace mu
 					// faces removed by the ignored removemask, but it is ok
 
 					// TODO: Swap instead of ignore, and implement removemask on a mask?
-					const ASTOpMeshRemoveMask* typedAt = dynamic_cast<const ASTOpMeshRemoveMask*>(at.get());
+					const ASTOpMeshRemoveMask* typedAt = static_cast<const ASTOpMeshRemoveMask*>(at.get());
 					newAt = Visit(typedAt->source.child());
 					break;
 				}
diff --git a/Engine/Plugins/Experimental/Mutable/Source/MutableTools/Private/MuT/ASTOpMeshMaskClipUVMask.h b/Engine/Plugins/Experimental/Mutable/Source/MutableTools/Private/MuT/ASTOpMeshMaskClipUVMask.h
--- a/Engine/Plugins/Experimental/Mutable/Source/MutableTools/Private/MuT/ASTOpMeshMaskClipUVMask.h
+++ b/Engine/Plugins/Experimental/Mutable/Source/MutableTools/Private/MuT/ASTOpMeshMaskClipUVMask.h
@@ -16,7 +16,7 @@ struct FProgram;
 	//---------------------------------------------------------------------------------------------
 	//!
 	//---------------------------------------------------------------------------------------------
-	class ASTOpMeshMaskClipUVMask : public ASTOp
+	class ASTOpMeshMaskClipUVMask final : public ASTOp
 	{
 	public:
 
diff --git a/Engine/Plugins/Experimental/Mutable/Source/MutableTools/Private/MuT/ASTOpMeshMorph.cpp b/Engine/Plugins/Experimental/Mutable/Source/MutableTools/Private/MuT/ASTOpMeshMorph.cpp
--- a/Engine/Plugins/Experimental/Mutable/Source/MutableTools/Private/MuT/ASTOpMeshMorph.cpp
+++ b/Engine/Plugins/Experimental/Mutable/Source/MutableTools/Private/MuT/ASTOpMeshMorph.cpp
@@ -31,8 +31,9 @@ namespace mu
 	//---------------------------------------------------------------------------------------------
 	bool ASTOpMeshMorph::IsEqual(const ASTOp& otherUntyped) const
 	{
-		if (const ASTOpMeshMorph* other = dynamic_cast<const ASTOpMeshMorph*>(&otherUntyped))
+		if (otherUntyped.GetOpType() == GetOpType())
 		{
+			const ASTOpMeshMorph* other = static_cast<const ASTOpMeshMorph*>(&otherUntyped);
 			return Factor == other->Factor && Base == other->Base && Target == other->Target;
 		}
 		return false;
@@ -139,7 +140,7 @@ namespace mu
 			case OP_TYPE::ME_ADDTAGS:
 			{
 				// Ignore the morph target tags
-				const ASTOpMeshAddTags* AddTags = dynamic_cast<const ASTOpMeshAddTags*>(Target.child().get());
+				const ASTOpMeshAddTags* AddTags = static_cast<const ASTOpMeshAddTags*>(Target.child().get());
 
 				Ptr<ASTOpMeshMorph> New = mu::Clone<ASTOpMeshMorph>(this);
 				New->Target = AddTags->Source.child();
diff --git a/Engine/Plugins/Experimental/Mutable/Source/MutableTools/Private/MuT/ASTOpMeshMorph.h b/Engine/Plugins/Experimental/Mutable/Source/MutableTools/Private/MuT/ASTOpMeshMorph.h
--- a/Engine/Plugins/Experimental/Mutable/Source/MutableTools/Private/MuT/ASTOpMeshMorph.h
+++ b/Engine/Plugins/Experimental/Mutable/Source/MutableTools/Private/MuT/ASTOpMeshMorph.h
@@ -15,7 +15,7 @@ namespace mu
 	//---------------------------------------------------------------------------------------------
 	//! From a source mesh, remove a list of fragments with a condition.
 	//---------------------------------------------------------------------------------------------
-	class ASTOpMeshMorph : public ASTOp
+	class ASTOpMeshMorph final : public ASTOp
 	{
 	public:
 
diff --git a/Engine/Plugins/Experimental/Mutable/Source/MutableTools/Private/MuT/ASTOpMeshMorphReshape.cpp b/Engine/Plugins/Experimental/Mutable/Source/MutableTools/Private/MuT/ASTOpMeshMorphReshape.cpp
--- a/Engine/Plugins/Experimental/Mutable/Source/MutableTools/Private/MuT/ASTOpMeshMorphReshape.cpp
+++ b/Engine/Plugins/Experimental/Mutable/Source/MutableTools/Private/MuT/ASTOpMeshMorphReshape.cpp
@@ -32,8 +32,9 @@ namespace mu
 
 	bool ASTOpMeshMorphReshape::IsEqual(const ASTOp& otherUntyped) const
 	{
-		if (auto other = dynamic_cast<const ASTOpMeshMorphReshape*>(&otherUntyped))
+		if (otherUntyped.GetOpType() == GetOpType())
 		{
+			const ASTOpMeshMorphReshape* other = static_cast<const ASTOpMeshMorphReshape*>(&otherUntyped);
 			return Morph == other->Morph && Reshape == other->Reshape;
 		}
 
diff --git a/Engine/Plugins/Experimental/Mutable/Source/MutableTools/Private/MuT/ASTOpMeshMorphReshape.h b/Engine/Plugins/Experimental/Mutable/Source/MutableTools/Private/MuT/ASTOpMeshMorphReshape.h
--- a/Engine/Plugins/Experimental/Mutable/Source/MutableTools/Private/MuT/ASTOpMeshMorphReshape.h
+++ b/Engine/Plugins/Experimental/Mutable/Source/MutableTools/Private/MuT/ASTOpMeshMorphReshape.h
@@ -15,7 +15,7 @@ struct FProgram;
 	//---------------------------------------------------------------------------------------------
 	//!
 	//---------------------------------------------------------------------------------------------
-	class ASTOpMeshMorphReshape : public ASTOp
+	class ASTOpMeshMorphReshape final : public ASTOp
 	{
 	public:
 
diff --git a/Engine/Plugins/Experimental/Mutable/Source/MutableTools/Private/MuT/ASTOpMeshOptimizeSkinning.cpp b/Engine/Plugins/Experimental/Mutable/Source/MutableTools/Private/MuT/ASTOpMeshOptimizeSkinning.cpp
--- a/Engine/Plugins/Experimental/Mutable/Source/MutableTools/Private/MuT/ASTOpMeshOptimizeSkinning.cpp
+++ b/Engine/Plugins/Experimental/Mutable/Source/MutableTools/Private/MuT/ASTOpMeshOptimizeSkinning.cpp
@@ -35,8 +35,9 @@ namespace mu
 	//-------------------------------------------------------------------------------------------------
 	bool ASTOpMeshOptimizeSkinning::IsEqual(const ASTOp& otherUntyped) const
 	{
-		if (auto other = dynamic_cast<const ASTOpMeshOptimizeSkinning*>(&otherUntyped))
+		if (otherUntyped.GetOpType() == GetOpType())
 		{
+			const ASTOpMeshOptimizeSkinning* other = static_cast<const ASTOpMeshOptimizeSkinning*>(&otherUntyped);
 			return source == other->source;
 		}
 		return false;
diff --git a/Engine/Plugins/Experimental/Mutable/Source/MutableTools/Private/MuT/ASTOpMeshOptimizeSkinning.h b/Engine/Plugins/Experimental/Mutable/Source/MutableTools/Private/MuT/ASTOpMeshOptimizeSkinning.h
--- a/Engine/Plugins/Experimental/Mutable/Source/MutableTools/Private/MuT/ASTOpMeshOptimizeSkinning.h
+++ b/Engine/Plugins/Experimental/Mutable/Source/MutableTools/Private/MuT/ASTOpMeshOptimizeSkinning.h
@@ -13,7 +13,7 @@ namespace mu
 	struct FProgram;
 
 	//---------------------------------------------------------------------------------------------
-	class ASTOpMeshOptimizeSkinning : public ASTOp
+	class ASTOpMeshOptimizeSkinning final : public ASTOp
 	{
 	public:
 
diff --git a/Engine/Plugins/Experimental/Mutable/Source/MutableTools/Private/MuT/ASTOpMeshRemoveMask.cpp b/Engine/Plugins/Experimental/Mutable/Source/MutableTools/Private/MuT/ASTOpMeshRemoveMask.cpp
--- a/Engine/Plugins/Experimental/Mutable/Source/MutableTools/Private/MuT/ASTOpMeshRemoveMask.cpp
+++ b/Engine/Plugins/Experimental/Mutable/Source/MutableTools/Private/MuT/ASTOpMeshRemoveMask.cpp
@@ -40,8 +40,9 @@ namespace mu
 	//---------------------------------------------------------------------------------------------
 	bool ASTOpMeshRemoveMask::IsEqual(const ASTOp& otherUntyped) const
 	{
-		if (const ASTOpMeshRemoveMask* other = dynamic_cast<const ASTOpMeshRemoveMask*>(&otherUntyped))
+		if (otherUntyped.GetOpType() == GetOpType())
 		{
+			const ASTOpMeshRemoveMask* other = static_cast<const ASTOpMeshRemoveMask*>(&otherUntyped);
 			return source == other->source && removes == other->removes;
 		}
 		return false;
@@ -256,7 +257,7 @@ namespace mu
 						NewOp = mu::Clone<ASTOpMeshRemoveMask>(this);
 					}
 
-					const ASTOpMeshAddTags* Add = dynamic_cast<const ASTOpMeshAddTags*>(removes[RemoveIndex].Value.child().get());
+					const ASTOpMeshAddTags* Add = static_cast<const ASTOpMeshAddTags*>(removes[RemoveIndex].Value.child().get());
 					NewOp->removes[RemoveIndex].Value = Add->Source.child();
 
 					break;
diff --git a/Engine/Plugins/Experimental/Mutable/Source/MutableTools/Private/MuT/ASTOpMeshRemoveMask.h b/Engine/Plugins/Experimental/Mutable/Source/MutableTools/Private/MuT/ASTOpMeshRemoveMask.h
--- a/Engine/Plugins/Experimental/Mutable/Source/MutableTools/Private/MuT/ASTOpMeshRemoveMask.h
+++ b/Engine/Plugins/Experimental/Mutable/Source/MutableTools/Private/MuT/ASTOpMeshRemoveMask.h
@@ -15,7 +15,7 @@ namespace mu
 	//---------------------------------------------------------------------------------------------
 	//! From a source mesh, remove a list of fragments with a condition.
 	//---------------------------------------------------------------------------------------------
-	class ASTOpMeshRemoveMask : public ASTOp
+	class ASTOpMeshRemoveMask final : public ASTOp
 	{
 	public:
 
diff --git a/Engine/Plugins/Experimental/Mutable/Source/MutableTools/Private/MuT/ASTOpMeshTransform.cpp b/Engine/Plugins/Experimental/Mutable/Source/MutableTools/Private/MuT/ASTOpMeshTransform.cpp
--- a/Engine/Plugins/Experimental/Mutable/Source/MutableTools/Private/MuT/ASTOpMeshTransform.cpp
+++ b/Engine/Plugins/Experimental/Mutable/Source/MutableTools/Private/MuT/ASTOpMeshTransform.cpp
@@ -30,8 +30,9 @@ namespace mu
 
 	bool ASTOpMeshTransform::IsEqual(const ASTOp& otherUntyped) const
 	{
-		if (auto other = dynamic_cast<const ASTOpMeshTransform*>(&otherUntyped))
+		if (otherUntyped.GetOpType() == GetOpType())
 		{
+			const ASTOpMeshTransform* other = static_cast<const ASTOpMeshTransform*>(&otherUntyped);
 			return source == other->source &&
 				matrix == other->matrix;
 		}
diff --git a/Engine/Plugins/Experimental/Mutable/Source/MutableTools/Private/MuT/ASTOpMeshTransform.h b/Engine/Plugins/Experimental/Mutable/Source/MutableTools/Private/MuT/ASTOpMeshTransform.h
--- a/Engine/Plugins/Experimental/Mutable/Source/MutableTools/Private/MuT/ASTOpMeshTransform.h
+++ b/Engine/Plugins/Experimental/Mutable/Source/MutableTools/Private/MuT/ASTOpMeshTransform.h
@@ -17,7 +17,7 @@ struct FProgram;
 	//---------------------------------------------------------------------------------------------
 	//!
 	//---------------------------------------------------------------------------------------------
-	class ASTOpMeshTransform : public ASTOp
+	class ASTOpMeshTransform final : public ASTOp
 	{
 	public:
 
diff --git a/Engine/Plugins/Experimental/Mutable/Source/MutableTools/Private/MuT/ASTOpParameter.cpp b/Engine/Plugins/Experimental/Mutable/Source/MutableTools/Private/MuT/ASTOpParameter.cpp
--- a/Engine/Plugins/Experimental/Mutable/Source/MutableTools/Private/MuT/ASTOpParameter.cpp
+++ b/Engine/Plugins/Experimental/Mutable/Source/MutableTools/Private/MuT/ASTOpParameter.cpp
@@ -36,8 +36,9 @@ namespace mu
 
 	bool ASTOpParameter::IsEqual(const ASTOp& otherUntyped) const
 	{
-		if (auto other = dynamic_cast<const ASTOpParameter*>(&otherUntyped))
+		if (otherUntyped.GetOpType() == GetOpType())
 		{
+			const ASTOpParameter* other = static_cast<const ASTOpParameter*>(&otherUntyped);
 			return type == other->type &&
 				parameter == other->parameter &&
 				ranges == other->ranges;
@@ -113,7 +114,7 @@ namespace mu
 		{
 			if (f->GetOpType() == OP_TYPE::BO_EQUAL_INT_CONST)
 			{
-				const auto typedFact = dynamic_cast<const ASTOpFixed*>(f.get());
+				const auto typedFact = static_cast<const ASTOpFixed*>(f.get());
 				auto value = typedFact->children[typedFact->op.args.BoolEqualScalarConst.value].child();
 				if (value.get() == this)
 				{
diff --git a/Engine/Plugins/Experimental/Mutable/Source/MutableTools/Private/MuT/ASTOpParameter.h b/Engine/Plugins/Experimental/Mutable/Source/MutableTools/Private/MuT/ASTOpParameter.h
--- a/Engine/Plugins/Experimental/Mutable/Source/MutableTools/Private/MuT/ASTOpParameter.h
+++ b/Engine/Plugins/Experimental/Mutable/Source/MutableTools/Private/MuT/ASTOpParameter.h
@@ -18,7 +18,7 @@ namespace mu
 	//---------------------------------------------------------------------------------------------
 	//! Parameter operation.
 	//---------------------------------------------------------------------------------------------
-	class ASTOpParameter : public ASTOp
+	class ASTOpParameter final : public ASTOp
 	{
 	public:
 
diff --git a/Engine/Plugins/Experimental/Mutable/Source/MutableTools/Private/MuT/ASTOpReferenceResource.cpp b/Engine/Plugins/Experimental/Mutable/Source/MutableTools/Private/MuT/ASTOpReferenceResource.cpp
--- a/Engine/Plugins/Experimental/Mutable/Source/MutableTools/Private/MuT/ASTOpReferenceResource.cpp
+++ b/Engine/Plugins/Experimental/Mutable/Source/MutableTools/Private/MuT/ASTOpReferenceResource.cpp
@@ -21,8 +21,9 @@ namespace mu
 	//-------------------------------------------------------------------------------------------------
 	bool ASTOpReferenceResource::IsEqual(const ASTOp& otherUntyped) const
 	{
-		if (const ASTOpReferenceResource* other = dynamic_cast<const ASTOpReferenceResource*>(&otherUntyped))
+		if (otherUntyped.GetOpType() == GetOpType())
 		{
+			const ASTOpReferenceResource* other = static_cast<const ASTOpReferenceResource*>(&otherUntyped);
 			return type == other->type && ID == other->ID && bForceLoad == other->bForceLoad && ImageDesc == other->ImageDesc;
 		}
 		return false;
diff --git a/Engine/Plugins/Experimental/Mutable/Source/MutableTools/Private/MuT/ASTOpReferenceResource.h b/Engine/Plugins/Experimental/Mutable/Source/MutableTools/Private/MuT/ASTOpReferenceResource.h
--- a/Engine/Plugins/Experimental/Mutable/Source/MutableTools/Private/MuT/ASTOpReferenceResource.h
+++ b/Engine/Plugins/Experimental/Mutable/Source/MutableTools/Private/MuT/ASTOpReferenceResource.h
@@ -19,7 +19,7 @@ namespace mu
 	//---------------------------------------------------------------------------------------------
 	//! A reference to an engine image (or other resources in the future)
 	//---------------------------------------------------------------------------------------------
-	class ASTOpReferenceResource : public ASTOp
+	class ASTOpReferenceResource final : public ASTOp
 	{
 	public:
 
diff --git a/Engine/Plugins/Experimental/Mutable/Source/MutableTools/Private/MuT/ASTOpScalarCurve.cpp b/Engine/Plugins/Experimental/Mutable/Source/MutableTools/Private/MuT/ASTOpScalarCurve.cpp
--- a/Engine/Plugins/Experimental/Mutable/Source/MutableTools/Private/MuT/ASTOpScalarCurve.cpp
+++ b/Engine/Plugins/Experimental/Mutable/Source/MutableTools/Private/MuT/ASTOpScalarCurve.cpp
@@ -45,8 +45,9 @@ namespace mu
 
 	bool ASTOpScalarCurve::IsEqual(const ASTOp& otherUntyped) const
 	{
-		if (auto other = dynamic_cast<const ASTOpScalarCurve*>(&otherUntyped))
+		if (otherUntyped.GetOpType() == GetOpType())
 		{
+			const ASTOpScalarCurve* other = static_cast<const ASTOpScalarCurve*>(&otherUntyped);
 			return time == other->time && curve == other->curve;
 		}
 		return false;
diff --git a/Engine/Plugins/Experimental/Mutable/Source/MutableTools/Private/MuT/ASTOpScalarCurve.h b/Engine/Plugins/Experimental/Mutable/Source/MutableTools/Private/MuT/ASTOpScalarCurve.h
--- a/Engine/Plugins/Experimental/Mutable/Source/MutableTools/Private/MuT/ASTOpScalarCurve.h
+++ b/Engine/Plugins/Experimental/Mutable/Source/MutableTools/Private/MuT/ASTOpScalarCurve.h
@@ -17,7 +17,7 @@ struct FProgram;
 	//---------------------------------------------------------------------------------------------
 	//!
 	//---------------------------------------------------------------------------------------------
-	class ASTOpScalarCurve : public ASTOp
+	class ASTOpScalarCurve final : public ASTOp
 	{
 	public:
 
diff --git a/Engine/Plugins/Experimental/Mutable/Source/MutableTools/Private/MuT/ASTOpSwitch.cpp b/Engine/Plugins/Experimental/Mutable/Source/MutableTools/Private/MuT/ASTOpSwitch.cpp
--- a/Engine/Plugins/Experimental/Mutable/Source/MutableTools/Private/MuT/ASTOpSwitch.cpp
+++ b/Engine/Plugins/Experimental/Mutable/Source/MutableTools/Private/MuT/ASTOpSwitch.cpp
@@ -39,8 +39,9 @@ namespace mu
 	//-------------------------------------------------------------------------------------------------
 	bool ASTOpSwitch::IsEqual(const ASTOp& otherUntyped) const
 	{
-		if (auto other = dynamic_cast<const ASTOpSwitch*>(&otherUntyped))
+		if (GetOpType()==otherUntyped.GetOpType())
 		{
+			const ASTOpSwitch* other = static_cast<const ASTOpSwitch*>(&otherUntyped);
 			return type == other->type && variable == other->variable &&
 				cases == other->cases && def == other->def;
 		}
@@ -419,7 +420,7 @@ namespace mu
 		{
 			Ptr<ASTOp> Branch = def.child();
 
-			auto typedCondition = dynamic_cast<const ASTOpFixed*>(variable.child().get());
+			auto typedCondition = static_cast<const ASTOpFixed*>(variable.child().get());
 			for (int32 o = 0; o < cases.Num(); ++o)
 			{
 				if (cases[o].branch &&
@@ -438,7 +439,7 @@ namespace mu
 		else if (variable->GetOpType() == OP_TYPE::NU_PARAMETER)
 		{
 			// If all the branches for the possible values are the same op remove the instruction
-			const ASTOpParameter* ParamOp = dynamic_cast<const ASTOpParameter*>(variable.child().get());
+			const ASTOpParameter* ParamOp = static_cast<const ASTOpParameter*>(variable.child().get());
 			check(ParamOp);
 			check(!ParamOp->parameter.m_possibleValues.IsEmpty());
 
@@ -518,7 +519,7 @@ namespace mu
 					// TODO: Probably it could be a any switch, it doesn't need to be of the same type.
 					if (Parent->GetOpType() == GetOpType())
 					{
-						const ASTOpSwitch* ParentSwitch = dynamic_cast<const ASTOpSwitch*>(Parent);
+						const ASTOpSwitch* ParentSwitch = static_cast<const ASTOpSwitch*>(Parent);
 						check(ParentSwitch);
 
 						// To be compatible the switch must be on the same variable
diff --git a/Engine/Plugins/Experimental/Mutable/Source/MutableTools/Private/MuT/ASTOpSwitch.h b/Engine/Plugins/Experimental/Mutable/Source/MutableTools/Private/MuT/ASTOpSwitch.h
--- a/Engine/Plugins/Experimental/Mutable/Source/MutableTools/Private/MuT/ASTOpSwitch.h
+++ b/Engine/Plugins/Experimental/Mutable/Source/MutableTools/Private/MuT/ASTOpSwitch.h
@@ -18,7 +18,7 @@ template <class SCALAR> class vec4;
 	//---------------------------------------------------------------------------------------------
 	//! Variable sized switch operation.
 	//---------------------------------------------------------------------------------------------
-	class ASTOpSwitch : public ASTOp
+	class ASTOpSwitch final : public ASTOp
 	{
 	public:
 
@@ -73,6 +73,7 @@ template <class SCALAR> class vec4;
 		bool GetNonBlackRect(FImageRect& maskUsage) const override;
 		bool IsImagePlainConstant(FVector4f& colour) const override;
 		Ptr<ImageSizeExpression> GetImageSizeExpression() const override;
+		virtual bool IsSwitch() const override { return true; }
 
 		// Own interface
 
diff --git a/Engine/Plugins/Experimental/Mutable/Source/MutableTools/Private/MuT/CodeGenerator.cpp b/Engine/Plugins/Experimental/Mutable/Source/MutableTools/Private/MuT/CodeGenerator.cpp
--- a/Engine/Plugins/Experimental/Mutable/Source/MutableTools/Private/MuT/CodeGenerator.cpp
+++ b/Engine/Plugins/Experimental/Mutable/Source/MutableTools/Private/MuT/CodeGenerator.cpp
@@ -168,29 +168,34 @@ namespace mu
 		}
 
         // Type-specific generation
-		if (const NodeScalar* ScalarNode = dynamic_cast<const NodeScalar*>(pNode.get()))
+		if (pNode->GetType()->IsA(NodeScalar::GetStaticType()))
 		{
+			const NodeScalar* ScalarNode = static_cast<const NodeScalar*>(pNode.get());
 			FScalarGenerationResult ScalarResult;
 			GenerateScalar(ScalarResult, Options, ScalarNode);
 			return ScalarResult.op;
 		}
 
-		else if (const NodeColour* ColorNode = dynamic_cast<const NodeColour*>(pNode.get()))
+		else if (pNode->GetType()->IsA(NodeColour::GetStaticType()))
 		{
+			const NodeColour* ColorNode = static_cast<const NodeColour*>(pNode.get());
 			FColorGenerationResult Result;
 			GenerateColor(Result, Options, ColorNode);
 			return Result.op;
 		}
 
-		else if (const NodeProjector* projNode = dynamic_cast<const NodeProjector*>(pNode.get()))
+		else if (pNode->GetType()->IsA(NodeProjector::GetStaticType()))
 		{
+			const NodeProjector* projNode = static_cast<const NodeProjector*>(pNode.get());
 			FProjectorGenerationResult ProjResult;
 			GenerateProjector(ProjResult, Options, projNode);
 			return ProjResult.op;
 		}
 
-		else if (const NodeSurfaceNew* surfNode = dynamic_cast<const NodeSurfaceNew*>(pNode.get()))
+		else if (pNode->GetType()->IsA(NodeSurfaceNew::GetStaticType()))
 		{
+			const NodeSurfaceNew* surfNode = static_cast<const NodeSurfaceNew*>(pNode.get());
+
 			// This happens only if we generate a node graph that has a NodeSurfaceNew at the root.
 			FSurfaceGenerationResult surfResult;
 			const TArray<FirstPassGenerator::FSurface::FEdit> edits;
@@ -198,25 +203,25 @@ namespace mu
 			return surfResult.surfaceOp;
 		}
 
-		else if (dynamic_cast<const NodeSurfaceVariation*>(pNode.get()))
+		else if (pNode->GetType()->IsA(NodeSurfaceVariation::GetStaticType()))
 		{
 			// This happens only if we generate a node graph that has a NodeSurfaceVariation at the root.
 			return nullptr;
 		}
 
-		else if (dynamic_cast<const NodeSurfaceSwitch*>(pNode.get()))
+		else if (pNode->GetType()->IsA(NodeSurfaceSwitch::GetStaticType()))
 		{
 			// This happens only if we generate a node graph that has a NodeSurfaceSwitch at the root.
 			return nullptr;
 		}
 
-		else if (dynamic_cast<const NodeSurfaceEdit*>(pNode.get()))
+		else if (pNode->GetType()->IsA(NodeSurfaceEdit::GetStaticType()))
 		{
 			// This happens only if we generate a node graph that has a NodeSurfaceEdit at the root.
 			return nullptr;
 		}
 
-		else if (dynamic_cast<const NodeModifier*>(pNode.get()))
+		else if (pNode->GetType()->IsA(NodeModifier::GetStaticType()))
 		{
 			// This happens only if we generate a node graph that has a modifier at the root.
 			return nullptr;
@@ -296,10 +301,11 @@ namespace mu
 			return;
 		}
 
-
 		// Generate for each different type of node
-		if (const NodeRangeFromScalar* FromScalar = dynamic_cast<const NodeRangeFromScalar*>(Untyped.get()))
+		if (Untyped->GetType()==NodeRangeFromScalar::GetStaticType())
 		{
+			const NodeRangeFromScalar* FromScalar = static_cast<const NodeRangeFromScalar*>(Untyped.get());
+
 			Result = FRangeGenerationResult();
 			Result.rangeName = FromScalar->GetName();
 
@@ -495,103 +501,6 @@ namespace mu
 	}
 
 
-	const NodeMeshConstant* FindSourceMesh(const Node* pNode)
-	{
-		const NodeMeshConstant* pResult = nullptr;
-
-		if (!pNode)
-		{
-			return pResult;
-		}
-
-		if (const NodeSurfaceNew* pTypedSN = dynamic_cast<const NodeSurfaceNew*>(pNode))
-        {
-            // TODO: 0...
-            pResult = FindSourceMesh( pTypedSN->GetMesh(0).get() );
-        }
- 		else if (const NodeSurfaceEdit* pTypedSE = dynamic_cast<const NodeSurfaceEdit*>(pNode))
-		{
-			NodePatchMesh* pPatch = pTypedSE->GetMesh();
-			if (pPatch)
-			{
-				pResult = FindSourceMesh(pPatch->GetAdd());
-			}
-		}
-		else if (const NodeSurfaceVariation* pTypedSV = dynamic_cast<const NodeSurfaceVariation*>(pNode))
-		{
-			if (pTypedSV->GetPrivate()->m_defaultSurfaces.Num())
-			{
-				pResult = FindSourceMesh(pTypedSV->GetPrivate()->m_defaultSurfaces[0].get());
-			}
-		}
-		else if (const NodeSurfaceSwitch* pTypedSS = dynamic_cast<const NodeSurfaceSwitch*>(pNode))
-		{
-			for (int32 i=0; i<pTypedSS->GetPrivate()->Options.Num(); ++i)
-			{
-				if (pTypedSS->GetPrivate()->Options[i])
-				{
-					pResult = FindSourceMesh(pTypedSS->GetPrivate()->Options[i].get());
-					if (pResult)
-					{
-						break;
-					}
-				}
-			}
-		}
-		else if (const NodeMeshInterpolate* pTypedMI = dynamic_cast<const NodeMeshInterpolate*>(pNode))
-        {
-            if ( pTypedMI->GetTargetCount()>0 )
-            {
-                pResult = FindSourceMesh( pTypedMI->GetTarget(0).get() );
-            }
-        }
-		else if (const NodeMeshMorph* pTypedMM = dynamic_cast<const NodeMeshMorph*>(pNode))
-        {
-            pResult = FindSourceMesh( pTypedMM->GetBase().get() );
-        }
-		else if (const NodeMeshGeometryOperation* pTypedGO = dynamic_cast<const NodeMeshGeometryOperation*>(pNode))
-		{
-			pResult = FindSourceMesh(pTypedGO->GetMeshA().get());
-		}
-		else if (const NodeMeshReshape* pTypedR = dynamic_cast<const NodeMeshReshape*>(pNode))
-		{
-			pResult = FindSourceMesh(pTypedR->GetBaseMesh().get());
-		}
-		else if (const NodeMeshFormat* pTypedMF = dynamic_cast<const NodeMeshFormat*>(pNode))
-        {
-            if ( pTypedMF->GetSource() )
-            {
-                pResult = FindSourceMesh( pTypedMF->GetSource().get() );
-            }
-        }
- 		else if (const NodeMeshFragment* pTypedMFrag = dynamic_cast<const NodeMeshFragment*>(pNode))
-        {
-            if ( pTypedMFrag->GetMesh() )
-            {
-                pResult = FindSourceMesh( pTypedMFrag->GetMesh().get() );
-            }
-        }
-		else if (const NodeMeshConstant* pCasted = dynamic_cast<const NodeMeshConstant*>(pNode))
-        {
-            pResult = pCasted;
-        }
- 		else if (const NodeMeshClipMorphPlane* pTypedCMP = dynamic_cast<const NodeMeshClipMorphPlane*>(pNode))
-		{
-			pResult = FindSourceMesh(pTypedCMP->GetSource().get());
-		}
-		else if (const NodeMeshClipWithMesh* pTypedCWM = dynamic_cast<const NodeMeshClipWithMesh*>(pNode))
-		{
-			pResult = FindSourceMesh(pTypedCWM->GetSource().get());
-		}
-		else
-		{
-			check(false);
-		}
-
-		return pResult;
-	}
-
-
 	void CodeGenerator::Generate_LOD(const FGenericGenerationOptions& Options, FGenericGenerationResult& Result, const NodeLOD* InNode)
 	{
 		const NodeLOD::Private& node = *InNode->GetPrivate();
@@ -613,7 +522,7 @@ namespace mu
 				{
 					check(componentOp->GetOpType() == OP_TYPE::IN_ADDCOMPONENT);
 
-					ASTOpInstanceAdd* typedOp = dynamic_cast<ASTOpInstanceAdd*>(componentOp.get());
+					ASTOpInstanceAdd* typedOp = static_cast<ASTOpInstanceAdd*>(componentOp.get());
 
 					// Complete the instruction adding the base
 					typedOp->instance = lastCompOp;
@@ -637,13 +546,13 @@ namespace mu
 
                 // First find the last op in the chain of IN_ADDCOMPONENT operations
                 check( cop->GetOpType() == OP_TYPE::IN_ADDCOMPONENT );
-				ASTOpInstanceAdd* typedOp = dynamic_cast<ASTOpInstanceAdd*>(cop.get());
+				ASTOpInstanceAdd* typedOp = static_cast<ASTOpInstanceAdd*>(cop.get());
 				ASTOpInstanceAdd* bottomOp = typedOp;
 				while (bottomOp->instance)
 				{
 					// Step down
 					check(bottomOp->instance->GetOpType() == OP_TYPE::IN_ADDCOMPONENT);
-					bottomOp = dynamic_cast<ASTOpInstanceAdd*>(bottomOp->instance.child().get());
+					bottomOp = static_cast<ASTOpInstanceAdd*>(bottomOp->instance.child().get());
 				}
 
 				// Chain
@@ -1196,24 +1105,27 @@ namespace mu
 				if (NodeImagePtr pImageNode = node.m_images[t].m_pImage)
 				{
 					// Any image-specific format or mipmapping needs to be applied at the end
-					NodeImageMipmapPtr mipmapNode;
-					NodeImageFormatPtr formatNode;
-					NodeImageSwizzlePtr swizzleNode;
+					Ptr<NodeImageMipmap> mipmapNode;
+					Ptr<NodeImageFormat> formatNode;
+					Ptr<NodeImageSwizzle> swizzleNode;
 					bool found = false;
 					while (!found)
 					{
-						if (NodeImageMipmap* tm = dynamic_cast<NodeImageMipmap*>(pImageNode.get()))
+						if (pImageNode->GetType()==NodeImageMipmap::GetStaticType())
 						{
+							NodeImageMipmap* tm = static_cast<NodeImageMipmap*>(pImageNode.get());
 							if (!mipmapNode) mipmapNode = tm;
 							pImageNode = tm->GetSource();
 						}
-						else if (NodeImageFormat* tf = dynamic_cast<NodeImageFormat*>(pImageNode.get()))
+						else if (pImageNode->GetType() == NodeImageFormat::GetStaticType())
 						{
+							NodeImageFormat* tf = static_cast<NodeImageFormat*>(pImageNode.get());
 							if (!formatNode) formatNode = tf;
 							pImageNode = tf->GetSource();
 						}
-						else if (NodeImageSwizzle* ts = dynamic_cast<NodeImageSwizzle*>(pImageNode.get()))
+						else if (pImageNode->GetType() == NodeImageSwizzle::GetStaticType())
 						{
+							NodeImageSwizzle* ts = static_cast<NodeImageSwizzle*>(pImageNode.get());
 							NodeImage* Source = ts->GetSource(0).get();
 							if (!swizzleNode && Source ==ts->GetSource(1) && Source==ts->GetSource(2) && Source==ts->GetSource(3))
 							{
@@ -2201,8 +2113,9 @@ namespace mu
 		Ptr<ASTOpMeshRemoveMask> removeOp;
 		for (const FirstPassGenerator::FModifier& m : modifiers)
 		{
-			if (const NodeModifierMeshClipWithMesh::Private* TypedClipNode = dynamic_cast<const NodeModifierMeshClipWithMesh::Private*>(m.node))
+			if (m.node->m_pNode->GetType()== NodeModifierMeshClipWithMesh::GetStaticType())
 			{
+				const NodeModifierMeshClipWithMesh::Private* TypedClipNode = static_cast<const NodeModifierMeshClipWithMesh::Private*>(m.node);
 				Ptr<ASTOpMeshMaskClipMesh> op = new ASTOpMeshMaskClipMesh();
 				op->source = preModifiersMesh;
 
@@ -2244,8 +2157,9 @@ namespace mu
 		// Process clip-with-mask modifiers
 		for (const FirstPassGenerator::FModifier& m : modifiers)
 		{
-			if (const NodeModifierMeshClipWithUVMask::Private* TypedClipNode = dynamic_cast<const NodeModifierMeshClipWithUVMask::Private*>(m.node))
+			if (m.node->m_pNode->GetType() == NodeModifierMeshClipWithUVMask::GetStaticType())
 			{
+				const NodeModifierMeshClipWithUVMask::Private* TypedClipNode = static_cast<const NodeModifierMeshClipWithUVMask::Private*>(m.node);
 				Ptr<ASTOpMeshMaskClipUVMask> op = new ASTOpMeshMaskClipUVMask();
 				op->Source = preModifiersMesh;
 				op->LayoutIndex = TypedClipNode->LayoutIndex;
@@ -2292,8 +2206,9 @@ namespace mu
 		{
 			Ptr<ASTOp> modifiedMeshOp;
 
-			if (const NodeModifierMeshClipMorphPlane::Private* TypedNode = dynamic_cast<const NodeModifierMeshClipMorphPlane::Private*>(m.node))
+			if (m.node->m_pNode->GetType() == NodeModifierMeshClipMorphPlane::GetStaticType())
 			{
+				const NodeModifierMeshClipMorphPlane::Private* TypedNode = static_cast<const NodeModifierMeshClipMorphPlane::Private*>(m.node);
 				Ptr<ASTOpMeshClipMorphPlane> op = new ASTOpMeshClipMorphPlane();
 				op->source = lastMeshOp;
 
@@ -2372,9 +2287,9 @@ namespace mu
 		{
 			Ptr<ASTOp> ModifiedMeshOp;
 
-			using ModifierType = NodeModifierMeshClipDeform::Private;
-			if ( const ModifierType* TypedClipNode = dynamic_cast<const ModifierType*>(M.node))
+			if (M.node->m_pNode->GetType()==NodeModifierMeshClipDeform::GetStaticType())
 			{
+				const NodeModifierMeshClipDeform::Private* TypedClipNode = static_cast<const NodeModifierMeshClipDeform::Private*>(M.node);
 				Ptr<ASTOpMeshBindShape>  BindOp = new ASTOpMeshBindShape();
 				Ptr<ASTOpMeshClipDeform> ClipOp = new ASTOpMeshClipDeform();
 
diff --git a/Engine/Plugins/Experimental/Mutable/Source/MutableTools/Private/MuT/CodeGenerator_Bool.cpp b/Engine/Plugins/Experimental/Mutable/Source/MutableTools/Private/MuT/CodeGenerator_Bool.cpp
--- a/Engine/Plugins/Experimental/Mutable/Source/MutableTools/Private/MuT/CodeGenerator_Bool.cpp
+++ b/Engine/Plugins/Experimental/Mutable/Source/MutableTools/Private/MuT/CodeGenerator_Bool.cpp
@@ -40,20 +40,24 @@ namespace mu
 		}
 
 		// Generate for each different type of node
-		if (auto Constant = dynamic_cast<const NodeBoolConstant*>(Untyped.get()))
+		if (Untyped->GetType()==NodeBoolConstant::GetStaticType())
 		{
+			const NodeBoolConstant* Constant = static_cast<const NodeBoolConstant*>(Untyped.get());
 			GenerateBool_Constant(Result, Options, Constant);
 		}
-		else if (auto Param = dynamic_cast<const NodeBoolParameter*>(Untyped.get()))
+		else if (Untyped->GetType() == NodeBoolParameter::GetStaticType())
 		{
+			const NodeBoolParameter* Param = static_cast<const NodeBoolParameter*>(Untyped.get());
 			GenerateBool_Parameter(Result, Options, Param);
 		}
-		else if (auto Sample = dynamic_cast<const NodeBoolNot*>(Untyped.get()))
+		else if (Untyped->GetType() == NodeBoolNot::GetStaticType())
 		{
+			const NodeBoolNot* Sample = static_cast<const NodeBoolNot*>(Untyped.get());
 			GenerateBool_Not(Result, Options, Sample);
 		}
-		else if (auto From = dynamic_cast<const NodeBoolAnd*>(Untyped.get()))
+		else if (Untyped->GetType() == NodeBoolAnd::GetStaticType())
 		{
+			const NodeBoolAnd* From = static_cast<const NodeBoolAnd*>(Untyped.get());
 			GenerateBool_And(Result, Options, From);
 		}
 		else
diff --git a/Engine/Plugins/Experimental/Mutable/Source/MutableTools/Private/MuT/CodeGenerator_Color.cpp b/Engine/Plugins/Experimental/Mutable/Source/MutableTools/Private/MuT/CodeGenerator_Color.cpp
--- a/Engine/Plugins/Experimental/Mutable/Source/MutableTools/Private/MuT/CodeGenerator_Color.cpp
+++ b/Engine/Plugins/Experimental/Mutable/Source/MutableTools/Private/MuT/CodeGenerator_Color.cpp
@@ -48,9 +48,6 @@
 #include "MuT/Table.h"
 #include "MuT/TablePrivate.h"
 
-#include <memory>
-#include <utility>
-
 
 namespace mu
 {
@@ -77,36 +74,44 @@ namespace mu
 		}
 
 		// Generate for each different type of node
-		if (auto Constant = dynamic_cast<const NodeColourConstant*>(Untyped.get()))
+		if (Untyped->GetType()==NodeColourConstant::GetStaticType())
 		{
+			const NodeColourConstant* Constant = static_cast<const NodeColourConstant*>(Untyped.get());
 			GenerateColor_Constant(Result, Options, Constant);
 		}
-		else if (auto Param = dynamic_cast<const NodeColourParameter*>(Untyped.get()))
+		else if (Untyped->GetType() == NodeColourParameter::GetStaticType())
 		{
+			const NodeColourParameter* Param = static_cast<const NodeColourParameter*>(Untyped.get());
 			GenerateColor_Parameter(Result, Options, Param);
 		}
-		else if (auto Switch = dynamic_cast<const NodeColourSwitch*>(Untyped.get()))
+		else if (Untyped->GetType() == NodeColourSwitch::GetStaticType())
 		{
+			const NodeColourSwitch* Switch = static_cast<const NodeColourSwitch*>(Untyped.get());
 			GenerateColor_Switch(Result, Options, Switch);
 		}
-		else if (auto Sample = dynamic_cast<const NodeColourSampleImage*>(Untyped.get()))
+		else if (Untyped->GetType() == NodeColourSampleImage::GetStaticType())
 		{
+			const NodeColourSampleImage* Sample = static_cast<const NodeColourSampleImage*>(Untyped.get());
 			GenerateColor_SampleImage(Result, Options, Sample);
 		}
-		else if (auto From = dynamic_cast<const NodeColourFromScalars*>(Untyped.get()))
+		else if (Untyped->GetType() == NodeColourFromScalars::GetStaticType())
 		{
+			const NodeColourFromScalars* From = static_cast<const NodeColourFromScalars*>(Untyped.get());
 			GenerateColor_FromScalars(Result, Options, From);
 		}
-		else if (auto Arithmetic = dynamic_cast<const NodeColourArithmeticOperation*>(Untyped.get()))
+		else if (Untyped->GetType() == NodeColourArithmeticOperation::GetStaticType())
 		{
+			const NodeColourArithmeticOperation* Arithmetic = static_cast<const NodeColourArithmeticOperation*>(Untyped.get());
 			GenerateColor_Arithmetic(Result, Options, Arithmetic);
 		}
-		else if (auto Variation = dynamic_cast<const NodeColourVariation*>(Untyped.get()))
+		else if (Untyped->GetType() == NodeColourVariation::GetStaticType())
 		{
+			const NodeColourVariation* Variation = static_cast<const NodeColourVariation*>(Untyped.get());
 			GenerateColor_Variation(Result, Options, Variation);
 		}
-		else if (auto Table = dynamic_cast<const NodeColourTable*>(Untyped.get()))
+		else if (Untyped->GetType() == NodeColourTable::GetStaticType())
 		{
+			const NodeColourTable* Table = static_cast<const NodeColourTable*>(Untyped.get());
 			GenerateColor_Table(Result, Options, Table);
 		}
 		else
diff --git a/Engine/Plugins/Experimental/Mutable/Source/MutableTools/Private/MuT/CodeGenerator_Projector.cpp b/Engine/Plugins/Experimental/Mutable/Source/MutableTools/Private/MuT/CodeGenerator_Projector.cpp
--- a/Engine/Plugins/Experimental/Mutable/Source/MutableTools/Private/MuT/CodeGenerator_Projector.cpp
+++ b/Engine/Plugins/Experimental/Mutable/Source/MutableTools/Private/MuT/CodeGenerator_Projector.cpp
@@ -20,9 +20,6 @@
 #include "MuT/NodeRange.h"
 #include "map"
 
-#include <memory>
-#include <utility>
-
 
 namespace mu
 {
diff --git a/Engine/Plugins/Experimental/Mutable/Source/MutableTools/Private/MuT/CodeGenerator_Scalar.cpp b/Engine/Plugins/Experimental/Mutable/Source/MutableTools/Private/MuT/CodeGenerator_Scalar.cpp
--- a/Engine/Plugins/Experimental/Mutable/Source/MutableTools/Private/MuT/CodeGenerator_Scalar.cpp
+++ b/Engine/Plugins/Experimental/Mutable/Source/MutableTools/Private/MuT/CodeGenerator_Scalar.cpp
@@ -67,38 +67,45 @@ namespace mu
 			return;
 		}
 
-
 		// Generate for each different type of node
-		if (auto Constant = dynamic_cast<const NodeScalarConstant*>(Untyped.get()))
+		if (Untyped->GetType() == NodeScalarConstant::GetStaticType())
 		{
+			const NodeScalarConstant* Constant = static_cast<const NodeScalarConstant*>(Untyped.get());
 			GenerateScalar_Constant(result, Options, Constant);
 		}
-		else if (auto Param = dynamic_cast<const NodeScalarParameter*>(Untyped.get()))
+		else if (Untyped->GetType() == NodeScalarParameter::GetStaticType())
 		{
+			const NodeScalarParameter* Param = static_cast<const NodeScalarParameter*>(Untyped.get());
 			GenerateScalar_Parameter(result, Options, Param);
 		}
-		else if (auto Switch = dynamic_cast<const NodeScalarSwitch*>(Untyped.get()))
+		else if (Untyped->GetType() == NodeScalarSwitch::GetStaticType())
 		{
+			const NodeScalarSwitch* Switch = static_cast<const NodeScalarSwitch*>(Untyped.get());
 			GenerateScalar_Switch(result, Options, Switch);
 		}
-		else if (auto EnumParam = dynamic_cast<const NodeScalarEnumParameter*>(Untyped.get()))
+		else if (Untyped->GetType() == NodeScalarEnumParameter::GetStaticType())
 		{
+			const NodeScalarEnumParameter* EnumParam = static_cast<const NodeScalarEnumParameter*>(Untyped.get());
 			GenerateScalar_EnumParameter(result, Options, EnumParam);
 		}
-		else if (auto Curve = dynamic_cast<const NodeScalarCurve*>(Untyped.get()))
+		else if (Untyped->GetType() == NodeScalarCurve::GetStaticType())
 		{
+			const NodeScalarCurve* Curve = static_cast<const NodeScalarCurve*>(Untyped.get());
 			GenerateScalar_Curve(result, Options, Curve);
 		}
-		else if (auto Arithmetic = dynamic_cast<const NodeScalarArithmeticOperation*>(Untyped.get()))
+		else if (Untyped->GetType() == NodeScalarArithmeticOperation::GetStaticType())
 		{
+			const NodeScalarArithmeticOperation* Arithmetic = static_cast<const NodeScalarArithmeticOperation*>(Untyped.get());
 			GenerateScalar_Arithmetic(result, Options, Arithmetic);
 		}
-		else if (auto Variation = dynamic_cast<const NodeScalarVariation*>(Untyped.get()))
+		else if (Untyped->GetType() == NodeScalarVariation::GetStaticType())
 		{
+			const NodeScalarVariation* Variation = static_cast<const NodeScalarVariation*>(Untyped.get());
 			GenerateScalar_Variation(result, Options, Variation);
 		}
-		else if (auto Table = dynamic_cast<const NodeScalarTable*>(Untyped.get()))
+		else if (Untyped->GetType() == NodeScalarTable::GetStaticType())
 		{
+			const NodeScalarTable* Table = static_cast<const NodeScalarTable*>(Untyped.get());
 			GenerateScalar_Table(result, Options, Table);
 		}
 		else
diff --git a/Engine/Plugins/Experimental/Mutable/Source/MutableTools/Private/MuT/CodeGenerator_SecondPass.cpp b/Engine/Plugins/Experimental/Mutable/Source/MutableTools/Private/MuT/CodeGenerator_SecondPass.cpp
--- a/Engine/Plugins/Experimental/Mutable/Source/MutableTools/Private/MuT/CodeGenerator_SecondPass.cpp
+++ b/Engine/Plugins/Experimental/Mutable/Source/MutableTools/Private/MuT/CodeGenerator_SecondPass.cpp
@@ -23,10 +23,6 @@ namespace mu
 		check(options);
 		m_pFirstPass = firstPass;
 		m_pCompilerOptions = options;
-
-		// Default conditions when there is no restriction accumulated.
-		CONDITION_CONTEXT noCondition;
-		m_currentCondition.push_back(noCondition);
 	}
 
 
@@ -102,8 +98,9 @@ namespace mu
 				negTag);
 
 			// If the surface is a constant false, we can skip adding it
-			if (auto* constOp = dynamic_cast<const ASTOpConstantBool*>(surfCondition.get()))
+			if (surfCondition && surfCondition->GetOpType()==OP_TYPE::BO_CONSTANT)
 			{
+				const ASTOpConstantBool* constOp = static_cast<const ASTOpConstantBool*>(surfCondition.get());
 				if (constOp->value == false)
 				{
 					continue;
@@ -279,13 +276,14 @@ namespace mu
 			// If the tag is a constant ...
 			bool isConstant = false;
 			bool constantValue = false;
-			if (auto* constOp = dynamic_cast<const ASTOpConstantBool*>(tagCondition.get()))
+			if (tagCondition->GetOpType()==OP_TYPE::BO_CONSTANT)
 			{
+				const ASTOpConstantBool* constOp = static_cast<const ASTOpConstantBool*>(tagCondition.get());
 				isConstant = true;
 				constantValue = constOp->value;
 			}
 
-			if (tagCondition && !isConstant)
+			if (!isConstant)
 			{
 				if (!c)
 				{
@@ -343,8 +341,9 @@ namespace mu
 			// If the tag is a constant ...
 			bool isConstant = false;
 			bool constantValue = false;
-			if (auto* constOp = dynamic_cast<const ASTOpConstantBool*>(tagCondition.get()))
+			if (tagCondition && tagCondition->GetOpType()== OP_TYPE::BO_CONSTANT)
 			{
+				const ASTOpConstantBool* constOp = static_cast<const ASTOpConstantBool*>(tagCondition.get());
 				isConstant = true;
 				constantValue = constOp->value;
 			}
@@ -416,8 +415,9 @@ namespace mu
 			// If the tag is a constant ...
 			bool isConstant = false;
 			bool constantValue = false;
-			if (auto* constOp = dynamic_cast<const ASTOpConstantBool*>(tagCondition.get()))
+			if (tagCondition && tagCondition->GetOpType() == OP_TYPE::BO_CONSTANT)
 			{
+				const ASTOpConstantBool* constOp = static_cast<const ASTOpConstantBool*>(tagCondition.get());
 				isConstant = true;
 				constantValue = constOp->value;
 			}
@@ -474,8 +474,9 @@ namespace mu
 				// If the tag is a constant ...
 				bool isConstant = false;
 				bool constantValue = false;
-				if (auto* constOp = dynamic_cast<const ASTOpConstantBool*>(tagCondition.get()))
+				if (tagCondition && tagCondition->GetOpType() == OP_TYPE::BO_CONSTANT)
 				{
+					const ASTOpConstantBool* constOp = static_cast<const ASTOpConstantBool*>(tagCondition.get());
 					isConstant = true;
 					constantValue = constOp->value;
 				}
diff --git a/Engine/Plugins/Experimental/Mutable/Source/MutableTools/Private/MuT/CodeGenerator_SecondPass.h b/Engine/Plugins/Experimental/Mutable/Source/MutableTools/Private/MuT/CodeGenerator_SecondPass.h
--- a/Engine/Plugins/Experimental/Mutable/Source/MutableTools/Private/MuT/CodeGenerator_SecondPass.h
+++ b/Engine/Plugins/Experimental/Mutable/Source/MutableTools/Private/MuT/CodeGenerator_SecondPass.h
@@ -13,9 +13,6 @@
 #include "map"
 #include "set"
 
-#include <memory>
-#include <utility>
-
 
 namespace mu
 {
@@ -38,16 +35,8 @@ namespace mu
         const CompilerOptions::Private *m_pCompilerOptions = nullptr;
 
 
-        struct CONDITION_CONTEXT
-        {
-            Ptr<ASTOp> surfaceCondition;
-            FirstPassGenerator::StateCondition stateCondition;
-        };
-        vector< CONDITION_CONTEXT > m_currentCondition;
-
-
         //!
-        ErrorLogPtr m_pErrorLog;
+        Ptr<ErrorLog> m_pErrorLog;
 
         //!
         struct CONDITION_GENERATION_KEY
diff --git a/Engine/Plugins/Experimental/Mutable/Source/MutableTools/Private/MuT/CodeOptimiser.cpp b/Engine/Plugins/Experimental/Mutable/Source/MutableTools/Private/MuT/CodeOptimiser.cpp
--- a/Engine/Plugins/Experimental/Mutable/Source/MutableTools/Private/MuT/CodeOptimiser.cpp
+++ b/Engine/Plugins/Experimental/Mutable/Source/MutableTools/Private/MuT/CodeOptimiser.cpp
@@ -110,37 +110,22 @@ namespace mu
 
 				case OP_TYPE::ME_CONSTANT:
 				{
-					auto typedNode = dynamic_cast<ASTOpConstantResource*>(n.get());
-					check(typedNode);
-
-					if (typedNode)
-					{
-						AllMeshOps.Add(typedNode);
-					}
+					ASTOpConstantResource* typedNode = static_cast<ASTOpConstantResource*>(n.get());
+					AllMeshOps.Add(typedNode);
 					break;
 				}
 
 				case OP_TYPE::IM_CONSTANT:
 				{
-					auto typedNode = dynamic_cast<ASTOpConstantResource*>(n.get());
-					check(typedNode);
-
-					if (typedNode)
-					{
-						AllImageOps.Add(typedNode);
-					}
+					ASTOpConstantResource* typedNode = static_cast<ASTOpConstantResource*>(n.get());
+					AllImageOps.Add(typedNode);
 					break;
 				}
 
 				case OP_TYPE::LA_CONSTANT:
 				{
-					auto typedNode = dynamic_cast<ASTOpConstantResource*>(n.get());
-					check(typedNode);
-
-					if (typedNode)
-					{
-						AllLayoutOps.Add(typedNode);
-					}
+					ASTOpConstantResource* typedNode = static_cast<ASTOpConstantResource*>(n.get());
+					AllLayoutOps.Add(typedNode);
 					break;
 				}
 
@@ -736,9 +721,9 @@ namespace mu
 	//                }
 
 	//                else
-					if (n->GetOpType()==OP_TYPE::IM_RASTERMESH)
+				if (n->GetOpType()==OP_TYPE::IM_RASTERMESH)
 				{
-					const ASTOpImageRasterMesh* Raster = dynamic_cast<const ASTOpImageRasterMesh*>(n.get());
+					const ASTOpImageRasterMesh* Raster = static_cast<const ASTOpImageRasterMesh*>(n.get());
 					if ( !Raster->image )
 					{
 						specialCase = true;
@@ -954,8 +939,7 @@ namespace mu
 
 			case OP_TYPE::ME_CONSTANT:
 			{
-				mu::Ptr<ASTOpConstantResource> typedOp = dynamic_cast<ASTOpConstantResource*>(node.get());
-				check(typedOp);
+				mu::Ptr<ASTOpConstantResource> typedOp = static_cast<ASTOpConstantResource*>(node.get());
 
 				if (currentProtected)
 				{
@@ -972,17 +956,12 @@ namespace mu
 
 			case OP_TYPE::ME_CLIPMORPHPLANE:
 			{
-				auto typedOp = dynamic_cast<ASTOpMeshClipMorphPlane*>(node.get());
-				if (typedOp)
+				ASTOpMeshClipMorphPlane* typedOp = static_cast<ASTOpMeshClipMorphPlane*>(node.get());
+				if (typedOp->vertexSelectionType == OP::MeshClipMorphPlaneArgs::VS_BONE_HIERARCHY)
 				{
-					if (typedOp->vertexSelectionType
-							==
-							OP::MeshClipMorphPlaneArgs::VS_BONE_HIERARCHY)
-					{
-						// We need the skeleton for the source mesh
-						RecurseWithState( typedOp->source.child(), true );
-						return false;
-					}
+					// We need the skeleton for the source mesh
+					RecurseWithState( typedOp->source.child(), true );
+					return false;
 				}
 
 				return true;
@@ -990,26 +969,20 @@ namespace mu
 
 			case OP_TYPE::ME_APPLYPOSE:
 			{
-				auto typedOp = dynamic_cast<ASTOpMeshApplyPose*>(node.get());
-				check(typedOp);
+				ASTOpMeshApplyPose* typedOp = static_cast<ASTOpMeshApplyPose*>(node.get());
 
-				if (typedOp)
-				{
-					// We need the skeleton for both meshes
-					RecurseWithState(typedOp->base.child(), true);
-					RecurseWithState(typedOp->pose.child(), true);
-					return false;
-				}
+				// We need the skeleton for both meshes
+				RecurseWithState(typedOp->base.child(), true);
+				RecurseWithState(typedOp->pose.child(), true);
+				return false;
 
 				break;
 			}
 
 			case OP_TYPE::ME_BINDSHAPE:
 			{
-				auto typedOp = dynamic_cast<ASTOpMeshBindShape*>(node.get());
-				check(typedOp);
-
-				if (typedOp && typedOp->bReshapeSkeleton)
+				ASTOpMeshBindShape* typedOp = static_cast<ASTOpMeshBindShape*>(node.get());
+				if (typedOp->bReshapeSkeleton)
 				{
 					RecurseWithState(typedOp->Mesh.child(), true);
 					return false;
@@ -1020,10 +993,8 @@ namespace mu
 
 			case OP_TYPE::ME_APPLYSHAPE:
 			{
-				auto typedOp = dynamic_cast<ASTOpMeshApplyShape*>(node.get());
-				check(typedOp);
-
-				if (typedOp && typedOp->bReshapeSkeleton)
+				ASTOpMeshApplyShape* typedOp = static_cast<ASTOpMeshApplyShape*>(node.get());
+				if (typedOp->bReshapeSkeleton)
 				{
 					RecurseWithState(typedOp->Mesh.child(), true);
 					return false;
@@ -1045,13 +1016,13 @@ namespace mu
 	//---------------------------------------------------------------------------------------------
 	// This stores an ADD_MESH op with the child meshes collected and the final skeleton to use
 	// for this op.
-	struct ADDMESH_SKELETON
+	struct FAddMeshSkeleton
 	{
 		mu::Ptr<ASTOp> m_pAddMeshOp;
 		TArray<mu::Ptr<ASTOpConstantResource>> m_contributingMeshes;
 		mu::Ptr<Skeleton> m_pFinalSkeleton;
 
-		ADDMESH_SKELETON( const mu::Ptr<ASTOp>& pAddMeshOp,
+		FAddMeshSkeleton( const mu::Ptr<ASTOp>& pAddMeshOp,
 						  TArray<mu::Ptr<ASTOpConstantResource>>& contributingMeshes,
 						  const mu::Ptr<Skeleton>& pFinalSkeleton )
 		{
@@ -1069,7 +1040,7 @@ namespace mu
 		// that require it.
 		CollectAllMeshesForSkeletonVisitorAST requireSkeletonCollector( roots );
 
-		TArray<ADDMESH_SKELETON> replacementsFound;
+		TArray<FAddMeshSkeleton> replacementsFound;
 
 		ASTOp::Traverse_TopDown_Unique_Imprecise( roots, [&](mu::Ptr<ASTOp>& at )
 		{
@@ -1078,7 +1049,7 @@ namespace mu
 
 			if ( at->GetOpType() == OP_TYPE::IN_ADDMESH )
 			{
-				auto typedNode = dynamic_cast<ASTOpInstanceAdd*>(at.get());
+				ASTOpInstanceAdd* typedNode = static_cast<ASTOpInstanceAdd*>(at.get());
 				mu::Ptr<ASTOp> meshRoot = typedNode->value.child();
 
 				if (meshRoot)
@@ -1092,9 +1063,7 @@ namespace mu
 						// \todo: refine to avoid instruction branches with irrelevant skeletons.
 						if ( lat->GetOpType() == OP_TYPE::ME_CONSTANT )
 						{
-							mu::Ptr<ASTOpConstantResource> typedOp = dynamic_cast<ASTOpConstantResource*>(lat.get());
-							check(typedOp);
-
+							mu::Ptr<ASTOpConstantResource> typedOp = static_cast<ASTOpConstantResource*>(lat.get());
 							if ( subtreeMeshes.Find(typedOp)
 								 ==
 								 INDEX_NONE )
@@ -1130,14 +1099,9 @@ namespace mu
 		{
 			if (at->GetOpType()==OP_TYPE::ME_CONSTANT)
 			{
-				auto typedOp = dynamic_cast<ASTOpConstantResource*>(at.get());
-				check(typedOp);
-				if (!typedOp)
-				{
-					return true;
-				}
+				ASTOpConstantResource* typedOp = static_cast<ASTOpConstantResource*>(at.get());
 
-				for(ADDMESH_SKELETON& Rep: replacementsFound) 
+				for(FAddMeshSkeleton& Rep: replacementsFound)
 				{
 					if (Rep.m_contributingMeshes.Contains(at))
 					{
@@ -1294,7 +1258,7 @@ namespace mu
 			{
 				if (n->GetOpType()==OP_TYPE::ME_CONSTANT)
 				{
-					auto typed = dynamic_cast<ASTOpConstantResource*>(n.get());
+					ASTOpConstantResource* typed = static_cast<ASTOpConstantResource*>(n.get());
 					auto pMesh = static_cast<const Mesh*>(typed->GetValue().get());
 					pMesh->ResetStaticFormatFlags();
 					typed->SetValue( pMesh,
