[mutable] Remove usage of RTTI.

diff --git a/Engine/Plugins/Experimental/Mutable/Source/MutableTools/Private/MuT/CodeOptimiser_Logic.cpp b/Engine/Plugins/Experimental/Mutable/Source/MutableTools/Private/MuT/CodeOptimiser_Logic.cpp
--- a/Engine/Plugins/Experimental/Mutable/Source/MutableTools/Private/MuT/CodeOptimiser_Logic.cpp
+++ b/Engine/Plugins/Experimental/Mutable/Source/MutableTools/Private/MuT/CodeOptimiser_Logic.cpp
@@ -1021,11 +1021,11 @@ namespace mu
 
                 if (n->GetOpType()==OP_TYPE::IN_CONDITIONAL)
                 {
-					ASTOpConditional* topConditional = dynamic_cast<ASTOpConditional*>(n.get());
+					ASTOpConditional* topConditional = static_cast<ASTOpConditional*>(n.get());
 					Ptr<ASTOp> yes = topConditional->yes.child();
                     if (yes && yes->GetOpType()==OP_TYPE::IN_ADDSURFACE)
                     {
-                        auto addSurface = dynamic_cast<ASTOpInstanceAdd*>(yes.get());
+						ASTOpInstanceAdd* addSurface = static_cast<ASTOpInstanceAdd*>(yes.get());
 
                         bool ended = false;
                         while (!ended)
@@ -1034,7 +1034,7 @@ namespace mu
 							Ptr<ASTOp> base = addSurface->instance.child();
                             if ( base && base->GetOpType()==OP_TYPE::IN_CONDITIONAL )
                             {
-								ASTOpConditional* bottomConditional = dynamic_cast<ASTOpConditional*>(base.get());
+								ASTOpConditional* bottomConditional = static_cast<ASTOpConditional*>(base.get());
 
                                 // Are the two conditions exclusive?
                                 bool conditionaAreExclusive = false;
@@ -1080,11 +1080,11 @@ namespace mu
 
                 else if (n->GetOpType()==OP_TYPE::ME_CONDITIONAL)
                 {
-					ASTOpConditional* topConditional = dynamic_cast<ASTOpConditional*>(n.get());
+					ASTOpConditional* topConditional = static_cast<ASTOpConditional*>(n.get());
                     Ptr<ASTOp> yes = topConditional->yes.child();
                     if ( yes && yes->GetOpType()==OP_TYPE::ME_MERGE )
                     {
-                        Ptr<ASTOpFixed> merge = dynamic_cast<ASTOpFixed*>(yes.get());
+                        Ptr<ASTOpFixed> merge = static_cast<ASTOpFixed*>(yes.get());
 
                         bool ended = false;
                         while (!ended)
@@ -1093,7 +1093,7 @@ namespace mu
 							Ptr<ASTOp> base = merge->children[merge->op.args.MeshMerge.base].child();
                             if ( base && base->GetOpType()==OP_TYPE::ME_CONDITIONAL )
                             {
-								ASTOpConditional* bottomConditional = dynamic_cast<ASTOpConditional*>(base.get());
+								ASTOpConditional* bottomConditional = static_cast<ASTOpConditional*>(base.get());
 
                                 // Are the two conditions exclusive?
                                 bool conditionaAreExclusive = false;
@@ -1200,10 +1200,14 @@ namespace mu
             {
                 bool recurse = true;
 
-                auto topConditional = dynamic_cast<ASTOpConditional*>(n.get());
+				ASTOpConditional* topConditional = nullptr;
+				if (n && n->IsConditional())
+				{
+					topConditional = static_cast<ASTOpConditional*>(n.get());
+				}
+
                 if (topConditional && topConditional->condition)
                 {
-
                     if ( topConditional->condition->GetOpType()==OP_TYPE::BO_EQUAL_INT_CONST
                          &&
                          topConditional->no
@@ -1214,7 +1218,7 @@ namespace mu
                         Ptr<ASTOpSwitch> switchOp = new ASTOpSwitch();
                         switchOp->type = GetSwitchForType(GetOpDataType(topConditional->GetOpType()));
 
-                        auto firstCompare = dynamic_cast<ASTOpFixed*>(topConditional->condition.child().get());
+						ASTOpFixed* firstCompare = static_cast<ASTOpFixed*>(topConditional->condition.child().get());
                         switchOp->variable = firstCompare->children[firstCompare->op.args.BoolEqualScalarConst.value].child();
 
                         auto current = n;
@@ -1222,7 +1226,12 @@ namespace mu
                         {
                             bool valid = false;
 
-                            auto conditional = dynamic_cast<ASTOpConditional*>(current.get());
+							ASTOpConditional* conditional = nullptr;
+							if (current && current->IsConditional())
+							{
+								conditional = static_cast<ASTOpConditional*>(current.get());
+							}
+
                             if ( conditional
                                  &&
                                  conditional->GetOpType()==topConditional->GetOpType()
@@ -1231,7 +1240,7 @@ namespace mu
                                  &&
                                  conditional->condition->GetOpType()==OP_TYPE::BO_EQUAL_INT_CONST )
                             {
-                                auto compare = dynamic_cast<ASTOpFixed*>(conditional->condition.child().get());
+								ASTOpFixed* compare = static_cast<ASTOpFixed*>(conditional->condition.child().get());
                                 check(compare);
                                 if (compare)
                                 {
@@ -1281,12 +1290,17 @@ namespace mu
             {
                 bool recurse = true;
 
-				ASTOpSwitch* topSwitch = dynamic_cast<ASTOpSwitch*>(n.get());
+				ASTOpSwitch* topSwitch = nullptr;
+				if (n && n->IsSwitch())
+				{
+					topSwitch = static_cast<ASTOpSwitch*>(n.get());
+				}
+
                 if (topSwitch)
                 {
                     bool first = true;
-                    auto caseType = OP_TYPE::NONE;
-                    for (const auto& c:topSwitch->cases)
+					OP_TYPE caseType = OP_TYPE::NONE;
+                    for (const ASTOpSwitch::FCase& c:topSwitch->cases)
                     {
                         if ( c.branch )
                         {
diff --git a/Engine/Plugins/Experimental/Mutable/Source/MutableTools/Private/MuT/CodeOptimiser_Parameters.cpp b/Engine/Plugins/Experimental/Mutable/Source/MutableTools/Private/MuT/CodeOptimiser_Parameters.cpp
--- a/Engine/Plugins/Experimental/Mutable/Source/MutableTools/Private/MuT/CodeOptimiser_Parameters.cpp
+++ b/Engine/Plugins/Experimental/Mutable/Source/MutableTools/Private/MuT/CodeOptimiser_Parameters.cpp
@@ -143,7 +143,7 @@ namespace mu
                     case OP_TYPE::PR_PARAMETER:
                     case OP_TYPE::IM_PARAMETER:
                     {
-						const ASTOpParameter* typed = dynamic_cast<const ASTOpParameter*>(at.get());
+						const ASTOpParameter* typed = static_cast<const ASTOpParameter*>(at.get());
                         const TArray<FString>& params = m_pState->nodeState.m_runtimeParams;
                         if ( params.Find( typed->parameter.m_name)
                              !=
@@ -157,7 +157,7 @@ namespace mu
 
                     case OP_TYPE::ME_INTERPOLATE:
                     {
-						const ASTOpFixed* typed = dynamic_cast<const ASTOpFixed*>(at.get());
+						const ASTOpFixed* typed = static_cast<const ASTOpFixed*>(at.get());
                         PENDING_ITEM childItem;
                         childItem.itemType = 0;
                         childItem.onlyLayoutsRelevant = item.onlyLayoutsRelevant;
@@ -383,7 +383,7 @@ namespace mu
         //-------------------------------------------------------------------------------------
         case OP_TYPE::IM_CONDITIONAL:
         {
-			const ASTOpConditional* typedAt = dynamic_cast<const ASTOpConditional*>(at.get());
+			const ASTOpConditional* typedAt = static_cast<const ASTOpConditional*>(at.get());
 
             // If the condition is not runtime, but the branches are, try to move the
             // conditional down
@@ -403,8 +403,8 @@ namespace mu
                     {
                         case OP_TYPE::IM_COMPOSE:
                         {
-						const ASTOpImageCompose* typedYes = dynamic_cast<const ASTOpImageCompose*>(typedAt->yes.child().get());
-						const ASTOpImageCompose* typedNo = dynamic_cast<const ASTOpImageCompose*>(typedAt->no.child().get());
+						const ASTOpImageCompose* typedYes = static_cast<const ASTOpImageCompose*>(typedAt->yes.child().get());
+						const ASTOpImageCompose* typedNo = static_cast<const ASTOpImageCompose*>(typedAt->no.child().get());
                         if ( typedYes->BlockIndex
                              ==
                              typedNo->BlockIndex
@@ -464,7 +464,7 @@ namespace mu
                         {
                             optimised = true;
 
-							const ASTOpImageLayerColor* typedYes = dynamic_cast<const ASTOpImageLayerColor*>(typedAt->yes.child().get());
+							const ASTOpImageLayerColor* typedYes = static_cast<const ASTOpImageLayerColor*>(typedAt->yes.child().get());
 
                             Ptr<ASTOpFixed> blackOp = new ASTOpFixed;
                             blackOp->op.type = OP_TYPE::CO_CONSTANT;
@@ -571,7 +571,7 @@ namespace mu
                         {
                             optimised = true;
 
-                            const ASTOpImageLayerColor* typedNo = dynamic_cast<const ASTOpImageLayerColor*>(typedAt->no.child().get());
+                            const ASTOpImageLayerColor* typedNo = static_cast<const ASTOpImageLayerColor*>(typedAt->no.child().get());
 
                             Ptr<ASTOpFixed> blackOp = new ASTOpFixed;
                             blackOp->op.type = OP_TYPE::CO_CONSTANT;
@@ -857,7 +857,7 @@ namespace mu
         //-----------------------------------------------------------------------------------------
         case OP_TYPE::IM_COMPOSE:
         {
-			const ASTOpImageCompose* typedAt = dynamic_cast<const ASTOpImageCompose*>(at.get());
+			const ASTOpImageCompose* typedAt = static_cast<const ASTOpImageCompose*>(at.get());
 
             Ptr<ASTOp> blockAt = typedAt->BlockImage.child();
 			Ptr<ASTOp> baseAt = typedAt->Base.child();
@@ -889,8 +889,8 @@ namespace mu
                     {
                         optimised = true;
 
-						const ASTOpImageLayerColor* typedBaseAt = dynamic_cast<const ASTOpImageLayerColor*>(baseAt.get());
-						const ASTOpImageLayerColor* typedBlockAt = dynamic_cast<const ASTOpImageLayerColor*>(blockAt.get());
+						const ASTOpImageLayerColor* typedBaseAt = static_cast<const ASTOpImageLayerColor*>(baseAt.get());
+						const ASTOpImageLayerColor* typedBlockAt = static_cast<const ASTOpImageLayerColor*>(blockAt.get());
 
                         // The mask is a compose of the block mask on the base mask, but if none has
                         // a mask we don't need to make one.
@@ -931,8 +931,8 @@ namespace mu
                     {
                         optimised = true;
 
-						const ASTOpImageLayer* typedBaseAt = dynamic_cast<const ASTOpImageLayer*>(baseAt.get());
-						const ASTOpImageLayer* typedBlockAt = dynamic_cast<const ASTOpImageLayer*>(blockAt.get());
+						const ASTOpImageLayer* typedBaseAt = static_cast<const ASTOpImageLayer*>(baseAt.get());
+						const ASTOpImageLayer* typedBlockAt = static_cast<const ASTOpImageLayer*>(blockAt.get());
 
                         // The mask is a compose of the block mask on the base mask, but if none has
                         // a mask we don't need to make one.
@@ -989,7 +989,7 @@ namespace mu
                  &&
                  baseType == OP_TYPE::IM_COMPOSE )
             {
-				const ASTOpImageCompose* typedBaseAt = dynamic_cast<const ASTOpImageCompose*>(baseAt.get());
+				const ASTOpImageCompose* typedBaseAt = static_cast<const ASTOpImageCompose*>(baseAt.get());
 
 				Ptr<ASTOp> baseBlockAt = typedBaseAt->BlockImage.child();
                 bool baseBlockHasAny = m_hasRuntimeParamVisitor.HasAny( baseBlockAt );
@@ -1026,7 +1026,7 @@ namespace mu
                 {
                     optimised = true;
 
-					const ASTOpImageLayerColor* typedBlockAt = dynamic_cast<const ASTOpImageLayerColor*>(blockAt.get());
+					const ASTOpImageLayerColor* typedBlockAt = static_cast<const ASTOpImageLayerColor*>(blockAt.get());
 
 					Ptr<ASTOp> blockImage = typedBlockAt->base.child();
 					Ptr<ASTOp> blockMask = typedBlockAt->mask.child();
@@ -1081,7 +1081,7 @@ namespace mu
                 {
                     optimised = true;
 
-                    const ASTOpImageLayer* typedBlockAt = dynamic_cast<const ASTOpImageLayer*>(blockAt.get());
+                    const ASTOpImageLayer* typedBlockAt = static_cast<const ASTOpImageLayer*>(blockAt.get());
 
 					Ptr<ASTOp> blockImage = typedBlockAt->base.child();
 					Ptr<ASTOp> blockBlended = typedBlockAt->blend.child();
@@ -1207,7 +1207,7 @@ namespace mu
                 {
                     optimised = true;
 
-                    const ASTOpImageLayerColor* typedBaseAt = dynamic_cast<const ASTOpImageLayerColor*>(baseAt.get());
+                    const ASTOpImageLayerColor* typedBaseAt = static_cast<const ASTOpImageLayerColor*>(baseAt.get());
 
 					Ptr<ASTOpImageCompose> maskOp = mu::Clone<ASTOpImageCompose>(at);
                     {
@@ -1254,7 +1254,7 @@ namespace mu
                 {
                     optimised = true;
 
-                    const ASTOpImageLayer* typedBaseAt = dynamic_cast<const ASTOpImageLayer*>(baseAt.get());
+                    const ASTOpImageLayer* typedBaseAt = static_cast<const ASTOpImageLayer*>(baseAt.get());
 
 					Ptr<ASTOpImageCompose> maskOp = mu::Clone<ASTOpImageCompose>(at);
                     {
@@ -1428,7 +1428,7 @@ namespace mu
         // Sink the mipmap if worth it.
         case OP_TYPE::IM_MIPMAP:
         {
-			const ASTOpImageMipmap* typedAt = dynamic_cast<const ASTOpImageMipmap*>(at.get());
+			const ASTOpImageMipmap* typedAt = static_cast<const ASTOpImageMipmap*>(at.get());
 
 			Ptr<ASTOp> sourceOp = typedAt->Source.child();
 
@@ -1436,7 +1436,7 @@ namespace mu
             {
             case OP_TYPE::IM_LAYERCOLOUR:
             {
-                const ASTOpImageLayerColor* typedSource = dynamic_cast<const ASTOpImageLayerColor*>(sourceOp.get());
+                const ASTOpImageLayerColor* typedSource = static_cast<const ASTOpImageLayerColor*>(sourceOp.get());
 
                 bool colourHasRuntime = m_hasRuntimeParamVisitor.HasAny( typedSource->color.child() );
 
@@ -1529,7 +1529,7 @@ namespace mu
             // TODO: Code shared with the constant data format optimisation visitor
             case OP_TYPE::IM_LAYERCOLOUR:
             {
-                const ASTOpImageLayerColor* typedAt = dynamic_cast<const ASTOpImageLayerColor*>(at.get());
+                const ASTOpImageLayerColor* typedAt = static_cast<const ASTOpImageLayerColor*>(at.get());
 
                 RecurseWithCurrentState( typedAt->base.child() );
                 RecurseWithCurrentState( typedAt->color.child() );
@@ -1548,7 +1548,7 @@ namespace mu
 
             case OP_TYPE::IM_LAYER:
             {
-                const ASTOpImageLayer* typedAt = dynamic_cast<const ASTOpImageLayer*>(at.get());
+                const ASTOpImageLayer* typedAt = static_cast<const ASTOpImageLayer*>(at.get());
 
                 RecurseWithCurrentState( typedAt->base.child() );
                 RecurseWithCurrentState( typedAt->blend.child() );
@@ -1568,7 +1568,7 @@ namespace mu
 
             case OP_TYPE::IM_DISPLACE:
             {
-				const ASTOpFixed* typedAt = dynamic_cast<const ASTOpFixed*>(at.get());
+				const ASTOpFixed* typedAt = static_cast<const ASTOpFixed*>(at.get());
 
                 RecurseWithCurrentState( typedAt->children[typedAt->op.args.ImageDisplace.source].child() );
 
@@ -1649,7 +1649,7 @@ namespace mu
                 case OP_TYPE::PR_PARAMETER:
                 case OP_TYPE::IM_PARAMETER:
                 {
-					const ASTOpParameter* typedAt = dynamic_cast<const ASTOpParameter*>(at.get());
+					const ASTOpParameter* typedAt = static_cast<const ASTOpParameter*>(at.get());
                     m_params.Add(typedAt->parameter.m_name);
 
                     // Not interested in the parameters from the parameters decorators.
@@ -1661,7 +1661,7 @@ namespace mu
 				case OP_TYPE::LA_FROMMESH:
 				{
 					// Manually choose how to recurse this op
-					const ASTOpLayoutFromMesh* pTyped = dynamic_cast<const ASTOpLayoutFromMesh*>(at.get());
+					const ASTOpLayoutFromMesh* pTyped = static_cast<const ASTOpLayoutFromMesh*>(at.get());
 
 					// For that mesh we only want to know about the layouts
 					if (const ASTChild& Mesh = pTyped->Mesh)
@@ -1675,7 +1675,7 @@ namespace mu
                 case OP_TYPE::ME_MORPH:
                 {
                     // Manually choose how to recurse this op
-					const ASTOpMeshMorph* pTyped = dynamic_cast<const ASTOpMeshMorph*>( at.get() );
+					const ASTOpMeshMorph* pTyped = static_cast<const ASTOpMeshMorph*>( at.get() );
 
                     if ( pTyped->Base )
                     {
@@ -1802,7 +1802,7 @@ namespace mu
                 case OP_TYPE::IN_ADDSCALAR:
                 case OP_TYPE::IN_ADDSTRING:
                 {
-					const ASTOpInstanceAdd* typedAt = dynamic_cast<const ASTOpInstanceAdd*>(at.get());
+					const ASTOpInstanceAdd* typedAt = static_cast<const ASTOpInstanceAdd*>(at.get());
 
                     TPair<bool,bool> newState;
                     newState.Key = false; //resource root
@@ -1929,7 +1929,7 @@ namespace mu
         // textures? We don't want them uncompressed.
         if( type==OP_TYPE::IN_ADDIMAGE )
         {
-			ASTOpInstanceAdd* typedAt = dynamic_cast<ASTOpInstanceAdd*>(at.get());
+			ASTOpInstanceAdd* typedAt = static_cast<ASTOpInstanceAdd*>(at.get());
 			Ptr<ASTOp> imageAt = typedAt->value.child();
 
             // Does it have a runtime parameter in its subtree?
@@ -1980,7 +1980,7 @@ namespace mu
 
         if( at->GetOpType()==OP_TYPE::IN_ADDLOD )
         {
-			ASTOpAddLOD* typedAt = dynamic_cast<ASTOpAddLOD*>(at.get());
+			ASTOpAddLOD* typedAt = static_cast<ASTOpAddLOD*>(at.get());
 
             if (typedAt->lods.Num()>(size_t)m_lodCount)
             {
diff --git a/Engine/Plugins/Experimental/Mutable/Source/MutableTools/Private/MuT/CodeOptimiser_Semantic.cpp b/Engine/Plugins/Experimental/Mutable/Source/MutableTools/Private/MuT/CodeOptimiser_Semantic.cpp
--- a/Engine/Plugins/Experimental/Mutable/Source/MutableTools/Private/MuT/CodeOptimiser_Semantic.cpp
+++ b/Engine/Plugins/Experimental/Mutable/Source/MutableTools/Private/MuT/CodeOptimiser_Semantic.cpp
@@ -102,7 +102,7 @@ namespace mu
             }
             else if ( aAt->GetOpType() == OP_TYPE::BO_CONSTANT )
             {
-                if ( dynamic_cast<const ASTOpConstantBool*>(aAt.get())->value )
+                if (static_cast<const ASTOpConstantBool*>(aAt.get())->value)
                 {
                     at = bAt;
                     changed = true;
@@ -115,7 +115,7 @@ namespace mu
             }
             else if ( bAt->GetOpType() == OP_TYPE::BO_CONSTANT )
             {
-                if ( dynamic_cast<const ASTOpConstantBool*>(bAt.get())->value )
+                if (static_cast<const ASTOpConstantBool*>(bAt.get())->value)
                 {
                     at = aAt;
                     changed = true;
@@ -130,7 +130,7 @@ namespace mu
             // Common cases of repeated branch in children
             else if ( aAt->GetOpType() == type )
             {
-                auto typedA = dynamic_cast<const ASTOpFixed*>(aAt.get());
+				const ASTOpFixed* typedA = static_cast<const ASTOpFixed*>(aAt.get());
                 if ( typedA->children[typedA->op.args.BoolBinary.a].child()==bAt
                      ||
                      typedA->children[typedA->op.args.BoolBinary.b].child()==bAt )
@@ -141,7 +141,7 @@ namespace mu
             }
             else if ( bAt->GetOpType() == type )
             {
-                auto typedB = dynamic_cast<const ASTOpFixed*>(bAt.get());
+				const ASTOpFixed* typedB = static_cast<const ASTOpFixed*>(bAt.get());
                 if ( typedB->children[typedB->op.args.BoolBinary.b].child()==aAt
                      ||
                      typedB->children[typedB->op.args.BoolBinary.b].child()==bAt )
@@ -184,7 +184,7 @@ namespace mu
             }
             else if ( aAt->GetOpType() == OP_TYPE::BO_CONSTANT )
             {
-                if ( dynamic_cast<const ASTOpConstantBool*>(aAt.get())->value )
+                if (static_cast<const ASTOpConstantBool*>(aAt.get())->value)
                 {
                     at = aAt;
                     changed = true;
@@ -197,7 +197,7 @@ namespace mu
             }
             else if ( bAt->GetOpType() == OP_TYPE::BO_CONSTANT )
             {
-                if ( dynamic_cast<const ASTOpConstantBool*>(bAt.get())->value )
+                if (static_cast<const ASTOpConstantBool*>(bAt.get())->value)
                 {
                     at = bAt;
                     changed = true;
@@ -212,7 +212,7 @@ namespace mu
             // Common cases of repeated branch in children
             else if ( aAt->GetOpType() == type )
             {
-                auto typedA = dynamic_cast<const ASTOpFixed*>(aAt.get());
+				const ASTOpFixed* typedA = static_cast<const ASTOpFixed*>(aAt.get());
                 if ( typedA->children[typedA->op.args.BoolBinary.a].child()==bAt
                      ||
                      typedA->children[typedA->op.args.BoolBinary.b].child()==bAt )
@@ -223,7 +223,7 @@ namespace mu
             }
             else if ( bAt->GetOpType() == type )
             {
-                auto typedB = dynamic_cast<const ASTOpFixed*>(bAt.get());
+				const ASTOpFixed* typedB = static_cast<const ASTOpFixed*>(bAt.get());
                 if ( typedB->children[typedB->op.args.BoolBinary.b].child()==aAt
                      ||
                      typedB->children[typedB->op.args.BoolBinary.b].child()==bAt )
@@ -274,7 +274,7 @@ namespace mu
 						{
 							NewAt = mu::Clone<ASTOpFixed>(this);
 						}
-						const ASTOpFixed* TypedCandidate = dynamic_cast<const ASTOpFixed*>(Candidate.get());
+						const ASTOpFixed* TypedCandidate = static_cast<const ASTOpFixed*>(Candidate.get());
 						int32 CandidateChannel = op.args.ColourSwizzle.sourceChannels[ChannelIndex];
 
 						NewAt->SetChild(NewAt->op.args.ColourSwizzle.sources[ChannelIndex], TypedCandidate->children[TypedCandidate->op.args.ColourSwizzle.sources[CandidateChannel]]);
@@ -329,9 +329,10 @@ namespace mu
 
 						for (int32 ChannelIndex = 0; ChannelIndex < MUTABLE_OP_MAX_SWIZZLE_CHANNELS; ++ChannelIndex)
 						{
-							const ASTOpFixed* SelectedSource = dynamic_cast<const ASTOpFixed*>(children[op.args.ColourSwizzle.sources[ChannelIndex]].child().get());
-							if (SelectedSource)
+							const ASTOp* SelectedSourceGeneric = children[op.args.ColourSwizzle.sources[ChannelIndex]].child().get();
+							if (SelectedSourceGeneric)
 							{
+								const ASTOpFixed* SelectedSource = static_cast<const ASTOpFixed*>(SelectedSourceGeneric);
 								int32 SelectedChannel = op.args.ColourSwizzle.sourceChannels[ChannelIndex];
 								Ptr<ASTOp> SelectedFloatInput = SelectedSource->children[SelectedSource->op.args.ColourFromScalars.v[SelectedChannel]].child();
 								NewAt->SetChild(NewAt->op.args.ColourFromScalars.v[ChannelIndex], SelectedFloatInput);
@@ -364,10 +365,8 @@ namespace mu
 			case OP_TYPE::IM_RESIZE:
 			{
 				Ptr<ASTOpFixed> NewOp = mu::Clone<ASTOpFixed>(sourceAt.get());
-				NewOp->op.args.ImageResize.size[0] =
-					int16_t(NewOp->op.args.ImageResize.size[0] * op.args.ImageResizeRel.factor[0]);
-				NewOp->op.args.ImageResize.size[1] =
-					int16_t(NewOp->op.args.ImageResize.size[1] * op.args.ImageResizeRel.factor[1]);
+				NewOp->op.args.ImageResize.size[0] = int16(NewOp->op.args.ImageResize.size[0] * op.args.ImageResizeRel.factor[0]);
+				NewOp->op.args.ImageResize.size[1] = int16(NewOp->op.args.ImageResize.size[1] * op.args.ImageResizeRel.factor[1]);
 
 				at = NewOp;
 				break;
@@ -401,7 +400,7 @@ namespace mu
 			case OP_TYPE::ME_ADDTAGS:
 			{
 				// Tags in the fragment can be ignored.
-				const ASTOpMeshAddTags* Add = dynamic_cast<const ASTOpMeshAddTags*>(Fragment.get());
+				const ASTOpMeshAddTags* Add = static_cast<const ASTOpMeshAddTags*>(Fragment.get());
 
 				Ptr<ASTOpFixed> NewAt = mu::Clone<ASTOpFixed>(this);
 				NewAt->SetChild(NewAt->op.args.MeshMaskDiff.fragment, Add->Source);
@@ -611,8 +610,8 @@ namespace mu
             {
                 if (displaceMapAt->GetOpType()==OP_TYPE::IM_CONDITIONAL)
                 {
-                    auto typedSource = dynamic_cast<const ASTOpConditional*>(sourceAt.get());
-                    auto typedDisplacementMap = dynamic_cast<const ASTOpConditional*>(displaceMapAt.get());
+					const ASTOpConditional* typedSource = static_cast<const ASTOpConditional*>(sourceAt.get());
+					const ASTOpConditional* typedDisplacementMap = static_cast<const ASTOpConditional*>(displaceMapAt.get());
 
                     if (typedSource->condition==typedDisplacementMap->condition)
                     {
@@ -638,8 +637,8 @@ namespace mu
             {
                 if (displaceMapAt->GetOpType()==OP_TYPE::IM_SWITCH)
                 {
-                    auto typedSource = dynamic_cast<const ASTOpSwitch*>(sourceAt.get());
-                    auto typedDisplacementMap = dynamic_cast<const ASTOpSwitch*>(displaceMapAt.get());
+					const ASTOpSwitch* typedSource = static_cast<const ASTOpSwitch*>(sourceAt.get());
+					const ASTOpSwitch* typedDisplacementMap = static_cast<const ASTOpSwitch*>(displaceMapAt.get());
 
                     if (typedSource->IsCompatibleWith(typedDisplacementMap))
                     {
@@ -753,7 +752,7 @@ namespace mu
 
             check(root->GetOpType()==OP_TYPE::IM_RESIZEREL);
 
-            auto typedRoot = dynamic_cast<const ASTOpFixed*>(root);
+			const ASTOpFixed* typedRoot = static_cast<const ASTOpFixed*>(root);
             m_initialSource = typedRoot->children[typedRoot->op.args.ImageResizeRel.source].child();
             Ptr<ASTOp> newSource = Visit( m_initialSource, typedRoot );
 
@@ -844,7 +843,7 @@ namespace mu
 				// still an integer after relative scale
 				bool acceptable = false;
 				{
-					const ASTOpImageCompose* typedAt = dynamic_cast<const ASTOpImageCompose*>(at.get());
+					const ASTOpImageCompose* typedAt = static_cast<const ASTOpImageCompose*>(at.get());
 					Ptr<ASTOp> originalBaseOp = typedAt->Base.child();
 
 					// \todo: recursion-proof cache?
@@ -1014,7 +1013,7 @@ namespace mu
 			case OP_TYPE::IM_TRANSFORM:
 			{
 				// It can only sink in the transform if it doesn't have it's own size.
-				const ASTOpImageTransform* TypedAt = dynamic_cast<const ASTOpImageTransform*>(at.get());
+				const ASTOpImageTransform* TypedAt = static_cast<const ASTOpImageTransform*>(at.get());
 				if (TypedAt->SizeX == 0 && TypedAt->SizeY == 0)
 				{
 					Ptr<ASTOpImageTransform> NewOp = mu::Clone<ASTOpImageTransform>(at);
@@ -1087,7 +1086,7 @@ namespace mu
             case OP_TYPE::IM_RESIZE:
             {
                 // Keep top resize
-                Ptr<const ASTOpFixed> sourceOp = dynamic_cast<const ASTOpFixed*>(sourceAt.get());
+                Ptr<const ASTOpFixed> sourceOp = static_cast<const ASTOpFixed*>(sourceAt.get());
 
                 Ptr<ASTOpFixed> newOp = mu::Clone<ASTOpFixed>(this);
                 newOp->SetChild( newOp->op.args.ImageResize.source, sourceOp->children[sourceOp->op.args.ImageResize.source]);
@@ -1270,7 +1269,7 @@ namespace mu
 				// \todo: only if shrinking?
 				
 				// Only sink the resize if we know that the pixelformat source image is uncompressed.
-				Ptr<ASTOpImagePixelFormat> SourceTyped = dynamic_cast<ASTOpImagePixelFormat*>(sourceAt.get());
+				Ptr<ASTOpImagePixelFormat> SourceTyped = static_cast<ASTOpImagePixelFormat*>(sourceAt.get());
 				FImageDesc PixelFormatSourceDesc = SourceTyped ->Source->GetImageDesc();
 				if (PixelFormatSourceDesc.m_format!=EImageFormat::IF_NONE
 					&&
@@ -1338,8 +1337,7 @@ namespace mu
 				// In the size optimization phase we can optimize the resize with the displace
 				// because the constants have not been collapsed yet.
 				// We will still check it and sink the size directly below the IM_MAKEGROWMAP op
-				Ptr<ASTOpFixed> SourceTyped = dynamic_cast<ASTOpFixed*>(sourceAt.get());
-				check(SourceTyped);
+				Ptr<ASTOpFixed> SourceTyped = static_cast<ASTOpFixed*>(sourceAt.get());
 				Ptr<ASTOp> OriginalDisplacementMapOp = SourceTyped->children[SourceTyped->op.args.ImageDisplace.displacementMap].m_child;
 				if (OriginalDisplacementMapOp->GetOpType() == OP_TYPE::IM_MAKEGROWMAP)
 				{
@@ -1412,7 +1410,7 @@ namespace mu
 			case OP_TYPE::IM_TRANSFORM:
 			{
 				// We can only optimize here if we know the transform result size, otherwise, we will sink the op in the sinker.
-				const ASTOpImageTransform* typedAt = dynamic_cast<const ASTOpImageTransform*>(sourceAt.get());
+				const ASTOpImageTransform* typedAt = static_cast<const ASTOpImageTransform*>(sourceAt.get());
 				if (typedAt->SizeX != 0 && typedAt->SizeY != 0)
 				{
 					// Set the size in the children and remove resize
diff --git a/Engine/Plugins/Experimental/Mutable/Source/MutableTools/Private/MuT/DataPacker.cpp b/Engine/Plugins/Experimental/Mutable/Source/MutableTools/Private/MuT/DataPacker.cpp
--- a/Engine/Plugins/Experimental/Mutable/Source/MutableTools/Private/MuT/DataPacker.cpp
+++ b/Engine/Plugins/Experimental/Mutable/Source/MutableTools/Private/MuT/DataPacker.cpp
@@ -149,7 +149,7 @@ namespace mu
             m_allParams.SetNumZeroed( program.m_parameters.Num() );
             FullTraverse( program );
 
-            for (size_t i=0; i<m_allParams.Num(); ++i )
+            for (int32 i=0; i<m_allParams.Num(); ++i )
             {
                 if ( m_allParams[i] )
                 {
@@ -186,10 +186,10 @@ namespace mu
                     {
                         m_conditionVisitor.Run( args.condition, program );
 
-                        TArray<int> currentParams = GetCurrentState();
-                        for (size_t p=0; p<m_conditionVisitor.m_params.Num(); ++p )
+                        TArray<int32> currentParams = GetCurrentState();
+                        for (int32 p=0; p<m_conditionVisitor.RelevantParams.Num(); ++p )
                         {
-                            currentParams[ m_conditionVisitor.m_params[p] ]++;
+                            currentParams[ m_conditionVisitor.RelevantParams[p] ]++;
                         }
 
                         // Set state for child recursion
@@ -207,17 +207,17 @@ namespace mu
                 case OP_TYPE::LA_SWITCH:
 				case OP_TYPE::ED_SWITCH:
                 {
-					const uint8_t* data = program.GetOpArgsPointer(at);
+					const uint8* data = program.GetOpArgsPointer(at);
 					
 					OP::ADDRESS VarAddress;
 					FMemory::Memcpy( &VarAddress, data, sizeof(OP::ADDRESS) );
 
                     m_conditionVisitor.Run( VarAddress, program );
 
-					TArray<int> currentParams = GetCurrentState();
-                    for (size_t p=0; p<m_conditionVisitor.m_params.Num(); ++p )
+					TArray<int32> currentParams = GetCurrentState();
+                    for (int32 p=0; p<m_conditionVisitor.RelevantParams.Num(); ++p )
                     {
-                        currentParams[ m_conditionVisitor.m_params[p] ]++;
+                        currentParams[ m_conditionVisitor.RelevantParams[p] ]++;
                     }
 
                     // Set state for child recursion
@@ -238,7 +238,7 @@ namespace mu
                         {
                             // Accumulate the currently relevant parameters
                             const TArray<int>& currentParams = GetCurrentState();
-                            for (size_t i=0; i<currentParams.Num(); ++i )
+                            for (int32 i=0; i<currentParams.Num(); ++i )
                             {
                                 m_allParams[i] += currentParams[i];
                             }
@@ -251,7 +251,7 @@ namespace mu
                         {
                             // Accumulate the currently relevant parameters
                             const TArray<int>& currentParams = GetCurrentState();
-                            for (size_t i=0; i<currentParams.Num(); ++i )
+                            for (int32 i=0; i<currentParams.Num(); ++i )
                             {
                                 m_allParams[i] += currentParams[i];
                             }
@@ -302,7 +302,7 @@ namespace mu
             //std::fill( m_supportedFormats.begin(), m_supportedFormats.end(), initial );
 
 			TArray<bool> defaultState;
-			defaultState.SetNumZeroed(size_t(EImageFormat::IF_COUNT));
+			defaultState.SetNumZeroed(int32(EImageFormat::IF_COUNT));
 
             Traverse( roots, defaultState );
         }
@@ -315,7 +315,7 @@ namespace mu
             const TArray<bool>& currentFormats = GetCurrentState();
 
 			TArray<bool> defaultState;
-			defaultState.SetNumZeroed(size_t(EImageFormat::IF_COUNT));
+			defaultState.SetNumZeroed(int32(EImageFormat::IF_COUNT));
 			bool allFalse = currentFormats == defaultState;
 
             // Can we use the cache?
@@ -334,11 +334,11 @@ namespace mu
             case OP_TYPE::IM_CONSTANT:
             {
                 // Remove unsupported formats
-				const ASTOpConstantResource* op = dynamic_cast<const ASTOpConstantResource*>(node.get());
+				const ASTOpConstantResource* op = static_cast<const ASTOpConstantResource*>(node.get());
                 if (!m_supportedFormats.count(op))
                 {
 					TArray<bool> initial;
-					initial.Init( true, size_t(EImageFormat::IF_COUNT) );
+					initial.Init( true, int32(EImageFormat::IF_COUNT) );
                     m_supportedFormats.insert( std::make_pair(op, std::move(initial)) );
                 }
 
@@ -361,17 +361,17 @@ namespace mu
             case OP_TYPE::IM_COMPOSE:
             {
                 recurse = false;
-				const ASTOpImageCompose* op = dynamic_cast<const ASTOpImageCompose*>(node.get());
+				const ASTOpImageCompose* op = static_cast<const ASTOpImageCompose*>(node.get());
 
 				TArray<bool> newState;
-				newState.Init(false, size_t(EImageFormat::IF_COUNT));
+				newState.Init(false, int32(EImageFormat::IF_COUNT));
                 RecurseWithState( op->Layout.child(), newState );
                 RecurseWithState( op->Base.child(), newState );
                 RecurseWithState( op->BlockImage.child(), newState );
 
                 if ( op->Mask )
                 {
-                    newState[(size_t)EImageFormat::IF_L_UBIT_RLE] = true;
+                    newState[(int32)EImageFormat::IF_L_UBIT_RLE] = true;
                     RecurseWithState( op->Mask.child(), newState );
                 }
                 break;
@@ -380,16 +380,16 @@ namespace mu
             case OP_TYPE::IM_LAYERCOLOUR:
             {
                 recurse = false;
-				const ASTOpImageLayerColor* op = dynamic_cast<const ASTOpImageLayerColor*>(node.get());
+				const ASTOpImageLayerColor* op = static_cast<const ASTOpImageLayerColor*>(node.get());
 				TArray<bool> newState;
-				newState.Init(false, size_t(EImageFormat::IF_COUNT));
+				newState.Init(false, int32(EImageFormat::IF_COUNT));
 				RecurseWithState( op->base.child(), newState );
                 RecurseWithState( op->color.child(), newState );
 
                 if ( op->mask )
                 {
-                    newState[(size_t)EImageFormat::IF_L_UBYTE] = true;
-                    newState[(size_t)EImageFormat::IF_L_UBYTE_RLE] = true;
+                    newState[(int32)EImageFormat::IF_L_UBYTE] = true;
+                    newState[(int32)EImageFormat::IF_L_UBYTE_RLE] = true;
 
                     RecurseWithState( op->mask.child(), newState );
                 }
@@ -399,16 +399,16 @@ namespace mu
             case OP_TYPE::IM_LAYER:
             {
                 recurse = false;
-				const ASTOpImageLayer* op = dynamic_cast<const ASTOpImageLayer*>(node.get());
+				const ASTOpImageLayer* op = static_cast<const ASTOpImageLayer*>(node.get());
 				TArray<bool> newState;
-				newState.Init(false, size_t(EImageFormat::IF_COUNT));
+				newState.Init(false, int32(EImageFormat::IF_COUNT));
 				RecurseWithState( op->base.child(), newState );
                 RecurseWithState( op->blend.child(), newState );
 
                 if (op->mask)
                 {
-                    newState[(size_t)EImageFormat::IF_L_UBYTE] = true;
-                    newState[(size_t)EImageFormat::IF_L_UBYTE_RLE] = true;
+                    newState[(int32)EImageFormat::IF_L_UBYTE] = true;
+                    newState[(int32)EImageFormat::IF_L_UBYTE_RLE] = true;
 
                     RecurseWithState( op->mask.child(), newState );
                 }
@@ -418,9 +418,9 @@ namespace mu
             case OP_TYPE::IM_MULTILAYER:
             {
                 recurse = false;
-				const ASTOpImageMultiLayer* op = dynamic_cast<const ASTOpImageMultiLayer*>(node.get());
+				const ASTOpImageMultiLayer* op = static_cast<const ASTOpImageMultiLayer*>(node.get());
 				TArray<bool> newState;
-				newState.Init(false, size_t(EImageFormat::IF_COUNT));
+				newState.Init(false, int32(EImageFormat::IF_COUNT));
 				RecurseWithState( op->base.child(), newState );
                 RecurseWithState( op->blend.child(), newState );
 
@@ -437,13 +437,13 @@ namespace mu
             case OP_TYPE::IM_DISPLACE:
             {
                 recurse = false;
-				const ASTOpFixed* op = dynamic_cast<const ASTOpFixed*>(node.get());
+				const ASTOpFixed* op = static_cast<const ASTOpFixed*>(node.get());
 				TArray<bool> newState;
-				newState.Init(false, size_t(EImageFormat::IF_COUNT));
+				newState.Init(false, int32(EImageFormat::IF_COUNT));
 				RecurseWithState( op->children[op->op.args.ImageDisplace.source].child(), newState );
 
-                newState[(size_t)EImageFormat::IF_L_UBYTE ] = true;
-                newState[(size_t)EImageFormat::IF_L_UBYTE_RLE ] = true;
+                newState[(int32)EImageFormat::IF_L_UBYTE ] = true;
+                newState[(int32)EImageFormat::IF_L_UBYTE_RLE ] = true;
 
                 RecurseWithState( op->children[op->op.args.ImageDisplace.displacementMap].child(), newState );
                 break;
@@ -456,7 +456,7 @@ namespace mu
                 //m_currentFormats.pop_back();
 
 				TArray<bool> newState;
-				newState.Init(false, size_t(EImageFormat::IF_COUNT));
+				newState.Init(false, int32(EImageFormat::IF_COUNT));
 				if (currentFormats != newState)
                 {
                     RecurseWithState(node, newState);
@@ -502,10 +502,10 @@ namespace mu
             MUTABLE_CPUPROFILER_SCOPE(AccumulateMeshChannelUsageAST);
 
             // Sanity check in case we add more semantics
-            static_assert(MBS_COUNT<sizeof(uint64_t)*8, "Too many mesh buffer semantics." );
+            static_assert(MBS_COUNT<sizeof(uint64)*8, "Too many mesh buffer semantics." );
 
             // Default state: we need everything except internal semantics
-            uint64_t defaultState = 0xffffffffffffffff;
+            uint64 defaultState = 0xffffffffffffffff;
             defaultState ^= (UINT64_C(1)<<MBS_LAYOUTBLOCK);
             defaultState ^= (UINT64_C(1)<<MBS_VERTEXINDEX);
 
@@ -517,7 +517,7 @@ namespace mu
         {
             bool recurse = true;
 
-            uint64_t currentSemantics = GetCurrentState();
+            uint64 currentSemantics = GetCurrentState();
 
             switch ( node->GetOpType() )
             {
@@ -525,8 +525,8 @@ namespace mu
             case OP_TYPE::ME_CONSTANT:
             {
                 // Accumulate necessary semantics
-				const ASTOpConstantResource* op = dynamic_cast<const ASTOpConstantResource*>(node.get());
-                uint64_t currentFlags = 0;
+				const ASTOpConstantResource* op = static_cast<const ASTOpConstantResource*>(node.get());
+                uint64 currentFlags = 0;
                 if (m_requiredSemantics.count(op))
                 {
                     currentFlags = m_requiredSemantics[op];
@@ -549,9 +549,9 @@ namespace mu
             {
                 recurse = false;
 
-				const ASTOpMeshDifference* op = dynamic_cast<const ASTOpMeshDifference*>(node.get());
+				const ASTOpMeshDifference* op = static_cast<const ASTOpMeshDifference*>(node.get());
 
-                uint64_t newState = currentSemantics;
+                uint64 newState = currentSemantics;
                 newState |= (UINT64_C(1)<<MBS_VERTEXINDEX);
                 RecurseWithState( op->Base.child(), newState );
 
@@ -563,9 +563,9 @@ namespace mu
             {
                 recurse = false;
 
-				const ASTOpMeshRemoveMask* op = dynamic_cast<const ASTOpMeshRemoveMask*>(node.get());
+				const ASTOpMeshRemoveMask* op = static_cast<const ASTOpMeshRemoveMask*>(node.get());
 
-                uint64_t newState = currentSemantics;
+                uint64 newState = currentSemantics;
                 newState |= (UINT64_C(1)<<MBS_VERTEXINDEX);
                 RecurseWithState( op->source.child(), newState );
                 for( const TPair<ASTChild, ASTChild>& r: op->removes )
@@ -579,9 +579,9 @@ namespace mu
             {
                 recurse = false;
 
-				const ASTOpMeshMorph* op = dynamic_cast<const ASTOpMeshMorph*>(node.get());
+				const ASTOpMeshMorph* op = static_cast<const ASTOpMeshMorph*>(node.get());
 
-                uint64_t newState = currentSemantics;
+                uint64 newState = currentSemantics;
                 newState |= (UINT64_C(1)<<MBS_VERTEXINDEX);
                 RecurseWithState( op->Base.child(), newState );
                 RecurseWithState( op->Target.child(), newState );
@@ -593,9 +593,9 @@ namespace mu
             {
                 recurse = false;
 
-				const ASTOpFixed* op = dynamic_cast<const ASTOpFixed*>(node.get());
+				const ASTOpFixed* op = static_cast<const ASTOpFixed*>(node.get());
 
-                uint64_t newState = currentSemantics;
+                uint64 newState = currentSemantics;
                 newState |= (UINT64_C(1)<<MBS_LAYOUTBLOCK);
                 RecurseWithState( op->children[op->op.args.MeshApplyLayout.mesh].child(), newState );
 
@@ -607,10 +607,10 @@ namespace mu
             {
                 recurse = false;
 
-				const ASTOpMeshExtractLayoutBlocks* op = dynamic_cast<const ASTOpMeshExtractLayoutBlocks*>(node.get());
+				const ASTOpMeshExtractLayoutBlocks* op = static_cast<const ASTOpMeshExtractLayoutBlocks*>(node.get());
 
                 // todo: check if we really need all of them
-                uint64_t newState = currentSemantics;
+                uint64 newState = currentSemantics;
                 newState |= (UINT64_C(1)<<MBS_LAYOUTBLOCK);
                 newState |= (UINT64_C(1)<<MBS_VERTEXINDEX);
 
@@ -622,9 +622,9 @@ namespace mu
 			{
 				recurse = false;
 
-				const ASTOpLayoutFromMesh* op = dynamic_cast<const ASTOpLayoutFromMesh*>(node.get());
+				const ASTOpLayoutFromMesh* op = static_cast<const ASTOpLayoutFromMesh*>(node.get());
 
-				uint64_t newState = currentSemantics;
+				uint64 newState = currentSemantics;
 				newState |= (UINT64_C(1) << MBS_LAYOUTBLOCK);
 
 				RecurseWithState(op->Mesh.child(), newState);
@@ -635,18 +635,18 @@ namespace mu
             {
                 recurse = false;
 
-				const ASTOpInstanceAdd* op = dynamic_cast<const ASTOpInstanceAdd*>(node.get());
+				const ASTOpInstanceAdd* op = static_cast<const ASTOpInstanceAdd*>(node.get());
 
                 RecurseWithState( op->instance.child(), currentSemantics );
 
-                uint64_t newState = GetDefaultState();
+                uint64 newState = GetDefaultState();
                 RecurseWithState( op->value.child(), newState );
                 break;
             }
 
             default:
                 // Unhandled op, we may need everything? Recurse with current state?
-                //uint64_t newState = 0xffffffffffffffff;
+                //uint64 newState = 0xffffffffffffffff;
                 break;
 
             }
@@ -658,7 +658,7 @@ namespace mu
 
         //! Result of this visitor:
         //! Used mesh channel semantics for each constant mesh
-        std::unordered_map< Ptr<const ASTOpConstantResource>, uint64_t > m_requiredSemantics;
+        std::unordered_map< Ptr<const ASTOpConstantResource>, uint64 > m_requiredSemantics;
 
     };
 
@@ -668,7 +668,7 @@ namespace mu
     //---------------------------------------------------------------------------------------------
 
     // Todo: move to its own file
-    inline void MeshRemoveUnusedBufferSemantics( Mesh* pMesh, uint64_t usedSemantics )
+    inline void MeshRemoveUnusedBufferSemantics( Mesh* pMesh, uint64 usedSemantics )
     {
         // right now we only remove entire buffers if no channel is used
         // TODO: remove from inside the buffer?
@@ -696,7 +696,7 @@ namespace mu
 
         // TODO: hack, if we don't need layouts, remove them.
         {
-            uint64_t layoutSemantics = 0;
+            uint64 layoutSemantics = 0;
             layoutSemantics |= (UINT64_C(1)<<MBS_LAYOUTBLOCK);
 
             if ( (usedSemantics&layoutSemantics) == 0)
@@ -722,7 +722,7 @@ namespace mu
         {
             if (n->GetOpType()==OP_TYPE::IM_CONSTANT)
             {
-				ASTOpConstantResource* typed = dynamic_cast<ASTOpConstantResource*>(n.get());
+				ASTOpConstantResource* typed = static_cast<ASTOpConstantResource*>(n.get());
                 Ptr<const Image> pOld = static_cast<const Image*>(typed->GetValue().get());
 
 				FImageOperator ImOp = FImageOperator::GetDefault( Options->GetPrivate()->ImageFormatFunc );
@@ -752,25 +752,25 @@ namespace mu
 					ASTOp::Replace(n, NewPlain);
 
 				}
-				else if ( accFormat.m_supportedFormats[typed][(size_t)EImageFormat::IF_L_UBIT_RLE] )
+				else if ( accFormat.m_supportedFormats[typed][(int32)EImageFormat::IF_L_UBIT_RLE] )
                 {
                     ImagePtr pNew = ImOp.ImagePixelFormat( ImageCompressionQuality, pOld.get(), EImageFormat::IF_L_UBIT_RLE );
 
                     // Only replace if the compression was worth!
-                    size_t oldSize = pOld->GetDataSize();
-                    size_t newSize = pNew->GetDataSize();
+                    int32 oldSize = pOld->GetDataSize();
+					int32 newSize = pNew->GetDataSize();
                     if (float(oldSize) > float(newSize) * OptimizeOptions.MinRLECompressionGain)
                     {
                         typed->SetValue(pNew, OptimizeOptions.bUseDiskCache);
                     }
                 }
-                else if ( accFormat.m_supportedFormats[typed][(size_t)EImageFormat::IF_L_UBYTE_RLE] )
+                else if ( accFormat.m_supportedFormats[typed][(int32)EImageFormat::IF_L_UBYTE_RLE] )
                 {
                     ImagePtr pNew = ImOp.ImagePixelFormat( ImageCompressionQuality, pOld.get(), EImageFormat::IF_L_UBYTE_RLE );
 
                     // Only replace if the compression was worth!
-                    size_t oldSize = pOld->GetDataSize();
-                    size_t newSize = pNew->GetDataSize();
+					int32 oldSize = pOld->GetDataSize();
+					int32 newSize = pNew->GetDataSize();
                     if (float(oldSize) > float(newSize) * OptimizeOptions.MinRLECompressionGain)
                     {
                         typed->SetValue(pNew, OptimizeOptions.bUseDiskCache);
@@ -789,7 +789,7 @@ namespace mu
         {
             if (n->GetOpType()==OP_TYPE::ME_CONSTANT)
             {
-				ASTOpConstantResource* typed = dynamic_cast<ASTOpConstantResource*>(n.get());
+				ASTOpConstantResource* typed = static_cast<ASTOpConstantResource*>(n.get());
                 Ptr<Mesh> pMesh = static_cast<const Mesh*>(typed->GetValue().get())->Clone();
                 MeshRemoveUnusedBufferSemantics( pMesh.get(), meshSemanticsVisitor.m_requiredSemantics[typed]);
                 typed->SetValue(pMesh, OptimizeOptions.bUseDiskCache);
diff --git a/Engine/Plugins/Experimental/Mutable/Source/MutableTools/Private/MuT/Node.cpp b/Engine/Plugins/Experimental/Mutable/Source/MutableTools/Private/MuT/Node.cpp
--- a/Engine/Plugins/Experimental/Mutable/Source/MutableTools/Private/MuT/Node.cpp
+++ b/Engine/Plugins/Experimental/Mutable/Source/MutableTools/Private/MuT/Node.cpp
@@ -28,31 +28,23 @@
 namespace mu
 {
 
-	//---------------------------------------------------------------------------------------------
 	// Static initialisation
-	//---------------------------------------------------------------------------------------------
-	static NODE_TYPE s_nodeType = 		NODE_TYPE( "Node", 0 );
+	static FNodeType s_nodeType = 		FNodeType( "Node", 0 );
 
-
-	//---------------------------------------------------------------------------------------------
-	NODE_TYPE::NODE_TYPE()
+	FNodeType::FNodeType()
 	{
 		m_strName = "";
-		m_pParent = 0;
+		m_pParent = nullptr;
 	}
 
-	//---------------------------------------------------------------------------------------------
-	NODE_TYPE::NODE_TYPE( const char* strName, const NODE_TYPE* pParent )
+
+	FNodeType::FNodeType( const char* strName, const FNodeType* pParent )
 	{
         m_strName = strName;
 		m_pParent = pParent;
 	}
 
 
-
-	//---------------------------------------------------------------------------------------------
-	//---------------------------------------------------------------------------------------------
-	//---------------------------------------------------------------------------------------------
 	void Node::Serialise( const Node* pNode, OutputArchive& arch )
 	{
         uint32_t ver = 0;
@@ -88,7 +80,6 @@ namespace mu
 	}
 
 
-	//---------------------------------------------------------------------------------------------
 	NodePtr Node::StaticUnserialise( InputArchive& arch )
 	{
         uint32_t ver;
@@ -123,28 +114,23 @@ namespace mu
 	}
 
 
-
-	//---------------------------------------------------------------------------------------------
-	const NODE_TYPE* Node::GetType() const
+	const FNodeType* Node::GetType() const
 	{
 		return GetStaticType();
 	}
 
 
-	//---------------------------------------------------------------------------------------------
-	const NODE_TYPE* Node::GetStaticType()
+	const FNodeType* Node::GetStaticType()
 	{
 		return &s_nodeType;
 	}
 
 
-	//---------------------------------------------------------------------------------------------
 	void Node::SetMessageContext( const void* context )
 	{
 		GetBasePrivate()->m_errorContext = context;
 	}
 
-
 }
 
 
diff --git a/Engine/Plugins/Experimental/Mutable/Source/MutableTools/Private/MuT/NodeBool.cpp b/Engine/Plugins/Experimental/Mutable/Source/MutableTools/Private/MuT/NodeBool.cpp
--- a/Engine/Plugins/Experimental/Mutable/Source/MutableTools/Private/MuT/NodeBool.cpp
+++ b/Engine/Plugins/Experimental/Mutable/Source/MutableTools/Private/MuT/NodeBool.cpp
@@ -16,13 +16,13 @@ namespace mu
 	//---------------------------------------------------------------------------------------------
 	// Static initialisation
 	//---------------------------------------------------------------------------------------------
-	static NODE_TYPE s_nodeBoolType = NODE_TYPE( "NodeBool", Node::GetStaticType() );
+	static FNodeType s_nodeBoolType = FNodeType( "NodeBool", Node::GetStaticType() );
 
 
 	//---------------------------------------------------------------------------------------------
 	void NodeBool::Serialise( const NodeBool* p, OutputArchive& arch )
 	{
-        uint32_t ver = 0;
+        uint32 ver = 0;
 		arch << ver;
 
 		arch << uint32_t(p->Type);
@@ -33,11 +33,11 @@ namespace mu
 	//---------------------------------------------------------------------------------------------
 	Ptr<NodeBool> NodeBool::StaticUnserialise( InputArchive& arch )
 	{
-        uint32_t ver;
+        uint32 ver;
 		arch >> ver;
 		check( ver == 0 );
 
-        uint32_t id;
+        uint32 id;
 		arch >> id;
 
 		switch (id)
@@ -55,14 +55,14 @@ namespace mu
 
 
 	//---------------------------------------------------------------------------------------------
-	const NODE_TYPE* NodeBool::GetType() const
+	const FNodeType* NodeBool::GetType() const
 	{
 		return GetStaticType();
 	}
 
 
 	//---------------------------------------------------------------------------------------------
-	const NODE_TYPE* NodeBool::GetStaticType()
+	const FNodeType* NodeBool::GetStaticType()
 	{
 		return &s_nodeBoolType;
 	}
diff --git a/Engine/Plugins/Experimental/Mutable/Source/MutableTools/Private/MuT/NodeBoolAnd.cpp b/Engine/Plugins/Experimental/Mutable/Source/MutableTools/Private/MuT/NodeBoolAnd.cpp
--- a/Engine/Plugins/Experimental/Mutable/Source/MutableTools/Private/MuT/NodeBoolAnd.cpp
+++ b/Engine/Plugins/Experimental/Mutable/Source/MutableTools/Private/MuT/NodeBoolAnd.cpp
@@ -14,8 +14,8 @@ namespace mu
 	//---------------------------------------------------------------------------------------------
 	// Static initialisation
 	//---------------------------------------------------------------------------------------------
-	NODE_TYPE NodeBoolAnd::Private::s_type =
-			NODE_TYPE( "BoolParameter", NodeBool::GetStaticType() );
+	FNodeType NodeBoolAnd::Private::s_type =
+			FNodeType( "BoolParameter", NodeBool::GetStaticType() );
 
 
 	//---------------------------------------------------------------------------------------------
diff --git a/Engine/Plugins/Experimental/Mutable/Source/MutableTools/Private/MuT/NodeBoolConstant.cpp b/Engine/Plugins/Experimental/Mutable/Source/MutableTools/Private/MuT/NodeBoolConstant.cpp
--- a/Engine/Plugins/Experimental/Mutable/Source/MutableTools/Private/MuT/NodeBoolConstant.cpp
+++ b/Engine/Plugins/Experimental/Mutable/Source/MutableTools/Private/MuT/NodeBoolConstant.cpp
@@ -15,8 +15,8 @@ namespace mu
 	//---------------------------------------------------------------------------------------------
 	// Static initialisation
 	//---------------------------------------------------------------------------------------------
-	NODE_TYPE NodeBoolConstant::Private::s_type =
-			NODE_TYPE( "BoolConstant", NodeBool::GetStaticType() );
+	FNodeType NodeBoolConstant::Private::s_type =
+			FNodeType( "BoolConstant", NodeBool::GetStaticType() );
 
 
 	//---------------------------------------------------------------------------------------------
diff --git a/Engine/Plugins/Experimental/Mutable/Source/MutableTools/Private/MuT/NodeBoolNot.cpp b/Engine/Plugins/Experimental/Mutable/Source/MutableTools/Private/MuT/NodeBoolNot.cpp
--- a/Engine/Plugins/Experimental/Mutable/Source/MutableTools/Private/MuT/NodeBoolNot.cpp
+++ b/Engine/Plugins/Experimental/Mutable/Source/MutableTools/Private/MuT/NodeBoolNot.cpp
@@ -14,8 +14,8 @@ namespace mu
 	//---------------------------------------------------------------------------------------------
 	// Static initialisation
 	//---------------------------------------------------------------------------------------------
-	NODE_TYPE NodeBoolNot::Private::s_type =
-			NODE_TYPE( "BoolParameter", NodeBool::GetStaticType() );
+	FNodeType NodeBoolNot::Private::s_type =
+			FNodeType( "BoolParameter", NodeBool::GetStaticType() );
 
 
 	//---------------------------------------------------------------------------------------------
diff --git a/Engine/Plugins/Experimental/Mutable/Source/MutableTools/Private/MuT/NodeBoolParameter.cpp b/Engine/Plugins/Experimental/Mutable/Source/MutableTools/Private/MuT/NodeBoolParameter.cpp
--- a/Engine/Plugins/Experimental/Mutable/Source/MutableTools/Private/MuT/NodeBoolParameter.cpp
+++ b/Engine/Plugins/Experimental/Mutable/Source/MutableTools/Private/MuT/NodeBoolParameter.cpp
@@ -20,8 +20,8 @@ namespace mu
 	//---------------------------------------------------------------------------------------------
 	// Static initialisation
 	//---------------------------------------------------------------------------------------------
-	NODE_TYPE NodeBoolParameter::Private::s_type =
-			NODE_TYPE( "BoolParameter", NodeBool::GetStaticType() );
+	FNodeType NodeBoolParameter::Private::s_type =
+			FNodeType( "BoolParameter", NodeBool::GetStaticType() );
 
 
 	//---------------------------------------------------------------------------------------------
diff --git a/Engine/Plugins/Experimental/Mutable/Source/MutableTools/Private/MuT/NodeBoolPrivate.h b/Engine/Plugins/Experimental/Mutable/Source/MutableTools/Private/MuT/NodeBoolPrivate.h
--- a/Engine/Plugins/Experimental/Mutable/Source/MutableTools/Private/MuT/NodeBoolPrivate.h
+++ b/Engine/Plugins/Experimental/Mutable/Source/MutableTools/Private/MuT/NodeBoolPrivate.h
@@ -23,7 +23,7 @@ namespace mu
 	{
 	public:
 
-		static NODE_TYPE s_type;
+		static FNodeType s_type;
 
 		bool m_value;
 
@@ -52,7 +52,7 @@ namespace mu
 	{
 	public:
 
-		static NODE_TYPE s_type;
+		static FNodeType s_type;
 
 		bool m_defaultValue;
 		FString m_name;
@@ -103,7 +103,7 @@ namespace mu
 	{
 	public:
 
-		static NODE_TYPE s_type;
+		static FNodeType s_type;
 
 		Ptr<NodeBool> m_pSource;
 
@@ -132,7 +132,7 @@ namespace mu
 	{
 	public:
 
-		static NODE_TYPE s_type;
+		static FNodeType s_type;
 
 		Ptr<NodeBool> m_pA;
 		Ptr<NodeBool> m_pB;
diff --git a/Engine/Plugins/Experimental/Mutable/Source/MutableTools/Private/MuT/NodeColour.cpp b/Engine/Plugins/Experimental/Mutable/Source/MutableTools/Private/MuT/NodeColour.cpp
--- a/Engine/Plugins/Experimental/Mutable/Source/MutableTools/Private/MuT/NodeColour.cpp
+++ b/Engine/Plugins/Experimental/Mutable/Source/MutableTools/Private/MuT/NodeColour.cpp
@@ -24,7 +24,7 @@ namespace mu
 	//---------------------------------------------------------------------------------------------
 	// Static initialisation
 	//---------------------------------------------------------------------------------------------
-	static NODE_TYPE s_nodeColourType = NODE_TYPE( "NodeColour", Node::GetStaticType() );
+	static FNodeType s_nodeColourType = FNodeType( "NodeColour", Node::GetStaticType() );
 
 
 	//---------------------------------------------------------------------------------------------
@@ -73,14 +73,14 @@ namespace mu
 
 
 	//---------------------------------------------------------------------------------------------
-	const NODE_TYPE* NodeColour::GetType() const
+	const FNodeType* NodeColour::GetType() const
 	{
 		return GetStaticType();
 	}
 
 
 	//---------------------------------------------------------------------------------------------
-	const NODE_TYPE* NodeColour::GetStaticType()
+	const FNodeType* NodeColour::GetStaticType()
 	{
 		return &s_nodeColourType;
 	}
diff --git a/Engine/Plugins/Experimental/Mutable/Source/MutableTools/Private/MuT/NodeColourArithmeticOperation.cpp b/Engine/Plugins/Experimental/Mutable/Source/MutableTools/Private/MuT/NodeColourArithmeticOperation.cpp
--- a/Engine/Plugins/Experimental/Mutable/Source/MutableTools/Private/MuT/NodeColourArithmeticOperation.cpp
+++ b/Engine/Plugins/Experimental/Mutable/Source/MutableTools/Private/MuT/NodeColourArithmeticOperation.cpp
@@ -15,8 +15,8 @@ namespace mu
 	//---------------------------------------------------------------------------------------------
 	// Static initialisation
 	//---------------------------------------------------------------------------------------------
-	NODE_TYPE NodeColourArithmeticOperation::Private::s_type =
-			NODE_TYPE( "ColourArithmenticOperation", NodeColour::GetStaticType() );
+	FNodeType NodeColourArithmeticOperation::Private::s_type =
+			FNodeType( "ColourArithmenticOperation", NodeColour::GetStaticType() );
 
 	//---------------------------------------------------------------------------------------------
 	//!
diff --git a/Engine/Plugins/Experimental/Mutable/Source/MutableTools/Private/MuT/NodeColourArithmeticOperationPrivate.h b/Engine/Plugins/Experimental/Mutable/Source/MutableTools/Private/MuT/NodeColourArithmeticOperationPrivate.h
--- a/Engine/Plugins/Experimental/Mutable/Source/MutableTools/Private/MuT/NodeColourArithmeticOperationPrivate.h
+++ b/Engine/Plugins/Experimental/Mutable/Source/MutableTools/Private/MuT/NodeColourArithmeticOperationPrivate.h
@@ -13,7 +13,7 @@ namespace mu
 	{
 	public:
 
-		static NODE_TYPE s_type;
+		static FNodeType s_type;
 
 		NodeColourArithmeticOperation::OPERATION m_operation;
 		NodeColourPtr m_pA;
diff --git a/Engine/Plugins/Experimental/Mutable/Source/MutableTools/Private/MuT/NodeColourConstant.cpp b/Engine/Plugins/Experimental/Mutable/Source/MutableTools/Private/MuT/NodeColourConstant.cpp
--- a/Engine/Plugins/Experimental/Mutable/Source/MutableTools/Private/MuT/NodeColourConstant.cpp
+++ b/Engine/Plugins/Experimental/Mutable/Source/MutableTools/Private/MuT/NodeColourConstant.cpp
@@ -15,8 +15,8 @@ namespace mu
 	//---------------------------------------------------------------------------------------------
 	// Static initialisation
 	//---------------------------------------------------------------------------------------------
-	NODE_TYPE NodeColourConstant::Private::s_type =
-			NODE_TYPE( "ColourConstant", NodeColour::GetStaticType() );
+	FNodeType NodeColourConstant::Private::s_type =
+			FNodeType( "ColourConstant", NodeColour::GetStaticType() );
 
 
 	//---------------------------------------------------------------------------------------------
diff --git a/Engine/Plugins/Experimental/Mutable/Source/MutableTools/Private/MuT/NodeColourConstantPrivate.h b/Engine/Plugins/Experimental/Mutable/Source/MutableTools/Private/MuT/NodeColourConstantPrivate.h
--- a/Engine/Plugins/Experimental/Mutable/Source/MutableTools/Private/MuT/NodeColourConstantPrivate.h
+++ b/Engine/Plugins/Experimental/Mutable/Source/MutableTools/Private/MuT/NodeColourConstantPrivate.h
@@ -14,7 +14,7 @@ namespace mu
 	{
 	public:
 
-		static NODE_TYPE s_type;
+		static FNodeType s_type;
 
 		FVector4f m_value;
 
diff --git a/Engine/Plugins/Experimental/Mutable/Source/MutableTools/Private/MuT/NodeColourFromScalars.cpp b/Engine/Plugins/Experimental/Mutable/Source/MutableTools/Private/MuT/NodeColourFromScalars.cpp
--- a/Engine/Plugins/Experimental/Mutable/Source/MutableTools/Private/MuT/NodeColourFromScalars.cpp
+++ b/Engine/Plugins/Experimental/Mutable/Source/MutableTools/Private/MuT/NodeColourFromScalars.cpp
@@ -15,8 +15,8 @@ namespace mu
 	//---------------------------------------------------------------------------------------------
 	// Static initialisation
 	//---------------------------------------------------------------------------------------------
-	NODE_TYPE NodeColourFromScalars::Private::s_type =
-			NODE_TYPE( "ColourFromScalars", NodeColour::GetStaticType() );
+	FNodeType NodeColourFromScalars::Private::s_type =
+			FNodeType( "ColourFromScalars", NodeColour::GetStaticType() );
 
 
 	//---------------------------------------------------------------------------------------------
diff --git a/Engine/Plugins/Experimental/Mutable/Source/MutableTools/Private/MuT/NodeColourParameter.cpp b/Engine/Plugins/Experimental/Mutable/Source/MutableTools/Private/MuT/NodeColourParameter.cpp
--- a/Engine/Plugins/Experimental/Mutable/Source/MutableTools/Private/MuT/NodeColourParameter.cpp
+++ b/Engine/Plugins/Experimental/Mutable/Source/MutableTools/Private/MuT/NodeColourParameter.cpp
@@ -16,8 +16,8 @@ namespace mu
 	//---------------------------------------------------------------------------------------------
 	// Static initialisation
 	//---------------------------------------------------------------------------------------------
-	NODE_TYPE NodeColourParameter::Private::s_type =
-			NODE_TYPE( "ColourParameter", NodeColour::GetStaticType() );
+	FNodeType NodeColourParameter::Private::s_type =
+			FNodeType( "ColourParameter", NodeColour::GetStaticType() );
 
 
 	//---------------------------------------------------------------------------------------------
diff --git a/Engine/Plugins/Experimental/Mutable/Source/MutableTools/Private/MuT/NodeColourSampleImage.cpp b/Engine/Plugins/Experimental/Mutable/Source/MutableTools/Private/MuT/NodeColourSampleImage.cpp
--- a/Engine/Plugins/Experimental/Mutable/Source/MutableTools/Private/MuT/NodeColourSampleImage.cpp
+++ b/Engine/Plugins/Experimental/Mutable/Source/MutableTools/Private/MuT/NodeColourSampleImage.cpp
@@ -16,8 +16,8 @@ namespace mu
 	//---------------------------------------------------------------------------------------------
 	// Static initialisation
 	//---------------------------------------------------------------------------------------------
-	NODE_TYPE NodeColourSampleImage::Private::s_type =
-			NODE_TYPE( "ColourSampleImage", NodeColour::GetStaticType() );
+	FNodeType NodeColourSampleImage::Private::s_type =
+			FNodeType( "ColourSampleImage", NodeColour::GetStaticType() );
 
 
 	//---------------------------------------------------------------------------------------------
diff --git a/Engine/Plugins/Experimental/Mutable/Source/MutableTools/Private/MuT/NodeColourSwitch.cpp b/Engine/Plugins/Experimental/Mutable/Source/MutableTools/Private/MuT/NodeColourSwitch.cpp
--- a/Engine/Plugins/Experimental/Mutable/Source/MutableTools/Private/MuT/NodeColourSwitch.cpp
+++ b/Engine/Plugins/Experimental/Mutable/Source/MutableTools/Private/MuT/NodeColourSwitch.cpp
@@ -14,8 +14,8 @@ namespace mu
 	//---------------------------------------------------------------------------------------------
 	// Static initialisation
 	//---------------------------------------------------------------------------------------------
-    NODE_TYPE NodeColourSwitch::Private::s_type =
-            NODE_TYPE( "ColourSwitch", NodeColour::GetStaticType() );
+    FNodeType NodeColourSwitch::Private::s_type =
+            FNodeType( "ColourSwitch", NodeColour::GetStaticType() );
 
 
 	//---------------------------------------------------------------------------------------------
diff --git a/Engine/Plugins/Experimental/Mutable/Source/MutableTools/Private/MuT/NodeColourTable.cpp b/Engine/Plugins/Experimental/Mutable/Source/MutableTools/Private/MuT/NodeColourTable.cpp
--- a/Engine/Plugins/Experimental/Mutable/Source/MutableTools/Private/MuT/NodeColourTable.cpp
+++ b/Engine/Plugins/Experimental/Mutable/Source/MutableTools/Private/MuT/NodeColourTable.cpp
@@ -17,8 +17,8 @@ namespace mu
 	//---------------------------------------------------------------------------------------------
 	// Static initialisation
 	//---------------------------------------------------------------------------------------------
-	NODE_TYPE NodeColourTable::Private::s_type =
-			NODE_TYPE( "TableColour", NodeColour::GetStaticType() );
+	FNodeType NodeColourTable::Private::s_type =
+			FNodeType( "TableColour", NodeColour::GetStaticType() );
 
 
 	//---------------------------------------------------------------------------------------------
diff --git a/Engine/Plugins/Experimental/Mutable/Source/MutableTools/Private/MuT/NodeColourVariation.cpp b/Engine/Plugins/Experimental/Mutable/Source/MutableTools/Private/MuT/NodeColourVariation.cpp
--- a/Engine/Plugins/Experimental/Mutable/Source/MutableTools/Private/MuT/NodeColourVariation.cpp
+++ b/Engine/Plugins/Experimental/Mutable/Source/MutableTools/Private/MuT/NodeColourVariation.cpp
@@ -14,8 +14,8 @@ namespace mu
     //---------------------------------------------------------------------------------------------
     // Static initialisation
     //---------------------------------------------------------------------------------------------
-    NODE_TYPE NodeColourVariation::Private::s_type =
-        NODE_TYPE( "ColourVariation", NodeColour::GetStaticType() );
+    FNodeType NodeColourVariation::Private::s_type =
+        FNodeType( "ColourVariation", NodeColour::GetStaticType() );
 
 
     MUTABLE_IMPLEMENT_NODE( NodeColourVariation, EType::Variation, Node, Node::EType::Colour)
diff --git a/Engine/Plugins/Experimental/Mutable/Source/MutableTools/Private/MuT/NodeComponent.cpp b/Engine/Plugins/Experimental/Mutable/Source/MutableTools/Private/MuT/NodeComponent.cpp
--- a/Engine/Plugins/Experimental/Mutable/Source/MutableTools/Private/MuT/NodeComponent.cpp
+++ b/Engine/Plugins/Experimental/Mutable/Source/MutableTools/Private/MuT/NodeComponent.cpp
@@ -19,21 +19,20 @@ namespace mu
 	//---------------------------------------------------------------------------------------------
 	// Static initialisation
 	//---------------------------------------------------------------------------------------------
-	static NODE_TYPE s_nodeComponentType =
-			NODE_TYPE( "NodeComponent", Node::GetStaticType() );
+	static FNodeType s_nodeComponentType = FNodeType( "NodeComponent", Node::GetStaticType() );
 
 
 	//---------------------------------------------------------------------------------------------
 	//---------------------------------------------------------------------------------------------
 	//---------------------------------------------------------------------------------------------
-	const NODE_TYPE* NodeComponent::GetType() const
+	const FNodeType* NodeComponent::GetType() const
 	{
 		return GetStaticType();
 	}
 
 
 	//---------------------------------------------------------------------------------------------
-	const NODE_TYPE* NodeComponent::GetStaticType()
+	const FNodeType* NodeComponent::GetStaticType()
 	{
 		return &s_nodeComponentType;
 	}
@@ -42,10 +41,10 @@ namespace mu
 	//---------------------------------------------------------------------------------------------
 	void NodeComponent::Serialise( const NodeComponent* p, OutputArchive& arch )
 	{
-        uint32_t ver = 0;
+        uint32 ver = 0;
 		arch << ver;
 
-		arch << uint32_t(p->Type);
+		arch << uint32(p->Type);
 		p->SerialiseWrapper(arch);
 	}
 
@@ -53,11 +52,11 @@ namespace mu
 	//---------------------------------------------------------------------------------------------
 	NodeComponentPtr NodeComponent::StaticUnserialise( InputArchive& arch )
 	{
-        uint32_t ver;
+        uint32 ver;
 		arch >> ver;
 		check( ver == 0 );
 
-        uint32_t id;
+        uint32 id;
 		arch >> id;
 
 		switch (id)
diff --git a/Engine/Plugins/Experimental/Mutable/Source/MutableTools/Private/MuT/NodeComponentEdit.cpp b/Engine/Plugins/Experimental/Mutable/Source/MutableTools/Private/MuT/NodeComponentEdit.cpp
--- a/Engine/Plugins/Experimental/Mutable/Source/MutableTools/Private/MuT/NodeComponentEdit.cpp
+++ b/Engine/Plugins/Experimental/Mutable/Source/MutableTools/Private/MuT/NodeComponentEdit.cpp
@@ -17,8 +17,8 @@ namespace mu
 	//---------------------------------------------------------------------------------------------
 	// Static initialisation
 	//---------------------------------------------------------------------------------------------
-	NODE_TYPE NodeComponentEdit::Private::s_type =
-			NODE_TYPE( "EditComponent", NodeComponent::GetStaticType() );
+	FNodeType NodeComponentEdit::Private::s_type =
+			FNodeType( "EditComponent", NodeComponent::GetStaticType() );
 
 
 	//---------------------------------------------------------------------------------------------
@@ -83,7 +83,7 @@ namespace mu
 		const NodeComponentNew::Private* parent = nullptr;
 		if (m_pParent)
 		{
-			NodeComponent::Private* ParentPrivate = dynamic_cast<NodeComponent::Private*>(m_pParent->GetBasePrivate());
+			NodeComponent::Private* ParentPrivate = static_cast<NodeComponent::Private*>(m_pParent->GetBasePrivate());
 			parent = ParentPrivate->GetParentComponentNew();
 		}
 
diff --git a/Engine/Plugins/Experimental/Mutable/Source/MutableTools/Private/MuT/NodeComponentNew.cpp b/Engine/Plugins/Experimental/Mutable/Source/MutableTools/Private/MuT/NodeComponentNew.cpp
--- a/Engine/Plugins/Experimental/Mutable/Source/MutableTools/Private/MuT/NodeComponentNew.cpp
+++ b/Engine/Plugins/Experimental/Mutable/Source/MutableTools/Private/MuT/NodeComponentNew.cpp
@@ -16,8 +16,8 @@ namespace mu
 	//---------------------------------------------------------------------------------------------
 	// Static initialisation
 	//---------------------------------------------------------------------------------------------
-	NODE_TYPE NodeComponentNew::Private::s_type =
-			NODE_TYPE( "NewComponent", NodeComponent::GetStaticType() );
+	FNodeType NodeComponentNew::Private::s_type =
+			FNodeType( "NewComponent", NodeComponent::GetStaticType() );
 
 
 	//---------------------------------------------------------------------------------------------
diff --git a/Engine/Plugins/Experimental/Mutable/Source/MutableTools/Private/MuT/NodeExtensionData.cpp b/Engine/Plugins/Experimental/Mutable/Source/MutableTools/Private/MuT/NodeExtensionData.cpp
--- a/Engine/Plugins/Experimental/Mutable/Source/MutableTools/Private/MuT/NodeExtensionData.cpp
+++ b/Engine/Plugins/Experimental/Mutable/Source/MutableTools/Private/MuT/NodeExtensionData.cpp
@@ -15,7 +15,7 @@ namespace mu
 	//---------------------------------------------------------------------------------------------
 	// Static initialisation
 	//---------------------------------------------------------------------------------------------
-	static NODE_TYPE s_type = NODE_TYPE("NodeExtensionData", Node::GetStaticType());
+	static FNodeType s_type = FNodeType("NodeExtensionData", Node::GetStaticType());
 
 	//---------------------------------------------------------------------------------------------
 	//!
@@ -55,14 +55,14 @@ namespace mu
 
 
 	//---------------------------------------------------------------------------------------------
-	const NODE_TYPE* NodeExtensionData::GetType() const
+	const FNodeType* NodeExtensionData::GetType() const
 	{
 		return GetStaticType();
 	}
 
 
 	//---------------------------------------------------------------------------------------------
-	const NODE_TYPE* NodeExtensionData::GetStaticType()
+	const FNodeType* NodeExtensionData::GetStaticType()
 	{
 		return &s_type;
 	}
diff --git a/Engine/Plugins/Experimental/Mutable/Source/MutableTools/Private/MuT/NodeExtensionDataConstant.cpp b/Engine/Plugins/Experimental/Mutable/Source/MutableTools/Private/MuT/NodeExtensionDataConstant.cpp
--- a/Engine/Plugins/Experimental/Mutable/Source/MutableTools/Private/MuT/NodeExtensionDataConstant.cpp
+++ b/Engine/Plugins/Experimental/Mutable/Source/MutableTools/Private/MuT/NodeExtensionDataConstant.cpp
@@ -14,8 +14,8 @@ namespace mu
 	//---------------------------------------------------------------------------------------------
 	// Static initialisation
 	//---------------------------------------------------------------------------------------------
-	NODE_TYPE NodeExtensionDataConstant::Private::s_type =
-			NODE_TYPE("ExtensionDataConstant", NodeExtensionData::GetStaticType());
+	FNodeType NodeExtensionDataConstant::Private::s_type =
+			FNodeType("ExtensionDataConstant", NodeExtensionData::GetStaticType());
 
 
 	//---------------------------------------------------------------------------------------------
diff --git a/Engine/Plugins/Experimental/Mutable/Source/MutableTools/Private/MuT/NodeExtensionDataSwitch.cpp b/Engine/Plugins/Experimental/Mutable/Source/MutableTools/Private/MuT/NodeExtensionDataSwitch.cpp
--- a/Engine/Plugins/Experimental/Mutable/Source/MutableTools/Private/MuT/NodeExtensionDataSwitch.cpp
+++ b/Engine/Plugins/Experimental/Mutable/Source/MutableTools/Private/MuT/NodeExtensionDataSwitch.cpp
@@ -12,8 +12,8 @@ namespace mu
 	//---------------------------------------------------------------------------------------------
 	// Static initialisation
 	//---------------------------------------------------------------------------------------------
-	NODE_TYPE NodeExtensionDataSwitch::Private::s_type =
-		NODE_TYPE("ExtensionDataSwitch", NodeExtensionData::GetStaticType());
+	FNodeType NodeExtensionDataSwitch::Private::s_type =
+		FNodeType("ExtensionDataSwitch", NodeExtensionData::GetStaticType());
 
 	//---------------------------------------------------------------------------------------------
 	//!
diff --git a/Engine/Plugins/Experimental/Mutable/Source/MutableTools/Private/MuT/NodeExtensionDataVariation.cpp b/Engine/Plugins/Experimental/Mutable/Source/MutableTools/Private/MuT/NodeExtensionDataVariation.cpp
--- a/Engine/Plugins/Experimental/Mutable/Source/MutableTools/Private/MuT/NodeExtensionDataVariation.cpp
+++ b/Engine/Plugins/Experimental/Mutable/Source/MutableTools/Private/MuT/NodeExtensionDataVariation.cpp
@@ -10,8 +10,8 @@ namespace mu
 	//---------------------------------------------------------------------------------------------
 	// Static initialisation
 	//---------------------------------------------------------------------------------------------
-	NODE_TYPE NodeExtensionDataVariation::Private::s_type =
-		NODE_TYPE("ExtensionDataVariation", NodeExtensionData::GetStaticType());
+	FNodeType NodeExtensionDataVariation::Private::s_type =
+		FNodeType("ExtensionDataVariation", NodeExtensionData::GetStaticType());
 
 	//---------------------------------------------------------------------------------------------
 	//!
diff --git a/Engine/Plugins/Experimental/Mutable/Source/MutableTools/Private/MuT/NodeImage.cpp b/Engine/Plugins/Experimental/Mutable/Source/MutableTools/Private/MuT/NodeImage.cpp
--- a/Engine/Plugins/Experimental/Mutable/Source/MutableTools/Private/MuT/NodeImage.cpp
+++ b/Engine/Plugins/Experimental/Mutable/Source/MutableTools/Private/MuT/NodeImage.cpp
@@ -1,6 +1,5 @@
 // Copyright Epic Games, Inc. All Rights Reserved.
 
-
 #include "MuT/NodeImage.h"
 
 #include "Misc/AssertionMacros.h"
@@ -38,7 +37,7 @@ namespace mu
 	//---------------------------------------------------------------------------------------------
 	// Static initialisation
 	//---------------------------------------------------------------------------------------------
-	static NODE_TYPE s_nodeImageType = 	NODE_TYPE( "NodeImage", Node::GetStaticType() );
+	static FNodeType s_nodeImageType = 	FNodeType( "NodeImage", Node::GetStaticType() );
 
 
 	//---------------------------------------------------------------------------------------------
@@ -103,14 +102,14 @@ namespace mu
 
 
 	//---------------------------------------------------------------------------------------------
-	const NODE_TYPE* NodeImage::GetType() const
+	const FNodeType* NodeImage::GetType() const
 	{
 		return GetStaticType();
 	}
 
 
 	//---------------------------------------------------------------------------------------------
-	const NODE_TYPE* NodeImage::GetStaticType()
+	const FNodeType* NodeImage::GetStaticType()
 	{
 		return &s_nodeImageType;
 	}
diff --git a/Engine/Plugins/Experimental/Mutable/Source/MutableTools/Private/MuT/NodeImageBinarise.cpp b/Engine/Plugins/Experimental/Mutable/Source/MutableTools/Private/MuT/NodeImageBinarise.cpp
--- a/Engine/Plugins/Experimental/Mutable/Source/MutableTools/Private/MuT/NodeImageBinarise.cpp
+++ b/Engine/Plugins/Experimental/Mutable/Source/MutableTools/Private/MuT/NodeImageBinarise.cpp
@@ -15,8 +15,8 @@ namespace mu
 	//---------------------------------------------------------------------------------------------
 	// Static initialisation
 	//---------------------------------------------------------------------------------------------
-	NODE_TYPE NodeImageBinarise::Private::s_type =
-			NODE_TYPE( "ImageMultiply", NodeImage::GetStaticType() );
+	FNodeType NodeImageBinarise::Private::s_type =
+			FNodeType( "ImageMultiply", NodeImage::GetStaticType() );
 
 	//---------------------------------------------------------------------------------------------
 
diff --git a/Engine/Plugins/Experimental/Mutable/Source/MutableTools/Private/MuT/NodeImageColourMap.cpp b/Engine/Plugins/Experimental/Mutable/Source/MutableTools/Private/MuT/NodeImageColourMap.cpp
--- a/Engine/Plugins/Experimental/Mutable/Source/MutableTools/Private/MuT/NodeImageColourMap.cpp
+++ b/Engine/Plugins/Experimental/Mutable/Source/MutableTools/Private/MuT/NodeImageColourMap.cpp
@@ -14,8 +14,8 @@ namespace mu
 	//---------------------------------------------------------------------------------------------
 	// Static initialisation
 	//---------------------------------------------------------------------------------------------
-	NODE_TYPE NodeImageColourMap::Private::s_type =
-			NODE_TYPE( "ImageColourMap", NodeImage::GetStaticType() );
+	FNodeType NodeImageColourMap::Private::s_type =
+			FNodeType( "ImageColourMap", NodeImage::GetStaticType() );
 
 
 	//---------------------------------------------------------------------------------------------
diff --git a/Engine/Plugins/Experimental/Mutable/Source/MutableTools/Private/MuT/NodeImageConditional.cpp b/Engine/Plugins/Experimental/Mutable/Source/MutableTools/Private/MuT/NodeImageConditional.cpp
--- a/Engine/Plugins/Experimental/Mutable/Source/MutableTools/Private/MuT/NodeImageConditional.cpp
+++ b/Engine/Plugins/Experimental/Mutable/Source/MutableTools/Private/MuT/NodeImageConditional.cpp
@@ -15,8 +15,8 @@ namespace mu
 	//---------------------------------------------------------------------------------------------
 	// Static initialisation
 	//---------------------------------------------------------------------------------------------
-    NODE_TYPE NodeImageConditional::Private::s_type =
-            NODE_TYPE( "ImageConditional", NodeImage::GetStaticType() );
+    FNodeType NodeImageConditional::Private::s_type =
+            FNodeType( "ImageConditional", NodeImage::GetStaticType() );
 
 
 	//---------------------------------------------------------------------------------------------
diff --git a/Engine/Plugins/Experimental/Mutable/Source/MutableTools/Private/MuT/NodeImageConstant.cpp b/Engine/Plugins/Experimental/Mutable/Source/MutableTools/Private/MuT/NodeImageConstant.cpp
--- a/Engine/Plugins/Experimental/Mutable/Source/MutableTools/Private/MuT/NodeImageConstant.cpp
+++ b/Engine/Plugins/Experimental/Mutable/Source/MutableTools/Private/MuT/NodeImageConstant.cpp
@@ -16,8 +16,8 @@ namespace mu
 	//---------------------------------------------------------------------------------------------
 	// Static initialisation
 	//---------------------------------------------------------------------------------------------
-	NODE_TYPE NodeImageConstant::Private::s_type =
-		NODE_TYPE("ImageConstant", NodeImage::GetStaticType());
+	FNodeType NodeImageConstant::Private::s_type =
+		FNodeType("ImageConstant", NodeImage::GetStaticType());
 
 
 	//---------------------------------------------------------------------------------------------
diff --git a/Engine/Plugins/Experimental/Mutable/Source/MutableTools/Private/MuT/NodeImageFormat.cpp b/Engine/Plugins/Experimental/Mutable/Source/MutableTools/Private/MuT/NodeImageFormat.cpp
--- a/Engine/Plugins/Experimental/Mutable/Source/MutableTools/Private/MuT/NodeImageFormat.cpp
+++ b/Engine/Plugins/Experimental/Mutable/Source/MutableTools/Private/MuT/NodeImageFormat.cpp
@@ -14,8 +14,8 @@ namespace mu
 	//---------------------------------------------------------------------------------------------
 	// Static initialisation
 	//---------------------------------------------------------------------------------------------
-	NODE_TYPE NodeImageFormat::Private::s_type =
-			NODE_TYPE( "ImageFormat", NodeImage::GetStaticType() );
+	FNodeType NodeImageFormat::Private::s_type =
+			FNodeType( "ImageFormat", NodeImage::GetStaticType() );
 
 
 	//---------------------------------------------------------------------------------------------
diff --git a/Engine/Plugins/Experimental/Mutable/Source/MutableTools/Private/MuT/NodeImageGradient.cpp b/Engine/Plugins/Experimental/Mutable/Source/MutableTools/Private/MuT/NodeImageGradient.cpp
--- a/Engine/Plugins/Experimental/Mutable/Source/MutableTools/Private/MuT/NodeImageGradient.cpp
+++ b/Engine/Plugins/Experimental/Mutable/Source/MutableTools/Private/MuT/NodeImageGradient.cpp
@@ -16,8 +16,8 @@ namespace mu
 	//---------------------------------------------------------------------------------------------
 	// Static initialisation
 	//---------------------------------------------------------------------------------------------
-	NODE_TYPE NodeImageGradient::Private::s_type =
-			NODE_TYPE( "ImageGradient", NodeImage::GetStaticType() );
+	FNodeType NodeImageGradient::Private::s_type =
+			FNodeType( "ImageGradient", NodeImage::GetStaticType() );
 
 
 	//---------------------------------------------------------------------------------------------
diff --git a/Engine/Plugins/Experimental/Mutable/Source/MutableTools/Private/MuT/NodeImageInterpolate.cpp b/Engine/Plugins/Experimental/Mutable/Source/MutableTools/Private/MuT/NodeImageInterpolate.cpp
--- a/Engine/Plugins/Experimental/Mutable/Source/MutableTools/Private/MuT/NodeImageInterpolate.cpp
+++ b/Engine/Plugins/Experimental/Mutable/Source/MutableTools/Private/MuT/NodeImageInterpolate.cpp
@@ -15,8 +15,8 @@ namespace mu
 	//---------------------------------------------------------------------------------------------
 	// Static initialisation
 	//---------------------------------------------------------------------------------------------
-	NODE_TYPE NodeImageInterpolate::Private::s_type =
-			NODE_TYPE( "ImageInterpolate", NodeImage::GetStaticType() );
+	FNodeType NodeImageInterpolate::Private::s_type =
+			FNodeType( "ImageInterpolate", NodeImage::GetStaticType() );
 
 
 	//---------------------------------------------------------------------------------------------
diff --git a/Engine/Plugins/Experimental/Mutable/Source/MutableTools/Private/MuT/NodeImageInvert.cpp b/Engine/Plugins/Experimental/Mutable/Source/MutableTools/Private/MuT/NodeImageInvert.cpp
--- a/Engine/Plugins/Experimental/Mutable/Source/MutableTools/Private/MuT/NodeImageInvert.cpp
+++ b/Engine/Plugins/Experimental/Mutable/Source/MutableTools/Private/MuT/NodeImageInvert.cpp
@@ -14,8 +14,8 @@ namespace mu
 	//---------------------------------------------------------------------------------------------
 	// Static initialisation
 	//---------------------------------------------------------------------------------------------
-	NODE_TYPE NodeImageInvert::Private::s_type =
-		NODE_TYPE("ImageInvert", NodeImage::GetStaticType());
+	FNodeType NodeImageInvert::Private::s_type =
+		FNodeType("ImageInvert", NodeImage::GetStaticType());
 
 	//---------------------------------------------------------------------------------------------
 
diff --git a/Engine/Plugins/Experimental/Mutable/Source/MutableTools/Private/MuT/NodeImageLayer.cpp b/Engine/Plugins/Experimental/Mutable/Source/MutableTools/Private/MuT/NodeImageLayer.cpp
--- a/Engine/Plugins/Experimental/Mutable/Source/MutableTools/Private/MuT/NodeImageLayer.cpp
+++ b/Engine/Plugins/Experimental/Mutable/Source/MutableTools/Private/MuT/NodeImageLayer.cpp
@@ -14,8 +14,8 @@ namespace mu
 	//---------------------------------------------------------------------------------------------
 	// Static initialisation
 	//---------------------------------------------------------------------------------------------
-	NODE_TYPE NodeImageLayer::Private::s_type =
-			NODE_TYPE( "ImageLayer", NodeImage::GetStaticType() );
+	FNodeType NodeImageLayer::Private::s_type =
+			FNodeType( "ImageLayer", NodeImage::GetStaticType() );
 
 
 	//---------------------------------------------------------------------------------------------
diff --git a/Engine/Plugins/Experimental/Mutable/Source/MutableTools/Private/MuT/NodeImageLayerColour.cpp b/Engine/Plugins/Experimental/Mutable/Source/MutableTools/Private/MuT/NodeImageLayerColour.cpp
--- a/Engine/Plugins/Experimental/Mutable/Source/MutableTools/Private/MuT/NodeImageLayerColour.cpp
+++ b/Engine/Plugins/Experimental/Mutable/Source/MutableTools/Private/MuT/NodeImageLayerColour.cpp
@@ -15,8 +15,8 @@ namespace mu
 	//---------------------------------------------------------------------------------------------
 	// Static initialisation
 	//---------------------------------------------------------------------------------------------
-	NODE_TYPE NodeImageLayerColour::Private::s_type =
-			NODE_TYPE( "ImageLayerColour", NodeImage::GetStaticType() );
+	FNodeType NodeImageLayerColour::Private::s_type =
+			FNodeType( "ImageLayerColour", NodeImage::GetStaticType() );
 
 
 	//---------------------------------------------------------------------------------------------
diff --git a/Engine/Plugins/Experimental/Mutable/Source/MutableTools/Private/MuT/NodeImageLuminance.cpp b/Engine/Plugins/Experimental/Mutable/Source/MutableTools/Private/MuT/NodeImageLuminance.cpp
--- a/Engine/Plugins/Experimental/Mutable/Source/MutableTools/Private/MuT/NodeImageLuminance.cpp
+++ b/Engine/Plugins/Experimental/Mutable/Source/MutableTools/Private/MuT/NodeImageLuminance.cpp
@@ -14,8 +14,8 @@ namespace mu
 	//---------------------------------------------------------------------------------------------
 	// Static initialisation
 	//---------------------------------------------------------------------------------------------
-	NODE_TYPE NodeImageLuminance::Private::s_type =
-			NODE_TYPE( "ImageLuminance", NodeImage::GetStaticType() );
+	FNodeType NodeImageLuminance::Private::s_type =
+			FNodeType( "ImageLuminance", NodeImage::GetStaticType() );
 
 
 	//---------------------------------------------------------------------------------------------
diff --git a/Engine/Plugins/Experimental/Mutable/Source/MutableTools/Private/MuT/NodeImageMipmap.cpp b/Engine/Plugins/Experimental/Mutable/Source/MutableTools/Private/MuT/NodeImageMipmap.cpp
--- a/Engine/Plugins/Experimental/Mutable/Source/MutableTools/Private/MuT/NodeImageMipmap.cpp
+++ b/Engine/Plugins/Experimental/Mutable/Source/MutableTools/Private/MuT/NodeImageMipmap.cpp
@@ -15,8 +15,8 @@ namespace mu
 	//---------------------------------------------------------------------------------------------
 	// Static initialisation
 	//---------------------------------------------------------------------------------------------
-    NODE_TYPE NodeImageMipmap::Private::s_type =
-            NODE_TYPE( "ImageMipmap", NodeImage::GetStaticType() );
+    FNodeType NodeImageMipmap::Private::s_type =
+            FNodeType( "ImageMipmap", NodeImage::GetStaticType() );
 
 
 	//---------------------------------------------------------------------------------------------
diff --git a/Engine/Plugins/Experimental/Mutable/Source/MutableTools/Private/MuT/NodeImageMultiLayer.cpp b/Engine/Plugins/Experimental/Mutable/Source/MutableTools/Private/MuT/NodeImageMultiLayer.cpp
--- a/Engine/Plugins/Experimental/Mutable/Source/MutableTools/Private/MuT/NodeImageMultiLayer.cpp
+++ b/Engine/Plugins/Experimental/Mutable/Source/MutableTools/Private/MuT/NodeImageMultiLayer.cpp
@@ -14,8 +14,8 @@ namespace mu
 	//---------------------------------------------------------------------------------------------
 	// Static initialisation
 	//---------------------------------------------------------------------------------------------
-    NODE_TYPE NodeImageMultiLayer::Private::s_type =
-            NODE_TYPE( "ImageMultiLayer", NodeImage::GetStaticType() );
+    FNodeType NodeImageMultiLayer::Private::s_type =
+            FNodeType( "ImageMultiLayer", NodeImage::GetStaticType() );
 
 
 	//---------------------------------------------------------------------------------------------
diff --git a/Engine/Plugins/Experimental/Mutable/Source/MutableTools/Private/MuT/NodeImageNormalComposite.cpp b/Engine/Plugins/Experimental/Mutable/Source/MutableTools/Private/MuT/NodeImageNormalComposite.cpp
--- a/Engine/Plugins/Experimental/Mutable/Source/MutableTools/Private/MuT/NodeImageNormalComposite.cpp
+++ b/Engine/Plugins/Experimental/Mutable/Source/MutableTools/Private/MuT/NodeImageNormalComposite.cpp
@@ -15,8 +15,8 @@ namespace mu
 	//---------------------------------------------------------------------------------------------
 	// Static initialisation
 	//---------------------------------------------------------------------------------------------
-	NODE_TYPE NodeImageNormalComposite::Private::s_type =
-			NODE_TYPE( "ImageNormalComposite", NodeImage::GetStaticType() );
+	FNodeType NodeImageNormalComposite::Private::s_type =
+			FNodeType( "ImageNormalComposite", NodeImage::GetStaticType() );
 
 
 	//---------------------------------------------------------------------------------------------
diff --git a/Engine/Plugins/Experimental/Mutable/Source/MutableTools/Private/MuT/NodeImageParameter.cpp b/Engine/Plugins/Experimental/Mutable/Source/MutableTools/Private/MuT/NodeImageParameter.cpp
--- a/Engine/Plugins/Experimental/Mutable/Source/MutableTools/Private/MuT/NodeImageParameter.cpp
+++ b/Engine/Plugins/Experimental/Mutable/Source/MutableTools/Private/MuT/NodeImageParameter.cpp
@@ -14,8 +14,8 @@ namespace mu
 	//---------------------------------------------------------------------------------------------
 	// Static initialisation
 	//---------------------------------------------------------------------------------------------
-    NODE_TYPE NodeImageParameter::Private::s_type =
-            NODE_TYPE( "ImageParameter", NodeImage::GetStaticType() );
+    FNodeType NodeImageParameter::Private::s_type =
+            FNodeType( "ImageParameter", NodeImage::GetStaticType() );
 
 
 	//---------------------------------------------------------------------------------------------
diff --git a/Engine/Plugins/Experimental/Mutable/Source/MutableTools/Private/MuT/NodeImagePlainColour.cpp b/Engine/Plugins/Experimental/Mutable/Source/MutableTools/Private/MuT/NodeImagePlainColour.cpp
--- a/Engine/Plugins/Experimental/Mutable/Source/MutableTools/Private/MuT/NodeImagePlainColour.cpp
+++ b/Engine/Plugins/Experimental/Mutable/Source/MutableTools/Private/MuT/NodeImagePlainColour.cpp
@@ -15,8 +15,8 @@ namespace mu
 	//---------------------------------------------------------------------------------------------
 	// Static initialisation
 	//---------------------------------------------------------------------------------------------
-	NODE_TYPE NodeImagePlainColour::Private::s_type =
-			NODE_TYPE( "ImagePlainColour", NodeImage::GetStaticType() );
+	FNodeType NodeImagePlainColour::Private::s_type =
+			FNodeType( "ImagePlainColour", NodeImage::GetStaticType() );
 
 
 	//---------------------------------------------------------------------------------------------
diff --git a/Engine/Plugins/Experimental/Mutable/Source/MutableTools/Private/MuT/NodeImageProject.cpp b/Engine/Plugins/Experimental/Mutable/Source/MutableTools/Private/MuT/NodeImageProject.cpp
--- a/Engine/Plugins/Experimental/Mutable/Source/MutableTools/Private/MuT/NodeImageProject.cpp
+++ b/Engine/Plugins/Experimental/Mutable/Source/MutableTools/Private/MuT/NodeImageProject.cpp
@@ -19,8 +19,8 @@ namespace mu
 	//---------------------------------------------------------------------------------------------
 	// Static initialisation
 	//---------------------------------------------------------------------------------------------
-	NODE_TYPE NodeImageProject::Private::s_type =
-			NODE_TYPE( "ImageProject", NodeImage::GetStaticType() );
+	FNodeType NodeImageProject::Private::s_type =
+			FNodeType( "ImageProject", NodeImage::GetStaticType() );
 
 
 	//---------------------------------------------------------------------------------------------
diff --git a/Engine/Plugins/Experimental/Mutable/Source/MutableTools/Private/MuT/NodeImageResize.cpp b/Engine/Plugins/Experimental/Mutable/Source/MutableTools/Private/MuT/NodeImageResize.cpp
--- a/Engine/Plugins/Experimental/Mutable/Source/MutableTools/Private/MuT/NodeImageResize.cpp
+++ b/Engine/Plugins/Experimental/Mutable/Source/MutableTools/Private/MuT/NodeImageResize.cpp
@@ -14,8 +14,8 @@ namespace mu
 	//---------------------------------------------------------------------------------------------
 	// Static initialisation
 	//---------------------------------------------------------------------------------------------
-	NODE_TYPE NodeImageResize::Private::s_type =
-			NODE_TYPE( "ImageResize", NodeImage::GetStaticType() );
+	FNodeType NodeImageResize::Private::s_type =
+			FNodeType( "ImageResize", NodeImage::GetStaticType() );
 
 	//---------------------------------------------------------------------------------------------
 
diff --git a/Engine/Plugins/Experimental/Mutable/Source/MutableTools/Private/MuT/NodeImageSaturate.cpp b/Engine/Plugins/Experimental/Mutable/Source/MutableTools/Private/MuT/NodeImageSaturate.cpp
--- a/Engine/Plugins/Experimental/Mutable/Source/MutableTools/Private/MuT/NodeImageSaturate.cpp
+++ b/Engine/Plugins/Experimental/Mutable/Source/MutableTools/Private/MuT/NodeImageSaturate.cpp
@@ -15,8 +15,8 @@ namespace mu
 	//---------------------------------------------------------------------------------------------
 	// Static initialisation
 	//---------------------------------------------------------------------------------------------
-	NODE_TYPE NodeImageSaturate::Private::s_type =
-			NODE_TYPE( "ImageSaturate", NodeImage::GetStaticType() );
+	FNodeType NodeImageSaturate::Private::s_type =
+			FNodeType( "ImageSaturate", NodeImage::GetStaticType() );
 
 
 	//---------------------------------------------------------------------------------------------
diff --git a/Engine/Plugins/Experimental/Mutable/Source/MutableTools/Private/MuT/NodeImageSwitch.cpp b/Engine/Plugins/Experimental/Mutable/Source/MutableTools/Private/MuT/NodeImageSwitch.cpp
--- a/Engine/Plugins/Experimental/Mutable/Source/MutableTools/Private/MuT/NodeImageSwitch.cpp
+++ b/Engine/Plugins/Experimental/Mutable/Source/MutableTools/Private/MuT/NodeImageSwitch.cpp
@@ -15,8 +15,8 @@ namespace mu
 	//---------------------------------------------------------------------------------------------
 	// Static initialisation
 	//---------------------------------------------------------------------------------------------
-	NODE_TYPE NodeImageSwitch::Private::s_type =
-			NODE_TYPE( "ImageSwitch", NodeImage::GetStaticType() );
+	FNodeType NodeImageSwitch::Private::s_type =
+			FNodeType( "ImageSwitch", NodeImage::GetStaticType() );
 
 
 	//---------------------------------------------------------------------------------------------
diff --git a/Engine/Plugins/Experimental/Mutable/Source/MutableTools/Private/MuT/NodeImageSwizzle.cpp b/Engine/Plugins/Experimental/Mutable/Source/MutableTools/Private/MuT/NodeImageSwizzle.cpp
--- a/Engine/Plugins/Experimental/Mutable/Source/MutableTools/Private/MuT/NodeImageSwizzle.cpp
+++ b/Engine/Plugins/Experimental/Mutable/Source/MutableTools/Private/MuT/NodeImageSwizzle.cpp
@@ -15,8 +15,8 @@ namespace mu
 	//---------------------------------------------------------------------------------------------
 	// Static initialisation
 	//---------------------------------------------------------------------------------------------
-	NODE_TYPE NodeImageSwizzle::Private::s_type =
-			NODE_TYPE( "ImageSwizzle", NodeImage::GetStaticType() );
+	FNodeType NodeImageSwizzle::Private::s_type =
+			FNodeType( "ImageSwizzle", NodeImage::GetStaticType() );
 
 
 	//---------------------------------------------------------------------------------------------
diff --git a/Engine/Plugins/Experimental/Mutable/Source/MutableTools/Private/MuT/NodeImageTable.cpp b/Engine/Plugins/Experimental/Mutable/Source/MutableTools/Private/MuT/NodeImageTable.cpp
--- a/Engine/Plugins/Experimental/Mutable/Source/MutableTools/Private/MuT/NodeImageTable.cpp
+++ b/Engine/Plugins/Experimental/Mutable/Source/MutableTools/Private/MuT/NodeImageTable.cpp
@@ -16,8 +16,8 @@ namespace mu
 	//---------------------------------------------------------------------------------------------
 	// Static initialisation
 	//---------------------------------------------------------------------------------------------
-	NODE_TYPE NodeImageTable::Private::s_type =
-			NODE_TYPE( "TableImage", NodeImage::GetStaticType() );
+	FNodeType NodeImageTable::Private::s_type =
+			FNodeType( "TableImage", NodeImage::GetStaticType() );
 
 
 	//---------------------------------------------------------------------------------------------
