[mutable] Remove usage of RTTI.

diff --git a/Engine/Plugins/Experimental/Mutable/Source/MutableTools/Private/MuT/NodeImageTransform.cpp b/Engine/Plugins/Experimental/Mutable/Source/MutableTools/Private/MuT/NodeImageTransform.cpp
--- a/Engine/Plugins/Experimental/Mutable/Source/MutableTools/Private/MuT/NodeImageTransform.cpp
+++ b/Engine/Plugins/Experimental/Mutable/Source/MutableTools/Private/MuT/NodeImageTransform.cpp
@@ -15,8 +15,8 @@ namespace mu
 	//---------------------------------------------------------------------------------------------
 	// Static initialisation
 	//---------------------------------------------------------------------------------------------
-	NODE_TYPE NodeImageTransform::Private::s_type =
-			NODE_TYPE( "ImageTransform", NodeImage::GetStaticType() );
+	FNodeType NodeImageTransform::Private::s_type =
+			FNodeType( "ImageTransform", NodeImage::GetStaticType() );
 
 	//---------------------------------------------------------------------------------------------
 
diff --git a/Engine/Plugins/Experimental/Mutable/Source/MutableTools/Private/MuT/NodeImageVariation.cpp b/Engine/Plugins/Experimental/Mutable/Source/MutableTools/Private/MuT/NodeImageVariation.cpp
--- a/Engine/Plugins/Experimental/Mutable/Source/MutableTools/Private/MuT/NodeImageVariation.cpp
+++ b/Engine/Plugins/Experimental/Mutable/Source/MutableTools/Private/MuT/NodeImageVariation.cpp
@@ -14,8 +14,8 @@ namespace mu
     //---------------------------------------------------------------------------------------------
     // Static initialisation
     //---------------------------------------------------------------------------------------------
-    NODE_TYPE NodeImageVariation::Private::s_type =
-        NODE_TYPE( "ImageVariation", NodeImage::GetStaticType() );
+    FNodeType NodeImageVariation::Private::s_type =
+        FNodeType( "ImageVariation", NodeImage::GetStaticType() );
 
 
     //---------------------------------------------------------------------------------------------
diff --git a/Engine/Plugins/Experimental/Mutable/Source/MutableTools/Private/MuT/NodeLOD.cpp b/Engine/Plugins/Experimental/Mutable/Source/MutableTools/Private/MuT/NodeLOD.cpp
--- a/Engine/Plugins/Experimental/Mutable/Source/MutableTools/Private/MuT/NodeLOD.cpp
+++ b/Engine/Plugins/Experimental/Mutable/Source/MutableTools/Private/MuT/NodeLOD.cpp
@@ -1,6 +1,5 @@
 // Copyright Epic Games, Inc. All Rights Reserved.
 
-
 #include "MuT/NodeLOD.h"
 
 #include "Misc/AssertionMacros.h"
@@ -18,15 +17,14 @@ namespace mu
 	//---------------------------------------------------------------------------------------------
 	// Static initialisation
 	//---------------------------------------------------------------------------------------------
-	NODE_TYPE NodeLOD::Private::s_type =
-			NODE_TYPE( "LOD", NodeLOD::GetStaticType() );
+	FNodeType NodeLOD::Private::s_type = FNodeType( "LOD", Node::GetStaticType() );
 
 
 	//---------------------------------------------------------------------------------------------
 	//!
 	//---------------------------------------------------------------------------------------------
 
-	MUTABLE_IMPLEMENT_NODE( NodeLOD, EType::LOD, Node, Node::EType::LOD);
+	MUTABLE_IMPLEMENT_NODE( NodeLOD, EType::LOD, Node, Node::EType::None);
 
 
 	//---------------------------------------------------------------------------------------------
diff --git a/Engine/Plugins/Experimental/Mutable/Source/MutableTools/Private/MuT/NodeLayout.cpp b/Engine/Plugins/Experimental/Mutable/Source/MutableTools/Private/MuT/NodeLayout.cpp
--- a/Engine/Plugins/Experimental/Mutable/Source/MutableTools/Private/MuT/NodeLayout.cpp
+++ b/Engine/Plugins/Experimental/Mutable/Source/MutableTools/Private/MuT/NodeLayout.cpp
@@ -1,6 +1,5 @@
 // Copyright Epic Games, Inc. All Rights Reserved.
 
-
 #include "MuT/NodeLayout.h"
 
 #include "Math/IntPoint.h"
@@ -24,18 +23,19 @@ namespace mu
 	//---------------------------------------------------------------------------------------------
 	// Static initialisation
 	//---------------------------------------------------------------------------------------------
-	static NODE_TYPE s_nodeLayoutType = NODE_TYPE( "NodeLayout", Node::GetStaticType() );
+	static FNodeType s_nodeLayoutType = FNodeType( "NodeLayout", Node::GetStaticType() );
 
 
 	//---------------------------------------------------------------------------------------------
 	void NodeLayout::Serialise( const NodeLayout* p, OutputArchive& arch )
 	{
-        uint32_t ver = 0;
+        uint32 ver = 0;
 		arch << ver;
 
 	#define SERIALISE_CHILDREN( C, ID ) \
-		( const C* pTyped = dynamic_cast<const C*>(p) )			\
+		( p->GetType()==C::GetStaticType() )					\
 		{ 														\
+			const C* pTyped = static_cast<const C*>(p);			\
             arch << (uint32_t)ID;								\
 			C::Serialise( pTyped, arch );						\
 		}														\
@@ -50,11 +50,11 @@ namespace mu
 	//---------------------------------------------------------------------------------------------
 	NodeLayoutPtr NodeLayout::StaticUnserialise( InputArchive& arch )
 	{
-        uint32_t ver;
+        uint32 ver;
 		arch >> ver;
 		check( ver == 0 );
 
-        uint32_t id;
+        uint32 id;
 		arch >> id;
 
 		switch (id)
@@ -68,14 +68,14 @@ namespace mu
 
 
 	//---------------------------------------------------------------------------------------------
-	const NODE_TYPE* NodeLayout::GetType() const
+	const FNodeType* NodeLayout::GetType() const
 	{
 		return GetStaticType();
 	}
 
 
 	//---------------------------------------------------------------------------------------------
-	const NODE_TYPE* NodeLayout::GetStaticType()
+	const FNodeType* NodeLayout::GetStaticType()
 	{
 		return &s_nodeLayoutType;
 	}
@@ -87,8 +87,8 @@ namespace mu
 	//---------------------------------------------------------------------------------------------
 	// Static initialisation
 	//---------------------------------------------------------------------------------------------
-	NODE_TYPE NodeLayoutBlocks::Private::s_type =
-			NODE_TYPE( "LayoutBlocks", NodeLayout::GetStaticType() );
+	FNodeType NodeLayoutBlocks::Private::s_type =
+			FNodeType( "LayoutBlocks", NodeLayout::GetStaticType() );
 
 
 	//---------------------------------------------------------------------------------------------
diff --git a/Engine/Plugins/Experimental/Mutable/Source/MutableTools/Private/MuT/NodeMesh.cpp b/Engine/Plugins/Experimental/Mutable/Source/MutableTools/Private/MuT/NodeMesh.cpp
--- a/Engine/Plugins/Experimental/Mutable/Source/MutableTools/Private/MuT/NodeMesh.cpp
+++ b/Engine/Plugins/Experimental/Mutable/Source/MutableTools/Private/MuT/NodeMesh.cpp
@@ -23,7 +23,6 @@
 #include "MuT/NodeMeshTransform.h"
 #include "MuT/NodeMeshVariation.h"
 
-#include <stdint.h>
 
 namespace mu
 {
@@ -31,16 +30,16 @@ namespace mu
 	//---------------------------------------------------------------------------------------------
 	// Static initialisation
 	//---------------------------------------------------------------------------------------------
-	static NODE_TYPE s_nodeMeshType = NODE_TYPE( "NodeMesh", Node::GetStaticType() );
+	static FNodeType s_nodeMeshType = FNodeType( "NodeMesh", Node::GetStaticType() );
 
 
 	//---------------------------------------------------------------------------------------------
 	void NodeMesh::Serialise( const NodeMesh* p, OutputArchive& arch )
 	{
-        uint32_t ver = 0;
+        uint32 ver = 0;
 		arch << ver;
 
-		arch << uint32_t(p->Type);
+		arch << uint32(p->Type);
 		p->SerialiseWrapper( arch );
     }
 
@@ -48,11 +47,11 @@ namespace mu
     //---------------------------------------------------------------------------------------------
 	NodeMeshPtr NodeMesh::StaticUnserialise( InputArchive& arch )
 	{
-        uint32_t ver;
+        uint32 ver;
 		arch >> ver;
 		check( ver == 0 );
 
-        uint32_t id;
+        uint32 id;
 		arch >> id;
 
 		switch (EType(id))
@@ -82,14 +81,14 @@ namespace mu
 
 
 	//---------------------------------------------------------------------------------------------
-	const NODE_TYPE* NodeMesh::GetType() const
+	const FNodeType* NodeMesh::GetType() const
 	{
 		return GetStaticType();
 	}
 
 
 	//---------------------------------------------------------------------------------------------
-	const NODE_TYPE* NodeMesh::GetStaticType()
+	const FNodeType* NodeMesh::GetStaticType()
 	{
 		return &s_nodeMeshType;
 	}
diff --git a/Engine/Plugins/Experimental/Mutable/Source/MutableTools/Private/MuT/NodeMeshApplyPose.cpp b/Engine/Plugins/Experimental/Mutable/Source/MutableTools/Private/MuT/NodeMeshApplyPose.cpp
--- a/Engine/Plugins/Experimental/Mutable/Source/MutableTools/Private/MuT/NodeMeshApplyPose.cpp
+++ b/Engine/Plugins/Experimental/Mutable/Source/MutableTools/Private/MuT/NodeMeshApplyPose.cpp
@@ -1,6 +1,5 @@
 // Copyright Epic Games, Inc. All Rights Reserved.
 
-
 #include "MuT/NodeMeshApplyPose.h"
 
 #include "Misc/AssertionMacros.h"
@@ -10,15 +9,14 @@
 #include "MuT/NodePrivate.h"
 
 
-
 namespace mu
 {
 
 	//---------------------------------------------------------------------------------------------
 	// Static initialisation
 	//---------------------------------------------------------------------------------------------
-    NODE_TYPE NodeMeshApplyPose::Private::s_type =
-            NODE_TYPE( "MeshApplyPose", NodeMesh::GetStaticType() );
+    FNodeType NodeMeshApplyPose::Private::s_type =
+            FNodeType( "MeshApplyPose", NodeMesh::GetStaticType() );
 
 
 	//---------------------------------------------------------------------------------------------
@@ -65,8 +63,7 @@ namespace mu
 
 		if ( m_pBase )
 		{
-			NodeMesh::Private* pPrivate =
-					dynamic_cast<NodeMesh::Private*>( m_pBase->GetBasePrivate() );
+			NodeMesh::Private* pPrivate = static_cast<NodeMesh::Private*>( m_pBase->GetBasePrivate() );
 
 			pResult = pPrivate->GetLayout( index );
 		}
diff --git a/Engine/Plugins/Experimental/Mutable/Source/MutableTools/Private/MuT/NodeMeshClipDeform.cpp b/Engine/Plugins/Experimental/Mutable/Source/MutableTools/Private/MuT/NodeMeshClipDeform.cpp
--- a/Engine/Plugins/Experimental/Mutable/Source/MutableTools/Private/MuT/NodeMeshClipDeform.cpp
+++ b/Engine/Plugins/Experimental/Mutable/Source/MutableTools/Private/MuT/NodeMeshClipDeform.cpp
@@ -16,8 +16,8 @@ namespace mu
 	//---------------------------------------------------------------------------------------------
 	// Static initialisation
 	//---------------------------------------------------------------------------------------------
-	NODE_TYPE NodeMeshClipDeform::Private::s_type =
-			NODE_TYPE( "MeshClipDeform", NodeMesh::GetStaticType() );
+	FNodeType NodeMeshClipDeform::Private::s_type =
+			FNodeType( "MeshClipDeform", NodeMesh::GetStaticType() );
 
 
 	//---------------------------------------------------------------------------------------------
@@ -64,9 +64,7 @@ namespace mu
 
 		if (m_pBaseMesh)
 		{
-			NodeMesh::Private* pPrivate =
-					dynamic_cast<NodeMesh::Private*>(m_pBaseMesh->GetBasePrivate() );
-
+			NodeMesh::Private* pPrivate = static_cast<NodeMesh::Private*>(m_pBaseMesh->GetBasePrivate() );
 			pResult = pPrivate->GetLayout( index );
 		}
 
diff --git a/Engine/Plugins/Experimental/Mutable/Source/MutableTools/Private/MuT/NodeMeshClipMorphPlane.cpp b/Engine/Plugins/Experimental/Mutable/Source/MutableTools/Private/MuT/NodeMeshClipMorphPlane.cpp
--- a/Engine/Plugins/Experimental/Mutable/Source/MutableTools/Private/MuT/NodeMeshClipMorphPlane.cpp
+++ b/Engine/Plugins/Experimental/Mutable/Source/MutableTools/Private/MuT/NodeMeshClipMorphPlane.cpp
@@ -1,6 +1,5 @@
 // Copyright Epic Games, Inc. All Rights Reserved.
 
-
 #include "MuT/NodeMeshClipMorphPlane.h"
 
 #include "Misc/AssertionMacros.h"
@@ -17,8 +16,8 @@ namespace mu
 	//---------------------------------------------------------------------------------------------
 	// Static initialisation
 	//---------------------------------------------------------------------------------------------
-    NODE_TYPE NodeMeshClipMorphPlane::Private::s_type =
-            NODE_TYPE( "MeshClipMorphPlane", NodeMesh::GetStaticType() );
+    FNodeType NodeMeshClipMorphPlane::Private::s_type =
+            FNodeType( "MeshClipMorphPlane", NodeMesh::GetStaticType() );
 
 
 	//---------------------------------------------------------------------------------------------
@@ -90,9 +89,7 @@ namespace mu
 
 		if ( m_pSource )
 		{
-			NodeMesh::Private* pPrivate =
-					dynamic_cast<NodeMesh::Private*>( m_pSource->GetBasePrivate() );
-
+			NodeMesh::Private* pPrivate = static_cast<NodeMesh::Private*>( m_pSource->GetBasePrivate() );
 			pResult = pPrivate->GetLayout( index );
 		}
 
diff --git a/Engine/Plugins/Experimental/Mutable/Source/MutableTools/Private/MuT/NodeMeshClipWithMesh.cpp b/Engine/Plugins/Experimental/Mutable/Source/MutableTools/Private/MuT/NodeMeshClipWithMesh.cpp
--- a/Engine/Plugins/Experimental/Mutable/Source/MutableTools/Private/MuT/NodeMeshClipWithMesh.cpp
+++ b/Engine/Plugins/Experimental/Mutable/Source/MutableTools/Private/MuT/NodeMeshClipWithMesh.cpp
@@ -16,8 +16,8 @@ namespace mu
 	//---------------------------------------------------------------------------------------------
 	// Static initialisation
 	//---------------------------------------------------------------------------------------------
-    NODE_TYPE NodeMeshClipWithMesh::Private::s_type =
-            NODE_TYPE( "MeshClipMorphPlane", NodeMesh::GetStaticType() );
+    FNodeType NodeMeshClipWithMesh::Private::s_type =
+            FNodeType( "MeshClipMorphPlane", NodeMesh::GetStaticType() );
 
 
 	//---------------------------------------------------------------------------------------------
@@ -57,9 +57,7 @@ namespace mu
 
 		if ( m_pSource )
 		{
-			NodeMesh::Private* pPrivate =
-					dynamic_cast<NodeMesh::Private*>( m_pSource->GetBasePrivate() );
-
+			NodeMesh::Private* pPrivate = static_cast<NodeMesh::Private*>( m_pSource->GetBasePrivate() );
 			pResult = pPrivate->GetLayout( index );
 		}
 
diff --git a/Engine/Plugins/Experimental/Mutable/Source/MutableTools/Private/MuT/NodeMeshConstant.cpp b/Engine/Plugins/Experimental/Mutable/Source/MutableTools/Private/MuT/NodeMeshConstant.cpp
--- a/Engine/Plugins/Experimental/Mutable/Source/MutableTools/Private/MuT/NodeMeshConstant.cpp
+++ b/Engine/Plugins/Experimental/Mutable/Source/MutableTools/Private/MuT/NodeMeshConstant.cpp
@@ -15,8 +15,8 @@ namespace mu
 	//---------------------------------------------------------------------------------------------
 	// Static initialisation
 	//---------------------------------------------------------------------------------------------
-	NODE_TYPE NodeMeshConstant::Private::s_type =
-			NODE_TYPE( "MeshConstant", NodeMesh::GetStaticType() );
+	FNodeType NodeMeshConstant::Private::s_type =
+			FNodeType( "MeshConstant", NodeMesh::GetStaticType() );
 
 
 	//---------------------------------------------------------------------------------------------
diff --git a/Engine/Plugins/Experimental/Mutable/Source/MutableTools/Private/MuT/NodeMeshFormat.cpp b/Engine/Plugins/Experimental/Mutable/Source/MutableTools/Private/MuT/NodeMeshFormat.cpp
--- a/Engine/Plugins/Experimental/Mutable/Source/MutableTools/Private/MuT/NodeMeshFormat.cpp
+++ b/Engine/Plugins/Experimental/Mutable/Source/MutableTools/Private/MuT/NodeMeshFormat.cpp
@@ -17,8 +17,8 @@ class FMeshBufferSet;
 	//---------------------------------------------------------------------------------------------
 	// Static initialisation
 	//---------------------------------------------------------------------------------------------
-	NODE_TYPE NodeMeshFormat::Private::s_type =
-			NODE_TYPE( "MeshFormat", NodeMesh::GetStaticType() );
+	FNodeType NodeMeshFormat::Private::s_type =
+			FNodeType( "MeshFormat", NodeMesh::GetStaticType() );
 
 
 	//---------------------------------------------------------------------------------------------
@@ -72,9 +72,7 @@ class FMeshBufferSet;
 
 		if ( m_pSource )
 		{
-			NodeMesh::Private* pPrivate =
-					dynamic_cast<NodeMesh::Private*>( m_pSource->GetBasePrivate() );
-
+			NodeMesh::Private* pPrivate = static_cast<NodeMesh::Private*>( m_pSource->GetBasePrivate() );
 			pResult = pPrivate->GetLayout( index );
 		}
 
diff --git a/Engine/Plugins/Experimental/Mutable/Source/MutableTools/Private/MuT/NodeMeshFragment.cpp b/Engine/Plugins/Experimental/Mutable/Source/MutableTools/Private/MuT/NodeMeshFragment.cpp
--- a/Engine/Plugins/Experimental/Mutable/Source/MutableTools/Private/MuT/NodeMeshFragment.cpp
+++ b/Engine/Plugins/Experimental/Mutable/Source/MutableTools/Private/MuT/NodeMeshFragment.cpp
@@ -16,8 +16,8 @@ namespace mu
 	//---------------------------------------------------------------------------------------------
 	// Static initialisation
 	//---------------------------------------------------------------------------------------------
-	NODE_TYPE NodeMeshFragment::Private::s_type =
-			NODE_TYPE( "MeshFragment", NodeMesh::GetStaticType() );
+	FNodeType NodeMeshFragment::Private::s_type =
+			FNodeType( "MeshFragment", NodeMesh::GetStaticType() );
 
 
 	//---------------------------------------------------------------------------------------------
@@ -102,10 +102,7 @@ namespace mu
 		if ( m_pMesh )
 		{
 			// TODO: Cut a fragment out of the layout.
-
-			NodeMesh::Private* pPrivate =
-					dynamic_cast<NodeMesh::Private*>( m_pMesh->GetBasePrivate() );
-
+			NodeMesh::Private* pPrivate = static_cast<NodeMesh::Private*>( m_pMesh->GetBasePrivate() );
 			pResult = pPrivate->GetLayout( index );
 		}
 
diff --git a/Engine/Plugins/Experimental/Mutable/Source/MutableTools/Private/MuT/NodeMeshGeometryOperation.cpp b/Engine/Plugins/Experimental/Mutable/Source/MutableTools/Private/MuT/NodeMeshGeometryOperation.cpp
--- a/Engine/Plugins/Experimental/Mutable/Source/MutableTools/Private/MuT/NodeMeshGeometryOperation.cpp
+++ b/Engine/Plugins/Experimental/Mutable/Source/MutableTools/Private/MuT/NodeMeshGeometryOperation.cpp
@@ -18,8 +18,8 @@ namespace mu
 	//---------------------------------------------------------------------------------------------
 	// Static initialisation
 	//---------------------------------------------------------------------------------------------
-	NODE_TYPE NodeMeshGeometryOperation::Private::s_type =
-			NODE_TYPE( "MeshGeometryOperation", NodeMesh::GetStaticType() );
+	FNodeType NodeMeshGeometryOperation::Private::s_type =
+			FNodeType( "MeshGeometryOperation", NodeMesh::GetStaticType() );
 
 
 	//---------------------------------------------------------------------------------------------
@@ -94,9 +94,7 @@ namespace mu
 
 		if (m_pMeshA)
 		{
-			NodeMesh::Private* pPrivate =
-					dynamic_cast<NodeMesh::Private*>(m_pMeshA->GetBasePrivate() );
-
+			NodeMesh::Private* pPrivate = static_cast<NodeMesh::Private*>(m_pMeshA->GetBasePrivate() );
 			pResult = pPrivate->GetLayout( index );
 		}
 
diff --git a/Engine/Plugins/Experimental/Mutable/Source/MutableTools/Private/MuT/NodeMeshInterpolate.cpp b/Engine/Plugins/Experimental/Mutable/Source/MutableTools/Private/MuT/NodeMeshInterpolate.cpp
--- a/Engine/Plugins/Experimental/Mutable/Source/MutableTools/Private/MuT/NodeMeshInterpolate.cpp
+++ b/Engine/Plugins/Experimental/Mutable/Source/MutableTools/Private/MuT/NodeMeshInterpolate.cpp
@@ -17,8 +17,8 @@ namespace mu
 	//---------------------------------------------------------------------------------------------
 	// Static initialisation
 	//---------------------------------------------------------------------------------------------
-	NODE_TYPE NodeMeshInterpolate::Private::s_type =
-			NODE_TYPE( "MeshInterpolate", NodeMesh::GetStaticType() );
+	FNodeType NodeMeshInterpolate::Private::s_type =
+			FNodeType( "MeshInterpolate", NodeMesh::GetStaticType() );
 
 
 	//---------------------------------------------------------------------------------------------
@@ -105,9 +105,7 @@ namespace mu
 		// TODO: Substract layouts too? Usually they are ignored.
 		if ( m_targets.Num()>0 && m_targets[0] )
 		{
-			NodeMesh::Private* pPrivate =
-					dynamic_cast<NodeMesh::Private*>( m_targets[0]->GetBasePrivate() );
-
+			NodeMesh::Private* pPrivate = static_cast<NodeMesh::Private*>( m_targets[0]->GetBasePrivate() );
 			pResult = pPrivate->GetLayout( index );
 		}
 
diff --git a/Engine/Plugins/Experimental/Mutable/Source/MutableTools/Private/MuT/NodeMeshMakeMorph.cpp b/Engine/Plugins/Experimental/Mutable/Source/MutableTools/Private/MuT/NodeMeshMakeMorph.cpp
--- a/Engine/Plugins/Experimental/Mutable/Source/MutableTools/Private/MuT/NodeMeshMakeMorph.cpp
+++ b/Engine/Plugins/Experimental/Mutable/Source/MutableTools/Private/MuT/NodeMeshMakeMorph.cpp
@@ -16,8 +16,8 @@ namespace mu
 	//---------------------------------------------------------------------------------------------
 	// Static initialisation
 	//---------------------------------------------------------------------------------------------
-    NODE_TYPE NodeMeshMakeMorph::Private::s_type =
-			NODE_TYPE( "MeshMakeMorph", NodeMesh::GetStaticType() );
+    FNodeType NodeMeshMakeMorph::Private::s_type =
+			FNodeType( "MeshMakeMorph", NodeMesh::GetStaticType() );
 
 
 	//---------------------------------------------------------------------------------------------
@@ -77,9 +77,7 @@ namespace mu
 		// TODO: Substract layouts too? Usually they are ignored.
 		if ( m_pBase )
 		{
-			NodeMesh::Private* pPrivate =
-					dynamic_cast<NodeMesh::Private*>( m_pBase->GetBasePrivate() );
-
+			NodeMesh::Private* pPrivate = static_cast<NodeMesh::Private*>( m_pBase->GetBasePrivate() );
 			pResult = pPrivate->GetLayout( index );
 		}
 
diff --git a/Engine/Plugins/Experimental/Mutable/Source/MutableTools/Private/MuT/NodeMeshMorph.cpp b/Engine/Plugins/Experimental/Mutable/Source/MutableTools/Private/MuT/NodeMeshMorph.cpp
--- a/Engine/Plugins/Experimental/Mutable/Source/MutableTools/Private/MuT/NodeMeshMorph.cpp
+++ b/Engine/Plugins/Experimental/Mutable/Source/MutableTools/Private/MuT/NodeMeshMorph.cpp
@@ -1,6 +1,5 @@
 // Copyright Epic Games, Inc. All Rights Reserved.
 
-
 #include "MuT/NodeMeshMorph.h"
 
 #include "Containers/Array.h"
@@ -19,7 +18,7 @@ namespace mu
 	//---------------------------------------------------------------------------------------------
 	// Static initialisation
 	//---------------------------------------------------------------------------------------------
-	NODE_TYPE NodeMeshMorph::Private::s_type = NODE_TYPE( "MeshMorph", NodeMesh::GetStaticType() );
+	FNodeType NodeMeshMorph::Private::s_type = FNodeType( "MeshMorph", NodeMesh::GetStaticType() );
 
 
 	//---------------------------------------------------------------------------------------------
@@ -102,9 +101,7 @@ namespace mu
 
 		if ( Base )
 		{
-			NodeMesh::Private* pPrivate =
-					dynamic_cast<NodeMesh::Private*>( Base->GetBasePrivate() );
-
+			NodeMesh::Private* pPrivate = static_cast<NodeMesh::Private*>( Base->GetBasePrivate() );
 			pResult = pPrivate->GetLayout( index );
 		}
 
diff --git a/Engine/Plugins/Experimental/Mutable/Source/MutableTools/Private/MuT/NodeMeshReshape.cpp b/Engine/Plugins/Experimental/Mutable/Source/MutableTools/Private/MuT/NodeMeshReshape.cpp
--- a/Engine/Plugins/Experimental/Mutable/Source/MutableTools/Private/MuT/NodeMeshReshape.cpp
+++ b/Engine/Plugins/Experimental/Mutable/Source/MutableTools/Private/MuT/NodeMeshReshape.cpp
@@ -17,8 +17,8 @@ namespace mu
 	//---------------------------------------------------------------------------------------------
 	// Static initialisation
 	//---------------------------------------------------------------------------------------------
-	NODE_TYPE NodeMeshReshape::Private::s_type =
-			NODE_TYPE("MeshReshape", NodeMesh::GetStaticType());
+	FNodeType NodeMeshReshape::Private::s_type =
+			FNodeType("MeshReshape", NodeMesh::GetStaticType());
 
 
 	//---------------------------------------------------------------------------------------------
@@ -124,9 +124,7 @@ namespace mu
 
 		if (BaseMesh)
 		{
-			NodeMesh::Private* pPrivate =
-					dynamic_cast<NodeMesh::Private*>(BaseMesh->GetBasePrivate() );
-
+			NodeMesh::Private* pPrivate = static_cast<NodeMesh::Private*>(BaseMesh->GetBasePrivate() );
 			pResult = pPrivate->GetLayout(Index);
 		}
 
diff --git a/Engine/Plugins/Experimental/Mutable/Source/MutableTools/Private/MuT/NodeMeshSwitch.cpp b/Engine/Plugins/Experimental/Mutable/Source/MutableTools/Private/MuT/NodeMeshSwitch.cpp
--- a/Engine/Plugins/Experimental/Mutable/Source/MutableTools/Private/MuT/NodeMeshSwitch.cpp
+++ b/Engine/Plugins/Experimental/Mutable/Source/MutableTools/Private/MuT/NodeMeshSwitch.cpp
@@ -1,6 +1,5 @@
 // Copyright Epic Games, Inc. All Rights Reserved.
 
-
 #include "MuT/NodeMeshSwitch.h"
 
 #include "Misc/AssertionMacros.h"
@@ -17,8 +16,8 @@ namespace mu
 	//---------------------------------------------------------------------------------------------
 	// Static initialisation
 	//---------------------------------------------------------------------------------------------
-	NODE_TYPE NodeMeshSwitch::Private::s_type =
-			NODE_TYPE( "MeshSwitch", NodeMesh::GetStaticType() );
+	FNodeType NodeMeshSwitch::Private::s_type =
+			FNodeType( "MeshSwitch", NodeMesh::GetStaticType() );
 
 
 	//---------------------------------------------------------------------------------------------
@@ -74,7 +73,7 @@ namespace mu
 		{
 			if (m_options[i])
 			{
-				NodeMesh::Private* pPrivate = dynamic_cast<NodeMesh::Private*>(m_options[i]->GetBasePrivate());
+				NodeMesh::Private* pPrivate = static_cast<NodeMesh::Private*>(m_options[i]->GetBasePrivate());
 				NodeLayoutPtr pResult = pPrivate->GetLayout(index);
 				return pResult;
 			}
diff --git a/Engine/Plugins/Experimental/Mutable/Source/MutableTools/Private/MuT/NodeMeshTable.cpp b/Engine/Plugins/Experimental/Mutable/Source/MutableTools/Private/MuT/NodeMeshTable.cpp
--- a/Engine/Plugins/Experimental/Mutable/Source/MutableTools/Private/MuT/NodeMeshTable.cpp
+++ b/Engine/Plugins/Experimental/Mutable/Source/MutableTools/Private/MuT/NodeMeshTable.cpp
@@ -17,8 +17,8 @@ namespace mu
 	//---------------------------------------------------------------------------------------------
 	// Static initialisation
 	//---------------------------------------------------------------------------------------------
-	NODE_TYPE NodeMeshTable::Private::s_type =
-			NODE_TYPE( "TableMesh", NodeMesh::GetStaticType() );
+	FNodeType NodeMeshTable::Private::s_type =
+			FNodeType( "TableMesh", NodeMesh::GetStaticType() );
 
 
 	//---------------------------------------------------------------------------------------------
diff --git a/Engine/Plugins/Experimental/Mutable/Source/MutableTools/Private/MuT/NodeMeshTangents.cpp b/Engine/Plugins/Experimental/Mutable/Source/MutableTools/Private/MuT/NodeMeshTangents.cpp
--- a/Engine/Plugins/Experimental/Mutable/Source/MutableTools/Private/MuT/NodeMeshTangents.cpp
+++ b/Engine/Plugins/Experimental/Mutable/Source/MutableTools/Private/MuT/NodeMeshTangents.cpp
@@ -16,8 +16,8 @@ namespace mu
 	//---------------------------------------------------------------------------------------------
 	// Static initialisation
 	//---------------------------------------------------------------------------------------------
-	NODE_TYPE NodeMeshTangents::Private::s_type =
-			NODE_TYPE( "MeshTangents", NodeMesh::GetStaticType() );
+	FNodeType NodeMeshTangents::Private::s_type =
+			FNodeType( "MeshTangents", NodeMesh::GetStaticType() );
 
 
 	//---------------------------------------------------------------------------------------------
@@ -50,9 +50,7 @@ namespace mu
 
 		if ( m_pSource )
 		{
-			NodeMesh::Private* pPrivate =
-					dynamic_cast<NodeMesh::Private*>( m_pSource->GetBasePrivate() );
-
+			NodeMesh::Private* pPrivate = static_cast<NodeMesh::Private*>( m_pSource->GetBasePrivate() );
 			pResult = pPrivate->GetLayout( index );
 		}
 
diff --git a/Engine/Plugins/Experimental/Mutable/Source/MutableTools/Private/MuT/NodeMeshTransform.cpp b/Engine/Plugins/Experimental/Mutable/Source/MutableTools/Private/MuT/NodeMeshTransform.cpp
--- a/Engine/Plugins/Experimental/Mutable/Source/MutableTools/Private/MuT/NodeMeshTransform.cpp
+++ b/Engine/Plugins/Experimental/Mutable/Source/MutableTools/Private/MuT/NodeMeshTransform.cpp
@@ -18,8 +18,8 @@ namespace mu
 	//---------------------------------------------------------------------------------------------
 	// Static initialisation
 	//---------------------------------------------------------------------------------------------
-    NODE_TYPE NodeMeshTransform::Private::s_type =
-            NODE_TYPE( "MeshTransform", NodeMesh::GetStaticType() );
+    FNodeType NodeMeshTransform::Private::s_type =
+            FNodeType( "MeshTransform", NodeMesh::GetStaticType() );
 
 
 	//---------------------------------------------------------------------------------------------
@@ -66,9 +66,7 @@ namespace mu
 
 		if ( m_pSource )
 		{
-			NodeMesh::Private* pPrivate =
-					dynamic_cast<NodeMesh::Private*>( m_pSource->GetBasePrivate() );
-
+			NodeMesh::Private* pPrivate = static_cast<NodeMesh::Private*>( m_pSource->GetBasePrivate() );
 			pResult = pPrivate->GetLayout( index );
 		}
 
diff --git a/Engine/Plugins/Experimental/Mutable/Source/MutableTools/Private/MuT/NodeMeshVariation.cpp b/Engine/Plugins/Experimental/Mutable/Source/MutableTools/Private/MuT/NodeMeshVariation.cpp
--- a/Engine/Plugins/Experimental/Mutable/Source/MutableTools/Private/MuT/NodeMeshVariation.cpp
+++ b/Engine/Plugins/Experimental/Mutable/Source/MutableTools/Private/MuT/NodeMeshVariation.cpp
@@ -13,8 +13,8 @@ namespace mu
     //---------------------------------------------------------------------------------------------
     // Static initialisation
     //---------------------------------------------------------------------------------------------
-    NODE_TYPE NodeMeshVariation::Private::s_type =
-        NODE_TYPE( "MeshVariation", NodeMesh::GetStaticType() );
+    FNodeType NodeMeshVariation::Private::s_type =
+        FNodeType( "MeshVariation", NodeMesh::GetStaticType() );
 
 
     //---------------------------------------------------------------------------------------------
diff --git a/Engine/Plugins/Experimental/Mutable/Source/MutableTools/Private/MuT/NodeModifier.cpp b/Engine/Plugins/Experimental/Mutable/Source/MutableTools/Private/MuT/NodeModifier.cpp
--- a/Engine/Plugins/Experimental/Mutable/Source/MutableTools/Private/MuT/NodeModifier.cpp
+++ b/Engine/Plugins/Experimental/Mutable/Source/MutableTools/Private/MuT/NodeModifier.cpp
@@ -17,21 +17,20 @@
 namespace mu
 {
 	// Static initialisation
-	static NODE_TYPE s_nodeModifierType =
-			NODE_TYPE( "NodeModifier", Node::GetStaticType() );
+	static FNodeType s_nodeModifierType = FNodeType( "NodeModifier", Node::GetStaticType() );
 
 	MUTABLE_IMPLEMENT_ENUM_SERIALISABLE(EMutableMultipleTagPolicy);
 
-	const NODE_TYPE* NodeModifier::GetType() const
+	const FNodeType* NodeModifier::GetType() const
 	{
 		return GetStaticType();
 	}
 
 
-	const NODE_TYPE* NodeModifier::GetStaticType()
+	const FNodeType* NodeModifier::GetStaticType()
 	{
-          return &s_nodeModifierType;
-        }
+		return &s_nodeModifierType;
+    }
 
 
 	void NodeModifier::Serialise( const NodeModifier* p, OutputArchive& arch )
@@ -39,12 +38,13 @@ namespace mu
         uint32 ver = 0;
 		arch << ver;
 
-#define SERIALISE_CHILDREN( C, ID ) \
-        ( const C* pTyped##ID = dynamic_cast<const C*>(p) )			\
-        {                                                           \
+	#define SERIALISE_CHILDREN( C, ID ) \
+		( p->GetType()==C::GetStaticType() )					\
+		{ 														\
+			const C* pTyped = static_cast<const C*>(p);			\
             arch << (uint32)ID;									\
-            C::Serialise( pTyped##ID, arch );						\
-		}
+			C::Serialise( pTyped, arch );						\
+		}														\
 
 		if SERIALISE_CHILDREN(NodeModifierMeshClipMorphPlane, 			0 )
 		else if SERIALISE_CHILDREN(NodeModifierMeshClipWithMesh, 		1 )
@@ -79,27 +79,21 @@ namespace mu
 
 	void NodeModifier::AddTag(const FString& Value)
 	{
-		NodeModifier::Private* pD = dynamic_cast<NodeModifier::Private*>(GetBasePrivate());
-		check(pD);
-
+		NodeModifier::Private* pD = static_cast<NodeModifier::Private*>(GetBasePrivate());
 		pD->RequiredTags.Add(Value);
 	}
 
 
 	void NodeModifier::SetMultipleTagPolicy(EMutableMultipleTagPolicy Value)
 	{
-		NodeModifier::Private* pD = dynamic_cast<NodeModifier::Private*>(GetBasePrivate());
-		check(pD);
-
+		NodeModifier::Private* pD = static_cast<NodeModifier::Private*>(GetBasePrivate());
 		pD->MultipleTagsPolicy = Value;
 	}
 
 
 	void NodeModifier::SetStage(bool bBeforeNormalOperation)
 	{
-		NodeModifier::Private* pD = dynamic_cast<NodeModifier::Private*>(GetBasePrivate());
-		check(pD);
-
+		NodeModifier::Private* pD = static_cast<NodeModifier::Private*>(GetBasePrivate());
 		pD->bApplyBeforeNormalOperations = bBeforeNormalOperation; 
 	}
 
diff --git a/Engine/Plugins/Experimental/Mutable/Source/MutableTools/Private/MuT/NodeModifierMeshClipDeform.cpp b/Engine/Plugins/Experimental/Mutable/Source/MutableTools/Private/MuT/NodeModifierMeshClipDeform.cpp
--- a/Engine/Plugins/Experimental/Mutable/Source/MutableTools/Private/MuT/NodeModifierMeshClipDeform.cpp
+++ b/Engine/Plugins/Experimental/Mutable/Source/MutableTools/Private/MuT/NodeModifierMeshClipDeform.cpp
@@ -16,8 +16,8 @@ namespace mu
 	//---------------------------------------------------------------------------------------------
 	// Static initialisation
 	//---------------------------------------------------------------------------------------------
-	NODE_TYPE NodeModifierMeshClipDeform::Private::s_type =
-			NODE_TYPE( "NodeModifierMeshClipDeform", NodeModifier::GetStaticType() );
+	FNodeType NodeModifierMeshClipDeform::Private::s_type =
+			FNodeType( "NodeModifierMeshClipDeform", NodeModifier::GetStaticType() );
 
 
 	//---------------------------------------------------------------------------------------------
diff --git a/Engine/Plugins/Experimental/Mutable/Source/MutableTools/Private/MuT/NodeModifierMeshClipMorphPlane.cpp b/Engine/Plugins/Experimental/Mutable/Source/MutableTools/Private/MuT/NodeModifierMeshClipMorphPlane.cpp
--- a/Engine/Plugins/Experimental/Mutable/Source/MutableTools/Private/MuT/NodeModifierMeshClipMorphPlane.cpp
+++ b/Engine/Plugins/Experimental/Mutable/Source/MutableTools/Private/MuT/NodeModifierMeshClipMorphPlane.cpp
@@ -15,8 +15,8 @@ namespace mu
 	//---------------------------------------------------------------------------------------------
 	// Static initialisation
 	//---------------------------------------------------------------------------------------------
-    NODE_TYPE NodeModifierMeshClipMorphPlane::Private::s_type =
-            NODE_TYPE( "ModifierMeshClipMorphPlane", NodeModifier::GetStaticType() );
+    FNodeType NodeModifierMeshClipMorphPlane::Private::s_type =
+            FNodeType( "ModifierMeshClipMorphPlane", NodeModifier::GetStaticType() );
 
 
 	//---------------------------------------------------------------------------------------------
diff --git a/Engine/Plugins/Experimental/Mutable/Source/MutableTools/Private/MuT/NodeModifierMeshClipWithMesh.cpp b/Engine/Plugins/Experimental/Mutable/Source/MutableTools/Private/MuT/NodeModifierMeshClipWithMesh.cpp
--- a/Engine/Plugins/Experimental/Mutable/Source/MutableTools/Private/MuT/NodeModifierMeshClipWithMesh.cpp
+++ b/Engine/Plugins/Experimental/Mutable/Source/MutableTools/Private/MuT/NodeModifierMeshClipWithMesh.cpp
@@ -15,8 +15,8 @@ namespace mu
 	//---------------------------------------------------------------------------------------------
 	// Static initialisation
 	//---------------------------------------------------------------------------------------------
-    NODE_TYPE NodeModifierMeshClipWithMesh::Private::s_type =
-            NODE_TYPE( "NodeModifierMeshClipWithMesh", NodeModifier::GetStaticType() );
+    FNodeType NodeModifierMeshClipWithMesh::Private::s_type =
+            FNodeType( "NodeModifierMeshClipWithMesh", NodeModifier::GetStaticType() );
 
 
 	//---------------------------------------------------------------------------------------------
diff --git a/Engine/Plugins/Experimental/Mutable/Source/MutableTools/Private/MuT/NodeModifierMeshClipWithUVMask.cpp b/Engine/Plugins/Experimental/Mutable/Source/MutableTools/Private/MuT/NodeModifierMeshClipWithUVMask.cpp
--- a/Engine/Plugins/Experimental/Mutable/Source/MutableTools/Private/MuT/NodeModifierMeshClipWithUVMask.cpp
+++ b/Engine/Plugins/Experimental/Mutable/Source/MutableTools/Private/MuT/NodeModifierMeshClipWithUVMask.cpp
@@ -15,8 +15,8 @@ namespace mu
 	//---------------------------------------------------------------------------------------------
 	// Static initialisation
 	//---------------------------------------------------------------------------------------------
-    NODE_TYPE NodeModifierMeshClipWithUVMask::Private::s_type =
-            NODE_TYPE( "NodeModifierMeshClipWithUVMask", NodeModifier::GetStaticType() );
+    FNodeType NodeModifierMeshClipWithUVMask::Private::s_type =
+            FNodeType( "NodeModifierMeshClipWithUVMask", NodeModifier::GetStaticType() );
 
 
 	//---------------------------------------------------------------------------------------------
diff --git a/Engine/Plugins/Experimental/Mutable/Source/MutableTools/Private/MuT/NodeObject.cpp b/Engine/Plugins/Experimental/Mutable/Source/MutableTools/Private/MuT/NodeObject.cpp
--- a/Engine/Plugins/Experimental/Mutable/Source/MutableTools/Private/MuT/NodeObject.cpp
+++ b/Engine/Plugins/Experimental/Mutable/Source/MutableTools/Private/MuT/NodeObject.cpp
@@ -1,6 +1,5 @@
 // Copyright Epic Games, Inc. All Rights Reserved.
 
-
 #include "MuT/NodeObject.h"
 
 #include "Misc/AssertionMacros.h"
@@ -17,7 +16,7 @@ namespace mu
 	//---------------------------------------------------------------------------------------------
 	// Static initialisation
 	//---------------------------------------------------------------------------------------------
-	NODE_TYPE s_nodeObjectType = NODE_TYPE( "NodeObject", NodeObject::GetStaticType() );
+	FNodeType s_nodeObjectType = FNodeType( "NodeObject", Node::GetStaticType() );
 
 
 	//---------------------------------------------------------------------------------------------
@@ -54,14 +53,14 @@ namespace mu
 
 
 	//---------------------------------------------------------------------------------------------
-	const NODE_TYPE* NodeObject::GetType() const
+	const FNodeType* NodeObject::GetType() const
 	{
 		return GetStaticType();
 	}
 
 
 	//---------------------------------------------------------------------------------------------
-	const NODE_TYPE* NodeObject::GetStaticType()
+	const FNodeType* NodeObject::GetStaticType()
 	{
 		return &s_nodeObjectType;
 	}
diff --git a/Engine/Plugins/Experimental/Mutable/Source/MutableTools/Private/MuT/NodeObjectGroup.cpp b/Engine/Plugins/Experimental/Mutable/Source/MutableTools/Private/MuT/NodeObjectGroup.cpp
--- a/Engine/Plugins/Experimental/Mutable/Source/MutableTools/Private/MuT/NodeObjectGroup.cpp
+++ b/Engine/Plugins/Experimental/Mutable/Source/MutableTools/Private/MuT/NodeObjectGroup.cpp
@@ -1,6 +1,5 @@
 // Copyright Epic Games, Inc. All Rights Reserved.
 
-
 #include "MuT/NodeObjectGroup.h"
 
 #include "HAL/PlatformCrt.h"
@@ -10,9 +9,6 @@
 #include "MuT/NodeObjectPrivate.h"
 #include "MuT/NodePrivate.h"
 
-#include <memory>
-#include <utility>
-
 
 namespace mu
 {
@@ -22,8 +18,7 @@ namespace mu
 	//---------------------------------------------------------------------------------------------
 	// Static initialisation
 	//---------------------------------------------------------------------------------------------
-	NODE_TYPE NodeObjectGroup::Private::s_type =
-            NODE_TYPE( "ObjectGroup", NodeObjectGroup::GetStaticType() );
+	FNodeType NodeObjectGroup::Private::s_type = FNodeType( "ObjectGroup", NodeObject::GetStaticType() );
 
 
 	//---------------------------------------------------------------------------------------------
@@ -118,9 +113,7 @@ namespace mu
 		{
 			if (m_children[i])
 			{
-				NodeObject::Private* pPrivate =
-						dynamic_cast<NodeObject::Private*>( m_children[i]->GetBasePrivate() );
-
+				NodeObject::Private* pPrivate = static_cast<NodeObject::Private*>( m_children[i]->GetBasePrivate() );
                 pLayout = pPrivate->GetLayout( lod, component, surface, texture );
 			}
 		}
diff --git a/Engine/Plugins/Experimental/Mutable/Source/MutableTools/Private/MuT/NodeObjectNew.cpp b/Engine/Plugins/Experimental/Mutable/Source/MutableTools/Private/MuT/NodeObjectNew.cpp
--- a/Engine/Plugins/Experimental/Mutable/Source/MutableTools/Private/MuT/NodeObjectNew.cpp
+++ b/Engine/Plugins/Experimental/Mutable/Source/MutableTools/Private/MuT/NodeObjectNew.cpp
@@ -1,6 +1,5 @@
 // Copyright Epic Games, Inc. All Rights Reserved.
 
-
 #include "MuT/NodeObjectNew.h"
 
 #include "HAL/PlatformString.h"
@@ -18,9 +17,6 @@
 #include "MuT/NodeSurface.h"
 #include "MuT/NodeSurfaceNew.h"
 
-#include <memory>
-#include <utility>
-
 
 namespace mu
 {
@@ -29,8 +25,7 @@ namespace mu
 	//---------------------------------------------------------------------------------------------
 	// Static initialisation
 	//---------------------------------------------------------------------------------------------
-	NODE_TYPE NodeObjectNew::Private::s_type =
-			NODE_TYPE( "Object", NodeObjectNew::GetStaticType() );
+	FNodeType NodeObjectNew::Private::s_type = FNodeType( "Object", NodeObject::GetStaticType() );
 
 
 	//---------------------------------------------------------------------------------------------
@@ -146,28 +141,24 @@ namespace mu
 
 		NodeLayoutPtr pLayout;
 
-		const NodeComponentNew* pComp =
-			dynamic_cast<const NodeComponentNew*> ( m_lods[lod]->GetComponent( component ).get() );
+		const NodeComponent* pCompGeneric = m_lods[lod]->GetComponent( component ).get();
 
-        if (pComp)
+        if (pCompGeneric->GetType()==NodeComponentNew::GetStaticType())
         {
-            const NodeSurfaceNew* pSurface =
-                    dynamic_cast<const NodeSurfaceNew*> ( pComp->GetSurface( surface ) );
-
-            if ( pSurface )
+			const NodeComponentNew* pComp = static_cast<const NodeComponentNew*> (pCompGeneric);
+			const NodeSurface* pSurfaceGeneric = pComp->GetSurface( surface );
+            if (pSurfaceGeneric->GetType()==NodeSurfaceNew::GetStaticType())
             {
-                // TODO: Look for the layout index of the given texture
+				const NodeSurfaceNew* pSurface = static_cast<const NodeSurfaceNew*> (pSurfaceGeneric);
+				
+				// TODO: Look for the layout index of the given texture
                 // TODO: Multiple meshes
                 if ( pSurface->GetMeshCount()>0 )
                 {
-                     const NodeMesh* pMesh =
-                             dynamic_cast<const NodeMesh*> ( pSurface->GetMesh( 0 ).get() );
-
+                    const NodeMesh* pMesh = pSurface->GetMesh( 0 ).get();
                     if ( pMesh )
                     {
-                        NodeMesh::Private* pPrivate =
-                                dynamic_cast<NodeMesh::Private*>( pMesh->GetBasePrivate() );
-
+                        NodeMesh::Private* pPrivate = static_cast<NodeMesh::Private*>( pMesh->GetBasePrivate() );
                         pLayout = pPrivate->GetLayout( 0 );
                     }
                 }
diff --git a/Engine/Plugins/Experimental/Mutable/Source/MutableTools/Private/MuT/NodePatchImage.cpp b/Engine/Plugins/Experimental/Mutable/Source/MutableTools/Private/MuT/NodePatchImage.cpp
--- a/Engine/Plugins/Experimental/Mutable/Source/MutableTools/Private/MuT/NodePatchImage.cpp
+++ b/Engine/Plugins/Experimental/Mutable/Source/MutableTools/Private/MuT/NodePatchImage.cpp
@@ -1,6 +1,5 @@
 // Copyright Epic Games, Inc. All Rights Reserved.
 
-
 #include "MuT/NodePatchImage.h"
 
 #include "Misc/AssertionMacros.h"
@@ -8,9 +7,6 @@
 #include "MuT/NodePatchImagePrivate.h"
 #include "MuT/NodePrivate.h"
 
-#include <memory>
-#include <utility>
-
 
 namespace mu
 {
@@ -19,15 +15,14 @@ namespace mu
     //---------------------------------------------------------------------------------------------
     // Static initialisation
     //---------------------------------------------------------------------------------------------
-    NODE_TYPE NodePatchImage::Private::s_type =
-            NODE_TYPE( "PatchTexture", NodePatchImage::GetStaticType() );
+    FNodeType NodePatchImage::Private::s_type = FNodeType( "PatchTexture", Node::GetStaticType() );
 
 
     //---------------------------------------------------------------------------------------------
     //!
     //---------------------------------------------------------------------------------------------
 
-    MUTABLE_IMPLEMENT_NODE( NodePatchImage, EType::PatchImage, Node, Node::EType::PatchImage);
+    MUTABLE_IMPLEMENT_NODE( NodePatchImage, EType::PatchImage, Node, Node::EType::None)
 
 
     //---------------------------------------------------------------------------------------------
diff --git a/Engine/Plugins/Experimental/Mutable/Source/MutableTools/Private/MuT/NodePatchMesh.cpp b/Engine/Plugins/Experimental/Mutable/Source/MutableTools/Private/MuT/NodePatchMesh.cpp
--- a/Engine/Plugins/Experimental/Mutable/Source/MutableTools/Private/MuT/NodePatchMesh.cpp
+++ b/Engine/Plugins/Experimental/Mutable/Source/MutableTools/Private/MuT/NodePatchMesh.cpp
@@ -15,15 +15,15 @@ namespace mu
 	//---------------------------------------------------------------------------------------------
 	// Static initialisation
 	//---------------------------------------------------------------------------------------------
-	NODE_TYPE NodePatchMesh::Private::s_type =
-			NODE_TYPE( "MeshPatch", Node::GetStaticType() );
+	FNodeType NodePatchMesh::Private::s_type =
+			FNodeType( "MeshPatch", Node::GetStaticType() );
 
 
 	//---------------------------------------------------------------------------------------------
 	//!
 	//---------------------------------------------------------------------------------------------
 
-	MUTABLE_IMPLEMENT_NODE( NodePatchMesh, EType::PatchMesh, Node, Node::EType::PatchMesh);
+	MUTABLE_IMPLEMENT_NODE( NodePatchMesh, EType::PatchMesh, Node, Node::EType::None);
 
 
 	//---------------------------------------------------------------------------------------------
diff --git a/Engine/Plugins/Experimental/Mutable/Source/MutableTools/Private/MuT/NodePrivate.h b/Engine/Plugins/Experimental/Mutable/Source/MutableTools/Private/MuT/NodePrivate.h
--- a/Engine/Plugins/Experimental/Mutable/Source/MutableTools/Private/MuT/NodePrivate.h
+++ b/Engine/Plugins/Experimental/Mutable/Source/MutableTools/Private/MuT/NodePrivate.h
@@ -69,12 +69,12 @@ namespace mu
         return m_pD;										\
     }														\
                                                             \
-    const NODE_TYPE* N::GetType() const						\
+    const FNodeType* N::GetType() const						\
     {														\
         return GetStaticType();								\
     }														\
                                                             \
-    const NODE_TYPE* N::GetStaticType()						\
+    const FNodeType* N::GetStaticType()						\
     {														\
         return &Private::s_type;							\
     }
diff --git a/Engine/Plugins/Experimental/Mutable/Source/MutableTools/Private/MuT/NodeProjector.cpp b/Engine/Plugins/Experimental/Mutable/Source/MutableTools/Private/MuT/NodeProjector.cpp
--- a/Engine/Plugins/Experimental/Mutable/Source/MutableTools/Private/MuT/NodeProjector.cpp
+++ b/Engine/Plugins/Experimental/Mutable/Source/MutableTools/Private/MuT/NodeProjector.cpp
@@ -1,6 +1,5 @@
 // Copyright Epic Games, Inc. All Rights Reserved.
 
-
 #include "MuT/NodeProjector.h"
 
 #include "Misc/AssertionMacros.h"
@@ -8,8 +7,6 @@
 #include "MuR/SerialisationPrivate.h"
 #include "MuT/Node.h"
 
-#include <stdint.h>
-
 
 namespace mu
 {
@@ -17,13 +14,13 @@ namespace mu
 	//---------------------------------------------------------------------------------------------
 	// Static initialisation
 	//---------------------------------------------------------------------------------------------
-	static NODE_TYPE s_nodeProjectorType = NODE_TYPE( "NodeProjector", Node::GetStaticType() );
+	static FNodeType s_nodeProjectorType = FNodeType( "NodeProjector", Node::GetStaticType() );
 
 
 	//---------------------------------------------------------------------------------------------
 	void NodeProjector::Serialise( const NodeProjector* p, OutputArchive& arch )
 	{
-        uint32_t ver = 0;
+        uint32 ver = 0;
 		arch << ver;
 
 		arch << uint32_t(p->Type);
@@ -34,11 +31,11 @@ namespace mu
 	//---------------------------------------------------------------------------------------------
 	NodeProjectorPtr NodeProjector::StaticUnserialise( InputArchive& arch )
 	{
-        uint32_t ver;
+        uint32 ver;
 		arch >> ver;
 		check( ver == 0 );
 
-        uint32_t id;
+        uint32 id;
 		arch >> id;
 
 		switch (id)
@@ -53,14 +50,14 @@ namespace mu
 
 
 	//---------------------------------------------------------------------------------------------
-	const NODE_TYPE* NodeProjector::GetType() const
+	const FNodeType* NodeProjector::GetType() const
 	{
 		return GetStaticType();
 	}
 
 
 	//---------------------------------------------------------------------------------------------
-	const NODE_TYPE* NodeProjector::GetStaticType()
+	const FNodeType* NodeProjector::GetStaticType()
 	{
 		return &s_nodeProjectorType;
 	}
diff --git a/Engine/Plugins/Experimental/Mutable/Source/MutableTools/Private/MuT/NodeProjectorConstant.cpp b/Engine/Plugins/Experimental/Mutable/Source/MutableTools/Private/MuT/NodeProjectorConstant.cpp
--- a/Engine/Plugins/Experimental/Mutable/Source/MutableTools/Private/MuT/NodeProjectorConstant.cpp
+++ b/Engine/Plugins/Experimental/Mutable/Source/MutableTools/Private/MuT/NodeProjectorConstant.cpp
@@ -16,8 +16,8 @@ namespace mu
 	//---------------------------------------------------------------------------------------------
 	// Static initialisation
 	//---------------------------------------------------------------------------------------------
-	NODE_TYPE NodeProjectorConstant::Private::s_type =
-			NODE_TYPE( "ProjectorConstant", NodeProjector::GetStaticType() );
+	FNodeType NodeProjectorConstant::Private::s_type =
+			FNodeType( "ProjectorConstant", NodeProjector::GetStaticType() );
 
 
 	//---------------------------------------------------------------------------------------------
diff --git a/Engine/Plugins/Experimental/Mutable/Source/MutableTools/Private/MuT/NodeProjectorParameter.cpp b/Engine/Plugins/Experimental/Mutable/Source/MutableTools/Private/MuT/NodeProjectorParameter.cpp
--- a/Engine/Plugins/Experimental/Mutable/Source/MutableTools/Private/MuT/NodeProjectorParameter.cpp
+++ b/Engine/Plugins/Experimental/Mutable/Source/MutableTools/Private/MuT/NodeProjectorParameter.cpp
@@ -19,8 +19,8 @@ namespace mu
 	//---------------------------------------------------------------------------------------------
 	// Static initialisation
 	//---------------------------------------------------------------------------------------------
-	NODE_TYPE NodeProjectorParameter::Private::s_type =
-			NODE_TYPE( "ProjectorParameter", NodeProjector::GetStaticType() );
+	FNodeType NodeProjectorParameter::Private::s_type =
+			FNodeType( "ProjectorParameter", NodeProjector::GetStaticType() );
 
 
 	//---------------------------------------------------------------------------------------------
diff --git a/Engine/Plugins/Experimental/Mutable/Source/MutableTools/Private/MuT/NodeProjectorPrivate.h b/Engine/Plugins/Experimental/Mutable/Source/MutableTools/Private/MuT/NodeProjectorPrivate.h
--- a/Engine/Plugins/Experimental/Mutable/Source/MutableTools/Private/MuT/NodeProjectorPrivate.h
+++ b/Engine/Plugins/Experimental/Mutable/Source/MutableTools/Private/MuT/NodeProjectorPrivate.h
@@ -22,7 +22,7 @@ namespace mu
 	{
 	public:
 
-		static NODE_TYPE s_type;
+		static FNodeType s_type;
 
         PROJECTOR_TYPE m_type = PROJECTOR_TYPE::PLANAR;
 		vec3<float> m_position;
@@ -67,7 +67,7 @@ namespace mu
 	{
 	public:
 
-		static NODE_TYPE s_type;
+		static FNodeType s_type;
 
 		FString m_name;
 		FString m_uid;
diff --git a/Engine/Plugins/Experimental/Mutable/Source/MutableTools/Private/MuT/NodeRange.cpp b/Engine/Plugins/Experimental/Mutable/Source/MutableTools/Private/MuT/NodeRange.cpp
--- a/Engine/Plugins/Experimental/Mutable/Source/MutableTools/Private/MuT/NodeRange.cpp
+++ b/Engine/Plugins/Experimental/Mutable/Source/MutableTools/Private/MuT/NodeRange.cpp
@@ -18,7 +18,7 @@ namespace mu
 	//---------------------------------------------------------------------------------------------
 	// Static initialisation
 	//---------------------------------------------------------------------------------------------
-    static NODE_TYPE s_nodeRangeType = 	NODE_TYPE( "NodeRange", Node::GetStaticType() );
+    static FNodeType s_nodeRangeType = 	FNodeType( "NodeRange", Node::GetStaticType() );
 
 
 	//---------------------------------------------------------------------------------------------
@@ -28,11 +28,12 @@ namespace mu
 		arch << ver;
 
 #define SERIALISE_CHILDREN( C, ID ) \
-        ( const C* pTyped##ID = dynamic_cast<const C*>(p) )			\
-        {                                                           \
-            arch << (uint32_t)ID;                                   \
-            C::Serialise( pTyped##ID, arch );						\
-        }                                                           \
+		( p->GetType()==C::GetStaticType() )					\
+		{ 														\
+			const C* pTyped = static_cast<const C*>(p);			\
+            arch << (uint32_t)ID;								\
+			C::Serialise( pTyped, arch );						\
+		}														\
 
         if SERIALISE_CHILDREN( NodeRangeFromScalar          , 0  )
         else check(false);
@@ -43,11 +44,11 @@ namespace mu
     //---------------------------------------------------------------------------------------------
     NodeRangePtr NodeRange::StaticUnserialise( InputArchive& arch )
 	{
-        uint32_t ver;
+        uint32 ver;
 		arch >> ver;
 		check( ver == 0 );
 
-        uint32_t id;
+        uint32 id;
 		arch >> id;
 
 		switch (id)
@@ -61,14 +62,14 @@ namespace mu
 
 
 	//---------------------------------------------------------------------------------------------
-    const NODE_TYPE* NodeRange::GetType() const
+    const FNodeType* NodeRange::GetType() const
 	{
 		return GetStaticType();
 	}
 
 
 	//---------------------------------------------------------------------------------------------
-    const NODE_TYPE* NodeRange::GetStaticType()
+    const FNodeType* NodeRange::GetStaticType()
 	{
         return &s_nodeRangeType;
 	}
diff --git a/Engine/Plugins/Experimental/Mutable/Source/MutableTools/Private/MuT/NodeRangeFromScalar.cpp b/Engine/Plugins/Experimental/Mutable/Source/MutableTools/Private/MuT/NodeRangeFromScalar.cpp
--- a/Engine/Plugins/Experimental/Mutable/Source/MutableTools/Private/MuT/NodeRangeFromScalar.cpp
+++ b/Engine/Plugins/Experimental/Mutable/Source/MutableTools/Private/MuT/NodeRangeFromScalar.cpp
@@ -16,8 +16,8 @@ namespace mu
 	//---------------------------------------------------------------------------------------------
 	// Static initialisation
 	//---------------------------------------------------------------------------------------------
-    NODE_TYPE NodeRangeFromScalar::Private::s_type =
-            NODE_TYPE( "RangeFromScalar", NodeRange::GetStaticType() );
+    FNodeType NodeRangeFromScalar::Private::s_type =
+            FNodeType( "RangeFromScalar", NodeRange::GetStaticType() );
 
 
 	//---------------------------------------------------------------------------------------------
diff --git a/Engine/Plugins/Experimental/Mutable/Source/MutableTools/Private/MuT/NodeScalar.cpp b/Engine/Plugins/Experimental/Mutable/Source/MutableTools/Private/MuT/NodeScalar.cpp
--- a/Engine/Plugins/Experimental/Mutable/Source/MutableTools/Private/MuT/NodeScalar.cpp
+++ b/Engine/Plugins/Experimental/Mutable/Source/MutableTools/Private/MuT/NodeScalar.cpp
@@ -25,13 +25,13 @@ namespace mu
 	//---------------------------------------------------------------------------------------------
 	// Static initialisation
 	//---------------------------------------------------------------------------------------------
-	static NODE_TYPE s_nodeScalarType = NODE_TYPE( "NodeScalar", Node::GetStaticType() );
+	static FNodeType s_nodeScalarType = FNodeType( "NodeScalar", Node::GetStaticType() );
 
 
 	//---------------------------------------------------------------------------------------------
 	void NodeScalar::Serialise( const NodeScalar* p, OutputArchive& arch )
 	{
-        uint32_t ver = 0;
+        uint32 ver = 0;
 		arch << ver;
 
 		arch << uint32_t(p->Type);
@@ -42,11 +42,11 @@ namespace mu
     //---------------------------------------------------------------------------------------------
 	NodeScalarPtr NodeScalar::StaticUnserialise( InputArchive& arch )
 	{
-        uint32_t ver;
+        uint32 ver;
 		arch >> ver;
 		check( ver == 0 );
 
-        uint32_t id;
+        uint32 id;
 		arch >> id;
 
 		switch (id)
@@ -67,14 +67,14 @@ namespace mu
 
 
 	//---------------------------------------------------------------------------------------------
-	const NODE_TYPE* NodeScalar::GetType() const
+	const FNodeType* NodeScalar::GetType() const
 	{
 		return GetStaticType();
 	}
 
 
 	//---------------------------------------------------------------------------------------------
-	const NODE_TYPE* NodeScalar::GetStaticType()
+	const FNodeType* NodeScalar::GetStaticType()
 	{
 		return &s_nodeScalarType;
 	}
diff --git a/Engine/Plugins/Experimental/Mutable/Source/MutableTools/Private/MuT/NodeScalarArithmeticOperation.cpp b/Engine/Plugins/Experimental/Mutable/Source/MutableTools/Private/MuT/NodeScalarArithmeticOperation.cpp
--- a/Engine/Plugins/Experimental/Mutable/Source/MutableTools/Private/MuT/NodeScalarArithmeticOperation.cpp
+++ b/Engine/Plugins/Experimental/Mutable/Source/MutableTools/Private/MuT/NodeScalarArithmeticOperation.cpp
@@ -18,8 +18,8 @@ namespace mu
 	//---------------------------------------------------------------------------------------------
 	// Static initialisation
 	//---------------------------------------------------------------------------------------------
-    NODE_TYPE NodeScalarArithmeticOperation::Private::s_type =
-            NODE_TYPE( "ScalarArithmenticOperation", NodeScalar::GetStaticType() );
+    FNodeType NodeScalarArithmeticOperation::Private::s_type =
+            FNodeType( "ScalarArithmenticOperation", NodeScalar::GetStaticType() );
 
 
 	//---------------------------------------------------------------------------------------------
diff --git a/Engine/Plugins/Experimental/Mutable/Source/MutableTools/Private/MuT/NodeScalarConstant.cpp b/Engine/Plugins/Experimental/Mutable/Source/MutableTools/Private/MuT/NodeScalarConstant.cpp
--- a/Engine/Plugins/Experimental/Mutable/Source/MutableTools/Private/MuT/NodeScalarConstant.cpp
+++ b/Engine/Plugins/Experimental/Mutable/Source/MutableTools/Private/MuT/NodeScalarConstant.cpp
@@ -14,8 +14,8 @@ namespace mu
 	//---------------------------------------------------------------------------------------------
 	// Static initialisation
 	//---------------------------------------------------------------------------------------------
-	NODE_TYPE NodeScalarConstant::Private::s_type =
-			NODE_TYPE( "ScalarConstant", NodeScalar::GetStaticType() );
+	FNodeType NodeScalarConstant::Private::s_type =
+			FNodeType( "ScalarConstant", NodeScalar::GetStaticType() );
 
 
 	//---------------------------------------------------------------------------------------------
diff --git a/Engine/Plugins/Experimental/Mutable/Source/MutableTools/Private/MuT/NodeScalarCurve.cpp b/Engine/Plugins/Experimental/Mutable/Source/MutableTools/Private/MuT/NodeScalarCurve.cpp
--- a/Engine/Plugins/Experimental/Mutable/Source/MutableTools/Private/MuT/NodeScalarCurve.cpp
+++ b/Engine/Plugins/Experimental/Mutable/Source/MutableTools/Private/MuT/NodeScalarCurve.cpp
@@ -17,8 +17,7 @@ namespace mu
 	//---------------------------------------------------------------------------------------------
 	// Static initialisation
 	//---------------------------------------------------------------------------------------------
-    NODE_TYPE NodeScalarCurve::Private::s_type =
-            NODE_TYPE( "Curve", NodeScalarCurve::GetStaticType() );
+    FNodeType NodeScalarCurve::Private::s_type = FNodeType( "Curve", NodeScalar::GetStaticType() );
 
 
 	//---------------------------------------------------------------------------------------------
diff --git a/Engine/Plugins/Experimental/Mutable/Source/MutableTools/Private/MuT/NodeScalarEnumParameter.cpp b/Engine/Plugins/Experimental/Mutable/Source/MutableTools/Private/MuT/NodeScalarEnumParameter.cpp
--- a/Engine/Plugins/Experimental/Mutable/Source/MutableTools/Private/MuT/NodeScalarEnumParameter.cpp
+++ b/Engine/Plugins/Experimental/Mutable/Source/MutableTools/Private/MuT/NodeScalarEnumParameter.cpp
@@ -14,8 +14,8 @@ namespace mu
 	//---------------------------------------------------------------------------------------------
 	// Static initialisation
 	//---------------------------------------------------------------------------------------------
-	NODE_TYPE NodeScalarEnumParameter::Private::s_type =
-			NODE_TYPE( "ScalarEnumParameter", NodeScalar::GetStaticType() );
+	FNodeType NodeScalarEnumParameter::Private::s_type =
+			FNodeType( "ScalarEnumParameter", NodeScalar::GetStaticType() );
 
 
 	//---------------------------------------------------------------------------------------------
diff --git a/Engine/Plugins/Experimental/Mutable/Source/MutableTools/Private/MuT/NodeScalarParameter.cpp b/Engine/Plugins/Experimental/Mutable/Source/MutableTools/Private/MuT/NodeScalarParameter.cpp
--- a/Engine/Plugins/Experimental/Mutable/Source/MutableTools/Private/MuT/NodeScalarParameter.cpp
+++ b/Engine/Plugins/Experimental/Mutable/Source/MutableTools/Private/MuT/NodeScalarParameter.cpp
@@ -16,8 +16,8 @@ namespace mu
 	//---------------------------------------------------------------------------------------------
 	// Static initialisation
 	//---------------------------------------------------------------------------------------------
-	NODE_TYPE NodeScalarParameter::Private::s_type =
-			NODE_TYPE( "ScalarParameter", NodeScalar::GetStaticType() );
+	FNodeType NodeScalarParameter::Private::s_type =
+			FNodeType( "ScalarParameter", NodeScalar::GetStaticType() );
 
 
 	//---------------------------------------------------------------------------------------------
diff --git a/Engine/Plugins/Experimental/Mutable/Source/MutableTools/Private/MuT/NodeScalarSwitch.cpp b/Engine/Plugins/Experimental/Mutable/Source/MutableTools/Private/MuT/NodeScalarSwitch.cpp
--- a/Engine/Plugins/Experimental/Mutable/Source/MutableTools/Private/MuT/NodeScalarSwitch.cpp
+++ b/Engine/Plugins/Experimental/Mutable/Source/MutableTools/Private/MuT/NodeScalarSwitch.cpp
@@ -18,8 +18,8 @@ namespace mu
 	//---------------------------------------------------------------------------------------------
 	// Static initialisation
 	//---------------------------------------------------------------------------------------------
-    NODE_TYPE NodeScalarSwitch::Private::s_type =
-            NODE_TYPE( "ScalarSwitch", NodeScalar::GetStaticType() );
+    FNodeType NodeScalarSwitch::Private::s_type =
+            FNodeType( "ScalarSwitch", NodeScalar::GetStaticType() );
 
 
 	//---------------------------------------------------------------------------------------------
diff --git a/Engine/Plugins/Experimental/Mutable/Source/MutableTools/Private/MuT/NodeScalarTable.cpp b/Engine/Plugins/Experimental/Mutable/Source/MutableTools/Private/MuT/NodeScalarTable.cpp
--- a/Engine/Plugins/Experimental/Mutable/Source/MutableTools/Private/MuT/NodeScalarTable.cpp
+++ b/Engine/Plugins/Experimental/Mutable/Source/MutableTools/Private/MuT/NodeScalarTable.cpp
@@ -16,8 +16,8 @@ namespace mu
 	//---------------------------------------------------------------------------------------------
 	// Static initialisation
 	//---------------------------------------------------------------------------------------------
-	NODE_TYPE NodeScalarTable::Private::s_type =
-			NODE_TYPE( "TableScalar", NodeScalar::GetStaticType() );
+	FNodeType NodeScalarTable::Private::s_type =
+			FNodeType( "TableScalar", NodeScalar::GetStaticType() );
 
 
 	//---------------------------------------------------------------------------------------------
diff --git a/Engine/Plugins/Experimental/Mutable/Source/MutableTools/Private/MuT/NodeScalarVariation.cpp b/Engine/Plugins/Experimental/Mutable/Source/MutableTools/Private/MuT/NodeScalarVariation.cpp
--- a/Engine/Plugins/Experimental/Mutable/Source/MutableTools/Private/MuT/NodeScalarVariation.cpp
+++ b/Engine/Plugins/Experimental/Mutable/Source/MutableTools/Private/MuT/NodeScalarVariation.cpp
@@ -14,8 +14,8 @@ namespace mu
     //---------------------------------------------------------------------------------------------
     // Static initialisation
     //---------------------------------------------------------------------------------------------
-    NODE_TYPE NodeScalarVariation::Private::s_type =
-        NODE_TYPE( "ScalarVariation", NodeScalar::GetStaticType() );
+    FNodeType NodeScalarVariation::Private::s_type =
+        FNodeType( "ScalarVariation", NodeScalar::GetStaticType() );
 
 
     //---------------------------------------------------------------------------------------------
diff --git a/Engine/Plugins/Experimental/Mutable/Source/MutableTools/Private/MuT/NodeString.cpp b/Engine/Plugins/Experimental/Mutable/Source/MutableTools/Private/MuT/NodeString.cpp
--- a/Engine/Plugins/Experimental/Mutable/Source/MutableTools/Private/MuT/NodeString.cpp
+++ b/Engine/Plugins/Experimental/Mutable/Source/MutableTools/Private/MuT/NodeString.cpp
@@ -19,22 +19,22 @@ namespace mu
 	//---------------------------------------------------------------------------------------------
 	// Static initialisation
 	//---------------------------------------------------------------------------------------------
-	static NODE_TYPE s_nodeStringType = NODE_TYPE( "NodeString", Node::GetStaticType() );
+	static FNodeType s_nodeStringType = FNodeType( "NodeString", Node::GetStaticType() );
 
 
 	//---------------------------------------------------------------------------------------------
 	void NodeString::Serialise( const NodeString* p, OutputArchive& arch )
 	{
-        uint32_t ver = 0;
+        uint32 ver = 0;
 		arch << ver;
 
-
 #define SERIALISE_CHILDREN( C, ID ) \
-        ( const C* pTyped##ID = dynamic_cast<const C*>(p) )			\
-        {                                                           \
-            arch << (uint32_t)ID;                                   \
-            C::Serialise( pTyped##ID, arch );						\
-        }                                                           \
+		( p->GetType()==C::GetStaticType() )					\
+		{ 														\
+			const C* pTyped = static_cast<const C*>(p);			\
+            arch << (uint32_t)ID;								\
+			C::Serialise( pTyped, arch );						\
+		}														\
 
         if SERIALISE_CHILDREN( NodeStringConstant                   , 0 )
         else if SERIALISE_CHILDREN( NodeStringParameter             , 1 )
@@ -46,11 +46,11 @@ namespace mu
     //---------------------------------------------------------------------------------------------
 	NodeStringPtr NodeString::StaticUnserialise( InputArchive& arch )
 	{
-        uint32_t ver;
+        uint32 ver;
 		arch >> ver;
 		check( ver == 0 );
 
-        uint32_t id;
+        uint32 id;
 		arch >> id;
 
 		switch (id)
@@ -65,14 +65,14 @@ namespace mu
 
 
 	//---------------------------------------------------------------------------------------------
-	const NODE_TYPE* NodeString::GetType() const
+	const FNodeType* NodeString::GetType() const
 	{
 		return GetStaticType();
 	}
 
 
 	//---------------------------------------------------------------------------------------------
-	const NODE_TYPE* NodeString::GetStaticType()
+	const FNodeType* NodeString::GetStaticType()
 	{
 		return &s_nodeStringType;
 	}
diff --git a/Engine/Plugins/Experimental/Mutable/Source/MutableTools/Private/MuT/NodeStringConstant.cpp b/Engine/Plugins/Experimental/Mutable/Source/MutableTools/Private/MuT/NodeStringConstant.cpp
--- a/Engine/Plugins/Experimental/Mutable/Source/MutableTools/Private/MuT/NodeStringConstant.cpp
+++ b/Engine/Plugins/Experimental/Mutable/Source/MutableTools/Private/MuT/NodeStringConstant.cpp
@@ -14,8 +14,8 @@ namespace mu
 	//---------------------------------------------------------------------------------------------
 	// Static initialisation
 	//---------------------------------------------------------------------------------------------
-	NODE_TYPE NodeStringConstant::Private::s_type =
-			NODE_TYPE( "StringConstant", NodeString::GetStaticType() );
+	FNodeType NodeStringConstant::Private::s_type =
+			FNodeType( "StringConstant", NodeString::GetStaticType() );
 
 
 	//---------------------------------------------------------------------------------------------
diff --git a/Engine/Plugins/Experimental/Mutable/Source/MutableTools/Private/MuT/NodeStringParameter.cpp b/Engine/Plugins/Experimental/Mutable/Source/MutableTools/Private/MuT/NodeStringParameter.cpp
--- a/Engine/Plugins/Experimental/Mutable/Source/MutableTools/Private/MuT/NodeStringParameter.cpp
+++ b/Engine/Plugins/Experimental/Mutable/Source/MutableTools/Private/MuT/NodeStringParameter.cpp
@@ -16,8 +16,8 @@ namespace mu
 	//---------------------------------------------------------------------------------------------
 	// Static initialisation
 	//---------------------------------------------------------------------------------------------
-	NODE_TYPE NodeStringParameter::Private::s_type =
-			NODE_TYPE( "StringParameter", NodeString::GetStaticType() );
+	FNodeType NodeStringParameter::Private::s_type =
+			FNodeType( "StringParameter", NodeString::GetStaticType() );
 
 
 	//---------------------------------------------------------------------------------------------
diff --git a/Engine/Plugins/Experimental/Mutable/Source/MutableTools/Private/MuT/NodeSurface.cpp b/Engine/Plugins/Experimental/Mutable/Source/MutableTools/Private/MuT/NodeSurface.cpp
--- a/Engine/Plugins/Experimental/Mutable/Source/MutableTools/Private/MuT/NodeSurface.cpp
+++ b/Engine/Plugins/Experimental/Mutable/Source/MutableTools/Private/MuT/NodeSurface.cpp
@@ -1,6 +1,5 @@
 // Copyright Epic Games, Inc. All Rights Reserved.
 
-
 #include "MuT/NodeSurface.h"
 
 #include "Misc/AssertionMacros.h"
@@ -11,8 +10,6 @@
 #include "MuT/NodeSurfaceVariation.h"
 #include "MuT/NodeSurfaceSwitch.h"
 
-#include <stdint.h>
-
 
 namespace mu
 {
@@ -21,21 +18,20 @@ namespace mu
 	//---------------------------------------------------------------------------------------------
 	// Static initialisation
 	//---------------------------------------------------------------------------------------------
-    static NODE_TYPE s_nodeSurfaceType =
-            NODE_TYPE( "NodeSurface", Node::GetStaticType() );
+    static FNodeType s_nodeSurfaceType = FNodeType( "NodeSurface", Node::GetStaticType() );
 
 
 	//---------------------------------------------------------------------------------------------
 	//---------------------------------------------------------------------------------------------
 	//---------------------------------------------------------------------------------------------
-    const NODE_TYPE* NodeSurface::GetType() const
+    const FNodeType* NodeSurface::GetType() const
 	{
 		return GetStaticType();
 	}
 
 
 	//---------------------------------------------------------------------------------------------
-    const NODE_TYPE* NodeSurface::GetStaticType()
+    const FNodeType* NodeSurface::GetStaticType()
 	{
         return &s_nodeSurfaceType;
 	}
@@ -44,7 +40,7 @@ namespace mu
 	//---------------------------------------------------------------------------------------------
     void NodeSurface::Serialise( const NodeSurface* p, OutputArchive& arch )
 	{
-        uint32_t ver = 0;
+        uint32 ver = 0;
 		arch << ver;
 
 		arch << uint32_t(p->Type);
@@ -55,11 +51,11 @@ namespace mu
     //---------------------------------------------------------------------------------------------
     NodeSurfacePtr NodeSurface::StaticUnserialise( InputArchive& arch )
 	{
-        uint32_t ver;
+        uint32 ver;
 		arch >> ver;
 		check( ver == 0 );
 
-        uint32_t id;
+        uint32 id;
 		arch >> id;
 
 		switch (id)
diff --git a/Engine/Plugins/Experimental/Mutable/Source/MutableTools/Private/MuT/NodeSurfaceEdit.cpp b/Engine/Plugins/Experimental/Mutable/Source/MutableTools/Private/MuT/NodeSurfaceEdit.cpp
--- a/Engine/Plugins/Experimental/Mutable/Source/MutableTools/Private/MuT/NodeSurfaceEdit.cpp
+++ b/Engine/Plugins/Experimental/Mutable/Source/MutableTools/Private/MuT/NodeSurfaceEdit.cpp
@@ -20,8 +20,8 @@ namespace mu
 	//---------------------------------------------------------------------------------------------
 	// Static initialisation
 	//---------------------------------------------------------------------------------------------
-    NODE_TYPE NodeSurfaceEdit::Private::s_type =
-            NODE_TYPE( "EditSurface", NodeSurface::GetStaticType() );
+    FNodeType NodeSurfaceEdit::Private::s_type =
+            FNodeType( "EditSurface", NodeSurface::GetStaticType() );
 
 
 	//---------------------------------------------------------------------------------------------
diff --git a/Engine/Plugins/Experimental/Mutable/Source/MutableTools/Private/MuT/NodeSurfaceNew.cpp b/Engine/Plugins/Experimental/Mutable/Source/MutableTools/Private/MuT/NodeSurfaceNew.cpp
--- a/Engine/Plugins/Experimental/Mutable/Source/MutableTools/Private/MuT/NodeSurfaceNew.cpp
+++ b/Engine/Plugins/Experimental/Mutable/Source/MutableTools/Private/MuT/NodeSurfaceNew.cpp
@@ -22,8 +22,8 @@ namespace mu
 	//---------------------------------------------------------------------------------------------
 	// Static initialisation
 	//---------------------------------------------------------------------------------------------
-    NODE_TYPE NodeSurfaceNew::Private::s_type =
-            NODE_TYPE( "NewSurface", NodeSurface::GetStaticType() );
+    FNodeType NodeSurfaceNew::Private::s_type =
+            FNodeType( "NewSurface", NodeSurface::GetStaticType() );
 
 
 	//---------------------------------------------------------------------------------------------
@@ -50,7 +50,7 @@ namespace mu
 
 
     //---------------------------------------------------------------------------------------------
-    void NodeSurfaceNew::SetCustomID( uint32_t id )
+    void NodeSurfaceNew::SetCustomID( uint32 id )
     {
         m_pD->ExternalId = id;
     }
diff --git a/Engine/Plugins/Experimental/Mutable/Source/MutableTools/Private/MuT/NodeSurfaceSwitch.cpp b/Engine/Plugins/Experimental/Mutable/Source/MutableTools/Private/MuT/NodeSurfaceSwitch.cpp
--- a/Engine/Plugins/Experimental/Mutable/Source/MutableTools/Private/MuT/NodeSurfaceSwitch.cpp
+++ b/Engine/Plugins/Experimental/Mutable/Source/MutableTools/Private/MuT/NodeSurfaceSwitch.cpp
@@ -15,8 +15,8 @@ namespace mu
 	//---------------------------------------------------------------------------------------------
 	// Static initialisation
 	//---------------------------------------------------------------------------------------------
-	NODE_TYPE NodeSurfaceSwitch::Private::s_type =
-			NODE_TYPE( "SurfaceSwitch", NodeSurface::GetStaticType() );
+	FNodeType NodeSurfaceSwitch::Private::s_type =
+			FNodeType( "SurfaceSwitch", NodeSurface::GetStaticType() );
 
 
 	//---------------------------------------------------------------------------------------------
diff --git a/Engine/Plugins/Experimental/Mutable/Source/MutableTools/Private/MuT/NodeSurfaceVariation.cpp b/Engine/Plugins/Experimental/Mutable/Source/MutableTools/Private/MuT/NodeSurfaceVariation.cpp
--- a/Engine/Plugins/Experimental/Mutable/Source/MutableTools/Private/MuT/NodeSurfaceVariation.cpp
+++ b/Engine/Plugins/Experimental/Mutable/Source/MutableTools/Private/MuT/NodeSurfaceVariation.cpp
@@ -17,8 +17,8 @@ namespace mu
 	//---------------------------------------------------------------------------------------------
 	// Static initialisation
 	//---------------------------------------------------------------------------------------------
-    NODE_TYPE NodeSurfaceVariation::Private::s_type =
-            NODE_TYPE( "SurfaceVariation", NodeSurface::GetStaticType() );
+    FNodeType NodeSurfaceVariation::Private::s_type =
+            FNodeType( "SurfaceVariation", NodeSurface::GetStaticType() );
 
 
 	//---------------------------------------------------------------------------------------------
diff --git a/Engine/Plugins/Experimental/Mutable/Source/MutableTools/Public/MuT/Node.h b/Engine/Plugins/Experimental/Mutable/Source/MutableTools/Public/MuT/Node.h
--- a/Engine/Plugins/Experimental/Mutable/Source/MutableTools/Public/MuT/Node.h
+++ b/Engine/Plugins/Experimental/Mutable/Source/MutableTools/Public/MuT/Node.h
@@ -41,13 +41,28 @@ namespace mu
 	//! Information about the type of a node, to provide some means to the tools to deal generically
 	//! with nodes.
 	//! \ingroup tools
-	struct NODE_TYPE
+	struct FNodeType
 	{
-		NODE_TYPE();
-		NODE_TYPE( const char* strName, const NODE_TYPE* pParent );
+		FNodeType();
+		FNodeType( const char* strName, const FNodeType* pParent );
 
 		const char* m_strName;
-		const NODE_TYPE* m_pParent;
+		const FNodeType* m_pParent;
+
+		inline bool IsA(const FNodeType* CandidateType) const
+		{
+			if (CandidateType == this)
+			{
+				return true;
+			}
+
+			if (m_pParent)
+			{
+				return m_pParent->IsA(CandidateType);
+			}
+
+			return false;
+		}
 	};
 
 
@@ -88,18 +103,18 @@ namespace mu
 
 		static void Serialise( const Node* pNode, OutputArchive& arch );
 		virtual void SerialiseWrapper(OutputArchive& arch) const = 0;
-		static NodePtr StaticUnserialise( InputArchive& arch );
+		static Ptr<Node> StaticUnserialise( InputArchive& arch );
 
 
 		//-----------------------------------------------------------------------------------------
 		// Own Interface
 		//-----------------------------------------------------------------------------------------
 
-		//! Node type hierarchy data.
-        virtual const NODE_TYPE* GetType() const;
-		static const NODE_TYPE* GetStaticType();
+		/** Node type hierarchy data. */
+        virtual const FNodeType* GetType() const;
+		static const FNodeType* GetStaticType();
 
-		//! Set the opaque context returned in messages in the compiler log.
+		/** Set the opaque context returned in messages in the compiler log. */
 		void SetMessageContext( const void* context );
 
 		//-----------------------------------------------------------------------------------------
diff --git a/Engine/Plugins/Experimental/Mutable/Source/MutableTools/Public/MuT/NodeBool.h b/Engine/Plugins/Experimental/Mutable/Source/MutableTools/Public/MuT/NodeBool.h
--- a/Engine/Plugins/Experimental/Mutable/Source/MutableTools/Public/MuT/NodeBool.h
+++ b/Engine/Plugins/Experimental/Mutable/Source/MutableTools/Public/MuT/NodeBool.h
@@ -44,8 +44,8 @@ namespace mu
 		// Node Interface
 		//-----------------------------------------------------------------------------------------
 
-        const NODE_TYPE* GetType() const override;
-		static const NODE_TYPE* GetStaticType();
+        const FNodeType* GetType() const override;
+		static const FNodeType* GetStaticType();
 
 
 		//-----------------------------------------------------------------------------------------
@@ -87,8 +87,8 @@ namespace mu
 		// Node Interface
 		//-----------------------------------------------------------------------------------------
 
-        const NODE_TYPE* GetType() const override;
-		static const NODE_TYPE* GetStaticType();
+        const FNodeType* GetType() const override;
+		static const FNodeType* GetStaticType();
 
 		//-----------------------------------------------------------------------------------------
 		// Own Interface
@@ -142,8 +142,8 @@ namespace mu
 		// Node Interface
 		//-----------------------------------------------------------------------------------------
 
-        const NODE_TYPE* GetType() const override;
-		static const NODE_TYPE* GetStaticType();
+        const FNodeType* GetType() const override;
+		static const FNodeType* GetStaticType();
 
 		//-----------------------------------------------------------------------------------------
 		// Own Interface
@@ -202,8 +202,8 @@ namespace mu
 		// Node Interface
 		//-----------------------------------------------------------------------------------------
 
-        const NODE_TYPE* GetType() const override;
-		static const NODE_TYPE* GetStaticType();
+        const FNodeType* GetType() const override;
+		static const FNodeType* GetStaticType();
 
 		//-----------------------------------------------------------------------------------------
 		// Own Interface
@@ -255,8 +255,8 @@ namespace mu
 		// Node Interface
 		//-----------------------------------------------------------------------------------------
 
-        const NODE_TYPE* GetType() const override;
-		static const NODE_TYPE* GetStaticType();
+        const FNodeType* GetType() const override;
+		static const FNodeType* GetStaticType();
 
 		//-----------------------------------------------------------------------------------------
 		// Own Interface
diff --git a/Engine/Plugins/Experimental/Mutable/Source/MutableTools/Public/MuT/NodeColour.h b/Engine/Plugins/Experimental/Mutable/Source/MutableTools/Public/MuT/NodeColour.h
--- a/Engine/Plugins/Experimental/Mutable/Source/MutableTools/Public/MuT/NodeColour.h
+++ b/Engine/Plugins/Experimental/Mutable/Source/MutableTools/Public/MuT/NodeColour.h
@@ -50,8 +50,8 @@ namespace mu
         // Node Interface
         //-----------------------------------------------------------------------------------------
 
-        const NODE_TYPE* GetType() const override;
-        static const NODE_TYPE* GetStaticType();
+        const FNodeType* GetType() const override;
+        static const FNodeType* GetStaticType();
 
 
         //-----------------------------------------------------------------------------------------
diff --git a/Engine/Plugins/Experimental/Mutable/Source/MutableTools/Public/MuT/NodeColourArithmeticOperation.h b/Engine/Plugins/Experimental/Mutable/Source/MutableTools/Public/MuT/NodeColourArithmeticOperation.h
--- a/Engine/Plugins/Experimental/Mutable/Source/MutableTools/Public/MuT/NodeColourArithmeticOperation.h
+++ b/Engine/Plugins/Experimental/Mutable/Source/MutableTools/Public/MuT/NodeColourArithmeticOperation.h
@@ -39,8 +39,8 @@ namespace mu
 		// Node Interface
 		//-----------------------------------------------------------------------------------------
 
-		virtual const NODE_TYPE* GetType() const override;
-		static const NODE_TYPE* GetStaticType();
+		virtual const FNodeType* GetType() const override;
+		static const FNodeType* GetStaticType();
 
 		//-----------------------------------------------------------------------------------------
 		// Own Interface
diff --git a/Engine/Plugins/Experimental/Mutable/Source/MutableTools/Public/MuT/NodeColourConstant.h b/Engine/Plugins/Experimental/Mutable/Source/MutableTools/Public/MuT/NodeColourConstant.h
--- a/Engine/Plugins/Experimental/Mutable/Source/MutableTools/Public/MuT/NodeColourConstant.h
+++ b/Engine/Plugins/Experimental/Mutable/Source/MutableTools/Public/MuT/NodeColourConstant.h
@@ -37,8 +37,8 @@ namespace mu
 		// Node Interface
 		//-----------------------------------------------------------------------------------------
 
-        const NODE_TYPE* GetType() const override;
-		static const NODE_TYPE* GetStaticType();
+        const FNodeType* GetType() const override;
+		static const FNodeType* GetStaticType();
 
 		//-----------------------------------------------------------------------------------------
 		// Own Interface
diff --git a/Engine/Plugins/Experimental/Mutable/Source/MutableTools/Public/MuT/NodeLayout.h b/Engine/Plugins/Experimental/Mutable/Source/MutableTools/Public/MuT/NodeLayout.h
--- a/Engine/Plugins/Experimental/Mutable/Source/MutableTools/Public/MuT/NodeLayout.h
+++ b/Engine/Plugins/Experimental/Mutable/Source/MutableTools/Public/MuT/NodeLayout.h
@@ -53,8 +53,8 @@ namespace mu
 		// Node Interface
 		//-----------------------------------------------------------------------------------------
 
-		const NODE_TYPE* GetType() const override;
-		static const NODE_TYPE* GetStaticType();
+		const FNodeType* GetType() const override;
+		static const FNodeType* GetStaticType();
 
 		//-----------------------------------------------------------------------------------------
 		// Own Interface
@@ -100,8 +100,8 @@ namespace mu
 		// Node Interface
 		//-----------------------------------------------------------------------------------------
 
-		const NODE_TYPE* GetType() const override;
-		static const NODE_TYPE* GetStaticType();
+		const FNodeType* GetType() const override;
+		static const FNodeType* GetStaticType();
 
 		//-----------------------------------------------------------------------------------------
 		// Own Interface
diff --git a/Engine/Plugins/Experimental/Mutable/Source/MutableTools/Public/MuT/NodeMeshClipDeform.h b/Engine/Plugins/Experimental/Mutable/Source/MutableTools/Public/MuT/NodeMeshClipDeform.h
--- a/Engine/Plugins/Experimental/Mutable/Source/MutableTools/Public/MuT/NodeMeshClipDeform.h
+++ b/Engine/Plugins/Experimental/Mutable/Source/MutableTools/Public/MuT/NodeMeshClipDeform.h
@@ -11,8 +11,7 @@
 namespace mu
 {
 
-	//! Node that morphs a base mesh with one or two weighted targets from a sequence.
-	//! \ingroup model
+	//! 
 	class MUTABLETOOLS_API NodeMeshClipDeform : public NodeMesh
 	{
 	public:
@@ -32,8 +31,8 @@ namespace mu
 		// Node Interface
 		//-----------------------------------------------------------------------------------------
 
-		const NODE_TYPE* GetType() const override;
-		static const NODE_TYPE* GetStaticType();
+		const FNodeType* GetType() const override;
+		static const FNodeType* GetStaticType();
 
 		//-----------------------------------------------------------------------------------------
 		// Own Interface
diff --git a/Engine/Plugins/Experimental/Mutable/Source/MutableTools/Public/MuT/NodeMeshClipMorphPlane.h b/Engine/Plugins/Experimental/Mutable/Source/MutableTools/Public/MuT/NodeMeshClipMorphPlane.h
--- a/Engine/Plugins/Experimental/Mutable/Source/MutableTools/Public/MuT/NodeMeshClipMorphPlane.h
+++ b/Engine/Plugins/Experimental/Mutable/Source/MutableTools/Public/MuT/NodeMeshClipMorphPlane.h
@@ -37,8 +37,8 @@ namespace mu
 		// Node Interface
 		//-----------------------------------------------------------------------------------------
 
-		virtual const NODE_TYPE* GetType() const;
-		static const NODE_TYPE* GetStaticType();
+		virtual const FNodeType* GetType() const;
+		static const FNodeType* GetStaticType();
 
 		//-----------------------------------------------------------------------------------------
 		// Own Interface
diff --git a/Engine/Plugins/Experimental/Mutable/Source/MutableTools/Public/MuT/NodeMeshClipWithMesh.h b/Engine/Plugins/Experimental/Mutable/Source/MutableTools/Public/MuT/NodeMeshClipWithMesh.h
--- a/Engine/Plugins/Experimental/Mutable/Source/MutableTools/Public/MuT/NodeMeshClipWithMesh.h
+++ b/Engine/Plugins/Experimental/Mutable/Source/MutableTools/Public/MuT/NodeMeshClipWithMesh.h
@@ -37,8 +37,8 @@ namespace mu
 		// Node Interface
 		//-----------------------------------------------------------------------------------------
 
-		virtual const NODE_TYPE* GetType() const;
-		static const NODE_TYPE* GetStaticType();
+		virtual const FNodeType* GetType() const;
+		static const FNodeType* GetStaticType();
 
 		//-----------------------------------------------------------------------------------------
 		// Own Interface
diff --git a/Engine/Plugins/Experimental/Mutable/Source/MutableTools/Public/MuT/NodeProjector.h b/Engine/Plugins/Experimental/Mutable/Source/MutableTools/Public/MuT/NodeProjector.h
--- a/Engine/Plugins/Experimental/Mutable/Source/MutableTools/Public/MuT/NodeProjector.h
+++ b/Engine/Plugins/Experimental/Mutable/Source/MutableTools/Public/MuT/NodeProjector.h
@@ -56,8 +56,8 @@ namespace mu
 		// Node Interface
 		//-----------------------------------------------------------------------------------------
 
-        const NODE_TYPE* GetType() const override;
-		static const NODE_TYPE* GetStaticType();
+        const FNodeType* GetType() const override;
+		static const FNodeType* GetStaticType();
 
 
 		//-----------------------------------------------------------------------------------------
@@ -96,8 +96,8 @@ namespace mu
 		// Node Interface
 		//-----------------------------------------------------------------------------------------
 
-        const NODE_TYPE* GetType() const override;
-		static const NODE_TYPE* GetStaticType();
+        const FNodeType* GetType() const override;
+		static const FNodeType* GetStaticType();
 
 		//-----------------------------------------------------------------------------------------
 		// Own Interface
@@ -158,8 +158,8 @@ namespace mu
 		// Node Interface
 		//-----------------------------------------------------------------------------------------
 
-        const NODE_TYPE* GetType() const override;
-		static const NODE_TYPE* GetStaticType();
+        const FNodeType* GetType() const override;
+		static const FNodeType* GetStaticType();
 
 		//-----------------------------------------------------------------------------------------
 		// Own Interface
diff --git a/Engine/Plugins/Experimental/Mutable/Source/MutableTools/Public/MuT/NodeSurfaceVariation.h b/Engine/Plugins/Experimental/Mutable/Source/MutableTools/Public/MuT/NodeSurfaceVariation.h
--- a/Engine/Plugins/Experimental/Mutable/Source/MutableTools/Public/MuT/NodeSurfaceVariation.h
+++ b/Engine/Plugins/Experimental/Mutable/Source/MutableTools/Public/MuT/NodeSurfaceVariation.h
@@ -40,8 +40,8 @@ namespace mu
 		// Node Interface
 		//-----------------------------------------------------------------------------------------
 
-        const NODE_TYPE* GetType() const override;
-		static const NODE_TYPE* GetStaticType();
+        const FNodeType* GetType() const override;
+		static const FNodeType* GetStaticType();
 
 		//-----------------------------------------------------------------------------------------
 		// Own Interface
@@ -58,7 +58,7 @@ namespace mu
 		int GetVariationCount() const;
 
         //!
-        enum class VariationType : uint8_t
+        enum class VariationType : uint8
         {
             //! The variation selection is controlled by tags defined in other surfaces.
             //! Default value.
@@ -86,7 +86,7 @@ namespace mu
 		//-----------------------------------------------------------------------------------------
 		class Private;
 		Private* GetPrivate() const;
-                Node::Private* GetBasePrivate() const override;
+        Node::Private* GetBasePrivate() const override;
 
 	protected:
 
