[mutable] Fix indentation issue causing some compilations to fail (!).
#rnx

diff --git a/Engine/Plugins/Experimental/Mutable/Source/MutableTools/Private/MuT/AST.cpp b/Engine/Plugins/Experimental/Mutable/Source/MutableTools/Private/MuT/AST.cpp
--- a/Engine/Plugins/Experimental/Mutable/Source/MutableTools/Private/MuT/AST.cpp
+++ b/Engine/Plugins/Experimental/Mutable/Source/MutableTools/Private/MuT/AST.cpp
@@ -223,7 +223,10 @@ mu::Ptr<ASTOp> ASTOp::DeepClone( const Ptr<ASTOp>& root )
 
     MapChildFunc m = [&](const Ptr<ASTOp>&n)
     {
-        if (!n) return Ptr<ASTOp>();
+		if (!n)
+		{
+			return Ptr<ASTOp>();
+		}
 		std::unordered_map<Ptr<const ASTOp>, Ptr<ASTOp>>::const_iterator it = visited.find(n);
         check(it!=visited.end());
         return it->second;
