[mutable] Exposed Bake instance method by extracting it from "CustomizableObjectEditorViewportClient"
- This will make possible the invocation of the baking externally to the viewport client.
- Replaced UE with Mutable notifications (prevents log spamming)
- Added Multicast delegate for mutable instance update with result
- Other small changes

Note : part of the chnages in this CL does not make sense without a context. The main body of the baking operation has been placed where it was originally (with some changes) to ease the review process. In a later CL, when the baking test commandled is implemented, the code will be placed where it should.

diff --git a/Engine/Plugins/Experimental/Mutable/Source/CustomizableObject/Classes/MuCO/CustomizableObjectInstance.h b/Engine/Plugins/Experimental/Mutable/Source/CustomizableObject/Classes/MuCO/CustomizableObjectInstance.h
--- a/Engine/Plugins/Experimental/Mutable/Source/CustomizableObject/Classes/MuCO/CustomizableObjectInstance.h
+++ b/Engine/Plugins/Experimental/Mutable/Source/CustomizableObject/Classes/MuCO/CustomizableObjectInstance.h
@@ -81,7 +81,7 @@ struct FUpdateContext
 
 
 DECLARE_DYNAMIC_DELEGATE_OneParam(FInstanceUpdateDelegate, const FUpdateContext&, Result);
-
+DECLARE_MULTICAST_DELEGATE_OneParam(FInstanceUpdateNativeDelegate, const FUpdateContext&);
 
 /* When creating new delegates use the following conventions:
  *
@@ -286,6 +286,9 @@ public:
 	UFUNCTION(BlueprintCallable, Category = CustomizableObjectInstance)
 	void UpdateSkeletalMeshAsyncResult(FInstanceUpdateDelegate Callback, bool bIgnoreCloseDist = false, bool bForceHighPriority = false);
 
+	void UpdateSkeletalMeshAsyncResult(FInstanceUpdateNativeDelegate Callback, bool bIgnoreCloseDist = false, bool bForceHighPriority = false);
+
+	
 	// Clones the instance creating a new identical transient instance.
 	UFUNCTION(BlueprintCallable, Category = CustomizableObjectInstance)
 	UCustomizableObjectInstance* Clone();
diff --git a/Engine/Plugins/Experimental/Mutable/Source/CustomizableObject/Private/MuCO/CustomizableObjectInstance.cpp b/Engine/Plugins/Experimental/Mutable/Source/CustomizableObject/Private/MuCO/CustomizableObjectInstance.cpp
--- a/Engine/Plugins/Experimental/Mutable/Source/CustomizableObject/Private/MuCO/CustomizableObjectInstance.cpp
+++ b/Engine/Plugins/Experimental/Mutable/Source/CustomizableObject/Private/MuCO/CustomizableObjectInstance.cpp
@@ -839,6 +839,18 @@ void UCustomizableObjectInstance::UpdateSkeletalMeshAsyncResult(FInstanceUpdateD
 }
 
 
+void UCustomizableObjectInstance::UpdateSkeletalMeshAsyncResult(FInstanceUpdateNativeDelegate Callback, bool bIgnoreCloseDist, bool bForceHighPriority)
+{
+	UCustomizableObjectSystemPrivate* SystemPrivate = UCustomizableObjectSystem::GetInstance()->GetPrivate();
+
+	const TSharedRef<FUpdateContextPrivate> Context = MakeShared<FUpdateContextPrivate>(*this);
+	Context->bIgnoreCloseDist = bIgnoreCloseDist;
+	Context->bForceHighPriority = bForceHighPriority;
+	Context->UpdateNativeCallback = Callback;
+	
+	SystemPrivate->EnqueueUpdateSkeletalMesh(Context);
+}
+
 void UCustomizableInstancePrivate::TickUpdateCloseCustomizableObjects(UCustomizableObjectInstance& Public, FMutableInstanceUpdateMap& InOutRequestedUpdates)
 {
 	if (!Public.CanUpdateInstance())
diff --git a/Engine/Plugins/Experimental/Mutable/Source/CustomizableObject/Private/MuCO/CustomizableObjectSystem.cpp b/Engine/Plugins/Experimental/Mutable/Source/CustomizableObject/Private/MuCO/CustomizableObjectSystem.cpp
--- a/Engine/Plugins/Experimental/Mutable/Source/CustomizableObject/Private/MuCO/CustomizableObjectSystem.cpp
+++ b/Engine/Plugins/Experimental/Mutable/Source/CustomizableObject/Private/MuCO/CustomizableObjectSystem.cpp
@@ -816,6 +816,8 @@ void FinishUpdateGlobal(const TSharedRef<FUpdateContextPrivate>& Context)
 	ContextPublic.UpdateResult = Context->UpdateResult;
 		
 	Context->UpdateCallback.ExecuteIfBound(ContextPublic);
+	Context->UpdateNativeCallback.Broadcast(ContextPublic);
+
 
 	if (SystemPrivate)
 	{
diff --git a/Engine/Plugins/Experimental/Mutable/Source/CustomizableObject/Private/MuCO/CustomizableObjectSystemPrivate.h b/Engine/Plugins/Experimental/Mutable/Source/CustomizableObject/Private/MuCO/CustomizableObjectSystemPrivate.h
--- a/Engine/Plugins/Experimental/Mutable/Source/CustomizableObject/Private/MuCO/CustomizableObjectSystemPrivate.h
+++ b/Engine/Plugins/Experimental/Mutable/Source/CustomizableObject/Private/MuCO/CustomizableObjectSystemPrivate.h
@@ -426,6 +426,7 @@ public:
 	EQueuePriorityType PriorityType = EQueuePriorityType::Low;
 	
 	FInstanceUpdateDelegate UpdateCallback;
+	FInstanceUpdateNativeDelegate UpdateNativeCallback;
 
 	/** Weak reference to the instance we are operating on.
 	 *It is weak because we don't want to lock it in case it becomes irrelevant in the game while operations are pending and it needs to be destroyed. */
diff --git a/Engine/Plugins/Experimental/Mutable/Source/CustomizableObjectEditor/CustomizableObjectEditor.Build.cs b/Engine/Plugins/Experimental/Mutable/Source/CustomizableObjectEditor/CustomizableObjectEditor.Build.cs
--- a/Engine/Plugins/Experimental/Mutable/Source/CustomizableObjectEditor/CustomizableObjectEditor.Build.cs
+++ b/Engine/Plugins/Experimental/Mutable/Source/CustomizableObjectEditor/CustomizableObjectEditor.Build.cs
@@ -87,7 +87,8 @@ public class CustomizableObjectEditor : ModuleRules
 			new string[] {
 				"PropertyEditor",
 				"CustomizableObject",
-				"StructUtils",
+				"StructUtils", 
+				"MutableTools",
 			}
 			);
 
diff --git a/Engine/Plugins/Experimental/Mutable/Source/CustomizableObjectEditor/Private/MuCOE/CustomizableObjectCustomSettings.cpp b/Engine/Plugins/Experimental/Mutable/Source/CustomizableObjectEditor/Private/MuCOE/CustomizableObjectCustomSettings.cpp
--- a/Engine/Plugins/Experimental/Mutable/Source/CustomizableObjectEditor/Private/MuCOE/CustomizableObjectCustomSettings.cpp
+++ b/Engine/Plugins/Experimental/Mutable/Source/CustomizableObjectEditor/Private/MuCOE/CustomizableObjectCustomSettings.cpp
@@ -2,6 +2,7 @@
 
 #include "MuCOE/CustomizableObjectCustomSettings.h"
 
+#include "CustomizableObjectEditor.h"
 #include "AssetRegistry/AssetRegistryModule.h"
 #include "Components/SpotLightComponent.h"
 #include "Framework/Views/TableViewMetadata.h"
diff --git a/Engine/Plugins/Experimental/Mutable/Source/CustomizableObjectEditor/Private/MuCOE/CustomizableObjectEditor.cpp b/Engine/Plugins/Experimental/Mutable/Source/CustomizableObjectEditor/Private/MuCOE/CustomizableObjectEditor.cpp
--- a/Engine/Plugins/Experimental/Mutable/Source/CustomizableObjectEditor/Private/MuCOE/CustomizableObjectEditor.cpp
+++ b/Engine/Plugins/Experimental/Mutable/Source/CustomizableObjectEditor/Private/MuCOE/CustomizableObjectEditor.cpp
@@ -2722,6 +2722,27 @@ void FCustomizableObjectEditor::CreatePreviewComponents()
 	}
 }
 
+void RemoveRestrictedChars(FString& String)
+{
+	// Remove restricted chars, according to FPaths::ValidatePath, RestrictedChars = "/?:&\\*\"<>|%#@^ ";
+
+	String = String.Replace(TEXT("/"), TEXT(""));
+	String = String.Replace(TEXT("?"), TEXT(""));
+	String = String.Replace(TEXT(":"), TEXT(""));
+	String = String.Replace(TEXT("&"), TEXT(""));
+	String = String.Replace(TEXT("\\"), TEXT(""));
+	String = String.Replace(TEXT("*"), TEXT(""));
+	String = String.Replace(TEXT("\""), TEXT(""));
+	String = String.Replace(TEXT("<"), TEXT(""));
+	String = String.Replace(TEXT(">"), TEXT(""));
+	String = String.Replace(TEXT("|"), TEXT(""));
+	String = String.Replace(TEXT("%"), TEXT(""));
+	String = String.Replace(TEXT("#"), TEXT(""));
+	String = String.Replace(TEXT("@"), TEXT(""));
+	String = String.Replace(TEXT("^"), TEXT(""));
+	String = String.Replace(TEXT(" "), TEXT(""));
+}
+
 
 void FCustomizableObjectEditor::OnCustomizableObjectStatusChanged(FCustomizableObjectStatus::EState, const FCustomizableObjectStatus::EState CurrentState)
 {
diff --git a/Engine/Plugins/Experimental/Mutable/Source/CustomizableObjectEditor/Private/MuCOE/CustomizableObjectEditor.h b/Engine/Plugins/Experimental/Mutable/Source/CustomizableObjectEditor/Private/MuCOE/CustomizableObjectEditor.h
--- a/Engine/Plugins/Experimental/Mutable/Source/CustomizableObjectEditor/Private/MuCOE/CustomizableObjectEditor.h
+++ b/Engine/Plugins/Experimental/Mutable/Source/CustomizableObjectEditor/Private/MuCOE/CustomizableObjectEditor.h
@@ -48,6 +48,7 @@ template <typename FuncType> class TFunction;
 
 DECLARE_DELEGATE(FCreatePreviewInstanceFlagDelegate);
 
+extern void RemoveRestrictedChars(FString& String);
 
 /**
 * Wrapper UObject class for the UCustomizableObjectInstance::FObjectInstanceUpdatedDelegate dynamic multicast delegate
diff --git a/Engine/Plugins/Experimental/Mutable/Source/CustomizableObjectEditor/Private/MuCOE/CustomizableObjectEditorViewportClient.cpp b/Engine/Plugins/Experimental/Mutable/Source/CustomizableObjectEditor/Private/MuCOE/CustomizableObjectEditorViewportClient.cpp
--- a/Engine/Plugins/Experimental/Mutable/Source/CustomizableObjectEditor/Private/MuCOE/CustomizableObjectEditorViewportClient.cpp
+++ b/Engine/Plugins/Experimental/Mutable/Source/CustomizableObjectEditor/Private/MuCOE/CustomizableObjectEditorViewportClient.cpp
@@ -2,6 +2,7 @@
 
 #include "MuCOE/CustomizableObjectEditorViewportClient.h"
 
+#include "MuCOE/CustomizableObjectInstanceBakingUtils.h"
 #include "Animation/DebugSkelMeshComponent.h"
 #include "Animation/PoseAsset.h"
 #include "Animation/AnimSingleNodeInstance.h"
@@ -12,6 +13,7 @@
 #include "Components/SpotLightComponent.h"
 #include "Components/StaticMeshComponent.h"
 #include "ContentBrowserModule.h"
+#include "CustomizableObjectEditor.h"
 #include "DynamicMeshBuilder.h"
 #include "Editor/EditorPerProjectUserSettings.h"
 #include "Editor/UnrealEdTypes.h"
@@ -19,13 +21,12 @@
 #include "FileHelpers.h"
 #include "IContentBrowserSingleton.h"
 #include "InputKeyEventArgs.h"
+#include "ObjectTools.h"
 #include "Materials/MaterialInstanceDynamic.h"
 #include "Misc/MessageDialog.h"
 #include "Misc/PackageName.h"
 #include "MuCO/CustomizableObjectInstance.h"
 #include "MuCO/CustomizableObjectSystem.h"
-#include "MuCO/CustomizableObjectMipDataProvider.h"
-#include "MuCOE/UnrealBakeHelpers.h"
 #include "MuCOE/CustomizableObjectPreviewScene.h"
 #include "MuCOE/ICustomizableObjectInstanceEditor.h"
 #include "MuCOE/Nodes/CustomizableObjectNodeMeshClipMorph.h"
@@ -34,17 +35,19 @@
 #include "MuCOE/Nodes/CustomizableObjectNodeProjectorParameter.h"
 #include "MuCOE/UnrealEditorPortabilityHelpers.h"
 #include "MuT/UnrealPixelFormatOverride.h"
-#include "ObjectTools.h"
 #include "Preferences/PersonaOptions.h"
 #include "ScopedTransaction.h"
 #include "SkeletalDebugRendering.h"
 #include "UnrealWidget.h"
+#include "Materials/MaterialInstanceConstant.h"
+#include "MuCO/CustomizableObjectMipDataProvider.h"
+#include "MuCOE/CustomizableObjectEditorLogger.h"
+#include "MuCOE/UnrealBakeHelpers.h"
 #include "MuCO/CustomizableObjectPrivate.h"
 #include "Widgets/Input/SButton.h"
 #include "Widgets/Input/SEditableTextBox.h"
 #include "Widgets/Layout/SUniformGridPanel.h"
 #include "Widgets/Notifications/SNotificationList.h"
-#include "Materials/MaterialInstanceConstant.h"
 
 class FMaterialRenderProxy;
 class UFont;
@@ -58,7 +61,6 @@ FCustomizableObjectEditorViewportClient::FCustomizableObjectEditorViewportClient
 	: FEditorViewportClient(&GLevelEditorModeTools(), InPreviewScene, EditorViewportWidget)
 	, CustomizableObjectEditorPtr(InCustomizableObjectEditor)
 	, CustomizableObject(nullptr)
-	, BakingOverwritePermission(false)
 {
 	// load config
 	ConfigOption = UPersonaOptions::StaticClass()->GetDefaultObject<UPersonaOptions>();
@@ -1580,8 +1582,8 @@ public:
 	/** FileName getter */
 	FString GetFileName();
 
-	bool GetExportAllResources();
-	bool GetGenerateConstantMaterialInstances();
+	bool GetExportAllResources() const;
+	bool GetGenerateConstantMaterialInstances() const;
 
 protected:
 	void OnPathChange(const FString& NewPath);
@@ -1600,664 +1602,765 @@ protected:
 
 //-------------------------------------------------------------------------------------------------
 
-void RemoveRestrictedChars(FString& String)
+// TODO: Move to CustomizableObjectInstanceBakingUtils.cpp after review to be able to enable testing commandlet
+namespace 
 {
-	// Remove restricted chars, according to FPaths::ValidatePath, RestrictedChars = "/?:&\\*\"<>|%#@^ ";
+	/**
+	 * Simple wrapper to be able to invoke the generation of a popup or log message depending on the execution context in which this code is being ran
+	 * @param InMessage The message to display
+	 * @param InTitle The title to be used for the popup or the log generated
+	 */
+	void ShowErrorNotification(const FText& InMessage,  const FText* InTitle = nullptr)
+	{
+		if (!FApp::IsUnattended())
+		{
+			FMessageDialog::Open(EAppMsgType::Ok, InMessage, *InTitle);
+		}
+		else
+		{
+			if (InTitle)
+			{
+				UE_LOG(LogMutable, Error, TEXT("%s - %s"), *InTitle->ToString(), *InMessage.ToString());
+			}
+			else
+			{
+				UE_LOG(LogMutable, Error, TEXT("%s"), *InMessage.ToString());
+			}
+		}
+	}
 
-	String = String.Replace(TEXT("/"), TEXT(""));
-	String = String.Replace(TEXT("?"), TEXT(""));
-	String = String.Replace(TEXT(":"), TEXT(""));
-	String = String.Replace(TEXT("&"), TEXT(""));
-	String = String.Replace(TEXT("\\"), TEXT(""));
-	String = String.Replace(TEXT("*"), TEXT(""));
-	String = String.Replace(TEXT("\""), TEXT(""));
-	String = String.Replace(TEXT("<"), TEXT(""));
-	String = String.Replace(TEXT(">"), TEXT(""));
-	String = String.Replace(TEXT("|"), TEXT(""));
-	String = String.Replace(TEXT("%"), TEXT(""));
-	String = String.Replace(TEXT("#"), TEXT(""));
-	String = String.Replace(TEXT("@"), TEXT(""));
-	String = String.Replace(TEXT("^"), TEXT(""));
-	String = String.Replace(TEXT(" "), TEXT(""));
-}
+	/**
+	 * Utility functions for the baking operation.
+	 */
+
+	/**
+	 * Validates the filename chosen for the baking data
+	 * @param FileName The filename chosen by the user
+	 * @return True if validation was successful, false otherwise
+	 */
+	bool ValidateProvidedFileName(const FString& FileName)
+	{
+		if (FileName.IsEmpty())
+		{
+			UE_LOG(LogMutable, Error, TEXT("Invalid baking configuration : FileName string is empty.."));
+			return false;
+		}
 
+		// Check for invalid characters in the name of the object to be serialized
+		{
+			TCHAR InvalidCharacter = '0';
+			{
+				FString InvalidCharacters = FPaths::GetInvalidFileSystemChars();
+				for (int32 i = 0; i < InvalidCharacters.Len(); ++i)
+				{
+					TCHAR Char = InvalidCharacters[i];
+					FString SearchedChar = FString::Chr(Char);
+					if (FileName.Contains(SearchedChar))
+					{
+						InvalidCharacter = InvalidCharacters[i];
+						break;
+					}
+				}
+			}
 
-//-------------------------------------------------------------------------------------------------
-void FCustomizableObjectEditorViewportClient::BakeInstance()
-{
-	BakeInstance(nullptr);
-}
+			if (InvalidCharacter != '0')
+			{
+				const FText InvalidCharacterText = FText::FromString(FString::Chr(InvalidCharacter));
+				const FText ErrorText = FText::Format(LOCTEXT("FCustomizableObjectEditorViewportClient_BakeInstance_InvalidCharacter", "The selected contains an invalid character ({0})."), InvalidCharacterText);
 
+				ShowErrorNotification(ErrorText);
+			
+				return false;
+			}
+		}
 
-//-------------------------------------------------------------------------------------------------
-void FCustomizableObjectEditorViewportClient::BakeInstance(UCustomizableObjectInstance* InInstance)
-{
-	if (CustomizableObject->GetPrivate()->Status.Get() == FCustomizableObjectStatus::EState::Loading)
+		return true;
+	}
+
+
+	/**
+	 * Validates the AssetPath chosen for the baking data
+	 * @param FileName The filename chosen by the user
+	 * @param AssetPath The AssetPath chosen by the user
+	 * @param InstanceCO The CustomizableObject from the provided COI
+	 * @return True if validation was successful, false otherwise
+	 */
+	bool ValidateProvidedAssetPath(const FString& FileName, const FString& AssetPath, const UCustomizableObject* InstanceCO)
 	{
-		FNotificationInfo Info(NSLOCTEXT("CustomizableObjectEditor", "CustomizableObjectCompileTryLater_BakeInstance", "Please wait unitl Customizable Object is loaded"));
-		Info.bFireAndForget = true;
-		Info.bUseThrobber = true;
-		Info.FadeOutDuration = 1.0f;
-		Info.ExpireDuration = 2.0f;
-		FSlateNotificationManager::Get().AddNotification(Info);
-		return;
+		if (AssetPath.IsEmpty())
+		{
+			UE_LOG(LogMutable, Error, TEXT("The AssetPath can not be empty!"));
+			return false;
+		}
+
+		// Ensure we are not overriding the parent CO
+		const FString FullAssetPath = AssetPath + FString("/") + FileName + FString(".") + FileName;		// Full asset path to the new asset we want to create
+		const bool bWouldOverrideParentCO = InstanceCO->GetPathName() == FullAssetPath;
+		if (bWouldOverrideParentCO)
+		{
+			const FText ErrorText = LOCTEXT("FCustomizableObjectEditorViewportClient_BakeInstance_OverwriteCO", "The selected path would overwrite the instance's parent Customizable Object.");
+
+			ShowErrorNotification(ErrorText);
+			
+			return false;
+		}
+
+		return true;
 	}
 
-	UCustomizableObjectInstance* Instance = InInstance ? InInstance : CustomizableObjectEditorPtr.Pin()->GetPreviewInstance();
 
-	UCustomizableObject* CustomizableObjectFromInstance = Instance ? Instance->GetCustomizableObject() : nullptr;
+	/**
+	 * Outputs a string that we know it is unique.
+	 * @param InResource The resource we are working with
+	 * @param ResourceName The name of the resource we have provided. This should have the name of the current resource and will have the unique name for the resource once the method exits
+	 * @param InCachedResources Collection with all the already processed resources.
+	 * @param InCachedResourceNames Collection with all the already processed resources name's
+	 * @return True if the generation of the unique resource name was successful, false otherwise.
+	 */
+	bool GetUniqueResourceName(const UObject* InResource, FString& ResourceName, TArray<UObject*>& InCachedResources, const TArray<FString>& InCachedResourceNames)
+	{
+		check(InResource);
+
+		int32 FindResult = InCachedResourceNames.Find(ResourceName);
+		if (FindResult != INDEX_NONE)
+		{
+			if (InResource == InCachedResources[FindResult])
+			{
+				return false;
+			}
+
+			uint32 Count = 0;
+			while (FindResult != INDEX_NONE)
+			{
+				FindResult = InCachedResourceNames.Find(ResourceName + "_" + FString::FromInt(Count));
+				Count++;
+			}
+
+			ResourceName += "_" + FString::FromInt(--Count);
+		}
+
+		return true;
+	}
+
 
-	if (!CustomizableObjectFromInstance || CustomizableObjectFromInstance->IsLocked())
+	/**
+	 * Ensures the resource we want to save is ready to be saved. It handles closing it's editor and warning the user about possible overriding of resources.
+	 * @param InAssetSavePath The directory path where to save the baked object
+	 * @param InObjName The name of the object to be baked
+	 * @param bUsedGrantedOverridingRights Control flag that determines if the user has given or not permission to override resources already in disk
+	 * @return True if the operation was successful, false otherwise
+	 */
+	bool ManageBakingAction(const FString& InAssetSavePath, const FString& InObjName, bool& bUsedGrantedOverridingRights)
 	{
-		FNotificationInfo Info(NSLOCTEXT("CustomizableObjectEditor", "CustomizableObjectCompilingTryLater", "Please wait until the Customizable Object is compiled"));
-		Info.bFireAndForget = true;
-		Info.bUseThrobber = true;
-		Info.FadeOutDuration = 1.0f;
-		Info.ExpireDuration = 2.0f;
-		FSlateNotificationManager::Get().AddNotification(Info);
-		return;
+		FString PackagePath = InAssetSavePath + "/" + InObjName;
+		UPackage* ExistingPackage = FindPackage(NULL, *PackagePath);
+
+		if (!ExistingPackage)
+		{
+			FString PackageFilePath = PackagePath + "." + InObjName;
+
+			FString PackageFileName;
+			if (FPackageName::DoesPackageExist(PackageFilePath, &PackageFileName))
+			{
+				ExistingPackage = LoadPackage(nullptr, *PackageFileName, LOAD_EditorOnly);
+			}
+			else
+			{
+				// if package does not exists
+				bUsedGrantedOverridingRights = false;
+				return true;
+			}
+		}
+
+		if (ExistingPackage)
+		{
+			// Checking if the asset is open in an editor
+			TArray<IAssetEditorInstance*> ObjectEditors = GEditor->GetEditorSubsystem<UAssetEditorSubsystem>()->FindEditorsForAssetAndSubObjects(ExistingPackage);
+			if (ObjectEditors.Num())
+			{
+				for (IAssetEditorInstance* ObjectEditorInstance : ObjectEditors)
+				{
+					// Close the editors that contains this asset
+					if (!ObjectEditorInstance->CloseWindow(EAssetEditorCloseReason::AssetEditorHostClosed))
+					{
+						const FText Caption = LOCTEXT("OpenExisitngFile", "Open File");
+						const FText Message = FText::Format(LOCTEXT("CantCloseAsset", "This Obejct \"{0}\" is open in an editor and can't be closed automatically. Please close the editor and try to bake it again"), FText::FromString(InObjName));
+
+						ShowErrorNotification(Message, &Caption);
+						
+						return false;
+					}
+				}
+			}
+
+			if (!bUsedGrantedOverridingRights)
+			{
+				const FText Caption = LOCTEXT("Already existing baked files", "Already existing baked files");
+				const FText Message = FText::Format(LOCTEXT("OverwriteBakedInstance", "Instance baked files already exist in selected destination \"{0}\", this action will overwrite them."), FText::AsCultureInvariant(InAssetSavePath));
+
+				// We need to guard this case since it may crash the editor if we are running an unattended commandlet and we try to generate this dialog
+				if (!FApp::IsUnattended())
+				{
+					if (FMessageDialog::Open(EAppMsgType::OkCancel, Message, Caption) == EAppReturnType::Cancel)
+					{
+						return false;
+					}
+				}
+				else
+				{
+					UE_LOG(LogMutable, Error, TEXT("%s - %s"), *Caption.ToString(), *Message.ToString());
+				}
+
+
+				bUsedGrantedOverridingRights = true;
+			}
+
+			UObject* ExistingObject = StaticFindObject(UObject::StaticClass(), ExistingPackage, *InObjName);
+			if (ExistingObject)
+			{
+				ExistingPackage->FullyLoad();
+
+				TArray<UObject*> ObjectsToDelete;
+				ObjectsToDelete.Add(ExistingObject);
+
+				// Delete objects in the package with the same name as the one we want to create
+				const uint32 NumObjectsDeleted = ObjectTools::ForceDeleteObjects(ObjectsToDelete, false);
+
+				return NumObjectsDeleted == ObjectsToDelete.Num();
+			}
+		}
+
+		return true;
 	}
+}
 
-	if (!Instance->HasAnySkeletalMesh())
+
+
+void FCustomizableObjectEditorViewportClient::BakeInstance()
+{
+	// Early exit if no instance is set in the editor 
+	UCustomizableObjectInstance* Instance = CustomizableObjectEditorPtr.Pin()->GetPreviewInstance();
+	if (!Instance)
 	{
+		UE_LOG(LogMutable, Error, TEXT("No Mutable Customizable Object Instnace was found in the current editor."));
 		return;
 	}
+	BakeTempInstance = Instance->Clone();
 
-	UCustomizableObjectSystem* System = UCustomizableObjectSystem::GetInstance();
-	check(System);
+	// Call the instance update async method
+	FInstanceUpdateNativeDelegate UpdateDelegate;
+	UpdateDelegate.AddRaw(this, &FCustomizableObjectEditorViewportClient::OnInstanceForBakingUpdate);
+	UpdateInstanceForBaking(*BakeTempInstance, UpdateDelegate);
+}
 
-	if (!InInstance)
-	{
-		// The instance in the editor viewport does not have high quality mips in the platform data because streaming is enabled.
-		// Disable streaming and retry with a newly generated temp instance.
-		System->SetProgressiveMipStreamingEnabled(false);
-		// Disable requested LOD generation as it will prevent the new instance from having all the LODs
-		System->SetOnlyGenerateRequestedLODsEnabled(false);
-		// Force high quality texture compression for this instance
-		PrepareUnrealCompression();
-		System->SetImagePixelFormatOverride(UnrealPixelFormatFunc);
 
-		BakeTempInstance = Instance->Clone();
-		BakeTempInstance->UpdatedNativeDelegate.AddSP(this, &FCustomizableObjectEditorViewportClient::BakeInstance);
-		BakeTempInstance->UpdateSkeletalMeshAsync(true, true);
+void BakeCustomizableObjectInstance(
+	UCustomizableObjectInstance& InInstance,
+	const FString& FileName,
+	const FString& AssetPath,
+	const bool bExportAllResources,
+	const bool bGenerateConstantMaterialInstances)
+{
+	// Ensure that the state of the COI provided is valid --------------------------------------------------------------------------------------------
+	UCustomizableObject* InstanceCO = InInstance.GetCustomizableObject();
+	check (InstanceCO);
+	
+	if (InstanceCO->GetPrivate()->Status.Get() == FCustomizableObjectStatus::EState::Loading)
+	{
+		FCustomizableObjectEditorLogger::CreateLog(
+			LOCTEXT("CustomizableObjectCompileTryLater_BakeInstance","Please wait unitl Customizable Object is loaded"))
+		.Category(ELoggerCategory::COInstanceBaking)
+		.CustomNotification()
+		.Notification(true)
+		.Log();
 
 		return;
 	}
-
-	// Warn if the CO looks like it wasn't compiled with high-quality texture compression.
-	// \TODO: This is a weak test, we should store the setting with the compiled data and check that instead.
-	UCustomizableObject* CO = Instance->GetCustomizableObject();
-	if (!CO)
+	
+	if (!ValidateProvidedFileName(FileName))
 	{
-		// Something is very wrong
+		UE_LOG(LogMutable, Error, TEXT("The FileName for the instance baking is not valid."));
 		return;
 	}
 
-	if (CO->CompileOptions.TextureCompression!=ECustomizableObjectTextureCompression::HighQuality)
+	if (!ValidateProvidedAssetPath(FileName,AssetPath,InstanceCO))
 	{
-		FNotificationInfo Info(NSLOCTEXT("CustomizableObjectEditor", "CustomizableObjectBakeLowQuality", "The Customizable Object wasn't compiled with high quality textures. For the best baking results, change the Texture Compression setting and recompile it."));
-		Info.bFireAndForget = true;
-		Info.bUseThrobber = true;
-		Info.FadeOutDuration = 1.0f;
-		Info.ExpireDuration = 6.0f;
-		FSlateNotificationManager::Get().AddNotification(Info);
+		UE_LOG(LogMutable, Error, TEXT("The AssetPath for the instance baking is not valid."));
+		return;
 	}
+	
+	// Ensure the CO of the COI is accessible 
+	if (!InstanceCO || InstanceCO->IsLocked())
+	{
+		FCustomizableObjectEditorLogger::CreateLog(
+		LOCTEXT("CustomizableObjectCompilingTryLater_Baking", "Please wait until the Customizable Object is compiled"))
+		.Category(ELoggerCategory::COInstanceBaking)
+		.CustomNotification()
+		.Notification(true)
+		.Log();
 
-	FString ObjectName = CO->GetName();
-	FText DefaultFileName = FText::Format(LOCTEXT("DefaultFileNameForBakeInstance", "{0}"), FText::AsCultureInvariant(ObjectName));
-	bool bExportAllResources = false;
-	bool bGenerateConstantMaterialInstances = false;
-
-	TSharedRef<SMutableSelectFolderDlg> FolderDlg =
-		SNew(SMutableSelectFolderDlg)
-		.DefaultAssetPath(FText())
-		.DefaultFileName(DefaultFileName);
+		return;
+	}
+	
+	// Exit early if the provided instance does not have a skeletal mesh
+	if (!InInstance.HasAnySkeletalMesh())
+	{
+		return;
+	}
 
-	if (FolderDlg->ShowModal() != EAppReturnType::Cancel)
+	// COI Validation completed : Proceed with the baking operation ----------------------------------------------------------------------------------
+	
+	// Notify of better configuration -> Continue operation normally
+	if (InstanceCO->CompileOptions.TextureCompression != ECustomizableObjectTextureCompression::HighQuality)
 	{
-		// Make sure we can create the asset without conflicts
-		//IAssetTools& AssetTools = FModuleManager::GetModuleChecked<FAssetToolsModule>("AssetTools").Get();
-		//if (!AssetTools.CanCreateAsset(ObjName, PkgName, LOCTEXT("BakeMutableInstance", "Baking a Mutable instance")))
-		//{
-		//	return;
-		//}
+		FCustomizableObjectEditorLogger::CreateLog(
+		LOCTEXT("CustomizableObjectBakeLowQuality", "The Customizable Object wasn't compiled with high quality textures. For the best baking results, change the Texture Compression setting and recompile it."))
+		.Category(ELoggerCategory::COInstanceBaking)
+		.CustomNotification()
+		.Notification(true)
+		.Log();
+	}
+	
+	
+	// Set the overriding flag to false wo we ask the user at least once about if he is willing to override old baked data
+	bool bUsedGrantedOverridingRights = false;
 
-		FString FileName = FolderDlg->GetFileName();
-		ObjectName = FileName;
+	TArray<UPackage*> PackagesToSave;
 
-		TCHAR InvalidCharacter = '0';
-		FString InvalidCharacters = FPaths::GetInvalidFileSystemChars();
+	const int32 NumComponents = InInstance.GetNumComponents();
+	for (int32 ComponentIndex = 0; ComponentIndex < NumComponents; ++ComponentIndex)
+	{
+		USkeletalMesh* Mesh = InInstance.GetSkeletalMesh(ComponentIndex);
 
-		for (int32 i = 0; i < InvalidCharacters.Len(); ++i)
+		if (!Mesh)
 		{
-			TCHAR Char = InvalidCharacters[i];
-			FString SearchedChar = FString::Chr(Char);
-			if (ObjectName.Contains(SearchedChar))
-			{
-				InvalidCharacter = InvalidCharacters[i];
-				break;
-			}
+			continue;
 		}
 
-		bExportAllResources = FolderDlg->GetExportAllResources();
-		bGenerateConstantMaterialInstances = FolderDlg->GetGenerateConstantMaterialInstances();
-
-		BakingOverwritePermission = false;
-		FString CustomObjectPath = CO->GetPathName();
-		FString AssetPath = FolderDlg->GetAssetPath();
-		FString FullAssetPath = AssetPath + FString("/") + ObjectName + FString(".") + ObjectName;
-
-		if (InvalidCharacter != '0')
+		FString ObjectName = FileName;
+		if (NumComponents > 1)
 		{
-			FText ErrorString = FText::FromString(FString::Chr(InvalidCharacter));
-			FMessageDialog::Open(EAppMsgType::Ok, FText::Format(LOCTEXT("FCustomizableObjectEditorViewportClient_BakeInstance_InvalidCharacter", "The selected contains an invalid character ({0})."), ErrorString));
+			ObjectName = FileName + "_Component_" + FString::FromInt(ComponentIndex);
 		}
-		else if (CustomObjectPath == FullAssetPath)
-		{
-			FMessageDialog::Open(EAppMsgType::Ok, LOCTEXT("FCustomizableObjectEditorViewportClient_BakeInstance_OverwriteCO", "The selected path would overwrite the instance's parent Customizable Object."));
-		}
-		else
-		{
-			TArray<UPackage*> PackagesToSave;
 
-			const int32 NumComponents = Instance->GetNumComponents();
-			for (int32 ComponentIndex = 0; ComponentIndex < NumComponents; ++ComponentIndex)
+		TMap<UObject*, UObject*> ReplacementMap;
+		TArray<FString> ArrayCachedElement;
+		TArray<UObject*> ArrayCachedObject;
+
+		if (bExportAllResources)
+		{
+			UMaterialInstance* Inst;
+			UMaterial* Material;
+			UTexture* Texture;
+			FString MaterialName;
+			FString ResourceName;
+			FString PackageName;
+			UObject* DuplicatedObject;
+			TArray<TMap<int, UTexture*>> TextureReplacementMaps;
+
+			// Duplicate Mutable generated textures
+			for (int32 m = 0; m < Mesh->GetMaterials().Num(); ++m)
 			{
-				USkeletalMesh* Mesh = Instance->GetSkeletalMesh(ComponentIndex);
+				UMaterialInterface* Interface = Mesh->GetMaterials()[m].MaterialInterface;
+				Material = Interface->GetMaterial();
+				MaterialName = Material ? Material->GetName() : "Material";
+				Inst = Cast<UMaterialInstance>(Mesh->GetMaterials()[m].MaterialInterface);
 
-				if (!Mesh)
-				{
-					continue;
-				}
+				TMap<int, UTexture*> ReplacementTextures;
+				TextureReplacementMaps.Add(ReplacementTextures);
 
-				if (NumComponents > 1)
+				// The material will only have Mutable generated textures if it's actually a UMaterialInstance
+				if (Material != nullptr && Inst != nullptr)
 				{
-					ObjectName = FileName + "_Component_" + FString::FromInt(ComponentIndex);
-				}
-
-				TMap<UObject*, UObject*> ReplacementMap;
-				TArray<FString> ArrayCachedElement;
-				TArray<UObject*> ArrayCachedObject;
+					TArray<FName> ParameterNames = FUnrealBakeHelpers::GetTextureParameterNames(Material);
 
-				if (bExportAllResources)
-				{
-					UMaterialInstance* Inst;
-					UMaterial* Material;
-					UTexture* Texture;
-					FString MaterialName;
-					FString ResourceName;
-					FString PackageName;
-					UObject* DuplicatedObject;
-					TArray<TMap<int, UTexture*>> TextureReplacementMaps;
-
-					// Duplicate Mutable generated textures
-					for (int32 m = 0; m < Mesh->GetMaterials().Num(); ++m)
+					for (int32 i = 0; i < ParameterNames.Num(); i++)
 					{
-						UMaterialInterface* Interface = Mesh->GetMaterials()[m].MaterialInterface;
-						Material = Interface->GetMaterial();
-						MaterialName = Material ? Material->GetName() : "Material";
-						Inst = Cast<UMaterialInstance>(Mesh->GetMaterials()[m].MaterialInterface);
-
-						TMap<int, UTexture*> ReplacementTextures;
-						TextureReplacementMaps.Add(ReplacementTextures);
-
-						// The material will only have Mutable generated textures if it's actually a UMaterialInstance
-						if (Material != nullptr && Inst != nullptr)
+						if (Inst->GetTextureParameterValue(ParameterNames[i], Texture))
 						{
-							TArray<FName> ParameterNames = FUnrealBakeHelpers::GetTextureParameterNames(Material);
+							UTexture2D* SrcTex = Cast<UTexture2D>(Texture);
+							if (!SrcTex) continue;
+
+							bool bIsMutableTexture = false;
 
-							for (int32 i = 0; i < ParameterNames.Num(); i++)
+							for (UAssetUserData* UserData : *SrcTex->GetAssetUserDataArray())
 							{
-								if (Inst->GetTextureParameterValue(ParameterNames[i], Texture))
+								UTextureMipDataProviderFactory* CustomMipDataProviderFactory = Cast<UMutableTextureMipDataProviderFactory>(UserData);
+								if (CustomMipDataProviderFactory)
 								{
-									UTexture2D* SrcTex = Cast<UTexture2D>(Texture);
-									if (!SrcTex) continue;
-
-									bool bIsMutableTexture = false;
-
-									for (UAssetUserData* UserData : *SrcTex->GetAssetUserDataArray())
-									{
-										UTextureMipDataProviderFactory* CustomMipDataProviderFactory = Cast<UMutableTextureMipDataProviderFactory>(UserData);
-										if (CustomMipDataProviderFactory)
-										{
-											bIsMutableTexture = true;
-										}
-									}
+									bIsMutableTexture = true;
+								}
+							}
 
-									if ((SrcTex->GetPlatformData() != nullptr) &&
-										(SrcTex->GetPlatformData()->Mips.Num() > 0) &&
-										bIsMutableTexture)
-									{
-										FString ParameterSanitized = ParameterNames[i].GetPlainNameString();
-										RemoveRestrictedChars(ParameterSanitized);
-										ResourceName = ObjectName + "_" + MaterialName + "_" + ParameterSanitized;
+							if ((SrcTex->GetPlatformData() != nullptr) &&
+								(SrcTex->GetPlatformData()->Mips.Num() > 0) &&
+								bIsMutableTexture)
+							{
+								FString ParameterSanitized = ParameterNames[i].GetPlainNameString();
+								RemoveRestrictedChars(ParameterSanitized);
+								ResourceName = ObjectName + "_" + MaterialName + "_" + ParameterSanitized;
 
-										if (!GetUniqueResourceName(SrcTex, ResourceName, ArrayCachedObject, ArrayCachedElement))
-										{
-											continue;
-										}
+								if (!GetUniqueResourceName(SrcTex, ResourceName, ArrayCachedObject, ArrayCachedElement))
+								{
+									continue;
+								}
 
-										if (!ManageBakingAction(AssetPath, ResourceName))
-										{
-											return;
-										}
+								if (!ManageBakingAction(AssetPath, ResourceName, bUsedGrantedOverridingRights))
+								{
+									return;
+								}
 
-										// Recover original name of the texture parameter value, now substituted by the generated Mutable texture
-										UTexture* OriginalTexture = nullptr;
-										UMaterialInstanceDynamic* InstDynamic = Cast<UMaterialInstanceDynamic>(Mesh->GetMaterials()[m].MaterialInterface);
-										if (InstDynamic != nullptr)
-										{
-											InstDynamic->Parent->GetTextureParameterValue(FName(*ParameterNames[i].GetPlainNameString()), OriginalTexture);
-										}
-										else
-										{
-											UMaterialInstanceConstant* InstConstant = Cast<UMaterialInstanceConstant>(Mesh->GetMaterials()[m].MaterialInterface);
+								// Recover original name of the texture parameter value, now substituted by the generated Mutable texture
+								UTexture* OriginalTexture = nullptr;
+								UMaterialInstanceDynamic* InstDynamic = Cast<UMaterialInstanceDynamic>(Mesh->GetMaterials()[m].MaterialInterface);
+								if (InstDynamic != nullptr)
+								{
+									InstDynamic->Parent->GetTextureParameterValue(FName(*ParameterNames[i].GetPlainNameString()), OriginalTexture);
+								}
+								else
+								{
+									UMaterialInstanceConstant* InstConstant = Cast<UMaterialInstanceConstant>(Mesh->GetMaterials()[m].MaterialInterface);
 
-											if (InstConstant != nullptr)
-											{
-												InstConstant->Parent->GetTextureParameterValue(FName(*ParameterNames[i].GetPlainNameString()), OriginalTexture);
-											}
-										}
+									if (InstConstant != nullptr)
+									{
+										InstConstant->Parent->GetTextureParameterValue(FName(*ParameterNames[i].GetPlainNameString()), OriginalTexture);
+									}
+								}
 
-										PackageName = FolderDlg->GetAssetPath() + FString("/") + ResourceName;
-										TMap<UObject*, UObject*> FakeReplacementMap;
-										UTexture2D* DupTex = FUnrealBakeHelpers::BakeHelper_CreateAssetTexture(SrcTex, ResourceName, PackageName, OriginalTexture, true, FakeReplacementMap, BakingOverwritePermission);
-										ArrayCachedElement.Add(ResourceName);
-										ArrayCachedObject.Add(DupTex);
-										PackagesToSave.Add(DupTex->GetPackage());
+								PackageName = AssetPath + FString("/") + ResourceName;
+								TMap<UObject*, UObject*> FakeReplacementMap;
+								UTexture2D* DupTex = FUnrealBakeHelpers::BakeHelper_CreateAssetTexture(SrcTex, ResourceName, PackageName, OriginalTexture, true, FakeReplacementMap, bUsedGrantedOverridingRights);
+								ArrayCachedElement.Add(ResourceName);
+								ArrayCachedObject.Add(DupTex);
+								PackagesToSave.Add(DupTex->GetPackage());
 
-										if (OriginalTexture != nullptr)
-										{
-											TextureReplacementMaps[m].Add(i, DupTex);
-										}
-									}
+								if (OriginalTexture != nullptr)
+								{
+									TextureReplacementMaps[m].Add(i, DupTex);
 								}
 							}
 						}
 					}
+				}
+			}
 
-					// Duplicate non-Mutable material textures
-					for (int32 m = 0; m < Mesh->GetMaterials().Num(); ++m)
-					{
-						UMaterialInterface* Interface = Mesh->GetMaterials()[m].MaterialInterface;
-						Material = Interface->GetMaterial();
-						MaterialName = Material ? Material->GetName() : "Material";
+			// Duplicate non-Mutable material textures
+			for (int32 m = 0; m < Mesh->GetMaterials().Num(); ++m)
+			{
+				UMaterialInterface* Interface = Mesh->GetMaterials()[m].MaterialInterface;
+				Material = Interface->GetMaterial();
+				MaterialName = Material ? Material->GetName() : "Material";
+
+				if (Material != nullptr)
+				{
+					TArray<FName> ParameterNames = FUnrealBakeHelpers::GetTextureParameterNames(Material);
 
-						if (Material != nullptr)
+					for (int32 i = 0; i < ParameterNames.Num(); i++)
+					{
+						TArray<FMaterialParameterInfo> InfoArray;
+						TArray<FGuid> GuidArray;
+						Material->GetAllTextureParameterInfo(InfoArray, GuidArray);
+						
+						if (Material->GetTextureParameterValue(InfoArray[i], Texture))
 						{
-							TArray<FName> ParameterNames = FUnrealBakeHelpers::GetTextureParameterNames(Material);
+							FString ParameterSanitized = ParameterNames[i].GetPlainNameString();
+							RemoveRestrictedChars(ParameterSanitized);
+							ResourceName = ObjectName + "_" + MaterialName + "_" + ParameterSanitized;
 
-							for (int32 i = 0; i < ParameterNames.Num(); i++)
+							if (ArrayCachedElement.Find(ResourceName) == INDEX_NONE)
 							{
-								TArray<FMaterialParameterInfo> InfoArray;
-								TArray<FGuid> GuidArray;
-								Material->GetAllTextureParameterInfo(InfoArray, GuidArray);
-								
-								if (Material->GetTextureParameterValue(InfoArray[i], Texture))
+								if (!ManageBakingAction(AssetPath, ResourceName, bUsedGrantedOverridingRights))
 								{
-									FString ParameterSanitized = ParameterNames[i].GetPlainNameString();
-									RemoveRestrictedChars(ParameterSanitized);
-									ResourceName = ObjectName + "_" + MaterialName + "_" + ParameterSanitized;
-
-									if (ArrayCachedElement.Find(ResourceName) == INDEX_NONE)
-									{
-										if (!ManageBakingAction(AssetPath, ResourceName))
-										{
-											return;
-										}
+									return;
+								}
 
-										PackageName = FolderDlg->GetAssetPath() + FString("/") + ResourceName;
-										TMap<UObject*, UObject*> FakeReplacementMap;
-										DuplicatedObject = FUnrealBakeHelpers::BakeHelper_DuplicateAsset(Texture, ResourceName, PackageName, true, FakeReplacementMap, BakingOverwritePermission, false);
-										ArrayCachedElement.Add(ResourceName);
-										ArrayCachedObject.Add(DuplicatedObject);
-										PackagesToSave.Add(DuplicatedObject->GetPackage());
+								PackageName = AssetPath + FString("/") + ResourceName;
+								TMap<UObject*, UObject*> FakeReplacementMap;
+								DuplicatedObject = FUnrealBakeHelpers::BakeHelper_DuplicateAsset(Texture, ResourceName, PackageName, true, FakeReplacementMap, bUsedGrantedOverridingRights, false);
+								ArrayCachedElement.Add(ResourceName);
+								ArrayCachedObject.Add(DuplicatedObject);
+								PackagesToSave.Add(DuplicatedObject->GetPackage());
 
-										UTexture* DupTexture = Cast<UTexture>(DuplicatedObject);
-										TextureReplacementMaps[m].Add(i, DupTexture);
-									}
-								}
+								UTexture* DupTexture = Cast<UTexture>(DuplicatedObject);
+								TextureReplacementMaps[m].Add(i, DupTexture);
 							}
 						}
 					}
+				}
+			}
 
 
-					// Duplicate the materials used by each material instance so that the replacement map has proper information 
-					// when duplicating the material instances
-					for (int32 m = 0; m < Mesh->GetMaterials().Num(); ++m)
-					{
-						UMaterialInterface* Interface = Mesh->GetMaterials()[m].MaterialInterface;
-						Material = Interface ? Interface->GetMaterial() : nullptr;
+			// Duplicate the materials used by each material instance so that the replacement map has proper information 
+			// when duplicating the material instances
+			for (int32 m = 0; m < Mesh->GetMaterials().Num(); ++m)
+			{
+				UMaterialInterface* Interface = Mesh->GetMaterials()[m].MaterialInterface;
+				Material = Interface ? Interface->GetMaterial() : nullptr;
 
-						if (Material)
-						{
-							ResourceName = ObjectName + "_Material_" + Material->GetName();
+				if (Material)
+				{
+					ResourceName = ObjectName + "_Material_" + Material->GetName();
 
-							if (!GetUniqueResourceName(Material, ResourceName, ArrayCachedObject, ArrayCachedElement))
-							{
-								continue;
-							}
+					if (!GetUniqueResourceName(Material, ResourceName, ArrayCachedObject, ArrayCachedElement))
+					{
+						continue;
+					}
 
-							if (!ManageBakingAction(AssetPath, ResourceName))
-							{
-								return;
-							}
+					if (!ManageBakingAction(AssetPath, ResourceName, bUsedGrantedOverridingRights))
+					{
+						return;
+					}
 
-							PackageName = FolderDlg->GetAssetPath() + FString("/") + ResourceName;
-							TMap<UObject*, UObject*> FakeReplacementMap;
-							DuplicatedObject = FUnrealBakeHelpers::BakeHelper_DuplicateAsset(Material, ResourceName, PackageName, 
-								false, FakeReplacementMap, BakingOverwritePermission, bGenerateConstantMaterialInstances);
-							ArrayCachedElement.Add(ResourceName);
-							ArrayCachedObject.Add(DuplicatedObject);
-							ReplacementMap.Add(Interface, DuplicatedObject);
-							PackagesToSave.Add(DuplicatedObject->GetPackage());
+					PackageName = AssetPath + FString("/") + ResourceName;
+					TMap<UObject*, UObject*> FakeReplacementMap;
+					DuplicatedObject = FUnrealBakeHelpers::BakeHelper_DuplicateAsset(Material, ResourceName, PackageName, 
+						false, FakeReplacementMap, bUsedGrantedOverridingRights, bGenerateConstantMaterialInstances);
+					ArrayCachedElement.Add(ResourceName);
+					ArrayCachedObject.Add(DuplicatedObject);
+					ReplacementMap.Add(Interface, DuplicatedObject);
+					PackagesToSave.Add(DuplicatedObject->GetPackage());
 
-							FUnrealBakeHelpers::CopyAllMaterialParameters(DuplicatedObject, Interface, TextureReplacementMaps[m]);
-						}
-					}
+					FUnrealBakeHelpers::CopyAllMaterialParameters(DuplicatedObject, Interface, TextureReplacementMaps[m]);
 				}
-				else
-				{
-					// Duplicate the material instances
-					for (int32 MaterialIndex = 0; MaterialIndex < Mesh->GetMaterials().Num(); ++MaterialIndex)
-					{
-						UMaterialInterface* Interface = Mesh->GetMaterials()[MaterialIndex].MaterialInterface;
-						UMaterial* ParentMaterial = Interface->GetMaterial();
-						FString MaterialName = ParentMaterial ? ParentMaterial->GetName() : "Material";
+			}
+		}
+		else
+		{
+			// Duplicate the material instances
+			for (int32 MaterialIndex = 0; MaterialIndex < Mesh->GetMaterials().Num(); ++MaterialIndex)
+			{
+				UMaterialInterface* Interface = Mesh->GetMaterials()[MaterialIndex].MaterialInterface;
+				UMaterial* ParentMaterial = Interface->GetMaterial();
+				FString MaterialName = ParentMaterial ? ParentMaterial->GetName() : "Material";
 
-						// Material
-						FString MatObjName = ObjectName + "_" + MaterialName;
+				// Material
+				FString MatObjName = ObjectName + "_" + MaterialName;
 
-						if (!GetUniqueResourceName(Interface, MatObjName, ArrayCachedObject, ArrayCachedElement))
-						{
-							continue;
-						}
+				if (!GetUniqueResourceName(Interface, MatObjName, ArrayCachedObject, ArrayCachedElement))
+				{
+					continue;
+				}
 
-						if (!ManageBakingAction(AssetPath, MatObjName))
-						{
-							return;
-						}
+				if (!ManageBakingAction(AssetPath, MatObjName, bUsedGrantedOverridingRights))
+				{
+					return;
+				}
 
-						FString MatPkgName = FolderDlg->GetAssetPath() + FString("/") + MatObjName;
-						UObject* DupMat = FUnrealBakeHelpers::BakeHelper_DuplicateAsset(Interface, MatObjName, 
-							MatPkgName, false, ReplacementMap, BakingOverwritePermission, bGenerateConstantMaterialInstances);
-						ArrayCachedObject.Add(DupMat);
-						ArrayCachedElement.Add(MatObjName);
-						PackagesToSave.Add(DupMat->GetPackage());
+				FString MatPkgName = AssetPath + FString("/") + MatObjName;
+				UObject* DupMat = FUnrealBakeHelpers::BakeHelper_DuplicateAsset(Interface, MatObjName, 
+					MatPkgName, false, ReplacementMap, bUsedGrantedOverridingRights, bGenerateConstantMaterialInstances);
+				ArrayCachedObject.Add(DupMat);
+				ArrayCachedElement.Add(MatObjName);
+				PackagesToSave.Add(DupMat->GetPackage());
 
-						UMaterialInstance* Inst = Cast<UMaterialInstance>(Interface);
+				UMaterialInstance* Inst = Cast<UMaterialInstance>(Interface);
 
-						// Only need to duplicate the generate textures if the original material is a dynamic instance
-						// If the material has Mutable textures, then it will be a dynamic material instance for sure
-						if (Inst)
-						{
-							// Duplicate generated textures
-							UMaterialInstanceDynamic* InstDynamic = Cast<UMaterialInstanceDynamic>(DupMat);
-							UMaterialInstanceConstant* InstConstant = Cast<UMaterialInstanceConstant>(DupMat);
+				// Only need to duplicate the generate textures if the original material is a dynamic instance
+				// If the material has Mutable textures, then it will be a dynamic material instance for sure
+				if (Inst)
+				{
+					// Duplicate generated textures
+					UMaterialInstanceDynamic* InstDynamic = Cast<UMaterialInstanceDynamic>(DupMat);
+					UMaterialInstanceConstant* InstConstant = Cast<UMaterialInstanceConstant>(DupMat);
 
-							if (InstDynamic || InstConstant)
+					if (InstDynamic || InstConstant)
+					{
+						for (int32 TextureIndex = 0; TextureIndex < Inst->TextureParameterValues.Num(); ++TextureIndex)
+						{
+							if (Inst->TextureParameterValues[TextureIndex].ParameterValue)
 							{
-								for (int32 TextureIndex = 0; TextureIndex < Inst->TextureParameterValues.Num(); ++TextureIndex)
+								if (Inst->TextureParameterValues[TextureIndex].ParameterValue->HasAnyFlags(RF_Transient))
 								{
-									if (Inst->TextureParameterValues[TextureIndex].ParameterValue)
+									UTexture2D* SrcTex = Cast<UTexture2D>(Inst->TextureParameterValues[TextureIndex].ParameterValue);
+
+									if (SrcTex)
 									{
-										if (Inst->TextureParameterValues[TextureIndex].ParameterValue->HasAnyFlags(RF_Transient))
-										{
-											UTexture2D* SrcTex = Cast<UTexture2D>(Inst->TextureParameterValues[TextureIndex].ParameterValue);
+										FString ParameterSanitized = Inst->TextureParameterValues[TextureIndex].ParameterInfo.Name.ToString();
+										RemoveRestrictedChars(ParameterSanitized);
 
-											if (SrcTex)
-											{
-												FString ParameterSanitized = Inst->TextureParameterValues[TextureIndex].ParameterInfo.Name.ToString();
-												RemoveRestrictedChars(ParameterSanitized);
-
-												FString TexObjName = ObjectName + "_" + MaterialName + "_" + ParameterSanitized;
-
-												if (!GetUniqueResourceName(SrcTex, TexObjName, ArrayCachedObject, ArrayCachedElement))
-												{
-													UTexture* PrevTexture = Cast<UTexture>(ArrayCachedObject[ArrayCachedElement.Find(TexObjName)]);
-
-													if (InstDynamic)
-													{
-														InstDynamic->SetTextureParameterValue(Inst->TextureParameterValues[TextureIndex].ParameterInfo.Name, PrevTexture);
-													}
-													else if (InstConstant)
-													{
-														InstConstant->SetTextureParameterValueEditorOnly(Inst->TextureParameterValues[TextureIndex].ParameterInfo.Name, PrevTexture);
-													}
-													
-													continue;
-												}
-
-												if (!ManageBakingAction(AssetPath, TexObjName))
-												{
-													return;
-												}
-
-												FString TexPkgName = FolderDlg->GetAssetPath() + FString("/") + TexObjName;
-												TMap<UObject*, UObject*> FakeReplacementMap;
-												UTexture2D* DupTex = FUnrealBakeHelpers::BakeHelper_CreateAssetTexture(SrcTex, TexObjName, TexPkgName, nullptr, false, FakeReplacementMap, BakingOverwritePermission);
-												ArrayCachedObject.Add(DupTex);
-												ArrayCachedElement.Add(TexObjName);
-												PackagesToSave.Add(DupTex->GetPackage());
-
-												if (InstDynamic)
-												{
-													InstDynamic->SetTextureParameterValue(Inst->TextureParameterValues[TextureIndex].ParameterInfo.Name, DupTex);
-												}
-												else if(InstConstant)
-												{
-													InstConstant->SetTextureParameterValueEditorOnly(Inst->TextureParameterValues[TextureIndex].ParameterInfo.Name, DupTex);
-												}
-											}
-											else
-											{
-												UE_LOG(LogMutable, Error, TEXT("A Mutable texture that is not a Texture2D has been found while baking a CustomizableObjectInstance."));
-											}
-										}
-										else
+										FString TexObjName = ObjectName + "_" + MaterialName + "_" + ParameterSanitized;
+
+										if (!GetUniqueResourceName(SrcTex, TexObjName, ArrayCachedObject, ArrayCachedElement))
 										{
-											// If it's not transient it's not a mutable texture, it's a pass-through texture
-											// Just set the original texture
+											UTexture* PrevTexture = Cast<UTexture>(ArrayCachedObject[ArrayCachedElement.Find(TexObjName)]);
+
 											if (InstDynamic)
 											{
-												InstDynamic->SetTextureParameterValue(Inst->TextureParameterValues[TextureIndex].ParameterInfo.Name, Inst->TextureParameterValues[TextureIndex].ParameterValue);
+												InstDynamic->SetTextureParameterValue(Inst->TextureParameterValues[TextureIndex].ParameterInfo.Name, PrevTexture);
 											}
 											else if (InstConstant)
 											{
-												InstConstant->SetTextureParameterValueEditorOnly(Inst->TextureParameterValues[TextureIndex].ParameterInfo.Name, Inst->TextureParameterValues[TextureIndex].ParameterValue);
+												InstConstant->SetTextureParameterValueEditorOnly(Inst->TextureParameterValues[TextureIndex].ParameterInfo.Name, PrevTexture);
 											}
+											
+											continue;
+										}
+
+										if (!ManageBakingAction(AssetPath, TexObjName, bUsedGrantedOverridingRights))
+										{
+											return;
+										}
+
+										FString TexPkgName = AssetPath + FString("/") + TexObjName;
+										TMap<UObject*, UObject*> FakeReplacementMap;
+										UTexture2D* DupTex = FUnrealBakeHelpers::BakeHelper_CreateAssetTexture(SrcTex, TexObjName, TexPkgName, nullptr, false, FakeReplacementMap, bUsedGrantedOverridingRights);
+										ArrayCachedObject.Add(DupTex);
+										ArrayCachedElement.Add(TexObjName);
+										PackagesToSave.Add(DupTex->GetPackage());
+
+										if (InstDynamic)
+										{
+											InstDynamic->SetTextureParameterValue(Inst->TextureParameterValues[TextureIndex].ParameterInfo.Name, DupTex);
 										}
+										else if(InstConstant)
+										{
+											InstConstant->SetTextureParameterValueEditorOnly(Inst->TextureParameterValues[TextureIndex].ParameterInfo.Name, DupTex);
+										}
+									}
+									else
+									{
+										UE_LOG(LogMutable, Error, TEXT("A Mutable texture that is not a Texture2D has been found while baking a CustomizableObjectInstance."));
+									}
+								}
+								else
+								{
+									// If it's not transient it's not a mutable texture, it's a pass-through texture
+									// Just set the original texture
+									if (InstDynamic)
+									{
+										InstDynamic->SetTextureParameterValue(Inst->TextureParameterValues[TextureIndex].ParameterInfo.Name, Inst->TextureParameterValues[TextureIndex].ParameterValue);
+									}
+									else if (InstConstant)
+									{
+										InstConstant->SetTextureParameterValueEditorOnly(Inst->TextureParameterValues[TextureIndex].ParameterInfo.Name, Inst->TextureParameterValues[TextureIndex].ParameterValue);
 									}
 								}
 							}
 						}
 					}
 				}
-				
-				// Skeletal Mesh's Skeleton
-				if (Mesh->GetSkeleton())
-				{
-					const bool bTransient = Mesh->GetSkeleton()->GetPackage() == GetTransientPackage();
-
-					// Don't duplicate if not transient or export all assets.
-					if (bTransient || bExportAllResources)
-					{
-						FString SkeletonName = ObjectName + "_Skeleton";
-						if (!ManageBakingAction(AssetPath, SkeletonName))
-						{
-							return;
-						}
-
-						FString SkeletonPkgName = FolderDlg->GetAssetPath() + FString("/") + SkeletonName;
-						UObject* DuplicatedSkeleton = FUnrealBakeHelpers::BakeHelper_DuplicateAsset(Mesh->GetSkeleton(), SkeletonName, 
-							SkeletonPkgName, false, ReplacementMap, BakingOverwritePermission, false);
-
-						ArrayCachedObject.Add(DuplicatedSkeleton);
-						PackagesToSave.Add(DuplicatedSkeleton->GetPackage());
-						ReplacementMap.Add(Mesh->GetSkeleton(), DuplicatedSkeleton);
-					}
-				}
+			}
+		}
+		
+		// Skeletal Mesh's Skeleton
+		if (Mesh->GetSkeleton())
+		{
+			const bool bTransient = Mesh->GetSkeleton()->GetPackage() == GetTransientPackage();
 
-				// Skeletal Mesh
-				if (!ManageBakingAction(AssetPath, ObjectName))
+			// Don't duplicate if not transient or export all assets.
+			if (bTransient || bExportAllResources)
+			{
+				FString SkeletonName = ObjectName + "_Skeleton";
+				if (!ManageBakingAction(AssetPath, SkeletonName, bUsedGrantedOverridingRights))
 				{
 					return;
 				}
 
-				FString PkgName = FolderDlg->GetAssetPath() + FString("/") + ObjectName;
-				UObject* DupObject = FUnrealBakeHelpers::BakeHelper_DuplicateAsset(Mesh, ObjectName, PkgName, 
-					false, ReplacementMap, BakingOverwritePermission, false);
-				ArrayCachedObject.Add(DupObject);
-				PackagesToSave.Add(DupObject->GetPackage());
+				FString SkeletonPkgName = AssetPath + FString("/") + SkeletonName;
+				UObject* DuplicatedSkeleton = FUnrealBakeHelpers::BakeHelper_DuplicateAsset(Mesh->GetSkeleton(), SkeletonName, 
+					SkeletonPkgName, false, ReplacementMap, bUsedGrantedOverridingRights, false);
 
-				Mesh->Build();
+				ArrayCachedObject.Add(DuplicatedSkeleton);
+				PackagesToSave.Add(DuplicatedSkeleton->GetPackage());
+				ReplacementMap.Add(Mesh->GetSkeleton(), DuplicatedSkeleton);
+			}
+		}
 
-				USkeletalMesh* SkeletalMesh = Cast<USkeletalMesh>(DupObject);
-				if (SkeletalMesh)
-				{
-					SkeletalMesh->GetLODInfoArray() = Mesh->GetLODInfoArray();
+		// Skeletal Mesh
+		if (!ManageBakingAction(AssetPath, ObjectName, bUsedGrantedOverridingRights))
+		{
+			return;
+		}
 
-					SkeletalMesh->GetImportedModel()->SkeletalMeshModelGUID = FGuid::NewGuid();
+		FString PkgName = AssetPath + FString("/") + ObjectName;
+		UObject* DupObject = FUnrealBakeHelpers::BakeHelper_DuplicateAsset(Mesh, ObjectName, PkgName, 
+			false, ReplacementMap, bUsedGrantedOverridingRights, false);
+		ArrayCachedObject.Add(DupObject);
+		PackagesToSave.Add(DupObject->GetPackage());
 
-					// Duplicate AssetUserData
-					{
-						const TArray<UAssetUserData*>* AssetUserDataArray = Mesh->GetAssetUserDataArray();
-						for (const UAssetUserData* AssetUserData : *AssetUserDataArray)
-						{
-							if (AssetUserData)
-							{
-								// Duplicate to change ownership
-								UAssetUserData* NewAssetUserData = Cast<UAssetUserData>(StaticDuplicateObject(AssetUserData, SkeletalMesh));
-								SkeletalMesh->AddAssetUserData(NewAssetUserData);
-							}
-						}
-					}
+		Mesh->Build();
 
-					// Generate render data
-					SkeletalMesh->Build();
-				}
+		USkeletalMesh* SkeletalMesh = Cast<USkeletalMesh>(DupObject);
+		if (SkeletalMesh)
+		{
+			SkeletalMesh->GetLODInfoArray() = Mesh->GetLODInfoArray();
+
+			SkeletalMesh->GetImportedModel()->SkeletalMeshModelGUID = FGuid::NewGuid();
 
-				// Remove duplicated UObjects from Root (previously added to avoid objects from beeing GC in the middle of the bake process)
-				for (UObject* Obj : ArrayCachedObject)
+			// Duplicate AssetUserData
+			{
+				const TArray<UAssetUserData*>* AssetUserDataArray = Mesh->GetAssetUserDataArray();
+				for (const UAssetUserData* AssetUserData : *AssetUserDataArray)
 				{
-					Obj->RemoveFromRoot();
+					if (AssetUserData)
+					{
+						// Duplicate to change ownership
+						UAssetUserData* NewAssetUserData = Cast<UAssetUserData>(StaticDuplicateObject(AssetUserData, SkeletalMesh));
+						SkeletalMesh->AddAssetUserData(NewAssetUserData);
+					}
 				}
 			}
 
-			if (PackagesToSave.Num())
-			{
-				FEditorFileUtils::PromptForCheckoutAndSave(PackagesToSave, false, true);
-			}
+			// Generate render data
+			SkeletalMesh->Build();
+		}
+
+		// Remove duplicated UObjects from Root (previously added to avoid objects from being GC in the middle of the bake process)
+		for (UObject* Obj : ArrayCachedObject)
+		{
+			Obj->RemoveFromRoot();
 		}
 	}
 
-	if (InInstance)
+	// Save the packages generated during the baking operation  --------------------------------------------------------------------------------------
+	
+	// Complete the baking by saving the packages we have cached during the baking operation
+	if (PackagesToSave.Num())
 	{
-		// Reenable Mutable texture streaming and requested LOD generation as they had been disabled to bake the textures
-		System->SetProgressiveMipStreamingEnabled(true);
-		System->SetOnlyGenerateRequestedLODsEnabled(true);
-		System->SetImagePixelFormatOverride(nullptr);
-		BakeTempInstance = nullptr;
+		FEditorFileUtils::PromptForCheckoutAndSave(PackagesToSave, false, !FApp::IsUnattended());
 	}
 }
 
 
-bool FCustomizableObjectEditorViewportClient::GetUniqueResourceName(UObject* Resource, FString& ResourceName, TArray<UObject*>& InCachedResources, TArray<FString>& InCachedResourceNames)
+void FCustomizableObjectEditorViewportClient::OnInstanceForBakingUpdate(const FUpdateContext& Result)
 {
-	int32 FindResult = InCachedResourceNames.Find(ResourceName);
-	if (FindResult != INDEX_NONE)
+	// Early exit if no instance was provided
+	if (!BakeTempInstance)
 	{
-		if (Resource == InCachedResources[FindResult])
-		{
-			return false;
-		}
-
-		uint32 Count = 0;
-		while (FindResult != INDEX_NONE)
-		{
-			FindResult = InCachedResourceNames.Find(ResourceName + "_" + FString::FromInt(Count));
-			Count++;
-		}
-
-		ResourceName += "_" + FString::FromInt(--Count);
+		UE_LOG(LogMutable, Error, TEXT("No Mutable Customizable Object Instnace was provided for the baking."));
+		return;
 	}
 
-	return true;
-}
-
-
-bool FCustomizableObjectEditorViewportClient::ManageBakingAction(const FString& Path, const FString& ObjName)
-{
-	FString PackagePath = Path + "/" + ObjName;
-	UPackage* ExistingPackage = FindPackage(NULL, *PackagePath);
-
-	if (!ExistingPackage)
+	// Early exit if update result is not success
+	if (Result.UpdateResult != EUpdateResult::Success)
 	{
-		FString PackageFilePath = PackagePath + "." + ObjName;
-
-		FString PackageFileName;
-		if (FPackageName::DoesPackageExist(PackageFilePath, &PackageFileName))
-		{
-			ExistingPackage = LoadPackage(nullptr, *PackageFileName, LOAD_EditorOnly);
-		}
-		else
-		{
-			// if package does not exists
-			BakingOverwritePermission = false;
-			return true;
-		}
+		UE_LOG(LogMutable, Warning ,TEXT("Instance finished update with an error state : %s. Skipping instance baking"), *UEnum::GetValueAsString(Result.UpdateResult));
+		BakeTempInstance = nullptr;
+		return;
 	}
-
-	if (ExistingPackage)
+	
+	const UCustomizableObject* CO = BakeTempInstance->GetCustomizableObject();
+	check (CO);
+	
+	// Let the user set some configurations at the editor level
+	const FText DefaultFileName = FText::Format(LOCTEXT("DefaultFileNameForBakeInstance", "{0}"), FText::AsCultureInvariant(CO->GetName()));
+	
+	TSharedRef<SMutableSelectFolderDlg> FolderDlg =
+		SNew(SMutableSelectFolderDlg)
+		.DefaultAssetPath(FText())
+		.DefaultFileName(DefaultFileName);
+	
+	if (FolderDlg->ShowModal() != EAppReturnType::Cancel)
 	{
-		// Checking if the asset is open in an editor
-		TArray<IAssetEditorInstance*> ObjectEditors = GEditor->GetEditorSubsystem<UAssetEditorSubsystem>()->FindEditorsForAssetAndSubObjects(ExistingPackage);
-		if (ObjectEditors.Num())
-		{
-			for (IAssetEditorInstance* ObjectEditorInstance : ObjectEditors)
-			{
-				// Close the editors that contains this asset
-				if (!ObjectEditorInstance->CloseWindow(EAssetEditorCloseReason::AssetEditorHostClosed))
-				{
-					FText Caption = LOCTEXT("OpenExisitngFile", "Open File");
-					FText Message = FText::Format(LOCTEXT("CantCloseAsset", "This Obejct \"{0}\" is open in an editor and can't be closed automatically. Please close the editor and try to bake it again"), FText::FromString(ObjName));
-
-					FMessageDialog::Open(EAppMsgType::Ok, Message, Caption);
-
-					return false;
-				}
-			}
-		}
-
-		if (!BakingOverwritePermission)
-		{
-			FText Caption = LOCTEXT("Already existing baked files", "Already existing baked files");
-			FText Message = FText::Format(LOCTEXT("OverwriteBakedInstance", "Instance baked files already exist in selected destination \"{0}\", this action will overwrite them."), FText::AsCultureInvariant(Path));
-			
-			if (FMessageDialog::Open(EAppMsgType::OkCancel, Message, Caption) == EAppReturnType::Cancel)
-			{
-				return false;
-			}
-
-			BakingOverwritePermission = true;
-		}
-
-		UObject* ExistingObject = StaticFindObject(UObject::StaticClass(), ExistingPackage, *ObjName);
-		if (ExistingObject)
-		{
-			ExistingPackage->FullyLoad();
-
-			TArray<UObject*> ObjectsToDelete;
-			ObjectsToDelete.Add(ExistingObject);
-
-			// Delete objects in the package with the same name as the one we want to create
-			const uint32 NumObjectsDeleted = ObjectTools::ForceDeleteObjects(ObjectsToDelete, false);
-
-			return NumObjectsDeleted == ObjectsToDelete.Num();
-		}
+		BakeCustomizableObjectInstance(
+			*BakeTempInstance,
+			FolderDlg->GetFileName(),
+			FolderDlg->GetAssetPath(),
+			FolderDlg->GetExportAllResources(),
+			FolderDlg->GetGenerateConstantMaterialInstances());
 	}
-
-	return true;
+	
+	BakeTempInstance = nullptr;
 }
 
 
@@ -2819,13 +2922,13 @@ FString SMutableSelectFolderDlg::GetFileName()
 }
 
 
-bool SMutableSelectFolderDlg::GetExportAllResources()
+bool SMutableSelectFolderDlg::GetExportAllResources() const
 {
 	return bExportAllResources;
 }
 
 
-bool SMutableSelectFolderDlg::GetGenerateConstantMaterialInstances()
+bool SMutableSelectFolderDlg::GetGenerateConstantMaterialInstances() const
 {
 	return bGenerateConstantMaterialInstances;
 }
diff --git a/Engine/Plugins/Experimental/Mutable/Source/CustomizableObjectEditor/Public/MuCOE/CustomizableObjectEditorLogger.h b/Engine/Plugins/Experimental/Mutable/Source/CustomizableObjectEditor/Public/MuCOE/CustomizableObjectEditorLogger.h
--- a/Engine/Plugins/Experimental/Mutable/Source/CustomizableObjectEditor/Public/MuCOE/CustomizableObjectEditorLogger.h
+++ b/Engine/Plugins/Experimental/Mutable/Source/CustomizableObjectEditor/Public/MuCOE/CustomizableObjectEditorLogger.h
@@ -15,7 +15,8 @@ enum class ELoggerCategory
 {
 	General,
 	GraphSearch,
-	Compilation
+	Compilation,
+	COInstanceBaking
 };
 
 
diff --git a/Engine/Plugins/Experimental/Mutable/Source/CustomizableObjectEditor/Public/MuCOE/CustomizableObjectEditorViewportClient.h b/Engine/Plugins/Experimental/Mutable/Source/CustomizableObjectEditor/Public/MuCOE/CustomizableObjectEditorViewportClient.h
--- a/Engine/Plugins/Experimental/Mutable/Source/CustomizableObjectEditor/Public/MuCOE/CustomizableObjectEditorViewportClient.h
+++ b/Engine/Plugins/Experimental/Mutable/Source/CustomizableObjectEditor/Public/MuCOE/CustomizableObjectEditorViewportClient.h
@@ -3,6 +3,7 @@
 #pragma once
 
 #include "EditorViewportClient.h"
+#include "MuCO/CustomizableObjectInstance.h"
 #include "MuCO/CustomizableObjectParameterTypeDefinitions.h"
 
 namespace EAnimationMode { enum Type : int; }
@@ -51,7 +52,6 @@ DECLARE_DELEGATE_RetVal(FColor, FWidgetColorDelegate);
 DECLARE_DELEGATE(FWidgetTrackingStartedDelegate);
 
 
-extern void RemoveRestrictedChars(FString& String);
 
 
 enum class EWidgetType : uint8
@@ -114,16 +114,21 @@ public:
 	 */
 	void DrawUVs(FViewport* InViewport, FCanvas* InCanvas, int32 InTextYPos, const FString& MaterialName);
 
-	// Callback for baking the current instance. If nullptr is passed, the instance in the editor will be taken.
+	/**
+	 * Bake the instance currently present in the editor. Internally will schedule the update of the instance before baking it's resources
+	 * Note : It updates the UCustomizableObjectSystem so it is in the desired config for baking
+	 */
 	void BakeInstance();
-	void BakeInstance(class UCustomizableObjectInstance* Instance);
 	
-	// Returns false if the resource has already been duplicated, otherwise, returns true and an unique ResourceName.
-	bool GetUniqueResourceName(UObject* InResource, FString& InOutResourceName, TArray<UObject*>& InCachedObjects, TArray<FString>& InCachedObjectNames);
+private:
 	
-	// If baking files that already exist, ask the user for permission to overwirte them
-	bool ManageBakingAction(const FString& Path, const FString& Name);
-
+	/**
+	 * Callback executed after the instance in the editor gets updated for baking its contents.
+	 * Note : It updates the UCustomizableObjectSystem so it resets it to the state it had prior to the invocation of BakeEditorInstance
+	 */
+	void OnInstanceForBakingUpdate(const FUpdateContext& Result);
+	
+public:
 	// Callback to show / hide instance geometry data
 	void StateChangeShowGeometryData();
 
@@ -335,9 +340,6 @@ private:
 	/** Flag to control whether to show / hide the instance geometry information data */
 	bool StateChangeShowGeometryDataFlag;
 
-	/** To know if the user has given permission to overwrite files when baking if already present */
-	bool BakingOverwritePermission;
-
 	/** To know if the orbital camera is being used or not */
 	bool bActivateOrbitalCamera;
 
@@ -354,6 +356,10 @@ private:
 	// instance does not have the high quality mips in the texture's platform data
 	TObjectPtr<UCustomizableObjectInstance> BakeTempInstance = nullptr;
 
+	// Cache System configuration cached before performing the mandatory instance update for baking so we can restore it after the bake operation
+	bool bIsProgressiveMipStreamingEnabled = false;
+	bool bIsOnlyGenerateRequestedLODsEnabled = false;
+	
 	// The following delegates currently are only used by the EWidgetType::Projector
 	FWidgetLocationDelegate WidgetLocationDelegate;
 	FOnWidgetLocationChangedDelegate OnWidgetLocationChangedDelegate;
diff --git a/Engine/Plugins/Experimental/Mutable/Source/CustomizableObjectEditor/Public/MuCOE/CustomizableObjectInstanceBakingUtils.cpp b/Engine/Plugins/Experimental/Mutable/Source/CustomizableObjectEditor/Public/MuCOE/CustomizableObjectInstanceBakingUtils.cpp
new file mode 100644
--- /dev/null
+++ b/Engine/Plugins/Experimental/Mutable/Source/CustomizableObjectEditor/Public/MuCOE/CustomizableObjectInstanceBakingUtils.cpp
@@ -0,0 +1,787 @@
+﻿// Copyright Epic Games, Inc. All Rights Reserved.
+
+#include "CustomizableObjectInstanceBakingUtils.h"
+#include "MuCOE/CustomizableObjectEditor.h"
+
+#include "Misc/MessageDialog.h"
+#include "FileHelpers.h"
+#include "ObjectTools.h"
+#include "UnrealBakeHelpers.h"
+#include "Engine/SkinnedAssetCommon.h"
+#include "Materials/MaterialInstanceConstant.h"
+#include "Materials/MaterialInstanceDynamic.h"
+#include "Materials/Material.h"
+
+#include "MuCO/CustomizableObject.h"
+#include "MuCO/CustomizableObjectMipDataProvider.h"
+#include "MuT/UnrealPixelFormatOverride.h"
+#include "Rendering/SkeletalMeshModel.h"
+
+#define LOCTEXT_NAMESPACE "CustomizableObjectEditor" 
+
+
+// /**
+//  * Simple wrapper to be able to invoke the generation of a popup or log message depending on the execution context in which this code is being ran
+//  * @param InMessage The message to display
+//  * @param InTitle The title to be used for the popup or the log generated
+//  */
+// void ShowErrorNotification(const FText& InMessage,  const FText* InTitle = nullptr)
+// {
+// 	if (!FApp::IsUnattended())
+// 	{
+// 		FMessageDialog::Open(EAppMsgType::Ok, InMessage, *InTitle);
+// 	}
+// 	else
+// 	{
+// 		if (InTitle)
+// 		{
+// 			UE_LOG(LogMutable, Error, TEXT("%s - %s"), *InTitle->ToString(), *InMessage.ToString());
+// 		}
+// 		else
+// 		{
+// 			UE_LOG(LogMutable, Error, TEXT("%s"), *InMessage.ToString());
+// 		}
+// 	}
+// }
+//
+// /**
+//  * Utility functions for the baking operation.
+//  */
+//
+// /**
+//  * Validates the filename chosen for the baking data
+//  * @param FileName The filename chosen by the user
+//  * @return True if validation was successful, false otherwise
+//  */
+// bool ValidateProvidedFileName(const FString& FileName)
+// {
+// 	if (FileName.IsEmpty())
+// 	{
+// 		UE_LOG(LogMutable, Error, TEXT("Invalid baking configuration : FileName string is empty.."));
+// 		return false;
+// 	}
+//
+// 	// Check for invalid characters in the name of the object to be serialized
+// 	{
+// 		TCHAR InvalidCharacter = '0';
+// 		{
+// 			FString InvalidCharacters = FPaths::GetInvalidFileSystemChars();
+// 			for (int32 i = 0; i < InvalidCharacters.Len(); ++i)
+// 			{
+// 				TCHAR Char = InvalidCharacters[i];
+// 				FString SearchedChar = FString::Chr(Char);
+// 				if (FileName.Contains(SearchedChar))
+// 				{
+// 					InvalidCharacter = InvalidCharacters[i];
+// 					break;
+// 				}
+// 			}
+// 		}
+//
+// 		if (InvalidCharacter != '0')
+// 		{
+// 			const FText InvalidCharacterText = FText::FromString(FString::Chr(InvalidCharacter));
+// 			const FText ErrorText = FText::Format(LOCTEXT("FCustomizableObjectEditorViewportClient_BakeInstance_InvalidCharacter", "The selected contains an invalid character ({0})."), InvalidCharacterText);
+//
+// 			ShowErrorNotification(ErrorText);
+// 		
+// 			return false;
+// 		}
+// 	}
+//
+// 	return true;
+// }
+//
+//
+// /**
+//  * Validates the AssetPath chosen for the baking data
+//  * @param FileName The filename chosen by the user
+//  * @param AssetPath The AssetPath chosen by the user
+//  * @param InstanceCO The CustomizableObject from the provided COI
+//  * @return True if validation was successful, false otherwise
+//  */
+// bool ValidateProvidedAssetPath(const FString& FileName, const FString& AssetPath, const UCustomizableObject* InstanceCO)
+// {
+// 	if (AssetPath.IsEmpty())
+// 	{
+// 		UE_LOG(LogMutable, Error, TEXT("The AssetPath can not be empty!"));
+// 		return false;
+// 	}
+//
+// 	// Ensure we are not overriding the parent CO
+// 	const FString FullAssetPath = AssetPath + FString("/") + FileName + FString(".") + FileName;		// Full asset path to the new asset we want to create
+// 	const bool bWouldOverrideParentCO = InstanceCO->GetPathName() == FullAssetPath;
+// 	if (bWouldOverrideParentCO)
+// 	{
+// 		const FText ErrorText = LOCTEXT("FCustomizableObjectEditorViewportClient_BakeInstance_OverwriteCO", "The selected path would overwrite the instance's parent Customizable Object.");
+//
+// 		ShowErrorNotification(ErrorText);
+// 		
+// 		return false;
+// 	}
+//
+// 	return true;
+// }
+//
+//
+// /**
+//  * Outputs a string that we know it is unique.
+//  * @param InResource The resource we are working with
+//  * @param ResourceName The name of the resource we have provided. This should have the name of the current resource and will have the unique name for the resource once the method exits
+//  * @param InCachedResources Collection with all the already processed resources.
+//  * @param InCachedResourceNames Collection with all the already processed resources name's
+//  * @return True if the generation of the unique resource name was successful, false otherwise.
+//  */
+// bool GetUniqueResourceName(const UObject* InResource, FString& ResourceName, TArray<UObject*>& InCachedResources, const TArray<FString>& InCachedResourceNames)
+// {
+// 	check(InResource);
+//
+// 	int32 FindResult = InCachedResourceNames.Find(ResourceName);
+// 	if (FindResult != INDEX_NONE)
+// 	{
+// 		if (InResource == InCachedResources[FindResult])
+// 		{
+// 			return false;
+// 		}
+//
+// 		uint32 Count = 0;
+// 		while (FindResult != INDEX_NONE)
+// 		{
+// 			FindResult = InCachedResourceNames.Find(ResourceName + "_" + FString::FromInt(Count));
+// 			Count++;
+// 		}
+//
+// 		ResourceName += "_" + FString::FromInt(--Count);
+// 	}
+//
+// 	return true;
+// }
+//
+//
+// /**
+//  * Ensures the resource we want to save is ready to be saved. It handles closing it's editor and warning the user about possible overriding of resources.
+//  * @param InAssetSavePath The directory path where to save the baked object
+//  * @param InObjName The name of the object to be baked
+//  * @param bUsedGrantedOverridingRights Control flag that determines if the user has given or not permission to override resources already in disk
+//  * @return True if the operation was successful, false otherwise
+//  */
+// bool ManageBakingAction(const FString& InAssetSavePath, const FString& InObjName, bool& bUsedGrantedOverridingRights)
+// {
+// 	FString PackagePath = InAssetSavePath + "/" + InObjName;
+// 	UPackage* ExistingPackage = FindPackage(NULL, *PackagePath);
+//
+// 	if (!ExistingPackage)
+// 	{
+// 		FString PackageFilePath = PackagePath + "." + InObjName;
+//
+// 		FString PackageFileName;
+// 		if (FPackageName::DoesPackageExist(PackageFilePath, &PackageFileName))
+// 		{
+// 			ExistingPackage = LoadPackage(nullptr, *PackageFileName, LOAD_EditorOnly);
+// 		}
+// 		else
+// 		{
+// 			// if package does not exists
+// 			bUsedGrantedOverridingRights = false;
+// 			return true;
+// 		}
+// 	}
+//
+// 	if (ExistingPackage)
+// 	{
+// 		// Checking if the asset is open in an editor
+// 		TArray<IAssetEditorInstance*> ObjectEditors = GEditor->GetEditorSubsystem<UAssetEditorSubsystem>()->FindEditorsForAssetAndSubObjects(ExistingPackage);
+// 		if (ObjectEditors.Num())
+// 		{
+// 			for (IAssetEditorInstance* ObjectEditorInstance : ObjectEditors)
+// 			{
+// 				// Close the editors that contains this asset
+// 				if (!ObjectEditorInstance->CloseWindow(EAssetEditorCloseReason::AssetEditorHostClosed))
+// 				{
+// 					const FText Caption = LOCTEXT("OpenExisitngFile", "Open File");
+// 					const FText Message = FText::Format(LOCTEXT("CantCloseAsset", "This Obejct \"{0}\" is open in an editor and can't be closed automatically. Please close the editor and try to bake it again"), FText::FromString(InObjName));
+//
+// 					ShowErrorNotification(Message, &Caption);
+// 					
+// 					return false;
+// 				}
+// 			}
+// 		}
+//
+// 		if (!bUsedGrantedOverridingRights)
+// 		{
+// 			const FText Caption = LOCTEXT("Already existing baked files", "Already existing baked files");
+// 			const FText Message = FText::Format(LOCTEXT("OverwriteBakedInstance", "Instance baked files already exist in selected destination \"{0}\", this action will overwrite them."), FText::AsCultureInvariant(InAssetSavePath));
+//
+// 			// We need to guard this case since it may crash the editor if we are running an unattended commandlet and we try to generate this dialog
+// 			if (!FApp::IsUnattended())
+// 			{
+// 				if (FMessageDialog::Open(EAppMsgType::OkCancel, Message, Caption) == EAppReturnType::Cancel)
+// 				{
+// 					return false;
+// 				}
+// 			}
+// 			else
+// 			{
+// 				UE_LOG(LogMutable, Error, TEXT("%s - %s"), *Caption.ToString(), *Message.ToString());
+// 			}
+//
+//
+// 			bUsedGrantedOverridingRights = true;
+// 		}
+//
+// 		UObject* ExistingObject = StaticFindObject(UObject::StaticClass(), ExistingPackage, *InObjName);
+// 		if (ExistingObject)
+// 		{
+// 			ExistingPackage->FullyLoad();
+//
+// 			TArray<UObject*> ObjectsToDelete;
+// 			ObjectsToDelete.Add(ExistingObject);
+//
+// 			// Delete objects in the package with the same name as the one we want to create
+// 			const uint32 NumObjectsDeleted = ObjectTools::ForceDeleteObjects(ObjectsToDelete, false);
+//
+// 			return NumObjectsDeleted == ObjectsToDelete.Num();
+// 		}
+// 	}
+//
+// 	return true;
+// }
+
+
+namespace PreBakeSystemSettings
+{
+	bool bIsProgressiveMipStreamingEnabled = false;
+	bool bIsOnlyGenerateRequestedLODsEnabled = false;
+}
+
+
+void PrepareForBaking()
+{
+	// Implementation of the bake operation
+	UCustomizableObjectSystem* System =  UCustomizableObjectSystem::GetInstanceChecked();
+	
+	// The instance in the editor viewport does not have high quality mips in the platform data because streaming is enabled.
+	// Disable streaming and retry with a newly generated temp instance.
+	PreBakeSystemSettings::bIsProgressiveMipStreamingEnabled = System->IsProgressiveMipStreamingEnabled();
+	System->SetProgressiveMipStreamingEnabled(false);
+	// Disable requested LOD generation as it will prevent the new instance from having all the LODs
+	PreBakeSystemSettings::bIsOnlyGenerateRequestedLODsEnabled = System->IsOnlyGenerateRequestedLODsEnabled();
+	System->SetOnlyGenerateRequestedLODsEnabled(false);
+	// Force high quality texture compression for this instance
+	PrepareUnrealCompression();
+	System->SetImagePixelFormatOverride(UnrealPixelFormatFunc);
+}
+
+
+void RestoreCustomizableObjectSettings(const FUpdateContext& Result)
+{
+	// Reenable Mutable texture streaming and requested LOD generation as they had been disabled to bake the textures
+	UCustomizableObjectSystem* System =  UCustomizableObjectSystem::GetInstanceChecked();
+	System->SetProgressiveMipStreamingEnabled(PreBakeSystemSettings::bIsProgressiveMipStreamingEnabled);
+	System->SetOnlyGenerateRequestedLODsEnabled(PreBakeSystemSettings::bIsOnlyGenerateRequestedLODsEnabled);
+	System->SetImagePixelFormatOverride(nullptr);
+}
+
+
+void UpdateInstanceForBaking(UCustomizableObjectInstance& InInstance, FInstanceUpdateNativeDelegate& InInstanceUpdateDelegate)
+{
+	// Prepare the customizable object system for baking
+	PrepareForBaking();
+	
+	// Ensure we clear the changes in the COSystem after performing the update so later updates do not get affected
+	InInstanceUpdateDelegate.AddStatic(&RestoreCustomizableObjectSettings);
+	
+	// Schedule the update
+	InInstance.UpdateSkeletalMeshAsyncResult(InInstanceUpdateDelegate,true,true);
+}
+
+/// TODO: it will be enabled in a future CL. Moved back to CustomizableObjectEditorViewportClient to ease review process
+/*
+void BakeCustomizableObjectInstance(
+	UCustomizableObjectInstance& InInstance,
+	const FString& FileName,
+	const FString& AssetPath,
+	const bool bExportAllResources,
+	const bool bGenerateConstantMaterialInstances, 
+	const bool bIsAssetRegistryLoaded)
+{
+	// Ensure that the state of the COI provided is valid --------------------------------------------------------------------------------------------
+	
+	if (!bIsAssetRegistryLoaded)
+	{
+		FCustomizableObjectEditorLogger::CreateLog(
+			LOCTEXT("CustomizableObjectCompileTryLater","Please wait until asset registry loads all assets"))
+		.Category(ELoggerCategory::COInstanceBaking)
+		.CustomNotification()
+		.Notification(true)
+		.Log();
+
+		return;
+	}
+	
+	if (!ValidateProvidedFileName(FileName))
+	{
+		UE_LOG(LogMutable, Error, TEXT("The FileName for the instance baking is not valid."));
+		return;
+	}
+	
+	UCustomizableObject* InstanceCO = InInstance.GetCustomizableObject();
+	check (InstanceCO);
+
+	if (!ValidateProvidedAssetPath(FileName,AssetPath,InstanceCO))
+	{
+		UE_LOG(LogMutable, Error, TEXT("The AssetPath for the instance baking is not valid."));
+		return;
+	}
+	
+	// Ensure the CO of the COI is accessible 
+	if (!InstanceCO || InstanceCO->IsLocked())
+	{
+		FCustomizableObjectEditorLogger::CreateLog(
+		LOCTEXT("CustomizableObjectCompilingTryLater", "Please wait until the Customizable Object is compiled"))
+		.Category(ELoggerCategory::COInstanceBaking)
+		.CustomNotification()
+		.Notification(true)
+		.Log();
+
+		return;
+	}
+	
+	// Exit early if the provided instance does not have a skeletal mesh
+	{
+		bool bHasSkeletalMesh = false;
+		
+		for (int32 ComponentIndex = 0; ComponentIndex < InInstance.SkeletalMeshes.Num(); ++ComponentIndex)
+		{
+			if (InInstance.SkeletalMeshes.IsValidIndex(ComponentIndex) && InInstance.SkeletalMeshes[ComponentIndex])
+			{
+				bHasSkeletalMesh = true;
+			}
+		}
+		
+		if (!bHasSkeletalMesh)
+		{
+			UE_LOG(LogMutable, Error, TEXT("The provided Customizable Object Instance has no skeletal meshes."));
+			return;
+		}
+	}
+
+	// COI Validation completed : Proceed with the baking operation ----------------------------------------------------------------------------------
+	
+	// Notify of better configuration -> Continue operation normally
+	if (InstanceCO->CompileOptions.TextureCompression != ECustomizableObjectTextureCompression::HighQuality)
+	{
+		FCustomizableObjectEditorLogger::CreateLog(
+		LOCTEXT("CustomizableObjectBakeLowQuality", "The Customizable Object wasn't compiled with high quality textures. For the best baking results, change the Texture Compression setting and recompile it."))
+		.Category(ELoggerCategory::COInstanceBaking)
+		.CustomNotification()
+		.Notification(true)
+		.Log();
+	}
+	
+	TArray<UPackage*> PackagesToSave;
+	
+	// Set the overriding flag to false wo we ask the user at least once about if he is willing to override old baked data
+	bool bUsedGrantedOverridingRights = false;
+
+	const int32 NumComponents = InInstance.SkeletalMeshes.Num();
+	for (int32 ComponentIndex = 0; ComponentIndex < NumComponents; ++ComponentIndex)
+	{
+		USkeletalMesh* Mesh = InInstance.SkeletalMeshes.IsValidIndex(ComponentIndex) && InInstance.SkeletalMeshes[ComponentIndex] ?
+			Cast<USkeletalMesh>(InInstance.SkeletalMeshes[ComponentIndex]) : nullptr;
+		
+		if (!Mesh)
+		{
+			continue;
+		}
+
+		FString ObjectName = FileName;
+		if (NumComponents > 1)
+		{
+			ObjectName = FileName + "_Component_" + FString::FromInt(ComponentIndex);
+		}
+		else
+		{
+			ObjectName = FileName;
+		}
+
+		TMap<UObject*, UObject*> ReplacementMap;
+		TArray<FString> ArrayCachedElement;
+		TArray<UObject*> ArrayCachedObject;
+
+		if (bExportAllResources)
+		{
+			UMaterialInstance* Inst;
+			UMaterial* Material;
+			UTexture* Texture;
+			FString MaterialName;
+			FString ResourceName;
+			FString PackageName;
+			UObject* DuplicatedObject;
+			TArray<TMap<int, UTexture*>> TextureReplacementMaps;
+
+			// Duplicate Mutable generated textures
+			for (int32 m = 0; m < Mesh->GetMaterials().Num(); ++m)
+			{
+				UMaterialInterface* Interface = Mesh->GetMaterials()[m].MaterialInterface;
+				Material = Interface->GetMaterial();
+				MaterialName = Material ? Material->GetName() : "Material";
+				Inst = Cast<UMaterialInstance>(Mesh->GetMaterials()[m].MaterialInterface);
+
+				TMap<int, UTexture*> ReplacementTextures;
+				TextureReplacementMaps.Add(ReplacementTextures);
+
+				// The material will only have Mutable generated textures if it's actually a UMaterialInstance
+				if (Material != nullptr && Inst != nullptr)
+				{
+					TArray<FName> ParameterNames = FUnrealBakeHelpers::GetTextureParameterNames(Material);
+
+					for (int32 i = 0; i < ParameterNames.Num(); i++)
+					{
+						if (Inst->GetTextureParameterValue(ParameterNames[i], Texture))
+						{
+							UTexture2D* SrcTex = Cast<UTexture2D>(Texture);
+							if (!SrcTex) continue;
+
+							bool bIsMutableTexture = false;
+
+							for (UAssetUserData* UserData : *SrcTex->GetAssetUserDataArray())
+							{
+								UTextureMipDataProviderFactory* CustomMipDataProviderFactory = Cast<UMutableTextureMipDataProviderFactory>(UserData);
+								if (CustomMipDataProviderFactory)
+								{
+									bIsMutableTexture = true;
+								}
+							}
+
+							if ((SrcTex->GetPlatformData() != nullptr) &&
+								(SrcTex->GetPlatformData()->Mips.Num() > 0) &&
+								bIsMutableTexture)
+							{
+								FString ParameterSanitized = ParameterNames[i].GetPlainNameString();
+								RemoveRestrictedChars(ParameterSanitized);
+								ResourceName = ObjectName + "_" + MaterialName + "_" + ParameterSanitized;
+
+								if (!GetUniqueResourceName(SrcTex, ResourceName, ArrayCachedObject, ArrayCachedElement))
+								{
+									continue;
+								}
+
+								if (!ManageBakingAction(AssetPath, ResourceName,bUsedGrantedOverridingRights))
+								{
+									return;
+								}
+
+								// Recover original name of the texture parameter value, now substituted by the generated Mutable texture
+								UTexture* OriginalTexture = nullptr;
+								UMaterialInstanceDynamic* InstDynamic = Cast<UMaterialInstanceDynamic>(Mesh->GetMaterials()[m].MaterialInterface);
+								if (InstDynamic != nullptr)
+								{
+									InstDynamic->Parent->GetTextureParameterValue(FName(*ParameterNames[i].GetPlainNameString()), OriginalTexture);
+								}
+								else
+								{
+									UMaterialInstanceConstant* InstConstant = Cast<UMaterialInstanceConstant>(Mesh->GetMaterials()[m].MaterialInterface);
+									if (InstConstant != nullptr)
+									{
+										InstConstant->Parent->GetTextureParameterValue(FName(*ParameterNames[i].GetPlainNameString()), OriginalTexture);
+									}
+								}
+
+								PackageName = AssetPath + FString("/") + ResourceName;
+								TMap<UObject*, UObject*> FakeReplacementMap;
+								UTexture2D* DupTex = FUnrealBakeHelpers::BakeHelper_CreateAssetTexture(SrcTex, ResourceName, PackageName, OriginalTexture, true, FakeReplacementMap, bUsedGrantedOverridingRights);
+								ArrayCachedElement.Add(ResourceName);
+								ArrayCachedObject.Add(DupTex);
+								PackagesToSave.Add(DupTex->GetPackage());
+
+								if (OriginalTexture != nullptr)
+								{
+									TextureReplacementMaps[m].Add(i, DupTex);
+								}
+							}
+						}
+					}
+				}
+			}
+
+			// Duplicate non-Mutable material textures
+			for (int32 m = 0; m < Mesh->GetMaterials().Num(); ++m)
+			{
+				UMaterialInterface* Interface = Mesh->GetMaterials()[m].MaterialInterface;
+				Material = Interface->GetMaterial();
+				MaterialName = Material ? Material->GetName() : "Material";
+
+				if (Material != nullptr)
+				{
+					TArray<FName> ParameterNames = FUnrealBakeHelpers::GetTextureParameterNames(Material);
+
+					for (int32 i = 0; i < ParameterNames.Num(); i++)
+					{
+						TArray<FMaterialParameterInfo> InfoArray;
+						TArray<FGuid> GuidArray;
+						Material->GetAllTextureParameterInfo(InfoArray, GuidArray);
+						
+						if (Material->GetTextureParameterValue(InfoArray[i], Texture))
+						{
+							FString ParameterSanitized = ParameterNames[i].GetPlainNameString();
+							RemoveRestrictedChars(ParameterSanitized);
+							ResourceName = ObjectName + "_" + MaterialName + "_" + ParameterSanitized;
+
+							if (ArrayCachedElement.Find(ResourceName) == INDEX_NONE)
+							{
+								if (!ManageBakingAction(AssetPath, ResourceName,bUsedGrantedOverridingRights))
+								{
+									return;
+								}
+
+								PackageName = AssetPath + FString("/") + ResourceName;
+								TMap<UObject*, UObject*> FakeReplacementMap;
+								DuplicatedObject = FUnrealBakeHelpers::BakeHelper_DuplicateAsset(Texture, ResourceName, PackageName, true, FakeReplacementMap, bUsedGrantedOverridingRights, false);
+								ArrayCachedElement.Add(ResourceName);
+								ArrayCachedObject.Add(DuplicatedObject);
+								PackagesToSave.Add(DuplicatedObject->GetPackage());
+
+								UTexture* DupTexture = Cast<UTexture>(DuplicatedObject);
+								TextureReplacementMaps[m].Add(i, DupTexture);
+							}
+						}
+					}
+				}
+			}
+
+
+			// Duplicate the materials used by each material instance so that the replacement map has proper information 
+			// when duplicating the material instances
+			for (int32 m = 0; m < Mesh->GetMaterials().Num(); ++m)
+			{
+				UMaterialInterface* Interface = Mesh->GetMaterials()[m].MaterialInterface;
+				Material = Interface ? Interface->GetMaterial() : nullptr;
+
+				if (Material)
+				{
+					ResourceName = ObjectName + "_Material_" + Material->GetName();
+
+					if (!GetUniqueResourceName(Material, ResourceName, ArrayCachedObject, ArrayCachedElement))
+					{
+						continue;
+					}
+
+					if (!ManageBakingAction(AssetPath, ResourceName,bUsedGrantedOverridingRights))
+					{
+						return;
+					}
+
+					PackageName = AssetPath + FString("/") + ResourceName;
+					TMap<UObject*, UObject*> FakeReplacementMap;
+					DuplicatedObject = FUnrealBakeHelpers::BakeHelper_DuplicateAsset(Material, ResourceName, PackageName, 
+						false, FakeReplacementMap, bUsedGrantedOverridingRights, bGenerateConstantMaterialInstances);
+					ArrayCachedElement.Add(ResourceName);
+					ArrayCachedObject.Add(DuplicatedObject);
+					ReplacementMap.Add(Interface, DuplicatedObject);
+					PackagesToSave.Add(DuplicatedObject->GetPackage());
+
+					FUnrealBakeHelpers::CopyAllMaterialParameters(DuplicatedObject, Interface, TextureReplacementMaps[m]);
+				}
+			}
+		}
+		else
+		{
+			// Duplicate the material instances
+			for (int32 MaterialIndex = 0; MaterialIndex < Mesh->GetMaterials().Num(); ++MaterialIndex)
+			{
+				UMaterialInterface* Interface = Mesh->GetMaterials()[MaterialIndex].MaterialInterface;
+				UMaterial* ParentMaterial = Interface->GetMaterial();
+				FString MaterialName = ParentMaterial ? ParentMaterial->GetName() : "Material";
+
+				// Material
+				FString MatObjName = ObjectName + "_" + MaterialName;
+
+				if (!GetUniqueResourceName(Interface, MatObjName, ArrayCachedObject, ArrayCachedElement))
+				{
+					continue;
+				}
+
+				if (!ManageBakingAction(AssetPath, MatObjName,bUsedGrantedOverridingRights))
+				{
+					return;
+				}
+
+				FString MatPkgName = AssetPath + FString("/") + MatObjName;
+				UObject* DupMat = FUnrealBakeHelpers::BakeHelper_DuplicateAsset(Interface, MatObjName, 
+					MatPkgName, false, ReplacementMap, bUsedGrantedOverridingRights, bGenerateConstantMaterialInstances);
+				ArrayCachedObject.Add(DupMat);
+				ArrayCachedElement.Add(MatObjName);
+				PackagesToSave.Add(DupMat->GetPackage());
+
+				UMaterialInstance* MaterialInstance = Cast<UMaterialInstance>(Interface);
+
+				// Only need to duplicate the generate textures if the original material is a dynamic instance
+				// If the material has Mutable textures, then it will be a dynamic material instance for sure
+				if (MaterialInstance)
+				{
+					// Duplicate generated textures
+					UMaterialInstanceDynamic* InstDynamic = Cast<UMaterialInstanceDynamic>(DupMat);
+					UMaterialInstanceConstant* InstConstant = Cast<UMaterialInstanceConstant>(DupMat);
+
+					if (InstDynamic || InstConstant)
+					{
+						for (int32 TextureIndex = 0; TextureIndex < MaterialInstance->TextureParameterValues.Num(); ++TextureIndex)
+						{
+							if (MaterialInstance->TextureParameterValues[TextureIndex].ParameterValue)
+							{
+								if (MaterialInstance->TextureParameterValues[TextureIndex].ParameterValue->HasAnyFlags(RF_Transient))
+								{
+									UTexture2D* SrcTex = Cast<UTexture2D>(MaterialInstance->TextureParameterValues[TextureIndex].ParameterValue);
+
+									if (SrcTex)
+									{
+										FString ParameterSanitized = MaterialInstance->TextureParameterValues[TextureIndex].ParameterInfo.Name.ToString();
+										RemoveRestrictedChars(ParameterSanitized);
+
+										FString TexObjName = ObjectName + "_" + MaterialName + "_" + ParameterSanitized;
+
+										if (!GetUniqueResourceName(SrcTex, TexObjName, ArrayCachedObject, ArrayCachedElement))
+										{
+											UTexture* PrevTexture = Cast<UTexture>(ArrayCachedObject[ArrayCachedElement.Find(TexObjName)]);
+
+											if (InstDynamic)
+											{
+												InstDynamic->SetTextureParameterValue(MaterialInstance->TextureParameterValues[TextureIndex].ParameterInfo.Name, PrevTexture);
+											}
+											else if (InstConstant)
+											{
+												InstConstant->SetTextureParameterValueEditorOnly(MaterialInstance->TextureParameterValues[TextureIndex].ParameterInfo.Name, PrevTexture);
+											}
+											
+											continue;
+										}
+
+										if (!ManageBakingAction(AssetPath, TexObjName,bUsedGrantedOverridingRights))
+										{
+											return;
+										}
+
+										FString TexPkgName = AssetPath + FString("/") + TexObjName;
+										TMap<UObject*, UObject*> FakeReplacementMap;
+										UTexture2D* DupTex = FUnrealBakeHelpers::BakeHelper_CreateAssetTexture(SrcTex, TexObjName, TexPkgName, nullptr, false, FakeReplacementMap, bUsedGrantedOverridingRights);
+										ArrayCachedObject.Add(DupTex);
+										ArrayCachedElement.Add(TexObjName);
+										PackagesToSave.Add(DupTex->GetPackage());
+
+										if (InstDynamic)
+										{
+											InstDynamic->SetTextureParameterValue(MaterialInstance->TextureParameterValues[TextureIndex].ParameterInfo.Name, DupTex);
+										}
+										else if(InstConstant)
+										{
+											InstConstant->SetTextureParameterValueEditorOnly(MaterialInstance->TextureParameterValues[TextureIndex].ParameterInfo.Name, DupTex);
+										}
+									}
+									else
+									{
+										UE_LOG(LogMutable, Error, TEXT("A Mutable texture that is not a Texture2D has been found while baking a CustomizableObjectInstance."));
+									}
+								}
+								else
+								{
+									// If it's not transient it's not a mutable texture, it's a pass-through texture
+									// Just set the original texture
+									if (InstDynamic)
+									{
+										InstDynamic->SetTextureParameterValue(MaterialInstance->TextureParameterValues[TextureIndex].ParameterInfo.Name, MaterialInstance->TextureParameterValues[TextureIndex].ParameterValue);
+									}
+									else if (InstConstant)
+									{
+										InstConstant->SetTextureParameterValueEditorOnly(MaterialInstance->TextureParameterValues[TextureIndex].ParameterInfo.Name, MaterialInstance->TextureParameterValues[TextureIndex].ParameterValue);
+									}
+								}
+							}
+						}
+					}
+				}
+			}
+		}
+		
+		// Skeletal Mesh's Skeleton
+		if (Mesh->GetSkeleton())
+		{
+			const bool bTransient = Mesh->GetSkeleton()->GetPackage() == GetTransientPackage();
+
+			// Don't duplicate if not transient or export all assets.
+			if (bTransient || bExportAllResources)
+			{
+				FString SkeletonName = ObjectName + "_Skeleton";
+				if (!ManageBakingAction(AssetPath, SkeletonName,bUsedGrantedOverridingRights))
+				{
+					return;
+				}
+
+				FString SkeletonPkgName = AssetPath + FString("/") + SkeletonName;
+				UObject* DuplicatedSkeleton = FUnrealBakeHelpers::BakeHelper_DuplicateAsset(Mesh->GetSkeleton(), SkeletonName, 
+					SkeletonPkgName, false, ReplacementMap, bUsedGrantedOverridingRights, false);
+
+				ArrayCachedObject.Add(DuplicatedSkeleton);
+				PackagesToSave.Add(DuplicatedSkeleton->GetPackage());
+				ReplacementMap.Add(Mesh->GetSkeleton(), DuplicatedSkeleton);
+			}
+		}
+
+		// Skeletal Mesh
+		if (!ManageBakingAction(AssetPath, ObjectName,bUsedGrantedOverridingRights))
+		{
+			return;
+		}
+
+		FString PkgName = AssetPath + FString("/") + ObjectName;
+		UObject* DupObject = FUnrealBakeHelpers::BakeHelper_DuplicateAsset(Mesh, ObjectName, PkgName, 
+			false, ReplacementMap, bUsedGrantedOverridingRights, false);
+		ArrayCachedObject.Add(DupObject);
+		PackagesToSave.Add(DupObject->GetPackage());
+
+		Mesh->Build();
+
+		USkeletalMesh* SkeletalMesh = Cast<USkeletalMesh>(DupObject);
+		if (SkeletalMesh)
+		{
+			SkeletalMesh->GetLODInfoArray() = Mesh->GetLODInfoArray();
+
+			SkeletalMesh->GetImportedModel()->SkeletalMeshModelGUID = FGuid::NewGuid();
+
+			// Duplicate AssetUserData
+			{
+				const TArray<UAssetUserData*>* AssetUserDataArray = Mesh->GetAssetUserDataArray();
+				for (const UAssetUserData* AssetUserData : *AssetUserDataArray)
+				{
+					if (AssetUserData)
+					{
+						// Duplicate to change ownership
+						UAssetUserData* NewAssetUserData = Cast<UAssetUserData>(StaticDuplicateObject(AssetUserData, SkeletalMesh));
+						SkeletalMesh->AddAssetUserData(NewAssetUserData);
+					}
+				}
+			}
+
+			// Generate render data
+			SkeletalMesh->Build();
+		}
+
+		// Remove duplicated UObjects from Root (previously added to avoid objects from beeing GC in the middle of the bake process)
+		for (UObject* Obj : ArrayCachedObject)
+		{
+			Obj->RemoveFromRoot();
+		}
+	}
+
+	// Save the packages generated during the baking operation  --------------------------------------------------------------------------------------
+	
+	// Complete the baking by saving the packages we have cached during the baking operation
+	if (PackagesToSave.Num())
+	{
+		FEditorFileUtils::PromptForCheckoutAndSave(PackagesToSave, false, !FApp::IsUnattended());
+	}
+}
+*/
+
+#undef LOCTEXT_NAMESPACE 
diff --git a/Engine/Plugins/Experimental/Mutable/Source/CustomizableObjectEditor/Public/MuCOE/CustomizableObjectInstanceBakingUtils.h b/Engine/Plugins/Experimental/Mutable/Source/CustomizableObjectEditor/Public/MuCOE/CustomizableObjectInstanceBakingUtils.h
new file mode 100644
--- /dev/null
+++ b/Engine/Plugins/Experimental/Mutable/Source/CustomizableObjectEditor/Public/MuCOE/CustomizableObjectInstanceBakingUtils.h
@@ -0,0 +1,30 @@
+﻿// Copyright Epic Games, Inc. All Rights Reserved.
+#pragma once
+
+#include "MuCO/CustomizableObjectInstance.h"
+
+
+/**
+ * Schedules the async update of the target instance. It will clear previous update data.
+ * This method will also take care of setting and later resetting the state of the Customizable Object System.
+ * @param InInstance The instance we want to update so we can later bake it's resources.
+ * @param InInstanceUpdateDelegate Delegate to be called after the instance gets updated.
+ */
+CUSTOMIZABLEOBJECTEDITOR_API
+void UpdateInstanceForBaking(UCustomizableObjectInstance& InInstance, FInstanceUpdateNativeDelegate& InInstanceUpdateDelegate);
+
+// TODO: enable after review -> Required for MTBL-1674
+/**
+ * Serializes onto disk the resources used by the targeted Customizable Object Instance. The operation can be configured in the FInstanceBakingSettings settings object.
+ * @param InInstance The mutable COI instance whose resources we want to bake onto disk
+ * @param FileName The base name for the all the files generated during the baking process
+ * @param AssetPath The path where to save the assets
+ * @param bExportAllResources Determines if we want a full or partial export
+ * @param bGenerateConstantMaterialInstances Determines if we want to generate constant material instances or not
+ * @param bIsAssetRegistryLoaded Status of the Asset Registry
+ */
+/*
+CUSTOMIZABLEOBJECTEDITOR_API
+void BakeCustomizableObjectInstance (UCustomizableObjectInstance& InInstance, const FString& FileName, const FString& AssetPath,
+	const bool bExportAllResources, const bool bGenerateConstantMaterialInstances, const bool bIsAssetRegistryLoaded);
+*/
\ No newline at end of file
