[Mutable] Remove unused COI private flag CreatingSkeletalMesh.
Remove old morph and cloth hotfix that's no longer required.
Remove CVarMutableReuseMaterialInstances.
Reinitialize the pose when we set a new mesh.

#rnx

diff --git a/Engine/Plugins/Experimental/Mutable/Source/CustomizableObject/Classes/MuCO/CustomizableObjectInstanceUsage.h b/Engine/Plugins/Experimental/Mutable/Source/CustomizableObject/Classes/MuCO/CustomizableObjectInstanceUsage.h
--- a/Engine/Plugins/Experimental/Mutable/Source/CustomizableObject/Classes/MuCO/CustomizableObjectInstanceUsage.h
+++ b/Engine/Plugins/Experimental/Mutable/Source/CustomizableObject/Classes/MuCO/CustomizableObjectInstanceUsage.h
@@ -56,7 +56,7 @@ public:
 	UFUNCTION(BlueprintCallable, Category = CustomizableObjectInstance)
 	void UpdateSkeletalMeshAsyncResult(FInstanceUpdateDelegate Callback, bool bIgnoreCloseDist = false, bool bForceHighPriority = false);
 
-	void SetSkeletalMesh(USkeletalMesh* SkeletalMesh, bool bReinitPose=true, bool bForceClothReset=true);
+	void SetSkeletalMesh(USkeletalMesh* SkeletalMesh);
 	void SetPhysicsAsset(class UPhysicsAsset* PhysicsAsset);
 	void UpdateDistFromComponentToPlayer(const AActor* const Pawn, bool bForceEvenIfNotBegunPlay = false);
 
diff --git a/Engine/Plugins/Experimental/Mutable/Source/CustomizableObject/Classes/MuCO/CustomizableSkeletalComponent.h b/Engine/Plugins/Experimental/Mutable/Source/CustomizableObject/Classes/MuCO/CustomizableSkeletalComponent.h
--- a/Engine/Plugins/Experimental/Mutable/Source/CustomizableObject/Classes/MuCO/CustomizableSkeletalComponent.h
+++ b/Engine/Plugins/Experimental/Mutable/Source/CustomizableObject/Classes/MuCO/CustomizableSkeletalComponent.h
@@ -45,7 +45,7 @@ public:
 	UFUNCTION(BlueprintCallable, Category = CustomizableObjectInstance)
 	void UpdateSkeletalMeshAsyncResult(FInstanceUpdateDelegate Callback, bool bIgnoreCloseDist = false, bool bForceHighPriority = false);
 
-	void SetSkeletalMesh(USkeletalMesh* SkeletalMesh, bool bReinitPose=true, bool bForceClothReset=true);
+	void SetSkeletalMesh(USkeletalMesh* SkeletalMesh);
 	void SetPhysicsAsset(class UPhysicsAsset* PhysicsAsset);
 	void UpdateDistFromComponentToPlayer(const AActor* const Pawn, bool bForceEvenIfNotBegunPlay = false);
 
diff --git a/Engine/Plugins/Experimental/Mutable/Source/CustomizableObject/Internal/MuCO/CustomizableInstancePrivateData.h b/Engine/Plugins/Experimental/Mutable/Source/CustomizableObject/Internal/MuCO/CustomizableInstancePrivateData.h
--- a/Engine/Plugins/Experimental/Mutable/Source/CustomizableObject/Internal/MuCO/CustomizableInstancePrivateData.h
+++ b/Engine/Plugins/Experimental/Mutable/Source/CustomizableObject/Internal/MuCO/CustomizableInstancePrivateData.h
@@ -37,7 +37,6 @@ enum ECOInstanceFlags
 	ECONone							= 0,  // Should not use the name None here.. it collides with other enum in global namespace
 
 	// Update process
-	CreatingSkeletalMesh			= 1 << 1,	//
 	ReuseTextures					= 1 << 3, 	// 
 	ReplacePhysicsAssets			= 1 << 4,	// Merge active PhysicsAssets and replace the base physics asset
 
diff --git a/Engine/Plugins/Experimental/Mutable/Source/CustomizableObject/Private/MuCO/CustomizableObjectInstance.cpp b/Engine/Plugins/Experimental/Mutable/Source/CustomizableObject/Private/MuCO/CustomizableObjectInstance.cpp
--- a/Engine/Plugins/Experimental/Mutable/Source/CustomizableObject/Private/MuCO/CustomizableObjectInstance.cpp
+++ b/Engine/Plugins/Experimental/Mutable/Source/CustomizableObject/Private/MuCO/CustomizableObjectInstance.cpp
@@ -1827,8 +1827,6 @@ bool UCustomizableInstancePrivate::UpdateSkeletalMesh_PostBeginUpdate0(UCustomiz
 		return true;
 	}
 
-	SetCOInstanceFlags(CreatingSkeletalMesh);
-
 	TextureReuseCache.Empty(); // Sections may have changed, so invalidate the texture reuse cache because it's indexed by section
 
 	TArray<TObjectPtr<USkeletalMesh>> OldSkeletalMeshes = SkeletalMeshes;
@@ -2125,8 +2123,6 @@ bool UCustomizableInstancePrivate::UpdateSkeletalMesh_PostBeginUpdate0(UCustomiz
 		}
 	}
 
-	ClearCOInstanceFlags(CreatingSkeletalMesh); // TODO MTBL-391: Review
-
 	if (!bSuccess)
 	{
 		InvalidateGeneratedData();
@@ -5382,11 +5378,6 @@ void UCustomizableInstancePrivate::ReuseTexture(UTexture2D* Texture, TSharedRef<
 	}
 }
 
-static bool bReuseMaterialInstances = true;
-FAutoConsoleVariableRef CVarMutableReuseMaterialInstances(
-	TEXT("Mutable.ReuseMaterialInstances"),
-	bReuseMaterialInstances,
-	TEXT("If true, allow reuse of MaterialInstances between updates."));
 
 void UCustomizableInstancePrivate::BuildMaterials(const TSharedRef<FUpdateContextPrivate>& OperationData, UCustomizableObjectInstance* Public)
 {
@@ -5402,10 +5393,7 @@ void UCustomizableInstancePrivate::BuildMaterials(const TSharedRef<FUpdateContex
 
 	// Temp copy to allow reuse of MaterialInstances
 	TArray<FGeneratedMaterial> OldGeneratedMaterials;
-	if (bReuseMaterialInstances)
-	{
-		Exchange(OldGeneratedMaterials, GeneratedMaterials);
-	}
+	Exchange(OldGeneratedMaterials, GeneratedMaterials);
 	
 	GeneratedMaterials.Reset();
 
@@ -5880,7 +5868,7 @@ void UCustomizableInstancePrivate::BuildMaterials(const TSharedRef<FUpdateContex
 
 					UMaterialInstanceDynamic* MaterialInstance = nullptr;
 					
-					if (const int32 OldMaterialIndex = OldGeneratedMaterials.Find(Material); bReuseMaterialInstances && OldMaterialIndex != INDEX_NONE)
+					if (const int32 OldMaterialIndex = OldGeneratedMaterials.Find(Material); OldMaterialIndex != INDEX_NONE)
 					{
 						const FGeneratedMaterial& OldMaterial = OldGeneratedMaterials[OldMaterialIndex];
 						MaterialInstance = Cast<UMaterialInstanceDynamic>(OldMaterial.MaterialInterface);
diff --git a/Engine/Plugins/Experimental/Mutable/Source/CustomizableObject/Private/MuCO/CustomizableObjectInstanceUsage.cpp b/Engine/Plugins/Experimental/Mutable/Source/CustomizableObject/Private/MuCO/CustomizableObjectInstanceUsage.cpp
--- a/Engine/Plugins/Experimental/Mutable/Source/CustomizableObject/Private/MuCO/CustomizableObjectInstanceUsage.cpp
+++ b/Engine/Plugins/Experimental/Mutable/Source/CustomizableObject/Private/MuCO/CustomizableObjectInstanceUsage.cpp
@@ -195,55 +195,13 @@ USkeletalMesh* UCustomizableObjectInstanceUsage::GetSkeletalMesh() const
 }
 
 
-void UCustomizableObjectInstanceUsage::SetSkeletalMesh(USkeletalMesh* SkeletalMesh, bool bReinitPose, bool bForceClothReset)
+void UCustomizableObjectInstanceUsage::SetSkeletalMesh(USkeletalMesh* SkeletalMesh)
 {
 	USkeletalMeshComponent* Parent = Cast<USkeletalMeshComponent>(GetAttachParent());
 
-	// LINUX_PLATFORM needs a more aggressive workaround so morph and cloth glitches
-	// are not visible. 
-	// TODO: Try to find a better way of setting BP Morphs and clothing and remove
-	// the workaround.  
 	if (Parent)
 	{
-#if PLATFORM_LINUX
-		const bool bDisableClothSimulation = Parent->bDisableClothSimulation;
-		if (bForceClothReset)
-		{
-			Parent->bDisableClothSimulation = true;
-		}
-#endif
-
-		if(UE_MUTABLE_GETSKINNEDASSET(Parent) == SkeletalMesh)
-		{
-			Parent->RecreateRenderState_Concurrent();
-		}
-		
-		TMap<FName, float> MorphTargetCurves = Parent->GetMorphTargetCurves();
-		Parent->SetSkeletalMesh(SkeletalMesh, bReinitPose);
-		
-		// USkeletalMeshCompoent MorphTargetCurves are reset when SetSkeletalMesh is called.
-		// Re-enable them if not bReinitPose.
-		if (!bReinitPose && MorphTargetCurves.Num() > 0)
-		{
-			for (const TPair<FName, float>& MorphTarget : MorphTargetCurves)
-			{
-				Parent->SetMorphTarget(MorphTarget.Key, MorphTarget.Value);
-			}
-
-#if PLATFORM_LINUX
-			FRenderStateRecreator RenderStateRecreator(Parent);
-			FAnimationRuntime::AppendActiveMorphTargets(SkeletalMesh, Parent->GetMorphTargetCurves(), Parent->ActiveMorphTargets, Parent->MorphTargetWeights);		
-#endif
-		}
-
-		if (bForceClothReset)
-		{
-			Parent->ForceClothNextUpdateTeleportAndReset();
-
-#if PLATFORM_LINUX
-			Parent->bDisableClothSimulation = bDisableClothSimulation;
-#endif
-		}
+		Parent->SetSkeletalMesh(SkeletalMesh, true);
 
 		const UCustomizableObjectInstance* Instance = GetCustomizableObjectInstance();
 		const UCustomizableObject* CustomizableObject = Instance ? Instance->GetCustomizableObject() : nullptr;
@@ -263,7 +221,7 @@ void UCustomizableObjectInstanceUsage::SetSkeletalMesh(USkeletalMesh* SkeletalMe
 				for (int32 Index = 0; Index < ComponentData->OverrideMaterials.Num(); ++Index)
 				{
 					Parent->SetMaterial(Index, ComponentData->OverrideMaterials[Index]);
-				}	
+				}
 			}
 		}	
 	}
@@ -348,17 +306,15 @@ void UCustomizableObjectInstanceUsage::UpdateDistFromComponentToLevelEditorCamer
 		}
 
 		USkeletalMesh* AttachedSkeletalMesh = GetAttachedSkeletalMesh();
-		int32 ComponentIndex = GetComponentIndex();
-		USkeletalMesh* GeneratedSKeletalMesh = CustomizableObjectInstance->GetSkeletalMesh(ComponentIndex);
-		if (!AttachedSkeletalMesh && !CustomizableObjectInstance->GetPrivate()->HasCOInstanceFlags(CreatingSkeletalMesh))
-		{
-			UCustomizableObject* Object = CustomizableObjectInstance->GetCustomizableObject(); 
-			USkeletalMesh* RefMesh = Object ? Object->GetRefSkeletalMesh(ComponentIndex) : nullptr;
-			SetSkeletalMesh(GeneratedSKeletalMesh ? GeneratedSKeletalMesh : RefMesh);
-		}
-		else if (GeneratedSKeletalMesh && AttachedSkeletalMesh != GeneratedSKeletalMesh)
+		const int32 ComponentIndex = GetComponentIndex();
+
+		const bool bInstanceGenerated = CustomizableObjectInstance->GetPrivate()->SkeletalMeshStatus != ESkeletalMeshStatus::NotGenerated;
+		USkeletalMesh* GeneratedSkeletalMesh = bInstanceGenerated ? CustomizableObjectInstance->GetSkeletalMesh(ComponentIndex) :
+			CustomizableObjectInstance->GetCustomizableObject()->GetRefSkeletalMesh(ComponentIndex);
+
+		if (AttachedSkeletalMesh != GeneratedSkeletalMesh)
 		{
-			SetSkeletalMesh(GeneratedSKeletalMesh);
+			SetSkeletalMesh(GeneratedSkeletalMesh);
 		}
 	}
 }
@@ -378,15 +334,15 @@ void UCustomizableObjectInstanceUsage::EditorUpdateComponent()
 		if (ParentActor)
 		{
 			USkeletalMesh* AttachedSkeletalMesh = GetAttachedSkeletalMesh();
-			int32 ComponentIndex = GetComponentIndex();
-			USkeletalMesh* GeneratedSKeletalMesh = CustomizableObjectInstance->GetSkeletalMesh(ComponentIndex);
-			if (!AttachedSkeletalMesh && !CustomizableObjectInstance->GetPrivate()->HasCOInstanceFlags(CreatingSkeletalMesh))
-			{
-				SetSkeletalMesh(GeneratedSKeletalMesh ? GeneratedSKeletalMesh : CustomizableObjectInstance->GetCustomizableObject()->GetRefSkeletalMesh(ComponentIndex));
-			}
-			else if (GeneratedSKeletalMesh && AttachedSkeletalMesh != GeneratedSKeletalMesh)
+			const int32 ComponentIndex = GetComponentIndex();
+
+			const bool bInstanceGenerated = CustomizableObjectInstance->GetPrivate()->SkeletalMeshStatus != ESkeletalMeshStatus::NotGenerated;
+			USkeletalMesh* GeneratedSkeletalMesh = bInstanceGenerated ? CustomizableObjectInstance->GetSkeletalMesh(ComponentIndex) :
+				CustomizableObjectInstance->GetCustomizableObject()->GetRefSkeletalMesh(ComponentIndex);
+
+			if (AttachedSkeletalMesh != GeneratedSkeletalMesh)
 			{
-				SetSkeletalMesh(GeneratedSKeletalMesh);
+				SetSkeletalMesh(GeneratedSkeletalMesh);
 			}
 		}
 	}
@@ -490,7 +446,7 @@ void UCustomizableObjectInstanceUsage::UpdateDistFromComponentToPlayer(const AAc
 
 		int32 ComponentIndex = GetComponentIndex();
 
-		if (ParentActor && GetAttachedSkeletalMesh() == nullptr && CustomizableObjectInstance->GetSkeletalMesh(ComponentIndex) && !CustomizableObjectInstance->GetPrivate()->HasCOInstanceFlags(CreatingSkeletalMesh))
+		if (ParentActor && GetAttachedSkeletalMesh() == nullptr && CustomizableObjectInstance->GetSkeletalMesh(ComponentIndex))
 		{
 			SetSkeletalMesh(CustomizableObjectInstance->GetSkeletalMesh(ComponentIndex));
 		}
diff --git a/Engine/Plugins/Experimental/Mutable/Source/CustomizableObject/Private/MuCO/CustomizableObjectSystem.cpp b/Engine/Plugins/Experimental/Mutable/Source/CustomizableObject/Private/MuCO/CustomizableObjectSystem.cpp
--- a/Engine/Plugins/Experimental/Mutable/Source/CustomizableObject/Private/MuCO/CustomizableObjectSystem.cpp
+++ b/Engine/Plugins/Experimental/Mutable/Source/CustomizableObject/Private/MuCO/CustomizableObjectSystem.cpp
@@ -893,8 +893,7 @@ void UpdateSkeletalMesh(const TSharedRef<FUpdateContextPrivate>& Context)
 		{
 			MUTABLE_CPUPROFILER_SCOPE(UpdateSkeletalMesh_SetSkeletalMesh);
 
-			const bool bIsCreatingSkeletalMesh = CustomizableObjectInstancePrivateData->HasCOInstanceFlags(CreatingSkeletalMesh); //TODO MTBL-391: Review
-			CustomizableObjectInstanceUsage->SetSkeletalMesh(CustomizableObjectInstancePrivateData->SkeletalMeshes[CustomizableObjectInstanceUsage->GetComponentIndex()], false, bIsCreatingSkeletalMesh);
+			CustomizableObjectInstanceUsage->SetSkeletalMesh(CustomizableObjectInstancePrivateData->SkeletalMeshes[CustomizableObjectInstanceUsage->GetComponentIndex()]);
 
 			if (CustomizableObjectInstancePrivateData->HasCOInstanceFlags(ReplacePhysicsAssets))
 			{
@@ -905,8 +904,6 @@ void UpdateSkeletalMesh(const TSharedRef<FUpdateContextPrivate>& Context)
 		}
 	}
 
-	CustomizableObjectInstancePrivateData->ClearCOInstanceFlags(CreatingSkeletalMesh);
-
 	CustomizableObjectInstance->bEditorPropertyChanged = false;
 }
 
diff --git a/Engine/Plugins/Experimental/Mutable/Source/CustomizableObject/Private/MuCO/CustomizableSkeletalComponent.cpp b/Engine/Plugins/Experimental/Mutable/Source/CustomizableObject/Private/MuCO/CustomizableSkeletalComponent.cpp
--- a/Engine/Plugins/Experimental/Mutable/Source/CustomizableObject/Private/MuCO/CustomizableSkeletalComponent.cpp
+++ b/Engine/Plugins/Experimental/Mutable/Source/CustomizableObject/Private/MuCO/CustomizableSkeletalComponent.cpp
@@ -52,11 +52,11 @@ USkeletalMesh* UCustomizableSkeletalComponent::GetSkeletalMesh() const
 }
 
 
-void UCustomizableSkeletalComponent::SetSkeletalMesh(USkeletalMesh* SkeletalMesh, bool bReinitPose, bool bForceClothReset)
+void UCustomizableSkeletalComponent::SetSkeletalMesh(USkeletalMesh* SkeletalMesh)
 {
 	if (CustomizableObjectInstanceUsage)
 	{
-		CustomizableObjectInstanceUsage->SetSkeletalMesh(SkeletalMesh, bReinitPose, bForceClothReset);
+		CustomizableObjectInstanceUsage->SetSkeletalMesh(SkeletalMesh);
 	}
 }
 
diff --git a/Engine/Plugins/Experimental/Mutable/Source/CustomizableObject/Private/MuCO/LogInformationUtil.cpp b/Engine/Plugins/Experimental/Mutable/Source/CustomizableObject/Private/MuCO/LogInformationUtil.cpp
--- a/Engine/Plugins/Experimental/Mutable/Source/CustomizableObject/Private/MuCO/LogInformationUtil.cpp
+++ b/Engine/Plugins/Experimental/Mutable/Source/CustomizableObject/Private/MuCO/LogInformationUtil.cpp
@@ -167,7 +167,6 @@ void LogInformationUtil::LogShowInstanceDataFull(const UCustomizableObjectInstan
 	LogData += MessageChunk;
 
 	MessageChunk += FString::Printf(TEXT("        bIsDowngradeLODUpdate = %d\n"), InstanceFlags & ECOInstanceFlags::PendingLODsDowngrade);
-	MessageChunk += FString::Printf(TEXT("        bIsCreatingSkeletalMesh = %d\n"), InstanceFlags & ECOInstanceFlags::CreatingSkeletalMesh);
 	LogData += MessageChunk;
 
 	for (int32 ComponentIndex = 0; ComponentIndex < CustomizableObjectInstance->GetNumComponents(); ++ComponentIndex)
diff --git a/Engine/Plugins/Experimental/Mutable/Source/CustomizableObjectEditor/Private/MuCOE/CustomizableObjectInstanceFactory.cpp b/Engine/Plugins/Experimental/Mutable/Source/CustomizableObjectEditor/Private/MuCOE/CustomizableObjectInstanceFactory.cpp
--- a/Engine/Plugins/Experimental/Mutable/Source/CustomizableObjectEditor/Private/MuCOE/CustomizableObjectInstanceFactory.cpp
+++ b/Engine/Plugins/Experimental/Mutable/Source/CustomizableObjectEditor/Private/MuCOE/CustomizableObjectInstanceFactory.cpp
@@ -58,7 +58,7 @@ void UCustomizableObjectInstanceFactory::PostSpawnActor(UObject* Asset, AActor*
 					CustomSkeletalComp->UnregisterComponent();
 					CustomSkeletalComp->CustomizableObjectInstance = COInstance;
 					CustomSkeletalComp->ComponentIndex = ComponentIndex;
-					CustomSkeletalComp->SetSkeletalMesh(SkeletalMesh, false);
+					CustomSkeletalComp->SetSkeletalMesh(SkeletalMesh);
 					CustomSkeletalComp->UpdateSkeletalMeshAsync();
 					CustomSkeletalComp->RegisterComponent();
 				}
