[Mutable] Respect Max In-game texture size when computing the LODBias and MaxTextureSize of a texture.
Fix texture size reduction not being applied on some paths and cases where it was applied twice.
Fix MaxTextureSize being ignored in some edge cases.

#rnx

diff --git a/Engine/Plugins/Experimental/Mutable/Source/CustomizableObjectEditor/Private/MuCOE/CustomizableObjectCompiler.cpp b/Engine/Plugins/Experimental/Mutable/Source/CustomizableObjectEditor/Private/MuCOE/CustomizableObjectCompiler.cpp
--- a/Engine/Plugins/Experimental/Mutable/Source/CustomizableObjectEditor/Private/MuCOE/CustomizableObjectCompiler.cpp
+++ b/Engine/Plugins/Experimental/Mutable/Source/CustomizableObjectEditor/Private/MuCOE/CustomizableObjectCompiler.cpp
@@ -1018,7 +1018,10 @@ void FCustomizableObjectCompiler::CompileInternal(UCustomizableObject* Object, c
 
 		Object->ImageProperties.Empty(GenerationContext.ImageProperties.Num());
 
-		for (const FGeneratedImageProperties& ImageProp : GenerationContext.ImageProperties)
+		TArray<FGeneratedImageProperties> ImageProperties;
+		GenerationContext.ImageProperties.GenerateValueArray(ImageProperties);
+
+		for (const FGeneratedImageProperties& ImageProp : ImageProperties)
 		{
 			Object->ImageProperties.Add({ ImageProp.TextureParameterName,
 										ImageProp.Filter,
diff --git a/Engine/Plugins/Experimental/Mutable/Source/CustomizableObjectEditor/Private/MuCOE/GenerateMutableSource/GenerateMutableSource.cpp b/Engine/Plugins/Experimental/Mutable/Source/CustomizableObjectEditor/Private/MuCOE/GenerateMutableSource/GenerateMutableSource.cpp
--- a/Engine/Plugins/Experimental/Mutable/Source/CustomizableObjectEditor/Private/MuCOE/GenerateMutableSource/GenerateMutableSource.cpp
+++ b/Engine/Plugins/Experimental/Mutable/Source/CustomizableObjectEditor/Private/MuCOE/GenerateMutableSource/GenerateMutableSource.cpp
@@ -1626,75 +1626,121 @@ void PopulateReferenceSkeletalMeshesData(FMutableGraphGenerationContext& Generat
 }
 
 
-int32 ComputeLODBias(const FMutableGraphGenerationContext& GenerationContext, const UTexture2D* ReferenceTexture, int32 MaxTextureSize,
-	const UCustomizableObjectNodeMaterial* MaterialNode, const int32 ImageIndex)
+uint32 GetBaseTextureSize(const FMutableGraphGenerationContext& GenerationContext, const UCustomizableObjectNodeMaterial* Material, uint32 ImageIndex)
+{
+	const FGeneratedImageProperties* ImageProperties = GenerationContext.ImageProperties.Find({ Material, ImageIndex });
+	return ImageProperties ? ImageProperties->TextureSize : 0;
+}
+
+
+// Find the LODBias to apply to stay within the MaxTextureSize limit of the TargetPlatform
+int32 GetPlatformLODBias(int32 TextureSize, int32 NumMips, int32 MaxPlatformSize)
+{
+	if (MaxPlatformSize > 0 && MaxPlatformSize < TextureSize)
+	{
+		const int32 MaxMipsAllowed = FMath::CeilLogTwo(MaxPlatformSize) + 1;
+		return NumMips - MaxMipsAllowed;
+	}
+
+	return 0;
+}
+
+
+uint32 ComputeLODBiasForTexture(const FMutableGraphGenerationContext& GenerationContext, const UTexture2D* Texture, const UTexture2D* ReferenceTexture, int32 BaseTextureSize)
 {
 	constexpr int32 MaxAllowedLODBias = 6;
 
+	// Force a large LODBias for debug
 	if (GenerationContext.Options.bForceLargeLODBias)
 	{
 		return FMath::Min(GenerationContext.Options.DebugBias, MaxAllowedLODBias);
 	}
-	
-	int32 LODBias = 0;
-
-	// This is not 100% correct because it makes assumptions about the final texture size
-	// that may not be correct, but if the reference texture is really representative of the average 
-	// case, then it is as good as we can do.
-	if (ReferenceTexture)
-	{
-		const UTextureLODSettings& LODSettings = GenerationContext.Options.TargetPlatform->GetTextureLODSettings();
 
-		LODBias = LODSettings.CalculateLODBias(ReferenceTexture->Source.GetSizeX(), ReferenceTexture->Source.GetSizeY(), MaxTextureSize,
-			ReferenceTexture->LODGroup, ReferenceTexture->LODBias, false, ReferenceTexture->MipGenSettings, ReferenceTexture->IsCurrentlyVirtualTextured());
-	}
+	// Max size and number of mips from Texture. 
+	const int32 SourceSize = (int32)FMath::Max3(Texture->Source.GetSizeX(),Texture->Source.GetSizeY(),(int64)1);
+	const int32 NumMipsSource = FMath::CeilLogTwo(SourceSize) + 1;
 
-	// Increment the LOD bias per each LOD if we are using automatic LODs
-	if (GenerationContext.CurrentLOD > 0 &&
-		GenerationContext.CurrentAutoLODStrategy == ECustomizableObjectAutomaticLODStrategy::AutomaticFromMesh &&
-		GenerationContext.Options.bUseLODAsBias)
+	// When the BaseTextureSize is known, skip mips until the texture is equal or smaller.
+	if (BaseTextureSize > 0)
 	{
-		// Only if the texture actually uses a layout. Otherwise it could be a special texture we shouldn't scale.
-		if (MaterialNode && MaterialNode->GetImageUVLayout(ImageIndex) >= 0)
+		if (BaseTextureSize < SourceSize)
 		{
-			// \todo: make it an object property to be tweaked
-			int MipsToSkipPerLOD = 1;
-			LODBias += MipsToSkipPerLOD * GenerationContext.CurrentLOD;
+			const int32 MaxNumMipsInGame = FMath::CeilLogTwo(BaseTextureSize) + 1;
+			return FMath::Max(NumMipsSource - MaxNumMipsInGame, 0);
 		}
+
+		return 0;
 	}
 
+	const UTextureLODSettings& LODSettings = GenerationContext.Options.TargetPlatform->GetTextureLODSettings();
+
+	// Get the MaxTextureSize for the TargetPlatform.
+	const int32 MaxTextureSize = GetMaxTextureSize(ReferenceTexture ? *ReferenceTexture : *Texture, LODSettings);
+
 	if (ReferenceTexture)
 	{
-		UE_LOG(LogMutable, Verbose, TEXT("Compiling texture with reference [%s] will have LOD Bias %d."), *ReferenceTexture->GetName(), LODBias);
+		// Max size and number of mips from ReferenceTexture. 
+		const int32 MaxRefSourceSize = (uint32)FMath::Max3(ReferenceTexture->Source.GetSizeX(), ReferenceTexture->Source.GetSizeY(), (int64)1);
+		const int32 NumMipsRefSource = FMath::CeilLogTwo(MaxRefSourceSize) + 1;
+
+		// Find the LODBias to apply to stay within the MaxTextureSize limit of the TargetPlatform
+		const int32 PlatformLODBias = GetPlatformLODBias(MaxRefSourceSize, NumMipsRefSource, MaxTextureSize);
+
+		// TextureSize in-game without any additional LOD bias.
+		const int64 ReferenceTextureSize = MaxRefSourceSize >> PlatformLODBias;
+
+		// Additional LODBias of the Texture
+		const int32 ReferenceTextureLODBias = LODSettings.CalculateLODBias(ReferenceTextureSize, ReferenceTextureSize, 0,	ReferenceTexture->LODGroup,
+			ReferenceTexture->LODBias, 0, ReferenceTexture->MipGenSettings, ReferenceTexture->IsCurrentlyVirtualTextured());
+
+		return FMath::Max(NumMipsSource - NumMipsRefSource + PlatformLODBias + ReferenceTextureLODBias, 0);
 	}
-	else
+
+	// Find the LODBias to apply to stay within the MaxTextureSize limit of the TargetPlatform
+	const int32 PlatformLODBias = GetPlatformLODBias(SourceSize, NumMipsSource, MaxTextureSize);
+
+	// TextureSize in-game without any additional LOD bias.
+	const int64 TextureSize = SourceSize >> PlatformLODBias;
+
+	// Additional LODBias of the Texture
+	const int32 TextureLODBias = LODSettings.CalculateLODBias(TextureSize, TextureSize, 0, Texture->LODGroup, Texture->LODBias, 0, Texture->MipGenSettings, Texture->IsCurrentlyVirtualTextured());
+
+	return FMath::Max(PlatformLODBias + TextureLODBias, 0);
+}
+
+
+int32 GetMaxTextureSize(const UTexture2D& ReferenceTexture, const UTextureLODSettings& LODSettings)
+{
+	// Setting the maximum texture size
+	FTextureLODGroup TextureGroupSettings = LODSettings.GetTextureLODGroup(ReferenceTexture.LODGroup);
+
+	if (TextureGroupSettings.MaxLODSize > 0)
 	{
-		UE_LOG(LogMutable, Verbose, TEXT("Compiling texture without reference will have LOD Bias %d."), LODBias);
+		return ReferenceTexture.MaxTextureSize == 0 ? TextureGroupSettings.MaxLODSize : FMath::Min(TextureGroupSettings.MaxLODSize, ReferenceTexture.MaxTextureSize);
 	}
 
-	return FMath::Min(LODBias, MaxAllowedLODBias);
+	return ReferenceTexture.MaxTextureSize;
 }
 
 
-int32 GetMaxTextureSize(const UTexture2D* ReferenceTexture, const FMutableGraphGenerationContext& GenerationContext)
+int32 GetTextureSizeInGame(const UTexture2D& Texture, const UTextureLODSettings& LODSettings, int32 SurfaceLODBias)
 {
-	if (ReferenceTexture)
-	{
-		// Setting the maximum texture size
-		const UTextureLODSettings& LODSettings = GenerationContext.Options.TargetPlatform->GetTextureLODSettings();
-		FTextureLODGroup TextureGroupSettings = LODSettings.GetTextureLODGroup(ReferenceTexture->LODGroup);
+	const int32 SourceSize = (uint32)FMath::Max3(Texture.Source.GetSizeX(), Texture.Source.GetSizeY(), (int64)1);
+	const int32 NumMipsSource = FMath::CeilLogTwo(SourceSize) + 1;
 
-		if (TextureGroupSettings.MaxLODSize > 0)
-		{
-			return ReferenceTexture->MaxTextureSize == 0 ? TextureGroupSettings.MaxLODSize : FMath::Min(TextureGroupSettings.MaxLODSize, ReferenceTexture->MaxTextureSize);
-		}
-		else
-		{
-			return ReferenceTexture->MaxTextureSize;
-		}
-	}
+	// Max size allowed on the TargetPlatform
+	const int32 MaxTextureSize = GetMaxTextureSize(Texture, LODSettings);
 
-	return 0;
+	// Find the LODBias to apply to stay within the MaxTextureSize limit of the TargetPlatform
+	const int32 PlatformLODBias = GetPlatformLODBias(SourceSize, NumMipsSource, MaxTextureSize);
+	
+	// MaxTextureSize in-game without any additional LOD bias.
+	const int32 MaxTextureSizeAllowed = SourceSize >> PlatformLODBias;
+
+	// Calculate the LODBias specific for this texture 
+	const int32 TextureLODBias = LODSettings.CalculateLODBias(MaxTextureSizeAllowed, MaxTextureSizeAllowed, 0, Texture.LODGroup, Texture.LODBias, 0, Texture.MipGenSettings, Texture.IsCurrentlyVirtualTextured());
+
+	return MaxTextureSizeAllowed >> (TextureLODBias + SurfaceLODBias);
 }
 
 
diff --git a/Engine/Plugins/Experimental/Mutable/Source/CustomizableObjectEditor/Private/MuCOE/GenerateMutableSource/GenerateMutableSource.h b/Engine/Plugins/Experimental/Mutable/Source/CustomizableObjectEditor/Private/MuCOE/GenerateMutableSource/GenerateMutableSource.h
--- a/Engine/Plugins/Experimental/Mutable/Source/CustomizableObjectEditor/Private/MuCOE/GenerateMutableSource/GenerateMutableSource.h
+++ b/Engine/Plugins/Experimental/Mutable/Source/CustomizableObjectEditor/Private/MuCOE/GenerateMutableSource/GenerateMutableSource.h
@@ -29,6 +29,7 @@ class UCustomizableObjectNodeTable;
 struct FCustomizableObjectClothingAssetData;
 
 class FCustomizableObjectCompiler;
+class UTextureLODSettings;
 class UAnimInstance;
 class UCompositeDataTable;
 class UCustomizableObjectNodeMaterial;
@@ -69,6 +70,9 @@ struct FGeneratedImageProperties
 	TEnumAsByte<TextureAddress> AddressY = TA_Clamp;
 
 	bool bIsPassThrough = false;
+
+	// ReferenceTexture source size.
+	int32 TextureSize = 0;
 };
 
 
@@ -159,6 +163,25 @@ struct FGeneratedImageKey
 	const UEdGraphPin* Pin;
 };
 
+
+struct FGeneratedImagePropertiesKey
+{
+	FGeneratedImagePropertiesKey(const UCustomizableObjectNodeMaterial* InMaterial, uint32 InImageIndex)
+	{
+		MaterialReferenceId = (PTRINT)InMaterial;
+		ImageIndex = InImageIndex;
+	}
+
+	bool operator==(const FGeneratedImagePropertiesKey& Other) const
+	{
+		return MaterialReferenceId == Other.MaterialReferenceId && ImageIndex == Other.ImageIndex;
+	}
+
+
+	PTRINT MaterialReferenceId = 0;
+	uint32 ImageIndex = 0;
+};
+
 /** Structure used to store the results of the recursive GenerateMutableSourceSurface calls. */
 struct FMutableGraphSurfaceGenerationData
 {
@@ -238,6 +261,13 @@ inline uint32 GetTypeHash(const FGeneratedImageKey& Key)
 }
 
 
+inline uint32 GetTypeHash(const FGeneratedImagePropertiesKey& Key)
+{
+	uint32 GuidHash = HashCombineFast(GetTypeHash(Key.MaterialReferenceId), Key.ImageIndex);
+
+	return GuidHash;
+}
+
 struct FPoseBoneData
 {
 	TArray<FName> ArrayBoneName;
@@ -257,6 +287,8 @@ struct FGroupProjectorTempData
 	TArray<FPoseBoneData> PoseBoneDataArray;
 
 	bool bAlternateResStateNameWarningDisplayed = false; // Used to display this warning only once
+
+	int32 TextureSize = 512;
 };
 
 
@@ -643,7 +675,7 @@ struct FMutableGraphGenerationContext
 	
 	TArray<UMaterialInterface*> ReferencedMaterials;
 	TArray<FName> ReferencedMaterialSlotNames;
-	TArray<FGeneratedImageProperties> ImageProperties;
+	TMap<FGeneratedImagePropertiesKey, FGeneratedImageProperties> ImageProperties;
 	TMap<FString, TArray<const UObject*>> ParameterNamesMap;
 	TArray<const UCustomizableObjectNode*> NoNameNodeObjectArray;
 	TMap<FString, FCustomizableObjectIdPair> GroupNodeMap;
@@ -697,7 +729,6 @@ struct FMutableGraphGenerationContext
 	int32 NumLODsInRoot = 0;
 	int32 CurrentMeshComponent = 0;
 	int32 NumMeshComponentsInRoot = 0;
-	int32 CurrentTextureLODBias = 0;
 
 	int32 FirstLODAvailable = MAX_MESH_LOD_COUNT;
 	int32 NumMaxLODsToStream = MAX_MESH_LOD_COUNT;
@@ -833,15 +864,17 @@ FString GenerateAnimationInstanceTag(const FString& AnimInstance, const FName& S
 
 FString GenerateGameplayTag(const FString& GameplayTag);
 
-FString GenerateAssetUserDataTag(const FString& AssetUserData);
+uint32 GetBaseTextureSize(const FMutableGraphGenerationContext& GenerationContext, const UCustomizableObjectNodeMaterial* Material, uint32 ImageIndex);
 
 // Computes the LOD bias for a texture given the current mesh LOD and automatic LOD settings, the reference texture settings
 // and whether it's being built for a server or not
-int32 ComputeLODBias(const FMutableGraphGenerationContext& GenerationContext, const UTexture2D* ReferenceTexture, int32 MaxTextureSize,
-	const UCustomizableObjectNodeMaterial* MaterialNode, const int32 ImageIndex);
+uint32 ComputeLODBiasForTexture(const FMutableGraphGenerationContext& GenerationContext, const UTexture2D* Texture,const UTexture2D* ReferenceTexture = nullptr, int32 MaxTextureSizeInGame = 0);
 
+// Max texture size to set on the ImageProperties
+int32 GetMaxTextureSize(const UTexture2D& ReferenceTexture, const UTextureLODSettings& LODSettings);
 
-int32 GetMaxTextureSize(const UTexture2D* ReferenceTexture, const FMutableGraphGenerationContext& GenerationContext);
+// Max texture size of the texture with per platform MaxTextureSize and LODBias applied.
+int32 GetTextureSizeInGame(const UTexture2D& Texture, const UTextureLODSettings& LODSettings, int32 SurfaceLODBias = 0);
 
 mu::Ptr<mu::Image> GenerateImageConstant( UTexture*, FMutableGraphGenerationContext&, bool bIsReference);
 
diff --git a/Engine/Plugins/Experimental/Mutable/Source/CustomizableObjectEditor/Private/MuCOE/GenerateMutableSource/GenerateMutableSourceColor.cpp b/Engine/Plugins/Experimental/Mutable/Source/CustomizableObjectEditor/Private/MuCOE/GenerateMutableSource/GenerateMutableSourceColor.cpp
--- a/Engine/Plugins/Experimental/Mutable/Source/CustomizableObjectEditor/Private/MuCOE/GenerateMutableSource/GenerateMutableSourceColor.cpp
+++ b/Engine/Plugins/Experimental/Mutable/Source/CustomizableObjectEditor/Private/MuCOE/GenerateMutableSource/GenerateMutableSourceColor.cpp
@@ -151,7 +151,7 @@ mu::NodeColourPtr GenerateMutableSourceColor(const UEdGraphPin* Pin, FMutableGra
 
 		if (const UEdGraphPin* ConnectedPin = FollowInputPin(*TypedNodeTexSample->TexturePin()))
 		{
-			mu::NodeImagePtr TextureNode = GenerateMutableSourceImage(ConnectedPin, GenerationContext, 0.f);
+			mu::NodeImagePtr TextureNode = GenerateMutableSourceImage(ConnectedPin, GenerationContext, 0);
 			ColorNode->SetImage(TextureNode);
 		}
 
diff --git a/Engine/Plugins/Experimental/Mutable/Source/CustomizableObjectEditor/Private/MuCOE/GenerateMutableSource/GenerateMutableSourceGroupProjector.cpp b/Engine/Plugins/Experimental/Mutable/Source/CustomizableObjectEditor/Private/MuCOE/GenerateMutableSource/GenerateMutableSourceGroupProjector.cpp
--- a/Engine/Plugins/Experimental/Mutable/Source/CustomizableObjectEditor/Private/MuCOE/GenerateMutableSource/GenerateMutableSourceGroupProjector.cpp
+++ b/Engine/Plugins/Experimental/Mutable/Source/CustomizableObjectEditor/Private/MuCOE/GenerateMutableSource/GenerateMutableSourceGroupProjector.cpp
@@ -177,39 +177,9 @@ mu::NodeImagePtr GenerateMutableGroupProjection(const int32 NodeLOD, const int32
 
 			ImageNode->SetMesh(MeshSwitchNode);
 			ImageNode->SetProjector(ProjectorTempData.NodeProjectorParameterPtr);
+			ImageNode->SetImage(ProjectorTempData.NodeImagePtr);
 
-			mu::NodeImagePtr ImageToProject = ProjectorTempData.NodeImagePtr;
-			{
-				// Apply LODBias to the images to project
-				const UCustomizableObjectNodeMaterial* NodeMaterial = TypedNodeMat ? TypedNodeMat : ParentMaterial;
-				const int32 MaxTextureSize = GroupProjectionReferenceTexture ? GroupProjectionReferenceTexture->MaxTextureSize : 0;
-				if (const int32 LODBias = ComputeLODBias(GenerationContext, GroupProjectionReferenceTexture, MaxTextureSize, NodeMaterial, ImageIndex);
-					LODBias > 0)
-				{
-					mu::NodeImageResizePtr ResizeImage = new mu::NodeImageResize();
-					ResizeImage->SetBase(ImageToProject.get());
-					ResizeImage->SetRelative(true);
-					float factor = FMath::Pow(0.5f, LODBias);
-					ResizeImage->SetSize(factor, factor);
-					ImageToProject = ResizeImage;
-				}
-			}
-
-			ImageNode->SetImage(ImageToProject);
-
-			TextureSize = [&]
-			{
-				if (const int32 Size = ProjectorTempData.CustomizableObjectNodeGroupProjectorParameter->ProjectionTextureSize;
-					Size <= 0 || !FMath::IsPowerOfTwo(Size))
-				{
-					// \todo: closest power of 2 bigger than the set value?
-					return 512;
-				}
-				else
-				{
-					return Size;
-				}
-			}();
+			TextureSize = ProjectorTempData.TextureSize;
 			ImageNode->SetImageSize( FUintVector2(TextureSize, TextureSize) );
 
 			ImageNodes.Add(ImageNode);
@@ -225,11 +195,6 @@ mu::NodeImagePtr GenerateMutableGroupProjection(const int32 NodeLOD, const int32
 	mu::NodeColourConstantPtr ZeroColorNode = new mu::NodeColourConstant();
 	ZeroColorNode->SetValue(FVector4f(0.f, 0.f, 0.f, 1.0f));
 
-	if (TextureSize <= 0 || !FMath::IsPowerOfTwo(TextureSize))
-	{
-		TextureSize = 512;
-	}
-
 	mu::NodeImagePlainColourPtr ZeroPlainColourNode = new mu::NodeImagePlainColour;
 	ZeroPlainColourNode->SetSize(TextureSize, TextureSize);
 	ZeroPlainColourNode->SetColour(ZeroColorNode);
@@ -488,31 +453,37 @@ bool GenerateMutableSourceGroupProjector(const UEdGraphPin* Pin, FMutableGraphGe
 			SwitchNode->SetParameter(EnumParameterNode);
 			SwitchNode->SetOptionCount(ArrayOptionImage.Num());
 
-			bool bFoundUnlinkedPin = false; 
-
-			for (int SelectorIndex = 0; SelectorIndex < ArrayOptionImage.Num(); ++SelectorIndex)
+			const uint32 AdditionalLODBias = GenerationContext.Options.bUseLODAsBias ? GenerationContext.FirstLODAvailable : 0;
+			for (int32 SelectorIndex = 0; SelectorIndex < ArrayOptionImage.Num(); ++SelectorIndex)
 			{
-				if (ArrayOptionImage[SelectorIndex].OptionImage)
+				if (const TObjectPtr<UTexture2D>& Texture = ArrayOptionImage[SelectorIndex].OptionImage)
 				{
 					mu::Ptr<mu::Image> ImageConstant = GenerateImageConstant(ArrayOptionImage[SelectorIndex].OptionImage, GenerationContext, false);
 
 					mu::NodeImageConstantPtr ImageNode = new mu::NodeImageConstant();
 					ImageNode->SetValue(ImageConstant.get());
 
-					SwitchNode->SetOption(SelectorIndex, ImageNode);
+					const uint32 MipsToSkip = ComputeLODBiasForTexture(GenerationContext, Texture, ProjParamNode->ReferenceTexture) + AdditionalLODBias;
+					SwitchNode->SetOption(SelectorIndex, ResizeTextureByNumMips(ImageNode, MipsToSkip));
 				}
 				else
 				{
-					bFoundUnlinkedPin = true;
+					FString msg = FString::Printf(TEXT("The group projection node must have a texture for all the options. Please set a texture for all the options."));
+					GenerationContext.Compiler->CompilerLog(FText::FromString(msg), ProjParamNode);
 				}
 			}
 
-			if (bFoundUnlinkedPin)
+			int32 TextureSize = ProjParamNode->ProjectionTextureSize > 0 ? ProjParamNode->ProjectionTextureSize : 512;
+
+			// If TextureSize is not power of two, round up to the next power of two 
+			if (!FMath::IsPowerOfTwo(TextureSize))
 			{
-				FString msg = FString::Printf(TEXT("The group projection node must have a texture for all the options. Please set a texture for all the options."));
-				GenerationContext.Compiler->CompilerLog(FText::FromString(msg), ProjParamNode);
+				TextureSize = FMath::RoundUpToPowerOfTwo(TextureSize);
 			}
 
+			// Apply additional LODBias if necessary
+			GroupProjectorTempData.TextureSize = FMath::Max(TextureSize >> AdditionalLODBias, 1);
+
 			GroupProjectorTempData.NodeImagePtr = SwitchNode;
 
 			GenerationContext.ProjectorGroupMap.Add(originalGroup, GroupProjectorTempData);
diff --git a/Engine/Plugins/Experimental/Mutable/Source/CustomizableObjectEditor/Private/MuCOE/GenerateMutableSource/GenerateMutableSourceImage.cpp b/Engine/Plugins/Experimental/Mutable/Source/CustomizableObjectEditor/Private/MuCOE/GenerateMutableSource/GenerateMutableSourceImage.cpp
--- a/Engine/Plugins/Experimental/Mutable/Source/CustomizableObjectEditor/Private/MuCOE/GenerateMutableSource/GenerateMutableSourceImage.cpp
+++ b/Engine/Plugins/Experimental/Mutable/Source/CustomizableObjectEditor/Private/MuCOE/GenerateMutableSource/GenerateMutableSourceImage.cpp
@@ -103,30 +103,23 @@ mu::ImagePtr ConvertTextureUnrealToMutable(UTexture2D* Texture, const UCustomiza
 	return MutableImage;
 }
 
-
-mu::Ptr<mu::NodeImage> ResizeToMaxTextureSize(float MaxTextureSize, const UTexture2D* BaseTexture, const mu::Ptr<mu::NodeImage>& ImageNode)
+mu::Ptr<mu::NodeImage> ResizeTextureByNumMips(const mu::Ptr<mu::NodeImage>& ImageConstant, int32 MipsToSkip)
 {
-	// To scale when above MaxTextureSize if defined
-	if (MaxTextureSize > 0 && BaseTexture
-		&& (BaseTexture->GetImportedSize().X > MaxTextureSize
-			|| BaseTexture->GetImportedSize().Y > MaxTextureSize))
+	if (MipsToSkip > 0)
 	{
-		mu::NodeImageResizePtr ResizeImage = new mu::NodeImageResize();
-		ResizeImage->SetBase(ImageNode.get());
-		ResizeImage->SetRelative(false);
-		float Factor = FMath::Min(
-			MaxTextureSize / (float)(BaseTexture->GetImportedSize().X),
-			MaxTextureSize / (float)(BaseTexture->GetImportedSize().Y));
-		ResizeImage->SetSize(BaseTexture->GetImportedSize().X * Factor, BaseTexture->GetImportedSize().Y * Factor);
-		return ResizeImage;
+		mu::NodeImageResizePtr ImageResize = new mu::NodeImageResize();
+		ImageResize->SetBase(ImageConstant.get());
+		ImageResize->SetRelative(true);
+		const float factor = FMath::Pow(0.5f, MipsToSkip);
+		ImageResize->SetSize(factor, factor);
+		return ImageResize;
 	}
-	// TODO: MaxTextureSize = 0 indicates platform-specific maximum size, we should scale for it
 
-	return ImageNode;
+	return ImageConstant;
 }
 
 
-mu::NodeImagePtr GenerateMutableSourceImage(const UEdGraphPin* Pin, FMutableGraphGenerationContext& GenerationContext, float MaxTextureSize)
+mu::NodeImagePtr GenerateMutableSourceImage(const UEdGraphPin* Pin, FMutableGraphGenerationContext& GenerationContext, int32 ReferenceTextureSize)
 {
 	check(Pin)
 	RETURN_ON_CYCLE(*Pin, GenerationContext)
@@ -160,10 +153,10 @@ mu::NodeImagePtr GenerateMutableSourceImage(const UEdGraphPin* Pin, FMutableGrap
 			GenerationContext.AddParticipatingObject(*BaseTexture);
 			
 			// Check the specific image cache
-			FGeneratedImageKey imageKey = FGeneratedImageKey(Pin);
+			FGeneratedImageKey ImageKey = FGeneratedImageKey(Pin);
 			mu::NodeImagePtr ImageNode;
-			mu::NodeImagePtr* Cached = GenerationContext.GeneratedImages.Find(imageKey);
-			if (Cached)
+			
+			if (mu::NodeImagePtr* Cached = GenerationContext.GeneratedImages.Find(ImageKey))
 			{
 				ImageNode = *Cached;
 			}
@@ -175,11 +168,11 @@ mu::NodeImagePtr GenerateMutableSourceImage(const UEdGraphPin* Pin, FMutableGrap
 				ReferenceImageNode->SetValue( ImageConstant.get() );
 				ImageNode = ReferenceImageNode;
 
-				GenerationContext.GeneratedImages.Add(imageKey, ImageNode);
+				GenerationContext.GeneratedImages.Add(ImageKey, ImageNode);
 			}
 
-			Result = ResizeToMaxTextureSize(MaxTextureSize, BaseTexture, ImageNode);
-
+			const uint32 MipsToSkip = ComputeLODBiasForTexture(GenerationContext, BaseTexture, nullptr, ReferenceTextureSize);
+			Result = ResizeTextureByNumMips(ImageNode, MipsToSkip);
 		}
 		else
 		{
@@ -224,23 +217,35 @@ mu::NodeImagePtr GenerateMutableSourceImage(const UEdGraphPin* Pin, FMutableGrap
 		ResizeNode->SetBase(FormatNode);
 		ResizeNode->SetRelative(false);
 
+		FUintVector2 TextureSize(TypedNodeParam->TextureSizeX, TypedNodeParam->TextureSizeY);
+
 		const UTexture2D* ReferenceTexture = TypedNodeParam->ReferenceValue;
 		if (ReferenceTexture)
 		{
 			GenerationContext.AddParticipatingObject(*ReferenceTexture);
 
-			ResizeNode->SetSize(FMath::Max(ReferenceTexture->GetImportedSize().X,1), FMath::Max(ReferenceTexture->GetImportedSize().X, 1));
+			const uint32 LODBias = ComputeLODBiasForTexture(GenerationContext, TypedNodeParam->ReferenceValue, ReferenceTexture, ReferenceTextureSize);
+			TextureSize.X = FMath::Max(ReferenceTexture->Source.GetSizeX() >> LODBias, 1);
+			TextureSize.Y = FMath::Max(ReferenceTexture->Source.GetSizeY() >> LODBias, 1);
 		}
 		else
 		{
-			if (TypedNodeParam->TextureSizeX <= 0 || TypedNodeParam->TextureSizeY <= 0)
+			const int32 MaxNodeTextureSize = FMath::Max(TypedNodeParam->TextureSizeX, TypedNodeParam->TextureSizeY);
+			if (MaxNodeTextureSize <= 0)
 			{
+				TextureSize.X = TextureSize.Y = 1;
 				GenerationContext.Compiler->CompilerLog(LOCTEXT("TextureParameterSize0", "Texture size not specified. Add a reference texture or set a valid value to the Texture Size variables."), Node);
 			}
-
-			ResizeNode->SetSize(FMath::Max(TypedNodeParam->TextureSizeX, 1), FMath::Max(TypedNodeParam->TextureSizeY, 1));
+			else if (ReferenceTextureSize > 0 && ReferenceTextureSize < MaxNodeTextureSize)
+			{
+				const uint32 MipsToSkip = FMath::CeilLogTwo(MaxNodeTextureSize) - FMath::CeilLogTwo(ReferenceTextureSize);
+				TextureSize.X = FMath::Max(TextureSize.X >> MipsToSkip, (uint32)1);
+				TextureSize.Y = FMath::Max(TextureSize.Y >> MipsToSkip, (uint32)1);
+			}
 		}
 
+		ResizeNode->SetSize(TextureSize.X, TextureSize.Y);
+
 		Result = ResizeNode;
 	}
 
@@ -252,7 +257,8 @@ mu::NodeImagePtr GenerateMutableSourceImage(const UEdGraphPin* Pin, FMutableGrap
 		UTexture2D* Texture = TypedNodeMesh->FindTextureForPin(Pin);
 		ImageNode->SetValue( GenerateImageConstant(Texture, GenerationContext, false).get() );
 
-		Result = ResizeToMaxTextureSize(MaxTextureSize, Texture, ImageNode);
+		const uint32 MipsToSkip = ComputeLODBiasForTexture(GenerationContext, Texture, nullptr, ReferenceTextureSize);
+		Result = ResizeTextureByNumMips(ImageNode, MipsToSkip);
 	}
 
 	else if (const UCustomizableObjectNodeTextureInterpolate* TypedNodeInterp = Cast<UCustomizableObjectNodeTextureInterpolate>(Node))
@@ -271,7 +277,7 @@ mu::NodeImagePtr GenerateMutableSourceImage(const UEdGraphPin* Pin, FMutableGrap
 		{
 			if (const UEdGraphPin* ConnectedPin = FollowInputPin(*TypedNodeInterp->Targets(LayerIndex)))
 			{
-				mu::NodeImagePtr TargetNode = GenerateMutableSourceImage(ConnectedPin, GenerationContext, MaxTextureSize);
+				mu::NodeImagePtr TargetNode = GenerateMutableSourceImage(ConnectedPin, GenerationContext, ReferenceTextureSize);
 
 				if (TargetNode)
 				{
@@ -287,7 +293,7 @@ mu::NodeImagePtr GenerateMutableSourceImage(const UEdGraphPin* Pin, FMutableGrap
 	{
 		if (const UEdGraphPin* ConnectedPin = FollowInputPin(*TypedNodeLayer->BasePin()))
 		{
-			Result = GenerateMutableSourceImage(ConnectedPin, GenerationContext, MaxTextureSize);
+			Result = GenerateMutableSourceImage(ConnectedPin, GenerationContext, ReferenceTextureSize);
 		}
 
 		for (int LayerIndex = 0; LayerIndex < TypedNodeLayer->GetNumLayers(); ++LayerIndex)
@@ -297,7 +303,7 @@ mu::NodeImagePtr GenerateMutableSourceImage(const UEdGraphPin* Pin, FMutableGrap
 				mu::NodeImagePtr MaskNode = nullptr;
 				if (const UEdGraphPin* ConnectedPin = FollowInputPin(*TypedNodeLayer->MaskPin(LayerIndex)))
 				{
-					MaskNode = GenerateMutableSourceImage(ConnectedPin, GenerationContext, MaxTextureSize);
+					MaskNode = GenerateMutableSourceImage(ConnectedPin, GenerationContext, ReferenceTextureSize);
 				}
 
 				mu::EBlendType Type = mu::EBlendType::BT_BLEND;
@@ -325,7 +331,7 @@ mu::NodeImagePtr GenerateMutableSourceImage(const UEdGraphPin* Pin, FMutableGrap
 				
 				if (OtherPin->PinType.PinCategory == Schema->PC_Image)
 				{
-					mu::NodeImagePtr BlendNode = GenerateMutableSourceImage(OtherPin, GenerationContext, MaxTextureSize);
+					mu::NodeImagePtr BlendNode = GenerateMutableSourceImage(OtherPin, GenerationContext, ReferenceTextureSize);
 
 					mu::NodeImageLayerPtr LayerNode = new mu::NodeImageLayer;
 					LayerNode->SetBlendType(Type);
@@ -397,7 +403,7 @@ mu::NodeImagePtr GenerateMutableSourceImage(const UEdGraphPin* Pin, FMutableGrap
 				{
 					if (const UEdGraphPin* TexturePin = FollowInputPin(*TypedNodeTextureSwitch->GetElementPin(SelectorIndex)))
 					{
-						SwitchNode->SetOption(SelectorIndex, GenerateMutableSourceImage(TexturePin, GenerationContext, MaxTextureSize));
+						SwitchNode->SetOption(SelectorIndex, GenerateMutableSourceImage(TexturePin, GenerationContext, ReferenceTextureSize));
 					}
 					else
 					{
@@ -426,7 +432,7 @@ mu::NodeImagePtr GenerateMutableSourceImage(const UEdGraphPin* Pin, FMutableGrap
 
 		if (const UEdGraphPin* ConnectedPin = FollowInputPin(*TypedNodeImageVar->DefaultPin()))
 		{
-			mu::NodeImagePtr ChildNode = GenerateMutableSourceImage(ConnectedPin, GenerationContext, MaxTextureSize);
+			mu::NodeImagePtr ChildNode = GenerateMutableSourceImage(ConnectedPin, GenerationContext, ReferenceTextureSize);
 			if (ChildNode)
 			{
 				TextureNode->SetDefaultImage(ChildNode.get());
@@ -449,7 +455,7 @@ mu::NodeImagePtr GenerateMutableSourceImage(const UEdGraphPin* Pin, FMutableGrap
 			TextureNode->SetVariationTag(VariationIndex, TypedNodeImageVar->Variations[VariationIndex].Tag);
 			if (const UEdGraphPin* ConnectedPin = FollowInputPin(*VariationPin))
 			{
-				mu::NodeImagePtr ChildNode = GenerateMutableSourceImage(ConnectedPin, GenerationContext, MaxTextureSize);
+				mu::NodeImagePtr ChildNode = GenerateMutableSourceImage(ConnectedPin, GenerationContext, ReferenceTextureSize);
 				TextureNode->SetVariationImage(VariationIndex, ChildNode.get());
 			}
 		}
@@ -480,22 +486,22 @@ mu::NodeImagePtr GenerateMutableSourceImage(const UEdGraphPin* Pin, FMutableGrap
 
 		if (const UEdGraphPin* ConnectedPin = FollowInputPin(*TypedNodeFrom->RPin()))
 		{
-			RNode = GenerateMutableSourceImage(ConnectedPin, GenerationContext, MaxTextureSize);
+			RNode = GenerateMutableSourceImage(ConnectedPin, GenerationContext, ReferenceTextureSize);
 			RGB = true;
 		}
 		if (const UEdGraphPin* ConnectedPin = FollowInputPin(*TypedNodeFrom->GPin()))
 		{
-			GNode = GenerateMutableSourceImage(ConnectedPin, GenerationContext, MaxTextureSize);
+			GNode = GenerateMutableSourceImage(ConnectedPin, GenerationContext, ReferenceTextureSize);
 			RGB = true;
 		}
 		if (const UEdGraphPin* ConnectedPin = FollowInputPin(*TypedNodeFrom->BPin()))
 		{
-			BNode = GenerateMutableSourceImage(ConnectedPin, GenerationContext, MaxTextureSize);
+			BNode = GenerateMutableSourceImage(ConnectedPin, GenerationContext, ReferenceTextureSize);
 			RGB = true;
 		}
 		if (const UEdGraphPin* ConnectedPin = FollowInputPin(*TypedNodeFrom->APin()))
 		{
-			ANode = GenerateMutableSourceImage(ConnectedPin, GenerationContext, MaxTextureSize);
+			ANode = GenerateMutableSourceImage(ConnectedPin, GenerationContext, ReferenceTextureSize);
 		}
 
 		mu::NodeImageSwizzlePtr SwizzleNode = new mu::NodeImageSwizzle;
@@ -542,7 +548,7 @@ mu::NodeImagePtr GenerateMutableSourceImage(const UEdGraphPin* Pin, FMutableGrap
 		mu::NodeImagePtr BaseNode;
 		if (const UEdGraphPin* ConnectedPin = FollowInputPin(*TypedNodeTo->InputPin()))
 		{
-			BaseNode = GenerateMutableSourceImage(ConnectedPin, GenerationContext, MaxTextureSize);
+			BaseNode = GenerateMutableSourceImage(ConnectedPin, GenerationContext, ReferenceTextureSize);
 		}
 
 		mu::NodeImageSwizzlePtr SwizzleNode = new mu::NodeImageSwizzle;
@@ -584,25 +590,31 @@ mu::NodeImagePtr GenerateMutableSourceImage(const UEdGraphPin* Pin, FMutableGrap
 		}
 
 		ImageNode->SetLayout(TypedNodeProject->Layout);
-		FUintVector2 TextureSize(TypedNodeProject->TextureSizeX, TypedNodeProject->TextureSizeY);
 
-		// Calculating Texture size using Reference texture parameters
+		// Calculate the max TextureSize allowed using the ReferenceTextureSize and the Reference texture from the node
+		int32 MaxReferenceTextureSizeInGame = ReferenceTextureSize;
 		if (TypedNodeProject->ReferenceTexture)
 		{
 			GenerationContext.AddParticipatingObject(*TypedNodeProject->ReferenceTexture);
 
-			int32 LODBias = ComputeLODBias(GenerationContext, TypedNodeProject->ReferenceTexture, TypedNodeProject->ReferenceTexture->MaxTextureSize, nullptr, INDEX_NONE);
+			const UTextureLODSettings& TextureLODSettings = GenerationContext.Options.TargetPlatform->GetTextureLODSettings();
+			const int32 FirstLODAvailable = GenerationContext.Options.bUseLODAsBias ? GenerationContext.FirstLODAvailable : 0;
+			MaxReferenceTextureSizeInGame = GetTextureSizeInGame(*TypedNodeProject->ReferenceTexture, TextureLODSettings, FirstLODAvailable);
+		}
 
-			if (TextureSize.X > 0 && TextureSize.Y > 0)
-			{
-				TextureSize.X = TextureSize.X >> LODBias;
-				TextureSize.Y = TextureSize.Y >> LODBias;
-			}
-			else
-			{
-				TextureSize.X = TypedNodeProject->ReferenceTexture->GetImportedSize().X >> LODBias;
-				TextureSize.Y = TypedNodeProject->ReferenceTexture->GetImportedSize().Y >> LODBias;
-			}
+		FUintVector2 TextureSize(TypedNodeProject->TextureSizeX, TypedNodeProject->TextureSizeY);
+
+		// Max TextureSize allowed
+		const int32 MaxProjectedTextureSizeInGame = ReferenceTextureSize > 0 && ReferenceTextureSize < MaxReferenceTextureSizeInGame ? ReferenceTextureSize : MaxReferenceTextureSizeInGame;
+
+		const int32 ProjectorNodeTextureSize = FMath::Max(TextureSize.X, TextureSize.Y);
+		if (ProjectorNodeTextureSize > 0 && ProjectorNodeTextureSize > MaxProjectedTextureSizeInGame)
+		{
+			const int32 NumMips = FMath::CeilLogTwo(ProjectorNodeTextureSize) + 1;
+			const int32 MaxNumMips = FMath::CeilLogTwo(MaxProjectedTextureSizeInGame) + 1;
+
+			TextureSize.X = TextureSize.X >> (NumMips - MaxNumMips);
+			TextureSize.Y = TextureSize.Y >> (NumMips - MaxNumMips);
 		}
 
 		ImageNode->SetImageSize(TextureSize);
@@ -670,7 +682,7 @@ mu::NodeImagePtr GenerateMutableSourceImage(const UEdGraphPin* Pin, FMutableGrap
 
 		if (const UEdGraphPin* ConnectedPin = FollowInputPin(*TypedNodeProject->MeshMaskPin()))
 		{
-			mu::NodeImagePtr MeshMaskNode = GenerateMutableSourceImage(ConnectedPin, GenerationContext, MaxTextureSize);
+			mu::NodeImagePtr MeshMaskNode = GenerateMutableSourceImage(ConnectedPin, GenerationContext, ReferenceTextureSize);
 			ImageNode->SetTargetMask(MeshMaskNode);
 		}
 
@@ -694,20 +706,7 @@ mu::NodeImagePtr GenerateMutableSourceImage(const UEdGraphPin* Pin, FMutableGrap
 
 		if (const UEdGraphPin* ConnectedPin = FollowInputPin(*TypedNodeProject->TexturePins(TexIndex)))
 		{
-			mu::NodeImagePtr SourceNode = GenerateMutableSourceImage(ConnectedPin, GenerationContext, MaxTextureSize);
-
-			// Images that are projected won't be scaled by any means, so we need to apply the lodbias settings here
-			if (GenerationContext.CurrentTextureLODBias > 0)
-			{
-				mu::NodeImageResizePtr ResizeImage = new mu::NodeImageResize();
-				ResizeImage->SetBase(SourceNode.get());
-				ResizeImage->SetRelative(true);
-				float factor = FMath::Pow(0.5f, GenerationContext.CurrentTextureLODBias);
-				ResizeImage->SetSize(factor, factor);
-				ResizeImage->SetMessageContext(Node);
-				SourceNode = ResizeImage;
-			}
-
+			mu::NodeImagePtr SourceNode = GenerateMutableSourceImage(ConnectedPin, GenerationContext, FMath::Max(TextureSize.X, TextureSize.Y));
 			ImageNode->SetImage(SourceNode);
 		}
 	}
@@ -719,7 +718,7 @@ mu::NodeImagePtr GenerateMutableSourceImage(const UEdGraphPin* Pin, FMutableGrap
 
 		if (const UEdGraphPin* ConnectedPin = FollowInputPin(*TypedNodeTexBin->GetBaseImagePin()))
 		{
-			mu::NodeImagePtr BaseImageNode = GenerateMutableSourceImage(ConnectedPin, GenerationContext, MaxTextureSize);
+			mu::NodeImagePtr BaseImageNode = GenerateMutableSourceImage(ConnectedPin, GenerationContext, ReferenceTextureSize);
 			BinariseNode->SetBase(BaseImageNode);
 		}
 
@@ -737,7 +736,7 @@ mu::NodeImagePtr GenerateMutableSourceImage(const UEdGraphPin* Pin, FMutableGrap
 
 		if (const UEdGraphPin* ConnectedPin = FollowInputPin(*TypedNodeTexInv->GetBaseImagePin()))
 		{
-			mu::NodeImagePtr BaseImageNode = GenerateMutableSourceImage(ConnectedPin, GenerationContext, MaxTextureSize);
+			mu::NodeImagePtr BaseImageNode = GenerateMutableSourceImage(ConnectedPin, GenerationContext, ReferenceTextureSize);
 			InvertNode->SetBase(BaseImageNode);
 		}
 	}
@@ -750,19 +749,19 @@ mu::NodeImagePtr GenerateMutableSourceImage(const UEdGraphPin* Pin, FMutableGrap
 
 		if (const UEdGraphPin* ConnectedPin = FollowInputPin(*TypedNodeColourMap->GetMapPin()))
 		{
-			mu::NodeImagePtr GradientImageNode = GenerateMutableSourceImage(ConnectedPin, GenerationContext, MaxTextureSize);
+			mu::NodeImagePtr GradientImageNode = GenerateMutableSourceImage(ConnectedPin, GenerationContext, ReferenceTextureSize);
 			ColourMapNode->SetMap( GradientImageNode ); 
 		}
 
 		if (const UEdGraphPin* ConnectedPin = FollowInputPin(*TypedNodeColourMap->GetMaskPin()))
 		{
-			mu::NodeImagePtr GradientImageNode = GenerateMutableSourceImage(ConnectedPin, GenerationContext, MaxTextureSize);
+			mu::NodeImagePtr GradientImageNode = GenerateMutableSourceImage(ConnectedPin, GenerationContext, ReferenceTextureSize);
 			ColourMapNode->SetMask( GradientImageNode ); 
 		}
 
 		if (const UEdGraphPin* ConnectedPin = FollowInputPin(*TypedNodeColourMap->GetBasePin()))
 		{
-			mu::NodeImagePtr SourceImageNode = GenerateMutableSourceImage(ConnectedPin, GenerationContext, MaxTextureSize);
+			mu::NodeImagePtr SourceImageNode = GenerateMutableSourceImage(ConnectedPin, GenerationContext, ReferenceTextureSize);
 			ColourMapNode->SetBase( SourceImageNode );
 		}
 	}
@@ -774,7 +773,7 @@ mu::NodeImagePtr GenerateMutableSourceImage(const UEdGraphPin* Pin, FMutableGrap
 		
 		if ( UEdGraphPin* BaseImagePin = FollowInputPin(*TypedNodeTransform->GetBaseImagePin()) )
 		{
-			mu::NodeImagePtr ImageNode = GenerateMutableSourceImage( BaseImagePin, GenerationContext, MaxTextureSize);
+			mu::NodeImagePtr ImageNode = GenerateMutableSourceImage( BaseImagePin, GenerationContext, ReferenceTextureSize);
 			TransformNode->SetBase( ImageNode ); 
 		}
 
@@ -821,27 +820,30 @@ mu::NodeImagePtr GenerateMutableSourceImage(const UEdGraphPin* Pin, FMutableGrap
 
 		FUintVector2 TextureSize(TypedNodeTransform->TextureSizeX, TypedNodeTransform->TextureSizeY);
 
-		// Calculating Texture size using Reference texture parameters
+		// Calculate the max TextureSize allowed using the ReferenceTextureSize and the Reference texture from the node
+		int32 MaxReferenceTextureSizeInGame = ReferenceTextureSize;
 		if (TypedNodeTransform->ReferenceTexture)
 		{
 			GenerationContext.AddParticipatingObject(*TypedNodeTransform->ReferenceTexture);
-			
-			const int32 LODBias = ComputeLODBias(
-				GenerationContext, TypedNodeTransform->ReferenceTexture, TypedNodeTransform->ReferenceTexture->MaxTextureSize, 
-				nullptr, INDEX_NONE);
 
-			if (TextureSize.X > 0 && TextureSize.Y > 0)
-			{
-				TextureSize.X = TextureSize.X >> LODBias;
-				TextureSize.Y = TextureSize.Y >> LODBias;
-			}
-			else
-			{
-				TextureSize.X = TypedNodeTransform->ReferenceTexture->GetImportedSize().X >> LODBias;
-				TextureSize.Y = TypedNodeTransform->ReferenceTexture->GetImportedSize().Y >> LODBias;
-			}
+			const UTextureLODSettings& TextureLODSettings = GenerationContext.Options.TargetPlatform->GetTextureLODSettings();
+			const int32 FirstLODAvailable = GenerationContext.Options.bUseLODAsBias ? GenerationContext.FirstLODAvailable : 0;
+			MaxReferenceTextureSizeInGame = GetTextureSizeInGame(*TypedNodeTransform->ReferenceTexture, TextureLODSettings, FirstLODAvailable);
 		}
-		
+
+		// Max TextureSize allowed
+		const int32 MaxTransformTextureSizeInGame = ReferenceTextureSize > 0 && ReferenceTextureSize < MaxReferenceTextureSizeInGame ? ReferenceTextureSize : MaxReferenceTextureSizeInGame;
+
+		const int32 TransformNodeTextureSize = FMath::Max(TextureSize.X, TextureSize.Y);
+		if (TransformNodeTextureSize > 0 && TransformNodeTextureSize > MaxTransformTextureSizeInGame)
+		{
+			const int32 NumMips = FMath::CeilLogTwo(TransformNodeTextureSize) + 1;
+			const int32 MaxNumMips = FMath::CeilLogTwo(MaxTransformTextureSizeInGame) + 1;
+
+			TextureSize.X = TextureSize.X >> (NumMips - MaxNumMips);
+			TextureSize.Y = TextureSize.Y >> (NumMips - MaxNumMips);
+		}
+
 		TransformNode->SetKeepAspectRatio(TypedNodeTransform->bKeepAspectRatio);
 		TransformNode->SetSizeX(TextureSize.X);
 		TransformNode->SetSizeY(TextureSize.Y);
@@ -854,7 +856,7 @@ mu::NodeImagePtr GenerateMutableSourceImage(const UEdGraphPin* Pin, FMutableGrap
 	
 		if ( UEdGraphPin* BaseImagePin = FollowInputPin(*TypedNodeSaturate->GetBaseImagePin()) )
 		{
-			mu::NodeImagePtr ImageNode = GenerateMutableSourceImage(BaseImagePin, GenerationContext, MaxTextureSize);
+			mu::NodeImagePtr ImageNode = GenerateMutableSourceImage(BaseImagePin, GenerationContext, ReferenceTextureSize);
 			SaturateNode->SetSource(ImageNode); 
 		}
 
@@ -926,7 +928,7 @@ mu::NodeImagePtr GenerateMutableSourceImage(const UEdGraphPin* Pin, FMutableGrap
 				{
 					if (const UEdGraphPin* TexturePin = FollowInputPin(*TypedNodePassThroughTextureSwitch->GetElementPin(SelectorIndex)))
 					{
-						mu::Ptr<mu::NodeImage> PassThroughImage = GenerateMutableSourceImage(TexturePin, GenerationContext, MaxTextureSize);
+						mu::Ptr<mu::NodeImage> PassThroughImage = GenerateMutableSourceImage(TexturePin, GenerationContext, ReferenceTextureSize);
 						SwitchNode->SetOption(SelectorIndex, PassThroughImage);
 					}
 					else
@@ -1078,12 +1080,17 @@ mu::NodeImagePtr GenerateMutableSourceImage(const UEdGraphPin* Pin, FMutableGrap
 							ImageTableNode->SetParameterName(TypedNodeTable->ParameterName);
 							ImageTableNode->SetNoneOption(TypedNodeTable->bAddNoneOption);
 
-							// TextureArrays are passthrough textures and do not need this step
 							if (DefaultTexture2D)
 							{
-								int32 DefaultMaxTextureSize = GetMaxTextureSize(DefaultTexture2D, GenerationContext);
-								ImageTableNode->SetMaxTextureSize(FMath::Max(DefaultTexture2D->Source.GetSizeX(), DefaultTexture2D->Source.GetSizeY()));
-								ImageTableNode->SetReferenceImageDescriptor(GenerateImageDescriptor(DefaultTexture2D));
+								mu::FImageDesc ImageDesc = GenerateImageDescriptor(DefaultTexture2D);
+
+								uint32 LODBias = ReferenceTextureSize > 0 ? ComputeLODBiasForTexture(GenerationContext, DefaultTexture2D, nullptr, ReferenceTextureSize) : 0;
+								ImageDesc.m_size[0] = ImageDesc.m_size[0] >> LODBias;
+								ImageDesc.m_size[1] = ImageDesc.m_size[1] >> LODBias;
+								
+								const uint16 MaxTextureSize = FMath::Max3(ImageDesc.m_size[0], ImageDesc.m_size[1], (uint16)1);
+								ImageTableNode->SetMaxTextureSize(MaxTextureSize);
+								ImageTableNode->SetReferenceImageDescriptor(ImageDesc);
 							}
 						}
 					}
diff --git a/Engine/Plugins/Experimental/Mutable/Source/CustomizableObjectEditor/Private/MuCOE/GenerateMutableSource/GenerateMutableSourceImage.h b/Engine/Plugins/Experimental/Mutable/Source/CustomizableObjectEditor/Private/MuCOE/GenerateMutableSource/GenerateMutableSourceImage.h
--- a/Engine/Plugins/Experimental/Mutable/Source/CustomizableObjectEditor/Private/MuCOE/GenerateMutableSource/GenerateMutableSourceImage.h
+++ b/Engine/Plugins/Experimental/Mutable/Source/CustomizableObjectEditor/Private/MuCOE/GenerateMutableSource/GenerateMutableSourceImage.h
@@ -4,6 +4,7 @@
 
 
 #include "MuR/Ptr.h"
+#include "HAL/Platform.h"
 
 class FCustomizableObjectCompiler;
 class UCustomizableObjectNode;
@@ -21,9 +22,7 @@ namespace mu
 
 mu::Ptr<mu::Image> ConvertTextureUnrealToMutable(UTexture2D* Texture, const UCustomizableObjectNode* Node, FCustomizableObjectCompiler* Compiler, bool bIsNormalComposite);
 
-
-mu::Ptr<mu::NodeImage> ResizeToMaxTextureSize(float MaxTextureSize, const UTexture2D* BaseTexture, const mu::Ptr<mu::NodeImage>& ImageNode);
-
+mu::Ptr<mu::NodeImage> ResizeTextureByNumMips(const mu::Ptr<mu::NodeImage>& Image, int32 MipsToSkip);
 
 /** Convert a CustomizableObject Source Graph from an Image pin into a mutable source graph. */
-mu::Ptr<mu::NodeImage> GenerateMutableSourceImage(const UEdGraphPin* Pin, FMutableGraphGenerationContext& GenerationContext, float MaxTextureSize);
+mu::Ptr<mu::NodeImage> GenerateMutableSourceImage(const UEdGraphPin* Pin, FMutableGraphGenerationContext& GenerationContext, int32 ReferenceTextureSize);
diff --git a/Engine/Plugins/Experimental/Mutable/Source/CustomizableObjectEditor/Private/MuCOE/GenerateMutableSource/GenerateMutableSourceSurface.cpp b/Engine/Plugins/Experimental/Mutable/Source/CustomizableObjectEditor/Private/MuCOE/GenerateMutableSource/GenerateMutableSourceSurface.cpp
--- a/Engine/Plugins/Experimental/Mutable/Source/CustomizableObjectEditor/Private/MuCOE/GenerateMutableSource/GenerateMutableSourceSurface.cpp
+++ b/Engine/Plugins/Experimental/Mutable/Source/CustomizableObjectEditor/Private/MuCOE/GenerateMutableSource/GenerateMutableSourceSurface.cpp
@@ -454,9 +454,10 @@ mu::NodeSurfacePtr GenerateMutableSourceSurface(const UEdGraphPin * Pin, FMutabl
 
 					// We don't need a reference texture or props here, but we do need the parameter name.
 					FGeneratedImageProperties Props;
+					FGeneratedImagePropertiesKey PropsKey(TypedNodeMat,(uint32)ImageIndex);
 					Props.TextureParameterName = ImageName;
 					Props.bIsPassThrough = true;
-					GenerationContext.ImageProperties.Add(Props);
+					GenerationContext.ImageProperties.Add(PropsKey, Props);
 					SurfaceData.ImageProperties = Props;
 				}
 			}
@@ -520,6 +521,7 @@ mu::NodeSurfacePtr GenerateMutableSourceSurface(const UEdGraphPin * Pin, FMutabl
 						}
 
 						FGeneratedImageProperties Props;
+						FGeneratedImagePropertiesKey PropsKey(TypedNodeMat, ImageIndex);
 						if (ReferenceTexture)
 						{
 							GenerationContext.AddParticipatingObject(*ReferenceTexture);
@@ -531,12 +533,20 @@ mu::NodeSurfacePtr GenerateMutableSourceSurface(const UEdGraphPin * Pin, FMutabl
 							Props.SRGB = ReferenceTexture->SRGB;
 							Props.LODBias = 0;
 							Props.MipGenSettings = ReferenceTexture->MipGenSettings;
-							Props.MaxTextureSize = GetMaxTextureSize(ReferenceTexture, GenerationContext);
 							Props.LODGroup = ReferenceTexture->LODGroup;
 							Props.AddressX = ReferenceTexture->AddressX;
 							Props.AddressY = ReferenceTexture->AddressY;
 							Props.bFlipGreenChannel = ReferenceTexture->bFlipGreenChannel;
 
+
+							// MaxTextureSize setting. Based on the ReferenceTexture and Platform settings.
+							const UTextureLODSettings& TextureLODSettings = GenerationContext.Options.TargetPlatform->GetTextureLODSettings();
+							Props.MaxTextureSize = GetMaxTextureSize(*ReferenceTexture, TextureLODSettings);
+							
+							// ReferenceTexture source size. Textures contributing to this Image should be equal to or smaller than TextureSize. 
+							// The LOD Bias applied to the root node will be applied on top of it.
+							Props.TextureSize = (int32)FMath::Max3(ReferenceTexture->Source.GetSizeX(), ReferenceTexture->Source.GetSizeY(), 1LL);
+
 							// TODO: MTBL-1081
 							// TextureGroup::TEXTUREGROUP_UI does not support streaming. If we generate a texture that requires streaming and set this group, it will crash when initializing the resource. 
 							// If LODGroup == TEXTUREGROUP_UI, UTexture::IsPossibleToStream() will return false and UE will assume all mips are loaded, when they're not, and crash.
@@ -556,14 +566,9 @@ mu::NodeSurfacePtr GenerateMutableSourceSurface(const UEdGraphPin * Pin, FMutabl
 							GenerationContext.Compiler->CompilerLog(FText::FromString(msg), Node);
 						}
 
-						GenerationContext.ImageProperties.Add(Props);
+						GenerationContext.ImageProperties.Add(PropsKey, Props);
 						SurfaceData.ImageProperties = Props;
 
-						// Calculate the LODBias for this texture
-						int32 LODBias = ComputeLODBias(GenerationContext, ReferenceTexture, ReferenceTexture ? ReferenceTexture->MaxTextureSize : 0, TypedNodeMat, ImageIndex);
-
-						GenerationContext.CurrentTextureLODBias = LODBias;
-
 						// Generate the texture nodes
 						mu::NodeImagePtr ImageNode = [&]()
 						{
@@ -573,7 +578,7 @@ mu::NodeSurfacePtr GenerateMutableSourceSurface(const UEdGraphPin * Pin, FMutabl
 								{
 									if (const UEdGraphPin* ConnectedPin = FollowInputPin(*ImagePin))
 									{
-										return GenerateMutableSourceImage(ConnectedPin, GenerationContext, Props.MaxTextureSize);
+										return GenerateMutableSourceImage(ConnectedPin, GenerationContext, Props.TextureSize);
 									}
 								}
 
@@ -581,7 +586,7 @@ mu::NodeSurfacePtr GenerateMutableSourceSurface(const UEdGraphPin * Pin, FMutabl
 								{
 									if (const UEdGraphPin* ConnectedPin = FollowInputPin(*TypedNodeMat->GetMaterialAssetPin()))
 									{
-										return GenerateMutableSourceImage(ConnectedPin, GenerationContext, Props.MaxTextureSize);
+										return GenerateMutableSourceImage(ConnectedPin, GenerationContext, Props.TextureSize);
 									}
 								}
 
@@ -593,7 +598,8 @@ mu::NodeSurfacePtr GenerateMutableSourceSurface(const UEdGraphPin * Pin, FMutabl
 									mu::Ptr<mu::Image> ImageConstant = GenerateImageConstant(Texture2D, GenerationContext, false);
 									ConstImageNode->SetValue(ImageConstant.get());
 
-									return ResizeToMaxTextureSize(Props.MaxTextureSize, Texture2D, ConstImageNode);
+									const uint32 MipsToSkip = ComputeLODBiasForTexture(GenerationContext, Texture2D, nullptr, Props.TextureSize);
+									return ResizeTextureByNumMips(ConstImageNode, MipsToSkip);
 								}
 							}
 							else
@@ -607,23 +613,12 @@ mu::NodeSurfacePtr GenerateMutableSourceSurface(const UEdGraphPin * Pin, FMutabl
 							ImageNode = GroupProjectionImg;
 						}
 
-						if (GenerationContext.Options.TargetPlatform && ReferenceTexture)
+						if (ReferenceTexture)
 						{
-							int LayerIndex = 0;
-
-							mu::NodeImagePtr LastImage = ImageNode;
-
-							// To apply LOD bias						
-							if (LODBias > 0)
-							{
-								mu::NodeImageResizePtr ResizeImage = new mu::NodeImageResize();
-								ResizeImage->SetBase(LastImage.get());
-								ResizeImage->SetRelative(true);
-								float factor = FMath::Pow(0.5f, LODBias);
-								ResizeImage->SetSize(factor, factor);
-								ResizeImage->SetMessageContext(Node);
-								LastImage = ResizeImage;
-							}
+							// Apply base LODBias. It will be propagated to most images.
+							const uint32 FirstLODAvailable = GenerationContext.Options.bUseLODAsBias ? GenerationContext.FirstLODAvailable : 0;
+							const uint32 BaseLODBias = ComputeLODBiasForTexture(GenerationContext, ReferenceTexture) + FirstLODAvailable;
+							mu::NodeImagePtr LastImage = ResizeTextureByNumMips(ImageNode, BaseLODBias);
 
 							mu::NodeImageMipmapPtr MipmapImage = new mu::NodeImageMipmap();
 							MipmapImage->SetSource(LastImage.get());
@@ -666,7 +661,8 @@ mu::NodeSurfacePtr GenerateMutableSourceSurface(const UEdGraphPin * Pin, FMutabl
 									CompositeNormalImage->SetValue(ImageConstant.get());
 
 									mu::NodeImageMipmapPtr NormalCompositeMipmapImage = new mu::NodeImageMipmap();
-									NormalCompositeMipmapImage->SetSource(CompositeNormalImage);
+									const uint32 MipsToSkip = ComputeLODBiasForTexture(GenerationContext, ReferenceCompositeNormalTexture, ReferenceTexture);
+									NormalCompositeMipmapImage->SetSource(ResizeTextureByNumMips(CompositeNormalImage, MipsToSkip));
 									NormalCompositeMipmapImage->SetMipmapGenerationSettings(mu::EMipmapFilterType::MFT_SimpleAverage, mu::EAddressMode::None, 1.0f, true);
 
 									CompositedImage->SetNormal(NormalCompositeMipmapImage);
@@ -1154,25 +1150,12 @@ mu::NodeSurfacePtr GenerateMutableSourceSurface(const UEdGraphPin * Pin, FMutabl
 					{
 						check(ParentMaterialNode->IsImageMutableMode(ImageIndex)); // Ensured at graph time. If it fails, something is wrong.
 						
-						float MaxTextureSize = 0.f;
-
-						// The static_cast is correct because we now it's a material node due to the ParentMaterialNodeCast
-						if (const mu::NodeSurfaceNewPtr ParentNode2 = static_cast<mu::NodeSurfaceNew*>(ParentNode.get()))
-						{
-							FString Aux(ParentNode2->GetImageName(ImageIndex));
-
-							if (!Aux.IsEmpty())
-							{
-								check(Aux.IsNumeric());
-								const int32 ImagePropertiesIndex = FCString::Atoi(*Aux);
-								check(GenerationContext.ImageProperties.IsValidIndex(ImagePropertiesIndex));
-								MaxTextureSize = GenerationContext.ImageProperties[ImagePropertiesIndex].MaxTextureSize;
-							}
-						}
-
 						if (const UEdGraphPin* ConnectedPin = FollowInputPin(*TypedNodeExt->GetUsedImagePin(ImageId)))
 						{
-							ImageNode = GenerateMutableSourceImage(ConnectedPin, GenerationContext, DummySurfaceData.ImageProperties.MaxTextureSize);
+							// ReferenceTextureSize is used to limit the size of textures contributing to the final image.
+							const int32 ReferenceTextureSize = GetBaseTextureSize(GenerationContext, ParentMaterialNode, ImageIndex);
+
+							ImageNode = GenerateMutableSourceImage(ConnectedPin, GenerationContext, ReferenceTextureSize);
 						}
 					}
 				}
@@ -1350,19 +1333,10 @@ mu::NodeSurfacePtr GenerateMutableSourceSurface(const UEdGraphPin * Pin, FMutabl
 					ImagePatchNode->SetBlendType(mu::EBlendType::BT_BLEND);
 					ImagePatchNode->SetApplyToAlphaChannel(true);
 
-					// Calculate the LODBias for this texture
-					UTexture2D* ReferenceTexture = ParentMaterialNode->GetImageReferenceTexture(ImageIndex);
-
-					if (ReferenceTexture)
-					{
-						GenerationContext.AddParticipatingObject(*ReferenceTexture);
-					}
-					
-					int32 LODBias = ComputeLODBias(GenerationContext, ReferenceTexture, ReferenceTexture ? ReferenceTexture->MaxTextureSize : 0, ParentMaterialNode, ImageIndex);
-
-					GenerationContext.CurrentTextureLODBias = LODBias;
+					// ReferenceTextureSize is used to limit the size of textures contributing to the final image.
+					const int32 ReferenceTextureSize = GetBaseTextureSize(GenerationContext, ParentMaterialNode, ImageIndex);
 
-					mu::NodeImagePtr ImageNode = GenerateMutableSourceImage(ConnectedImagePin, GenerationContext, 0.f);
+					mu::NodeImagePtr ImageNode = GenerateMutableSourceImage(ConnectedImagePin, GenerationContext, ReferenceTextureSize);
 					ImagePatchNode->SetImage(ImageNode);
 
 					const UEdGraphPin* ImageMaskPin = TypedNodeEdit->GetUsedImageMaskPin(ImageId);
@@ -1370,7 +1344,7 @@ mu::NodeSurfacePtr GenerateMutableSourceSurface(const UEdGraphPin * Pin, FMutabl
 					
 					if (const UEdGraphPin* ConnectedMaskPin = FollowInputPin(*ImageMaskPin))
 					{
-						mu::NodeImagePtr MaskNode = GenerateMutableSourceImage(ConnectedMaskPin, GenerationContext, 0.f);
+						mu::NodeImagePtr MaskNode = GenerateMutableSourceImage(ConnectedMaskPin, GenerationContext, ReferenceTextureSize);
 						ImagePatchNode->SetMask(MaskNode);
 					}
 
