#rnx
#tests PIE

diff --git a/Engine/Plugins/Experimental/Mutable/Source/CustomizableObject/Public/MuCO/CustomizableObjectCustomVersion.h b/Engine/Plugins/Experimental/Mutable/Source/CustomizableObject/Public/MuCO/CustomizableObjectCustomVersion.h
--- a/Engine/Plugins/Experimental/Mutable/Source/CustomizableObject/Public/MuCO/CustomizableObjectCustomVersion.h
+++ b/Engine/Plugins/Experimental/Mutable/Source/CustomizableObject/Public/MuCO/CustomizableObjectCustomVersion.h
@@ -110,6 +110,8 @@ struct CUSTOMIZABLEOBJECT_API FCustomizableObjectCustomVersion
 
 		AddedColumnIdDataToTableNodePins,
 
+		AddedAnyTextureTypeToPassThroughTextures,
+
 		// -----<new versions can be added above this line>-------------------------------------------------
 		VersionPlusOne,
 		LatestVersion = VersionPlusOne - 1
diff --git a/Engine/Plugins/Experimental/Mutable/Source/CustomizableObjectEditor/Classes/MuCOE/Nodes/CustomizableObjectNodePassThroughTexture.h b/Engine/Plugins/Experimental/Mutable/Source/CustomizableObjectEditor/Classes/MuCOE/Nodes/CustomizableObjectNodePassThroughTexture.h
--- a/Engine/Plugins/Experimental/Mutable/Source/CustomizableObjectEditor/Classes/MuCOE/Nodes/CustomizableObjectNodePassThroughTexture.h
+++ b/Engine/Plugins/Experimental/Mutable/Source/CustomizableObjectEditor/Classes/MuCOE/Nodes/CustomizableObjectNodePassThroughTexture.h
@@ -7,12 +7,18 @@
 #include "CustomizableObjectNodePassThroughTexture.generated.h"
 
 
-UCLASS()
+UCLASS(hidecategories = ("Texture2D"))
 class CUSTOMIZABLEOBJECTEDITOR_API UCustomizableObjectNodePassThroughTexture : public UCustomizableObjectNodeTextureBase
 {
 public:
 	GENERATED_BODY()
 
+	UPROPERTY(EditAnywhere, Category = Texture, Meta = (DisplayName = Texture))
+	TObjectPtr<UTexture> PassThroughTexture = nullptr;
+
+	// UCustomizableObjectNode interface
+	virtual void BackwardsCompatibleFixup() override;
+
 	// UCustomizableObjectNode interface
 	void AllocateDefaultPins(UCustomizableObjectNodeRemapPins* RemapPins) override;
 
diff --git a/Engine/Plugins/Experimental/Mutable/Source/CustomizableObjectEditor/Classes/MuCOE/Nodes/CustomizableObjectNodeTable.h b/Engine/Plugins/Experimental/Mutable/Source/CustomizableObjectEditor/Classes/MuCOE/Nodes/CustomizableObjectNodeTable.h
--- a/Engine/Plugins/Experimental/Mutable/Source/CustomizableObjectEditor/Classes/MuCOE/Nodes/CustomizableObjectNodeTable.h
+++ b/Engine/Plugins/Experimental/Mutable/Source/CustomizableObjectEditor/Classes/MuCOE/Nodes/CustomizableObjectNodeTable.h
@@ -101,8 +101,10 @@ public:
 	bool IsDefaultImageMode() { return bIsDefault; }
 	void SetDefaultImageMode(bool bValue) { bIsDefault = bValue; }
 
-	bool IsArrayTexture() { return bIsArrayTexture; }
-	void SetIsArrayTexture(bool bValue) { bIsArrayTexture = bValue; }
+	bool IsNotTexture2D() { return bIsNotTexture2D; }
+	void SetIsNotTexture2D(bool bValue) { bIsNotTexture2D = bValue; }
+
+	void ConvertArrayTextureToAnyTextureFixup();
 
 	// Pin Type
 	UPROPERTY()
@@ -113,9 +115,12 @@ private:
 	UPROPERTY()
 	bool bIsDefault = true;
 
+	// Replaced by the more general bIsNot2DTexture
 	UPROPERTY()
-	bool bIsArrayTexture = false;
+	bool bIsArrayTexture_DEPRECATED = false;
 
+	UPROPERTY()
+	bool bIsNotTexture2D = false;
 };
 
 
@@ -362,8 +367,8 @@ public:
 	// Returns true if the pin is in the default mode (same as node)
 	bool IsImagePinDefault(const UEdGraphPin* Pin) const;
 
-	// Returns true if the pin is a texture array pin
-	bool IsImageArrayPin(const UEdGraphPin* Pin) const;
+	// Returns true if the pin is not a UTexture2D but a different kind of UTexture
+	bool IsNotTexture2DPin(const UEdGraphPin* Pin) const;
 
 	// Returns the image mode of the column
 	ETableTextureType GetColumnImageMode(const FString& ColumnName) const;
diff --git a/Engine/Plugins/Experimental/Mutable/Source/CustomizableObjectEditor/Classes/MuCOE/Nodes/CustomizableObjectNodeTexture.h b/Engine/Plugins/Experimental/Mutable/Source/CustomizableObjectEditor/Classes/MuCOE/Nodes/CustomizableObjectNodeTexture.h
--- a/Engine/Plugins/Experimental/Mutable/Source/CustomizableObjectEditor/Classes/MuCOE/Nodes/CustomizableObjectNodeTexture.h
+++ b/Engine/Plugins/Experimental/Mutable/Source/CustomizableObjectEditor/Classes/MuCOE/Nodes/CustomizableObjectNodeTexture.h
@@ -64,7 +64,7 @@ class CUSTOMIZABLEOBJECTEDITOR_API UCustomizableObjectNodeTextureBase : public U
 public:
 	GENERATED_BODY()
 
-	UPROPERTY(EditAnywhere, Category = CustomizableObject)
+	UPROPERTY(EditAnywhere, Category = Texture2D)
 	TObjectPtr<UTexture2D> Texture = nullptr;
 
 	// Determines if the Node is collapsed or not
diff --git a/Engine/Plugins/Experimental/Mutable/Source/CustomizableObjectEditor/Private/MuCOE/EdGraphSchema_CustomizableObject.cpp b/Engine/Plugins/Experimental/Mutable/Source/CustomizableObjectEditor/Private/MuCOE/EdGraphSchema_CustomizableObject.cpp
--- a/Engine/Plugins/Experimental/Mutable/Source/CustomizableObjectEditor/Private/MuCOE/EdGraphSchema_CustomizableObject.cpp
+++ b/Engine/Plugins/Experimental/Mutable/Source/CustomizableObjectEditor/Private/MuCOE/EdGraphSchema_CustomizableObject.cpp
@@ -867,7 +867,7 @@ void UEdGraphSchema_CustomizableObject::GetContextMenuActions(UToolMenu* Menu, U
 		const UCustomizableObjectNodeTable* TableNode = Cast<UCustomizableObjectNodeTable>(Context->Node);
 		UEdGraphPin* TexturePin = (UEdGraphPin*)Context->Pin;
 
-		if (TableNode && TexturePin && !TableNode->IsImageArrayPin(TexturePin) && !TexturePin->LinkedTo.Num() && (TexturePin->PinType.PinCategory == PC_Image || TexturePin->PinType.PinCategory == PC_PassThroughImage))
+		if (TableNode && TexturePin && !TableNode->IsNotTexture2DPin(TexturePin) && !TexturePin->LinkedTo.Num() && (TexturePin->PinType.PinCategory == PC_Image || TexturePin->PinType.PinCategory == PC_PassThroughImage))
 		{
 			FText ActionText = FText::Format(LOCTEXT("ChangeTexturePinMode_Label", 
 				"Set as {0} texture"), Context->Pin->PinType.PinCategory == PC_PassThroughImage ? FText::FromString("Mutable") : FText::FromString("Pass Through"));
diff --git a/Engine/Plugins/Experimental/Mutable/Source/CustomizableObjectEditor/Private/MuCOE/GenerateMutableSource/GenerateMutableSource.cpp b/Engine/Plugins/Experimental/Mutable/Source/CustomizableObjectEditor/Private/MuCOE/GenerateMutableSource/GenerateMutableSource.cpp
--- a/Engine/Plugins/Experimental/Mutable/Source/CustomizableObjectEditor/Private/MuCOE/GenerateMutableSource/GenerateMutableSource.cpp
+++ b/Engine/Plugins/Experimental/Mutable/Source/CustomizableObjectEditor/Private/MuCOE/GenerateMutableSource/GenerateMutableSource.cpp
@@ -486,7 +486,7 @@ UTexture2D* FindReferenceImage(const UEdGraphPin* Pin, FMutableGraphGenerationCo
 
 	else if (const UCustomizableObjectNodePassThroughTexture* TypedNodePassThroughTex = Cast<UCustomizableObjectNodePassThroughTexture>(Node))
 	{
-		Result = TypedNodePassThroughTex->Texture;
+		Result = Cast<UTexture2D>(TypedNodePassThroughTex->PassThroughTexture);
 	}
 
 	else if (const UCustomizableObjectNodeTextureParameter* ParamNodeTex = Cast<UCustomizableObjectNodeTextureParameter>(Node))
diff --git a/Engine/Plugins/Experimental/Mutable/Source/CustomizableObjectEditor/Private/MuCOE/GenerateMutableSource/GenerateMutableSourceImage.cpp b/Engine/Plugins/Experimental/Mutable/Source/CustomizableObjectEditor/Private/MuCOE/GenerateMutableSource/GenerateMutableSourceImage.cpp
--- a/Engine/Plugins/Experimental/Mutable/Source/CustomizableObjectEditor/Private/MuCOE/GenerateMutableSource/GenerateMutableSourceImage.cpp
+++ b/Engine/Plugins/Experimental/Mutable/Source/CustomizableObjectEditor/Private/MuCOE/GenerateMutableSource/GenerateMutableSourceImage.cpp
@@ -869,7 +869,7 @@ mu::NodeImagePtr GenerateMutableSourceImage(const UEdGraphPin* Pin, FMutableGrap
 
 	else if (const UCustomizableObjectNodePassThroughTexture* TypedNodePassThroughTex = Cast<UCustomizableObjectNodePassThroughTexture>(Node))
 	{
-		UTexture2D* BaseTexture = TypedNodePassThroughTex->Texture;
+		UTexture* BaseTexture = TypedNodePassThroughTex->PassThroughTexture;
 		if (BaseTexture)
 		{
 			mu::NodeImageConstantPtr ImageNode = new mu::NodeImageConstant();
@@ -1002,10 +1002,9 @@ mu::NodeImagePtr GenerateMutableSourceImage(const UEdGraphPin* Pin, FMutableGrap
 					bSuccess = false;
 				}
 
-				UTexture2D* DefaultTexture2D = TypedNodeTable->GetColumnDefaultAssetByType<UTexture2D>(Pin);
-				UTexture2DArray* DefaultTextureArray = TypedNodeTable->GetColumnDefaultAssetByType<UTexture2DArray>(Pin);
+				UTexture* DefaultTexture = TypedNodeTable->GetColumnDefaultAssetByType<UTexture>(Pin);
 
-				if (bSuccess && Pin->PinType.PinCategory != Schema->PC_MaterialAsset && !DefaultTexture2D && !DefaultTextureArray)
+				if (bSuccess && Pin->PinType.PinCategory != Schema->PC_MaterialAsset && !DefaultTexture)
 				{
 					FString Msg = FString::Printf(TEXT("Couldn't find a default value in the data table's struct for the column [%s]. The default value is null or not a supported Texture"), *ColumnName);
 					GenerationContext.Compiler->CompilerLog(FText::FromString(Msg), Node);
@@ -1080,7 +1079,7 @@ mu::NodeImagePtr GenerateMutableSourceImage(const UEdGraphPin* Pin, FMutableGrap
 							ImageTableNode->SetParameterName(TypedNodeTable->ParameterName);
 							ImageTableNode->SetNoneOption(TypedNodeTable->bAddNoneOption);
 
-							if (DefaultTexture2D)
+							if (UTexture2D* DefaultTexture2D = Cast<UTexture2D>(DefaultTexture))
 							{
 								mu::FImageDesc ImageDesc = GenerateImageDescriptor(DefaultTexture2D);
 
diff --git a/Engine/Plugins/Experimental/Mutable/Source/CustomizableObjectEditor/Private/MuCOE/GenerateMutableSource/GenerateMutableSourceTable.cpp b/Engine/Plugins/Experimental/Mutable/Source/CustomizableObjectEditor/Private/MuCOE/GenerateMutableSource/GenerateMutableSourceTable.cpp
--- a/Engine/Plugins/Experimental/Mutable/Source/CustomizableObjectEditor/Private/MuCOE/GenerateMutableSource/GenerateMutableSourceTable.cpp
+++ b/Engine/Plugins/Experimental/Mutable/Source/CustomizableObjectEditor/Private/MuCOE/GenerateMutableSource/GenerateMutableSourceTable.cpp
@@ -392,17 +392,14 @@ bool FillTableColumn(const UCustomizableObjectNodeTable* TableNode,	mu::TablePtr
 
 		else if (SoftObjectProperty->PropertyClass->IsChildOf(UTexture::StaticClass()))
 		{
-			// Two supported texture types
-			UTexture2D* Texture2D = Cast<UTexture2D>(Object);
-			UTexture2DArray* TextureArray = Cast<UTexture2DArray>(Object);
+			UTexture* Texture = Cast<UTexture>(Object);
 
 			// Removing encoding part
 			const FString PinName = ColumnName.Replace(TEXT("--PassThrough"), TEXT(""), ESearchCase::CaseSensitive);
 
-			if (!Texture2D && !TextureArray)
+			if (!Texture)
 			{
-				Texture2D = TableNode->GetColumnDefaultAssetByType<UTexture2D>(PinName);
-				TextureArray = TableNode->GetColumnDefaultAssetByType<UTexture2DArray>(PinName);
+				Texture = TableNode->GetColumnDefaultAssetByType<UTexture>(PinName);
 
 				FString Message = Cast<UObject>(Object) ? "not a suported Texture" : "null";
 				FString WarningMessage = FString::Printf(TEXT("Texture from column [%s] row [%s] is %s. The default texture will be used instead."), *PinName, *RowName, *Message);
@@ -410,8 +407,7 @@ bool FillTableColumn(const UCustomizableObjectNodeTable* TableNode,	mu::TablePtr
 			}
 
 			// There will be always one of the two options
-			check(Texture2D || TextureArray);
-			UTexture* Texture = Texture2D ? Cast<UTexture>(Texture2D) : Cast<UTexture>(TextureArray);
+			check(Texture);
 
 			// Getting column index from column name
 			CurrentColumn = MutableTable->FindColumn(ColumnName);
diff --git a/Engine/Plugins/Experimental/Mutable/Source/CustomizableObjectEditor/Private/MuCOE/Nodes/CustomizableObjectNodePassThroughTexture.cpp b/Engine/Plugins/Experimental/Mutable/Source/CustomizableObjectEditor/Private/MuCOE/Nodes/CustomizableObjectNodePassThroughTexture.cpp
--- a/Engine/Plugins/Experimental/Mutable/Source/CustomizableObjectEditor/Private/MuCOE/Nodes/CustomizableObjectNodePassThroughTexture.cpp
+++ b/Engine/Plugins/Experimental/Mutable/Source/CustomizableObjectEditor/Private/MuCOE/Nodes/CustomizableObjectNodePassThroughTexture.cpp
@@ -7,6 +7,27 @@
 #define LOCTEXT_NAMESPACE "CustomizableObjectEditor"
 
 
+void UCustomizableObjectNodePassThroughTexture::BackwardsCompatibleFixup()
+{
+	Super::BackwardsCompatibleFixup();
+
+	const int32 CustomizableObjectCustomVersion = GetLinkerCustomVersion(FCustomizableObjectCustomVersion::GUID);
+
+	if (CustomizableObjectCustomVersion < FCustomizableObjectCustomVersion::AddedAnyTextureTypeToPassThroughTextures)
+	{
+		if (Texture)
+		{
+			if (!PassThroughTexture)
+			{
+				PassThroughTexture = Texture;
+			}
+
+			Texture = nullptr;
+		}
+	}
+}
+
+
 void UCustomizableObjectNodePassThroughTexture::AllocateDefaultPins(UCustomizableObjectNodeRemapPins* RemapPins)
 {
 	const UEdGraphSchema_CustomizableObject* Schema = GetDefault<UEdGraphSchema_CustomizableObject>();
diff --git a/Engine/Plugins/Experimental/Mutable/Source/CustomizableObjectEditor/Private/MuCOE/Nodes/CustomizableObjectNodeTable.cpp b/Engine/Plugins/Experimental/Mutable/Source/CustomizableObjectEditor/Private/MuCOE/Nodes/CustomizableObjectNodeTable.cpp
--- a/Engine/Plugins/Experimental/Mutable/Source/CustomizableObjectEditor/Private/MuCOE/Nodes/CustomizableObjectNodeTable.cpp
+++ b/Engine/Plugins/Experimental/Mutable/Source/CustomizableObjectEditor/Private/MuCOE/Nodes/CustomizableObjectNodeTable.cpp
@@ -164,6 +164,19 @@ void UCustomizableObjectNodeTable::BackwardsCompatibleFixup()
 			}
 		}
 	}
+
+	if (CustomizableObjectCustomVersion < FCustomizableObjectCustomVersion::AddedAnyTextureTypeToPassThroughTextures)
+	{
+		for (UEdGraphPin* Pin : Pins)
+		{
+			UCustomizableObjectNodePinData* PinData = GetPinData(*Pin);
+
+			if (UCustomizableObjectNodeTableImagePinData* ImagePinData = Cast<UCustomizableObjectNodeTableImagePinData>(PinData))
+			{
+				ImagePinData->ConvertArrayTextureToAnyTextureFixup();
+			}
+		}
+	}
 }
 
 
@@ -250,7 +263,7 @@ void UCustomizableObjectNodeTable::PinConnectionListChanged(UEdGraphPin* Pin)
 	{
 		if (UCustomizableObjectNodeTableImagePinData* ImagePinData = Cast<UCustomizableObjectNodeTableImagePinData>(GetPinData(*(Pin))))
 		{
-			if (!ImagePinData->IsArrayTexture() && ImagePinData->IsDefaultImageMode() && ImagePinData->ImageMode != DefaultImageMode)
+			if (!ImagePinData->IsNotTexture2D() && ImagePinData->IsDefaultImageMode() && ImagePinData->ImageMode != DefaultImageMode)
 			{
 				ImagePinData->ImageMode = DefaultImageMode;
 				ReconstructNode();
@@ -280,7 +293,7 @@ void UCustomizableObjectNodeTable::PostEditChangeProperty(FPropertyChangedEvent&
 			{
 				if (UCustomizableObjectNodeTableImagePinData* TexturePinData = Cast<UCustomizableObjectNodeTableImagePinData>(GetPinData(*(Pin))))
 				{
-					if (!Pin->LinkedTo.Num() && !TexturePinData->IsArrayTexture() && TexturePinData->IsDefaultImageMode())
+					if (!Pin->LinkedTo.Num() && !TexturePinData->IsNotTexture2D() && TexturePinData->IsDefaultImageMode())
 					{
 						ETableTextureType Mode = TexturePinData->ImageMode == ETableTextureType::PASSTHROUGH_TEXTURE ? ETableTextureType::MUTABLE_TEXTURE : ETableTextureType::PASSTHROUGH_TEXTURE;
 						TexturePinData->ImageMode = Mode;
@@ -383,7 +396,7 @@ void UCustomizableObjectNodeTable::AllocateDefaultPins(UCustomizableObjectNodeRe
 					UCustomizableObjectNodeTableImagePinData* PinData = NewObject<UCustomizableObjectNodeTableImagePinData>(this);
 					PinData->ColumnName = ColumnName;
 					PinData->StructColumnId = ColumnPropertyId;
-					PinData->SetIsArrayTexture(false);
+					PinData->SetIsNotTexture2D(false);
 
 					FName PinCategory = DefaultImageMode == ETableTextureType::PASSTHROUGH_TEXTURE ? Schema->PC_PassThroughImage : Schema->PC_Image;
 
@@ -403,13 +416,13 @@ void UCustomizableObjectNodeTable::AllocateDefaultPins(UCustomizableObjectNodeRe
 					OutPin = CustomCreatePin(EGPD_Output, PinCategory, FName(*PinName), PinData);
 				}
 
-				else if (Object->IsA(UTexture2DArray::StaticClass()))
+				else if (Object->IsA(UTexture::StaticClass()))
 				{
 					UCustomizableObjectNodeTableImagePinData* PinData = NewObject<UCustomizableObjectNodeTableImagePinData>(this);
 					PinData->ColumnName = ColumnName;
 					PinData->StructColumnId = ColumnPropertyId;
 					PinData->ImageMode = ETableTextureType::PASSTHROUGH_TEXTURE;
-					PinData->SetIsArrayTexture(true);
+					PinData->SetIsNotTexture2D(true);
 
 					OutPin = CustomCreatePin(EGPD_Output, Schema->PC_PassThroughImage, FName(*PinName), PinData);
 				}
@@ -806,7 +819,7 @@ void UCustomizableObjectNodeTable::RemapPinsData(const TMap<UEdGraphPin*, UEdGra
 			{
 				PinDataNewPin->ImageMode = PinDataOldPin->ImageMode;
 				PinDataNewPin->SetDefaultImageMode(PinDataOldPin->IsDefaultImageMode());
-				PinDataNewPin->SetIsArrayTexture(PinDataOldPin->IsArrayTexture());
+				PinDataNewPin->SetIsNotTexture2D(PinDataOldPin->IsNotTexture2D());
 			}
 		}
 	}
@@ -1191,11 +1204,21 @@ TArray<FName> UCustomizableObjectNodeTable::GetRowNames(const UDataTable* DataTa
 }
 
 
+void UCustomizableObjectNodeTableImagePinData::ConvertArrayTextureToAnyTextureFixup()
+{
+	if (bIsArrayTexture_DEPRECATED)
+	{
+		bIsNotTexture2D = true;
+		bIsArrayTexture_DEPRECATED = false;
+	}
+}
+
+
 void UCustomizableObjectNodeTable::ChangeImagePinMode(UEdGraphPin* Pin, bool bSetDefault)
 {
 	if (UCustomizableObjectNodeTableImagePinData* PinData = Cast<UCustomizableObjectNodeTableImagePinData>(GetPinData(*(Pin))))
 	{
-		if (PinData->IsArrayTexture())
+		if (PinData->IsNotTexture2D())
 		{
 			return;
 		}
@@ -1244,17 +1267,6 @@ bool UCustomizableObjectNodeTable::IsImagePinDefault(const UEdGraphPin* Pin) con
 }
 
 
-bool UCustomizableObjectNodeTable::IsImageArrayPin(const UEdGraphPin* Pin) const
-{
-	if (UCustomizableObjectNodeTableImagePinData* PinData = Cast<UCustomizableObjectNodeTableImagePinData>(GetPinData(*(Pin))))
-	{
-		return PinData->IsArrayTexture();
-	}
-
-	return false;
-}
-
-
 ETableTextureType UCustomizableObjectNodeTable::GetColumnImageMode(const FString& ColumnName) const
 {
 	for (const UEdGraphPin* Pin : Pins)
@@ -1272,6 +1284,17 @@ ETableTextureType UCustomizableObjectNodeTable::GetColumnImageMode(const FString
 }
 
 
+bool UCustomizableObjectNodeTable::IsNotTexture2DPin(const UEdGraphPin* Pin) const
+{
+	if (UCustomizableObjectNodeTableImagePinData* PinData = Cast<UCustomizableObjectNodeTableImagePinData>(GetPinData(*(Pin))))
+	{
+		return PinData->IsNotTexture2D();
+	}
+
+	return false;
+}
+
+
 ETableMeshPinType UCustomizableObjectNodeTable::GetPinMeshType(const UEdGraphPin* Pin) const
 {
 	// Getting Struct Pointer
