[Mutable] Fix static analysis warning. GenerationContext.Options.TargetPlatform pointer was utilized before it was verified against nullptr.

#rnx

diff --git a/Engine/Plugins/Experimental/Mutable/Source/CustomizableObjectEditor/Private/MuCOE/GenerateMutableSource/GenerateMutableSourceSurface.cpp b/Engine/Plugins/Experimental/Mutable/Source/CustomizableObjectEditor/Private/MuCOE/GenerateMutableSource/GenerateMutableSourceSurface.cpp
--- a/Engine/Plugins/Experimental/Mutable/Source/CustomizableObjectEditor/Private/MuCOE/GenerateMutableSource/GenerateMutableSourceSurface.cpp
+++ b/Engine/Plugins/Experimental/Mutable/Source/CustomizableObjectEditor/Private/MuCOE/GenerateMutableSource/GenerateMutableSourceSurface.cpp
@@ -679,26 +679,29 @@ mu::NodeSurfacePtr GenerateMutableSourceSurface(const UEdGraphPin * Pin, FMutabl
 							LastImage = FormatImage;
 
 							TArray<TArray<FTextureBuildSettings>> BuildSettingsPerFormatPerLayer;
-							ReferenceTexture->GetTargetPlatformBuildSettings(GenerationContext.Options.TargetPlatform, BuildSettingsPerFormatPerLayer);
-							if (BuildSettingsPerFormatPerLayer.IsEmpty())
+							if (GenerationContext.Options.TargetPlatform)
 							{
-								const FString ReplacedImageFormatMsg = FString::Printf(TEXT("In object [%s] for platform [%s] the unsupported image format of texture [%s] is used, IF_RGBA_UBYTE will be used instead."),
-									*GenerationContext.Object->GetName(),
-									*GenerationContext.Options.TargetPlatform->PlatformName(),
-									*ReferenceTexture->GetName());
-								const FText ReplacedImageFormatText = FText::FromString(ReplacedImageFormatMsg);
-								GenerationContext.Compiler->CompilerLog(ReplacedImageFormatText, Node, EMessageSeverity::Info);
-								UE_LOG(LogMutable, Log, TEXT("%s"), *ReplacedImageFormatMsg);
-							}
-							else if (BuildSettingsPerFormatPerLayer.Num()>1)
-							{
-								const FString ReplacedImageFormatMsg = FString::Printf(TEXT("In object [%s] for platform [%s] the image format of texture [%s] has multiple target formats. Only one will be used.."),
-									*GenerationContext.Object->GetName(),
-									*GenerationContext.Options.TargetPlatform->PlatformName(),
-									*ReferenceTexture->GetName());
-								const FText ReplacedImageFormatText = FText::FromString(ReplacedImageFormatMsg);
-								GenerationContext.Compiler->CompilerLog(ReplacedImageFormatText, Node, EMessageSeverity::Info);
-								UE_LOG(LogMutable, Log, TEXT("%s"), *ReplacedImageFormatMsg);
+								ReferenceTexture->GetTargetPlatformBuildSettings(GenerationContext.Options.TargetPlatform, BuildSettingsPerFormatPerLayer);
+								if (BuildSettingsPerFormatPerLayer.IsEmpty())
+								{
+									const FString ReplacedImageFormatMsg = FString::Printf(TEXT("In object [%s] for platform [%s] the unsupported image format of texture [%s] is used, IF_RGBA_UBYTE will be used instead."),
+										*GenerationContext.Object->GetName(),
+										*GenerationContext.Options.TargetPlatform->PlatformName(),
+										*ReferenceTexture->GetName());
+									const FText ReplacedImageFormatText = FText::FromString(ReplacedImageFormatMsg);
+									GenerationContext.Compiler->CompilerLog(ReplacedImageFormatText, Node, EMessageSeverity::Info);
+									UE_LOG(LogMutable, Log, TEXT("%s"), *ReplacedImageFormatMsg);
+								}
+								else if (BuildSettingsPerFormatPerLayer.Num() > 1)
+								{
+									const FString ReplacedImageFormatMsg = FString::Printf(TEXT("In object [%s] for platform [%s] the image format of texture [%s] has multiple target formats. Only one will be used.."),
+										*GenerationContext.Object->GetName(),
+										*GenerationContext.Options.TargetPlatform->PlatformName(),
+										*ReferenceTexture->GetName());
+									const FText ReplacedImageFormatText = FText::FromString(ReplacedImageFormatMsg);
+									GenerationContext.Compiler->CompilerLog(ReplacedImageFormatText, Node, EMessageSeverity::Info);
+									UE_LOG(LogMutable, Log, TEXT("%s"), *ReplacedImageFormatMsg);
+								}
 							}
 
 							if (!BuildSettingsPerFormatPerLayer.IsEmpty())
