[mutable] Added commandlet to test CO Instnace baking
- Moved Baking code to CustomizableObjectBakingUtils to be accessible from the testing commandlet.
- This commandled requires a COI.
- Compiles the COI's CO, updates the COI and then bakes it onto disk

diff --git a/Engine/Plugins/Experimental/Mutable/Source/CustomizableObjectEditor/Private/MuCOE/CustomizableObjectEditorViewportClient.cpp b/Engine/Plugins/Experimental/Mutable/Source/CustomizableObjectEditor/Private/MuCOE/CustomizableObjectEditorViewportClient.cpp
--- a/Engine/Plugins/Experimental/Mutable/Source/CustomizableObjectEditor/Private/MuCOE/CustomizableObjectEditorViewportClient.cpp
+++ b/Engine/Plugins/Experimental/Mutable/Source/CustomizableObjectEditor/Private/MuCOE/CustomizableObjectEditorViewportClient.cpp
@@ -1602,239 +1602,6 @@ protected:
 
 //-------------------------------------------------------------------------------------------------
 
-// TODO: Move to CustomizableObjectInstanceBakingUtils.cpp after review to be able to enable testing commandlet
-namespace 
-{
-	/**
-	 * Simple wrapper to be able to invoke the generation of a popup or log message depending on the execution context in which this code is being ran
-	 * @param InMessage The message to display
-	 * @param InTitle The title to be used for the popup or the log generated
-	 */
-	void ShowErrorNotification(const FText& InMessage,  const FText* InTitle = nullptr)
-	{
-		if (!FApp::IsUnattended())
-		{
-			FMessageDialog::Open(EAppMsgType::Ok, InMessage, *InTitle);
-		}
-		else
-		{
-			if (InTitle)
-			{
-				UE_LOG(LogMutable, Error, TEXT("%s - %s"), *InTitle->ToString(), *InMessage.ToString());
-			}
-			else
-			{
-				UE_LOG(LogMutable, Error, TEXT("%s"), *InMessage.ToString());
-			}
-		}
-	}
-
-	/**
-	 * Utility functions for the baking operation.
-	 */
-
-	/**
-	 * Validates the filename chosen for the baking data
-	 * @param FileName The filename chosen by the user
-	 * @return True if validation was successful, false otherwise
-	 */
-	bool ValidateProvidedFileName(const FString& FileName)
-	{
-		if (FileName.IsEmpty())
-		{
-			UE_LOG(LogMutable, Error, TEXT("Invalid baking configuration : FileName string is empty.."));
-			return false;
-		}
-
-		// Check for invalid characters in the name of the object to be serialized
-		{
-			TCHAR InvalidCharacter = '0';
-			{
-				FString InvalidCharacters = FPaths::GetInvalidFileSystemChars();
-				for (int32 i = 0; i < InvalidCharacters.Len(); ++i)
-				{
-					TCHAR Char = InvalidCharacters[i];
-					FString SearchedChar = FString::Chr(Char);
-					if (FileName.Contains(SearchedChar))
-					{
-						InvalidCharacter = InvalidCharacters[i];
-						break;
-					}
-				}
-			}
-
-			if (InvalidCharacter != '0')
-			{
-				const FText InvalidCharacterText = FText::FromString(FString::Chr(InvalidCharacter));
-				const FText ErrorText = FText::Format(LOCTEXT("FCustomizableObjectEditorViewportClient_BakeInstance_InvalidCharacter", "The selected contains an invalid character ({0})."), InvalidCharacterText);
-
-				ShowErrorNotification(ErrorText);
-			
-				return false;
-			}
-		}
-
-		return true;
-	}
-
-
-	/**
-	 * Validates the AssetPath chosen for the baking data
-	 * @param FileName The filename chosen by the user
-	 * @param AssetPath The AssetPath chosen by the user
-	 * @param InstanceCO The CustomizableObject from the provided COI
-	 * @return True if validation was successful, false otherwise
-	 */
-	bool ValidateProvidedAssetPath(const FString& FileName, const FString& AssetPath, const UCustomizableObject* InstanceCO)
-	{
-		if (AssetPath.IsEmpty())
-		{
-			UE_LOG(LogMutable, Error, TEXT("The AssetPath can not be empty!"));
-			return false;
-		}
-
-		// Ensure we are not overriding the parent CO
-		const FString FullAssetPath = AssetPath + FString("/") + FileName + FString(".") + FileName;		// Full asset path to the new asset we want to create
-		const bool bWouldOverrideParentCO = InstanceCO->GetPathName() == FullAssetPath;
-		if (bWouldOverrideParentCO)
-		{
-			const FText ErrorText = LOCTEXT("FCustomizableObjectEditorViewportClient_BakeInstance_OverwriteCO", "The selected path would overwrite the instance's parent Customizable Object.");
-
-			ShowErrorNotification(ErrorText);
-			
-			return false;
-		}
-
-		return true;
-	}
-
-
-	/**
-	 * Outputs a string that we know it is unique.
-	 * @param InResource The resource we are working with
-	 * @param ResourceName The name of the resource we have provided. This should have the name of the current resource and will have the unique name for the resource once the method exits
-	 * @param InCachedResources Collection with all the already processed resources.
-	 * @param InCachedResourceNames Collection with all the already processed resources name's
-	 * @return True if the generation of the unique resource name was successful, false otherwise.
-	 */
-	bool GetUniqueResourceName(const UObject* InResource, FString& ResourceName, TArray<UObject*>& InCachedResources, const TArray<FString>& InCachedResourceNames)
-	{
-		check(InResource);
-
-		int32 FindResult = InCachedResourceNames.Find(ResourceName);
-		if (FindResult != INDEX_NONE)
-		{
-			if (InResource == InCachedResources[FindResult])
-			{
-				return false;
-			}
-
-			uint32 Count = 0;
-			while (FindResult != INDEX_NONE)
-			{
-				FindResult = InCachedResourceNames.Find(ResourceName + "_" + FString::FromInt(Count));
-				Count++;
-			}
-
-			ResourceName += "_" + FString::FromInt(--Count);
-		}
-
-		return true;
-	}
-
-
-	/**
-	 * Ensures the resource we want to save is ready to be saved. It handles closing it's editor and warning the user about possible overriding of resources.
-	 * @param InAssetSavePath The directory path where to save the baked object
-	 * @param InObjName The name of the object to be baked
-	 * @param bUsedGrantedOverridingRights Control flag that determines if the user has given or not permission to override resources already in disk
-	 * @return True if the operation was successful, false otherwise
-	 */
-	bool ManageBakingAction(const FString& InAssetSavePath, const FString& InObjName, bool& bUsedGrantedOverridingRights)
-	{
-		FString PackagePath = InAssetSavePath + "/" + InObjName;
-		UPackage* ExistingPackage = FindPackage(NULL, *PackagePath);
-
-		if (!ExistingPackage)
-		{
-			FString PackageFilePath = PackagePath + "." + InObjName;
-
-			FString PackageFileName;
-			if (FPackageName::DoesPackageExist(PackageFilePath, &PackageFileName))
-			{
-				ExistingPackage = LoadPackage(nullptr, *PackageFileName, LOAD_EditorOnly);
-			}
-			else
-			{
-				// if package does not exists
-				bUsedGrantedOverridingRights = false;
-				return true;
-			}
-		}
-
-		if (ExistingPackage)
-		{
-			// Checking if the asset is open in an editor
-			TArray<IAssetEditorInstance*> ObjectEditors = GEditor->GetEditorSubsystem<UAssetEditorSubsystem>()->FindEditorsForAssetAndSubObjects(ExistingPackage);
-			if (ObjectEditors.Num())
-			{
-				for (IAssetEditorInstance* ObjectEditorInstance : ObjectEditors)
-				{
-					// Close the editors that contains this asset
-					if (!ObjectEditorInstance->CloseWindow(EAssetEditorCloseReason::AssetEditorHostClosed))
-					{
-						const FText Caption = LOCTEXT("OpenExisitngFile", "Open File");
-						const FText Message = FText::Format(LOCTEXT("CantCloseAsset", "This Obejct \"{0}\" is open in an editor and can't be closed automatically. Please close the editor and try to bake it again"), FText::FromString(InObjName));
-
-						ShowErrorNotification(Message, &Caption);
-						
-						return false;
-					}
-				}
-			}
-
-			if (!bUsedGrantedOverridingRights)
-			{
-				const FText Caption = LOCTEXT("Already existing baked files", "Already existing baked files");
-				const FText Message = FText::Format(LOCTEXT("OverwriteBakedInstance", "Instance baked files already exist in selected destination \"{0}\", this action will overwrite them."), FText::AsCultureInvariant(InAssetSavePath));
-
-				// We need to guard this case since it may crash the editor if we are running an unattended commandlet and we try to generate this dialog
-				if (!FApp::IsUnattended())
-				{
-					if (FMessageDialog::Open(EAppMsgType::OkCancel, Message, Caption) == EAppReturnType::Cancel)
-					{
-						return false;
-					}
-				}
-				else
-				{
-					UE_LOG(LogMutable, Error, TEXT("%s - %s"), *Caption.ToString(), *Message.ToString());
-				}
-
-
-				bUsedGrantedOverridingRights = true;
-			}
-
-			UObject* ExistingObject = StaticFindObject(UObject::StaticClass(), ExistingPackage, *InObjName);
-			if (ExistingObject)
-			{
-				ExistingPackage->FullyLoad();
-
-				TArray<UObject*> ObjectsToDelete;
-				ObjectsToDelete.Add(ExistingObject);
-
-				// Delete objects in the package with the same name as the one we want to create
-				const uint32 NumObjectsDeleted = ObjectTools::ForceDeleteObjects(ObjectsToDelete, false);
-
-				return NumObjectsDeleted == ObjectsToDelete.Num();
-			}
-		}
-
-		return true;
-	}
-}
-
-
 
 void FCustomizableObjectEditorViewportClient::BakeInstance()
 {
@@ -1845,6 +1612,7 @@ void FCustomizableObjectEditorViewportClient::BakeInstance()
 		UE_LOG(LogMutable, Error, TEXT("No Mutable Customizable Object Instnace was found in the current editor."));
 		return;
 	}
+	
 	BakeTempInstance = Instance->Clone();
 
 	// Call the instance update async method
@@ -1854,474 +1622,6 @@ void FCustomizableObjectEditorViewportClient::BakeInstance()
 }
 
 
-void BakeCustomizableObjectInstance(
-	UCustomizableObjectInstance& InInstance,
-	const FString& FileName,
-	const FString& AssetPath,
-	const bool bExportAllResources,
-	const bool bGenerateConstantMaterialInstances)
-{
-	// Ensure that the state of the COI provided is valid --------------------------------------------------------------------------------------------
-	UCustomizableObject* InstanceCO = InInstance.GetCustomizableObject();
-	check (InstanceCO);
-
-	// Ensure the CO of the COI is accessible 
-	if (!InstanceCO || InstanceCO->IsLocked())
-	{
-		FCustomizableObjectEditorLogger::CreateLog(
-		LOCTEXT("CustomizableObjectCompilingTryLater_Baking", "Please wait until the Customizable Object is compiled"))
-		.Category(ELoggerCategory::COInstanceBaking)
-		.CustomNotification()
-		.Notification(true)
-		.Log();
-
-		return;
-	}
-	
-	if (InstanceCO->GetPrivate()->Status.Get() == FCustomizableObjectStatus::EState::Loading)
-	{
-		FCustomizableObjectEditorLogger::CreateLog(
-			LOCTEXT("CustomizableObjectCompileTryLater_BakeInstance","Please wait unitl Customizable Object is loaded"))
-		.Category(ELoggerCategory::COInstanceBaking)
-		.CustomNotification()
-		.Notification(true)
-		.Log();
-
-		return;
-	}
-	
-	if (!ValidateProvidedFileName(FileName))
-	{
-		UE_LOG(LogMutable, Error, TEXT("The FileName for the instance baking is not valid."));
-		return;
-	}
-
-	if (!ValidateProvidedAssetPath(FileName,AssetPath,InstanceCO))
-	{
-		UE_LOG(LogMutable, Error, TEXT("The AssetPath for the instance baking is not valid."));
-		return;
-	}
-	
-	// Exit early if the provided instance does not have a skeletal mesh
-	if (!InInstance.HasAnySkeletalMesh())
-	{
-		return;
-	}
-
-	// COI Validation completed : Proceed with the baking operation ----------------------------------------------------------------------------------
-	
-	// Notify of better configuration -> Continue operation normally
-	if (InstanceCO->CompileOptions.TextureCompression != ECustomizableObjectTextureCompression::HighQuality)
-	{
-		FCustomizableObjectEditorLogger::CreateLog(
-		LOCTEXT("CustomizableObjectBakeLowQuality", "The Customizable Object wasn't compiled with high quality textures. For the best baking results, change the Texture Compression setting and recompile it."))
-		.Category(ELoggerCategory::COInstanceBaking)
-		.CustomNotification()
-		.Notification(true)
-		.Log();
-	}
-	
-	
-	// Set the overriding flag to false wo we ask the user at least once about if he is willing to override old baked data
-	bool bUsedGrantedOverridingRights = false;
-
-	TArray<UPackage*> PackagesToSave;
-
-	const int32 NumComponents = InInstance.GetNumComponents();
-	for (int32 ComponentIndex = 0; ComponentIndex < NumComponents; ++ComponentIndex)
-	{
-		USkeletalMesh* Mesh = InInstance.GetSkeletalMesh(ComponentIndex);
-
-		if (!Mesh)
-		{
-			continue;
-		}
-
-		FString ObjectName = FileName;
-		if (NumComponents > 1)
-		{
-			ObjectName = FileName + "_Component_" + FString::FromInt(ComponentIndex);
-		}
-
-		TMap<UObject*, UObject*> ReplacementMap;
-		TArray<FString> ArrayCachedElement;
-		TArray<UObject*> ArrayCachedObject;
-
-		if (bExportAllResources)
-		{
-			UMaterialInstance* Inst;
-			UMaterial* Material;
-			UTexture* Texture;
-			FString MaterialName;
-			FString ResourceName;
-			FString PackageName;
-			UObject* DuplicatedObject;
-			TArray<TMap<int, UTexture*>> TextureReplacementMaps;
-
-			// Duplicate Mutable generated textures
-			for (int32 m = 0; m < Mesh->GetMaterials().Num(); ++m)
-			{
-				UMaterialInterface* Interface = Mesh->GetMaterials()[m].MaterialInterface;
-				Material = Interface->GetMaterial();
-				MaterialName = Material ? Material->GetName() : "Material";
-				Inst = Cast<UMaterialInstance>(Mesh->GetMaterials()[m].MaterialInterface);
-
-				TMap<int, UTexture*> ReplacementTextures;
-				TextureReplacementMaps.Add(ReplacementTextures);
-
-				// The material will only have Mutable generated textures if it's actually a UMaterialInstance
-				if (Material != nullptr && Inst != nullptr)
-				{
-					TArray<FName> ParameterNames = FUnrealBakeHelpers::GetTextureParameterNames(Material);
-
-					for (int32 i = 0; i < ParameterNames.Num(); i++)
-					{
-						if (Inst->GetTextureParameterValue(ParameterNames[i], Texture))
-						{
-							UTexture2D* SrcTex = Cast<UTexture2D>(Texture);
-							if (!SrcTex) continue;
-
-							bool bIsMutableTexture = false;
-
-							for (UAssetUserData* UserData : *SrcTex->GetAssetUserDataArray())
-							{
-								UTextureMipDataProviderFactory* CustomMipDataProviderFactory = Cast<UMutableTextureMipDataProviderFactory>(UserData);
-								if (CustomMipDataProviderFactory)
-								{
-									bIsMutableTexture = true;
-								}
-							}
-
-							if ((SrcTex->GetPlatformData() != nullptr) &&
-								(SrcTex->GetPlatformData()->Mips.Num() > 0) &&
-								bIsMutableTexture)
-							{
-								FString ParameterSanitized = ParameterNames[i].GetPlainNameString();
-								RemoveRestrictedChars(ParameterSanitized);
-								ResourceName = ObjectName + "_" + MaterialName + "_" + ParameterSanitized;
-
-								if (!GetUniqueResourceName(SrcTex, ResourceName, ArrayCachedObject, ArrayCachedElement))
-								{
-									continue;
-								}
-
-								if (!ManageBakingAction(AssetPath, ResourceName, bUsedGrantedOverridingRights))
-								{
-									return;
-								}
-
-								// Recover original name of the texture parameter value, now substituted by the generated Mutable texture
-								UTexture* OriginalTexture = nullptr;
-								UMaterialInstanceDynamic* InstDynamic = Cast<UMaterialInstanceDynamic>(Mesh->GetMaterials()[m].MaterialInterface);
-								if (InstDynamic != nullptr)
-								{
-									InstDynamic->Parent->GetTextureParameterValue(FName(*ParameterNames[i].GetPlainNameString()), OriginalTexture);
-								}
-								else
-								{
-									UMaterialInstanceConstant* InstConstant = Cast<UMaterialInstanceConstant>(Mesh->GetMaterials()[m].MaterialInterface);
-
-									if (InstConstant != nullptr)
-									{
-										InstConstant->Parent->GetTextureParameterValue(FName(*ParameterNames[i].GetPlainNameString()), OriginalTexture);
-									}
-								}
-
-								PackageName = AssetPath + FString("/") + ResourceName;
-								TMap<UObject*, UObject*> FakeReplacementMap;
-								UTexture2D* DupTex = FUnrealBakeHelpers::BakeHelper_CreateAssetTexture(SrcTex, ResourceName, PackageName, OriginalTexture, true, FakeReplacementMap, bUsedGrantedOverridingRights);
-								ArrayCachedElement.Add(ResourceName);
-								ArrayCachedObject.Add(DupTex);
-								PackagesToSave.Add(DupTex->GetPackage());
-
-								if (OriginalTexture != nullptr)
-								{
-									TextureReplacementMaps[m].Add(i, DupTex);
-								}
-							}
-						}
-					}
-				}
-			}
-
-			// Duplicate non-Mutable material textures
-			for (int32 m = 0; m < Mesh->GetMaterials().Num(); ++m)
-			{
-				UMaterialInterface* Interface = Mesh->GetMaterials()[m].MaterialInterface;
-				Material = Interface->GetMaterial();
-				MaterialName = Material ? Material->GetName() : "Material";
-
-				if (Material != nullptr)
-				{
-					TArray<FName> ParameterNames = FUnrealBakeHelpers::GetTextureParameterNames(Material);
-
-					for (int32 i = 0; i < ParameterNames.Num(); i++)
-					{
-						TArray<FMaterialParameterInfo> InfoArray;
-						TArray<FGuid> GuidArray;
-						Material->GetAllTextureParameterInfo(InfoArray, GuidArray);
-						
-						if (Material->GetTextureParameterValue(InfoArray[i], Texture))
-						{
-							FString ParameterSanitized = ParameterNames[i].GetPlainNameString();
-							RemoveRestrictedChars(ParameterSanitized);
-							ResourceName = ObjectName + "_" + MaterialName + "_" + ParameterSanitized;
-
-							if (ArrayCachedElement.Find(ResourceName) == INDEX_NONE)
-							{
-								if (!ManageBakingAction(AssetPath, ResourceName, bUsedGrantedOverridingRights))
-								{
-									return;
-								}
-
-								PackageName = AssetPath + FString("/") + ResourceName;
-								TMap<UObject*, UObject*> FakeReplacementMap;
-								DuplicatedObject = FUnrealBakeHelpers::BakeHelper_DuplicateAsset(Texture, ResourceName, PackageName, true, FakeReplacementMap, bUsedGrantedOverridingRights, false);
-								ArrayCachedElement.Add(ResourceName);
-								ArrayCachedObject.Add(DuplicatedObject);
-								PackagesToSave.Add(DuplicatedObject->GetPackage());
-
-								UTexture* DupTexture = Cast<UTexture>(DuplicatedObject);
-								TextureReplacementMaps[m].Add(i, DupTexture);
-							}
-						}
-					}
-				}
-			}
-
-
-			// Duplicate the materials used by each material instance so that the replacement map has proper information 
-			// when duplicating the material instances
-			for (int32 m = 0; m < Mesh->GetMaterials().Num(); ++m)
-			{
-				UMaterialInterface* Interface = Mesh->GetMaterials()[m].MaterialInterface;
-				Material = Interface ? Interface->GetMaterial() : nullptr;
-
-				if (Material)
-				{
-					ResourceName = ObjectName + "_Material_" + Material->GetName();
-
-					if (!GetUniqueResourceName(Material, ResourceName, ArrayCachedObject, ArrayCachedElement))
-					{
-						continue;
-					}
-
-					if (!ManageBakingAction(AssetPath, ResourceName, bUsedGrantedOverridingRights))
-					{
-						return;
-					}
-
-					PackageName = AssetPath + FString("/") + ResourceName;
-					TMap<UObject*, UObject*> FakeReplacementMap;
-					DuplicatedObject = FUnrealBakeHelpers::BakeHelper_DuplicateAsset(Material, ResourceName, PackageName, 
-						false, FakeReplacementMap, bUsedGrantedOverridingRights, bGenerateConstantMaterialInstances);
-					ArrayCachedElement.Add(ResourceName);
-					ArrayCachedObject.Add(DuplicatedObject);
-					ReplacementMap.Add(Interface, DuplicatedObject);
-					PackagesToSave.Add(DuplicatedObject->GetPackage());
-
-					FUnrealBakeHelpers::CopyAllMaterialParameters(DuplicatedObject, Interface, TextureReplacementMaps[m]);
-				}
-			}
-		}
-		else
-		{
-			// Duplicate the material instances
-			for (int32 MaterialIndex = 0; MaterialIndex < Mesh->GetMaterials().Num(); ++MaterialIndex)
-			{
-				UMaterialInterface* Interface = Mesh->GetMaterials()[MaterialIndex].MaterialInterface;
-				UMaterial* ParentMaterial = Interface->GetMaterial();
-				FString MaterialName = ParentMaterial ? ParentMaterial->GetName() : "Material";
-
-				// Material
-				FString MatObjName = ObjectName + "_" + MaterialName;
-
-				if (!GetUniqueResourceName(Interface, MatObjName, ArrayCachedObject, ArrayCachedElement))
-				{
-					continue;
-				}
-
-				if (!ManageBakingAction(AssetPath, MatObjName, bUsedGrantedOverridingRights))
-				{
-					return;
-				}
-
-				FString MatPkgName = AssetPath + FString("/") + MatObjName;
-				UObject* DupMat = FUnrealBakeHelpers::BakeHelper_DuplicateAsset(Interface, MatObjName, 
-					MatPkgName, false, ReplacementMap, bUsedGrantedOverridingRights, bGenerateConstantMaterialInstances);
-				ArrayCachedObject.Add(DupMat);
-				ArrayCachedElement.Add(MatObjName);
-				PackagesToSave.Add(DupMat->GetPackage());
-
-				UMaterialInstance* Inst = Cast<UMaterialInstance>(Interface);
-
-				// Only need to duplicate the generate textures if the original material is a dynamic instance
-				// If the material has Mutable textures, then it will be a dynamic material instance for sure
-				if (Inst)
-				{
-					// Duplicate generated textures
-					UMaterialInstanceDynamic* InstDynamic = Cast<UMaterialInstanceDynamic>(DupMat);
-					UMaterialInstanceConstant* InstConstant = Cast<UMaterialInstanceConstant>(DupMat);
-
-					if (InstDynamic || InstConstant)
-					{
-						for (int32 TextureIndex = 0; TextureIndex < Inst->TextureParameterValues.Num(); ++TextureIndex)
-						{
-							if (Inst->TextureParameterValues[TextureIndex].ParameterValue)
-							{
-								if (Inst->TextureParameterValues[TextureIndex].ParameterValue->HasAnyFlags(RF_Transient))
-								{
-									UTexture2D* SrcTex = Cast<UTexture2D>(Inst->TextureParameterValues[TextureIndex].ParameterValue);
-
-									if (SrcTex)
-									{
-										FString ParameterSanitized = Inst->TextureParameterValues[TextureIndex].ParameterInfo.Name.ToString();
-										RemoveRestrictedChars(ParameterSanitized);
-
-										FString TexObjName = ObjectName + "_" + MaterialName + "_" + ParameterSanitized;
-
-										if (!GetUniqueResourceName(SrcTex, TexObjName, ArrayCachedObject, ArrayCachedElement))
-										{
-											UTexture* PrevTexture = Cast<UTexture>(ArrayCachedObject[ArrayCachedElement.Find(TexObjName)]);
-
-											if (InstDynamic)
-											{
-												InstDynamic->SetTextureParameterValue(Inst->TextureParameterValues[TextureIndex].ParameterInfo.Name, PrevTexture);
-											}
-											else if (InstConstant)
-											{
-												InstConstant->SetTextureParameterValueEditorOnly(Inst->TextureParameterValues[TextureIndex].ParameterInfo.Name, PrevTexture);
-											}
-											
-											continue;
-										}
-
-										if (!ManageBakingAction(AssetPath, TexObjName, bUsedGrantedOverridingRights))
-										{
-											return;
-										}
-
-										FString TexPkgName = AssetPath + FString("/") + TexObjName;
-										TMap<UObject*, UObject*> FakeReplacementMap;
-										UTexture2D* DupTex = FUnrealBakeHelpers::BakeHelper_CreateAssetTexture(SrcTex, TexObjName, TexPkgName, nullptr, false, FakeReplacementMap, bUsedGrantedOverridingRights);
-										ArrayCachedObject.Add(DupTex);
-										ArrayCachedElement.Add(TexObjName);
-										PackagesToSave.Add(DupTex->GetPackage());
-
-										if (InstDynamic)
-										{
-											InstDynamic->SetTextureParameterValue(Inst->TextureParameterValues[TextureIndex].ParameterInfo.Name, DupTex);
-										}
-										else if(InstConstant)
-										{
-											InstConstant->SetTextureParameterValueEditorOnly(Inst->TextureParameterValues[TextureIndex].ParameterInfo.Name, DupTex);
-										}
-									}
-									else
-									{
-										UE_LOG(LogMutable, Error, TEXT("A Mutable texture that is not a Texture2D has been found while baking a CustomizableObjectInstance."));
-									}
-								}
-								else
-								{
-									// If it's not transient it's not a mutable texture, it's a pass-through texture
-									// Just set the original texture
-									if (InstDynamic)
-									{
-										InstDynamic->SetTextureParameterValue(Inst->TextureParameterValues[TextureIndex].ParameterInfo.Name, Inst->TextureParameterValues[TextureIndex].ParameterValue);
-									}
-									else if (InstConstant)
-									{
-										InstConstant->SetTextureParameterValueEditorOnly(Inst->TextureParameterValues[TextureIndex].ParameterInfo.Name, Inst->TextureParameterValues[TextureIndex].ParameterValue);
-									}
-								}
-							}
-						}
-					}
-				}
-			}
-		}
-		
-		// Skeletal Mesh's Skeleton
-		if (Mesh->GetSkeleton())
-		{
-			const bool bTransient = Mesh->GetSkeleton()->GetPackage() == GetTransientPackage();
-
-			// Don't duplicate if not transient or export all assets.
-			if (bTransient || bExportAllResources)
-			{
-				FString SkeletonName = ObjectName + "_Skeleton";
-				if (!ManageBakingAction(AssetPath, SkeletonName, bUsedGrantedOverridingRights))
-				{
-					return;
-				}
-
-				FString SkeletonPkgName = AssetPath + FString("/") + SkeletonName;
-				UObject* DuplicatedSkeleton = FUnrealBakeHelpers::BakeHelper_DuplicateAsset(Mesh->GetSkeleton(), SkeletonName, 
-					SkeletonPkgName, false, ReplacementMap, bUsedGrantedOverridingRights, false);
-
-				ArrayCachedObject.Add(DuplicatedSkeleton);
-				PackagesToSave.Add(DuplicatedSkeleton->GetPackage());
-				ReplacementMap.Add(Mesh->GetSkeleton(), DuplicatedSkeleton);
-			}
-		}
-
-		// Skeletal Mesh
-		if (!ManageBakingAction(AssetPath, ObjectName, bUsedGrantedOverridingRights))
-		{
-			return;
-		}
-
-		FString PkgName = AssetPath + FString("/") + ObjectName;
-		UObject* DupObject = FUnrealBakeHelpers::BakeHelper_DuplicateAsset(Mesh, ObjectName, PkgName, 
-			false, ReplacementMap, bUsedGrantedOverridingRights, false);
-		ArrayCachedObject.Add(DupObject);
-		PackagesToSave.Add(DupObject->GetPackage());
-
-		Mesh->Build();
-
-		USkeletalMesh* SkeletalMesh = Cast<USkeletalMesh>(DupObject);
-		if (SkeletalMesh)
-		{
-			SkeletalMesh->GetLODInfoArray() = Mesh->GetLODInfoArray();
-
-			SkeletalMesh->GetImportedModel()->SkeletalMeshModelGUID = FGuid::NewGuid();
-
-			// Duplicate AssetUserData
-			{
-				const TArray<UAssetUserData*>* AssetUserDataArray = Mesh->GetAssetUserDataArray();
-				for (const UAssetUserData* AssetUserData : *AssetUserDataArray)
-				{
-					if (AssetUserData)
-					{
-						// Duplicate to change ownership
-						UAssetUserData* NewAssetUserData = Cast<UAssetUserData>(StaticDuplicateObject(AssetUserData, SkeletalMesh));
-						SkeletalMesh->AddAssetUserData(NewAssetUserData);
-					}
-				}
-			}
-
-			// Generate render data
-			SkeletalMesh->Build();
-		}
-
-		// Remove duplicated UObjects from Root (previously added to avoid objects from being GC in the middle of the bake process)
-		for (UObject* Obj : ArrayCachedObject)
-		{
-			Obj->RemoveFromRoot();
-		}
-	}
-
-	// Save the packages generated during the baking operation  --------------------------------------------------------------------------------------
-	
-	// Complete the baking by saving the packages we have cached during the baking operation
-	if (PackagesToSave.Num())
-	{
-		FEditorFileUtils::PromptForCheckoutAndSave(PackagesToSave, false, !FApp::IsUnattended());
-	}
-}
-
-
 void FCustomizableObjectEditorViewportClient::OnInstanceForBakingUpdate(const FUpdateContext& Result)
 {
 	// Early exit if no instance was provided
diff --git a/Engine/Plugins/Experimental/Mutable/Source/CustomizableObjectEditor/Public/MuCOE/CustomizableObjectInstanceBakingUtils.cpp b/Engine/Plugins/Experimental/Mutable/Source/CustomizableObjectEditor/Public/MuCOE/CustomizableObjectInstanceBakingUtils.cpp
--- a/Engine/Plugins/Experimental/Mutable/Source/CustomizableObjectEditor/Public/MuCOE/CustomizableObjectInstanceBakingUtils.cpp
+++ b/Engine/Plugins/Experimental/Mutable/Source/CustomizableObjectEditor/Public/MuCOE/CustomizableObjectInstanceBakingUtils.cpp
@@ -20,233 +20,233 @@
 #define LOCTEXT_NAMESPACE "CustomizableObjectEditor" 
 
 
-// /**
-//  * Simple wrapper to be able to invoke the generation of a popup or log message depending on the execution context in which this code is being ran
-//  * @param InMessage The message to display
-//  * @param InTitle The title to be used for the popup or the log generated
-//  */
-// void ShowErrorNotification(const FText& InMessage,  const FText* InTitle = nullptr)
-// {
-// 	if (!FApp::IsUnattended())
-// 	{
-// 		FMessageDialog::Open(EAppMsgType::Ok, InMessage, *InTitle);
-// 	}
-// 	else
-// 	{
-// 		if (InTitle)
-// 		{
-// 			UE_LOG(LogMutable, Error, TEXT("%s - %s"), *InTitle->ToString(), *InMessage.ToString());
-// 		}
-// 		else
-// 		{
-// 			UE_LOG(LogMutable, Error, TEXT("%s"), *InMessage.ToString());
-// 		}
-// 	}
-// }
-//
-// /**
-//  * Utility functions for the baking operation.
-//  */
-//
-// /**
-//  * Validates the filename chosen for the baking data
-//  * @param FileName The filename chosen by the user
-//  * @return True if validation was successful, false otherwise
-//  */
-// bool ValidateProvidedFileName(const FString& FileName)
-// {
-// 	if (FileName.IsEmpty())
-// 	{
-// 		UE_LOG(LogMutable, Error, TEXT("Invalid baking configuration : FileName string is empty.."));
-// 		return false;
-// 	}
-//
-// 	// Check for invalid characters in the name of the object to be serialized
-// 	{
-// 		TCHAR InvalidCharacter = '0';
-// 		{
-// 			FString InvalidCharacters = FPaths::GetInvalidFileSystemChars();
-// 			for (int32 i = 0; i < InvalidCharacters.Len(); ++i)
-// 			{
-// 				TCHAR Char = InvalidCharacters[i];
-// 				FString SearchedChar = FString::Chr(Char);
-// 				if (FileName.Contains(SearchedChar))
-// 				{
-// 					InvalidCharacter = InvalidCharacters[i];
-// 					break;
-// 				}
-// 			}
-// 		}
-//
-// 		if (InvalidCharacter != '0')
-// 		{
-// 			const FText InvalidCharacterText = FText::FromString(FString::Chr(InvalidCharacter));
-// 			const FText ErrorText = FText::Format(LOCTEXT("FCustomizableObjectEditorViewportClient_BakeInstance_InvalidCharacter", "The selected contains an invalid character ({0})."), InvalidCharacterText);
-//
-// 			ShowErrorNotification(ErrorText);
-// 		
-// 			return false;
-// 		}
-// 	}
-//
-// 	return true;
-// }
-//
-//
-// /**
-//  * Validates the AssetPath chosen for the baking data
-//  * @param FileName The filename chosen by the user
-//  * @param AssetPath The AssetPath chosen by the user
-//  * @param InstanceCO The CustomizableObject from the provided COI
-//  * @return True if validation was successful, false otherwise
-//  */
-// bool ValidateProvidedAssetPath(const FString& FileName, const FString& AssetPath, const UCustomizableObject* InstanceCO)
-// {
-// 	if (AssetPath.IsEmpty())
-// 	{
-// 		UE_LOG(LogMutable, Error, TEXT("The AssetPath can not be empty!"));
-// 		return false;
-// 	}
-//
-// 	// Ensure we are not overriding the parent CO
-// 	const FString FullAssetPath = AssetPath + FString("/") + FileName + FString(".") + FileName;		// Full asset path to the new asset we want to create
-// 	const bool bWouldOverrideParentCO = InstanceCO->GetPathName() == FullAssetPath;
-// 	if (bWouldOverrideParentCO)
-// 	{
-// 		const FText ErrorText = LOCTEXT("FCustomizableObjectEditorViewportClient_BakeInstance_OverwriteCO", "The selected path would overwrite the instance's parent Customizable Object.");
-//
-// 		ShowErrorNotification(ErrorText);
-// 		
-// 		return false;
-// 	}
-//
-// 	return true;
-// }
-//
-//
-// /**
-//  * Outputs a string that we know it is unique.
-//  * @param InResource The resource we are working with
-//  * @param ResourceName The name of the resource we have provided. This should have the name of the current resource and will have the unique name for the resource once the method exits
-//  * @param InCachedResources Collection with all the already processed resources.
-//  * @param InCachedResourceNames Collection with all the already processed resources name's
-//  * @return True if the generation of the unique resource name was successful, false otherwise.
-//  */
-// bool GetUniqueResourceName(const UObject* InResource, FString& ResourceName, TArray<UObject*>& InCachedResources, const TArray<FString>& InCachedResourceNames)
-// {
-// 	check(InResource);
-//
-// 	int32 FindResult = InCachedResourceNames.Find(ResourceName);
-// 	if (FindResult != INDEX_NONE)
-// 	{
-// 		if (InResource == InCachedResources[FindResult])
-// 		{
-// 			return false;
-// 		}
-//
-// 		uint32 Count = 0;
-// 		while (FindResult != INDEX_NONE)
-// 		{
-// 			FindResult = InCachedResourceNames.Find(ResourceName + "_" + FString::FromInt(Count));
-// 			Count++;
-// 		}
-//
-// 		ResourceName += "_" + FString::FromInt(--Count);
-// 	}
-//
-// 	return true;
-// }
-//
-//
-// /**
-//  * Ensures the resource we want to save is ready to be saved. It handles closing it's editor and warning the user about possible overriding of resources.
-//  * @param InAssetSavePath The directory path where to save the baked object
-//  * @param InObjName The name of the object to be baked
-//  * @param bUsedGrantedOverridingRights Control flag that determines if the user has given or not permission to override resources already in disk
-//  * @return True if the operation was successful, false otherwise
-//  */
-// bool ManageBakingAction(const FString& InAssetSavePath, const FString& InObjName, bool& bUsedGrantedOverridingRights)
-// {
-// 	FString PackagePath = InAssetSavePath + "/" + InObjName;
-// 	UPackage* ExistingPackage = FindPackage(NULL, *PackagePath);
-//
-// 	if (!ExistingPackage)
-// 	{
-// 		FString PackageFilePath = PackagePath + "." + InObjName;
-//
-// 		FString PackageFileName;
-// 		if (FPackageName::DoesPackageExist(PackageFilePath, &PackageFileName))
-// 		{
-// 			ExistingPackage = LoadPackage(nullptr, *PackageFileName, LOAD_EditorOnly);
-// 		}
-// 		else
-// 		{
-// 			// if package does not exists
-// 			bUsedGrantedOverridingRights = false;
-// 			return true;
-// 		}
-// 	}
-//
-// 	if (ExistingPackage)
-// 	{
-// 		// Checking if the asset is open in an editor
-// 		TArray<IAssetEditorInstance*> ObjectEditors = GEditor->GetEditorSubsystem<UAssetEditorSubsystem>()->FindEditorsForAssetAndSubObjects(ExistingPackage);
-// 		if (ObjectEditors.Num())
-// 		{
-// 			for (IAssetEditorInstance* ObjectEditorInstance : ObjectEditors)
-// 			{
-// 				// Close the editors that contains this asset
-// 				if (!ObjectEditorInstance->CloseWindow(EAssetEditorCloseReason::AssetEditorHostClosed))
-// 				{
-// 					const FText Caption = LOCTEXT("OpenExisitngFile", "Open File");
-// 					const FText Message = FText::Format(LOCTEXT("CantCloseAsset", "This Obejct \"{0}\" is open in an editor and can't be closed automatically. Please close the editor and try to bake it again"), FText::FromString(InObjName));
-//
-// 					ShowErrorNotification(Message, &Caption);
-// 					
-// 					return false;
-// 				}
-// 			}
-// 		}
-//
-// 		if (!bUsedGrantedOverridingRights)
-// 		{
-// 			const FText Caption = LOCTEXT("Already existing baked files", "Already existing baked files");
-// 			const FText Message = FText::Format(LOCTEXT("OverwriteBakedInstance", "Instance baked files already exist in selected destination \"{0}\", this action will overwrite them."), FText::AsCultureInvariant(InAssetSavePath));
-//
-// 			// We need to guard this case since it may crash the editor if we are running an unattended commandlet and we try to generate this dialog
-// 			if (!FApp::IsUnattended())
-// 			{
-// 				if (FMessageDialog::Open(EAppMsgType::OkCancel, Message, Caption) == EAppReturnType::Cancel)
-// 				{
-// 					return false;
-// 				}
-// 			}
-// 			else
-// 			{
-// 				UE_LOG(LogMutable, Error, TEXT("%s - %s"), *Caption.ToString(), *Message.ToString());
-// 			}
-//
-//
-// 			bUsedGrantedOverridingRights = true;
-// 		}
-//
-// 		UObject* ExistingObject = StaticFindObject(UObject::StaticClass(), ExistingPackage, *InObjName);
-// 		if (ExistingObject)
-// 		{
-// 			ExistingPackage->FullyLoad();
-//
-// 			TArray<UObject*> ObjectsToDelete;
-// 			ObjectsToDelete.Add(ExistingObject);
-//
-// 			// Delete objects in the package with the same name as the one we want to create
-// 			const uint32 NumObjectsDeleted = ObjectTools::ForceDeleteObjects(ObjectsToDelete, false);
-//
-// 			return NumObjectsDeleted == ObjectsToDelete.Num();
-// 		}
-// 	}
-//
-// 	return true;
-// }
+/**
+ * Simple wrapper to be able to invoke the generation of a popup or log message depending on the execution context in which this code is being ran
+ * @param InMessage The message to display
+ * @param InTitle The title to be used for the popup or the log generated
+ */
+void ShowErrorNotification(const FText& InMessage,  const FText* InTitle = nullptr)
+{
+	if (!FApp::IsUnattended())
+	{
+		FMessageDialog::Open(EAppMsgType::Ok, InMessage, *InTitle);
+	}
+	else
+	{
+		if (InTitle)
+		{
+			UE_LOG(LogMutable, Error, TEXT("%s - %s"), *InTitle->ToString(), *InMessage.ToString());
+		}
+		else
+		{
+			UE_LOG(LogMutable, Error, TEXT("%s"), *InMessage.ToString());
+		}
+	}
+}
+
+/**
+ * Utility functions for the baking operation.
+ */
+
+/**
+ * Validates the filename chosen for the baking data
+ * @param FileName The filename chosen by the user
+ * @return True if validation was successful, false otherwise
+ */
+bool ValidateProvidedFileName(const FString& FileName)
+{
+	if (FileName.IsEmpty())
+	{
+		UE_LOG(LogMutable, Error, TEXT("Invalid baking configuration : FileName string is empty.."));
+		return false;
+	}
+
+	// Check for invalid characters in the name of the object to be serialized
+	{
+		TCHAR InvalidCharacter = '0';
+		{
+			FString InvalidCharacters = FPaths::GetInvalidFileSystemChars();
+			for (int32 i = 0; i < InvalidCharacters.Len(); ++i)
+			{
+				TCHAR Char = InvalidCharacters[i];
+				FString SearchedChar = FString::Chr(Char);
+				if (FileName.Contains(SearchedChar))
+				{
+					InvalidCharacter = InvalidCharacters[i];
+					break;
+				}
+			}
+		}
+
+		if (InvalidCharacter != '0')
+		{
+			const FText InvalidCharacterText = FText::FromString(FString::Chr(InvalidCharacter));
+			const FText ErrorText = FText::Format(LOCTEXT("FCustomizableObjectEditorViewportClient_BakeInstance_InvalidCharacter", "The selected contains an invalid character ({0})."), InvalidCharacterText);
+
+			ShowErrorNotification(ErrorText);
+		
+			return false;
+		}
+	}
+
+	return true;
+}
+
+
+/**
+ * Validates the AssetPath chosen for the baking data
+ * @param FileName The filename chosen by the user
+ * @param AssetPath The AssetPath chosen by the user
+ * @param InstanceCO The CustomizableObject from the provided COI
+ * @return True if validation was successful, false otherwise
+ */
+bool ValidateProvidedAssetPath(const FString& FileName, const FString& AssetPath, const UCustomizableObject* InstanceCO)
+{
+	if (AssetPath.IsEmpty())
+	{
+		UE_LOG(LogMutable, Error, TEXT("The AssetPath can not be empty!"));
+		return false;
+	}
+
+	// Ensure we are not overriding the parent CO
+	const FString FullAssetPath = AssetPath + FString("/") + FileName + FString(".") + FileName;		// Full asset path to the new asset we want to create
+	const bool bWouldOverrideParentCO = InstanceCO->GetPathName() == FullAssetPath;
+	if (bWouldOverrideParentCO)
+	{
+		const FText ErrorText = LOCTEXT("FCustomizableObjectEditorViewportClient_BakeInstance_OverwriteCO", "The selected path would overwrite the instance's parent Customizable Object.");
+
+		ShowErrorNotification(ErrorText);
+		
+		return false;
+	}
+
+	return true;
+}
+
+
+/**
+ * Outputs a string that we know it is unique.
+ * @param InResource The resource we are working with
+ * @param ResourceName The name of the resource we have provided. This should have the name of the current resource and will have the unique name for the resource once the method exits
+ * @param InCachedResources Collection with all the already processed resources.
+ * @param InCachedResourceNames Collection with all the already processed resources name's
+ * @return True if the generation of the unique resource name was successful, false otherwise.
+ */
+bool GetUniqueResourceName(const UObject* InResource, FString& ResourceName, TArray<UObject*>& InCachedResources, const TArray<FString>& InCachedResourceNames)
+{
+	check(InResource);
+
+	int32 FindResult = InCachedResourceNames.Find(ResourceName);
+	if (FindResult != INDEX_NONE)
+	{
+		if (InResource == InCachedResources[FindResult])
+		{
+			return false;
+		}
+
+		uint32 Count = 0;
+		while (FindResult != INDEX_NONE)
+		{
+			FindResult = InCachedResourceNames.Find(ResourceName + "_" + FString::FromInt(Count));
+			Count++;
+		}
+
+		ResourceName += "_" + FString::FromInt(--Count);
+	}
+
+	return true;
+}
+
+
+/**
+ * Ensures the resource we want to save is ready to be saved. It handles closing it's editor and warning the user about possible overriding of resources.
+ * @param InAssetSavePath The directory path where to save the baked object
+ * @param InObjName The name of the object to be baked
+ * @param bUsedGrantedOverridingRights Control flag that determines if the user has given or not permission to override resources already in disk
+ * @return True if the operation was successful, false otherwise
+ */
+bool ManageBakingAction(const FString& InAssetSavePath, const FString& InObjName, bool& bUsedGrantedOverridingRights)
+{
+	FString PackagePath = InAssetSavePath + "/" + InObjName;
+	UPackage* ExistingPackage = FindPackage(NULL, *PackagePath);
+
+	if (!ExistingPackage)
+	{
+		FString PackageFilePath = PackagePath + "." + InObjName;
+
+		FString PackageFileName;
+		if (FPackageName::DoesPackageExist(PackageFilePath, &PackageFileName))
+		{
+			ExistingPackage = LoadPackage(nullptr, *PackageFileName, LOAD_EditorOnly);
+		}
+		else
+		{
+			// if package does not exists
+			bUsedGrantedOverridingRights = false;
+			return true;
+		}
+	}
+
+	if (ExistingPackage)
+	{
+		// Checking if the asset is open in an editor
+		TArray<IAssetEditorInstance*> ObjectEditors = GEditor->GetEditorSubsystem<UAssetEditorSubsystem>()->FindEditorsForAssetAndSubObjects(ExistingPackage);
+		if (ObjectEditors.Num())
+		{
+			for (IAssetEditorInstance* ObjectEditorInstance : ObjectEditors)
+			{
+				// Close the editors that contains this asset
+				if (!ObjectEditorInstance->CloseWindow(EAssetEditorCloseReason::AssetEditorHostClosed))
+				{
+					const FText Caption = LOCTEXT("OpenExisitngFile", "Open File");
+					const FText Message = FText::Format(LOCTEXT("CantCloseAsset", "This Obejct \"{0}\" is open in an editor and can't be closed automatically. Please close the editor and try to bake it again"), FText::FromString(InObjName));
+
+					ShowErrorNotification(Message, &Caption);
+					
+					return false;
+				}
+			}
+		}
+
+		if (!bUsedGrantedOverridingRights)
+		{
+			const FText Caption = LOCTEXT("Already existing baked files", "Already existing baked files");
+			const FText Message = FText::Format(LOCTEXT("OverwriteBakedInstance", "Instance baked files already exist in selected destination \"{0}\", this action will overwrite them."), FText::AsCultureInvariant(InAssetSavePath));
+
+			// We need to guard this case since it may crash the editor if we are running an unattended commandlet and we try to generate this dialog
+			if (!FApp::IsUnattended())
+			{
+				if (FMessageDialog::Open(EAppMsgType::OkCancel, Message, Caption) == EAppReturnType::Cancel)
+				{
+					return false;
+				}
+			}
+			else
+			{
+				UE_LOG(LogMutable, Error, TEXT("%s - %s"), *Caption.ToString(), *Message.ToString());
+			}
+
+
+			bUsedGrantedOverridingRights = true;
+		}
+
+		UObject* ExistingObject = StaticFindObject(UObject::StaticClass(), ExistingPackage, *InObjName);
+		if (ExistingObject)
+		{
+			ExistingPackage->FullyLoad();
+
+			TArray<UObject*> ObjectsToDelete;
+			ObjectsToDelete.Add(ExistingObject);
+
+			// Delete objects in the package with the same name as the one we want to create
+			const uint32 NumObjectsDeleted = ObjectTools::ForceDeleteObjects(ObjectsToDelete, false);
+
+			return NumObjectsDeleted == ObjectsToDelete.Num();
+		}
+	}
+
+	return true;
+}
 
 
 namespace PreBakeSystemSettings
@@ -296,22 +296,35 @@ void UpdateInstanceForBaking(UCustomizableObjectInstance& InInstance, FInstanceU
 	InInstance.UpdateSkeletalMeshAsyncResult(InInstanceUpdateDelegate,true,true);
 }
 
-/// TODO: it will be enabled in a future CL. Moved back to CustomizableObjectEditorViewportClient to ease review process
-/*
+
 void BakeCustomizableObjectInstance(
 	UCustomizableObjectInstance& InInstance,
 	const FString& FileName,
 	const FString& AssetPath,
 	const bool bExportAllResources,
-	const bool bGenerateConstantMaterialInstances, 
-	const bool bIsAssetRegistryLoaded)
+	const bool bGenerateConstantMaterialInstances)
 {
 	// Ensure that the state of the COI provided is valid --------------------------------------------------------------------------------------------
+	UCustomizableObject* InstanceCO = InInstance.GetCustomizableObject();
+	check (InstanceCO);
+
+	// Ensure the CO of the COI is accessible 
+	if (!InstanceCO || InstanceCO->IsLocked())
+	{
+		FCustomizableObjectEditorLogger::CreateLog(
+		LOCTEXT("CustomizableObjectCompilingTryLater_Baking", "Please wait until the Customizable Object is compiled"))
+		.Category(ELoggerCategory::COInstanceBaking)
+		.CustomNotification()
+		.Notification(true)
+		.Log();
+
+		return;
+	}
 	
-	if (!bIsAssetRegistryLoaded)
+	if (InstanceCO->GetPrivate()->Status.Get() == FCustomizableObjectStatus::EState::Loading)
 	{
 		FCustomizableObjectEditorLogger::CreateLog(
-			LOCTEXT("CustomizableObjectCompileTryLater","Please wait until asset registry loads all assets"))
+			LOCTEXT("CustomizableObjectCompileTryLater_BakeInstance","Please wait unitl Customizable Object is loaded"))
 		.Category(ELoggerCategory::COInstanceBaking)
 		.CustomNotification()
 		.Notification(true)
@@ -325,9 +338,6 @@ void BakeCustomizableObjectInstance(
 		UE_LOG(LogMutable, Error, TEXT("The FileName for the instance baking is not valid."));
 		return;
 	}
-	
-	UCustomizableObject* InstanceCO = InInstance.GetCustomizableObject();
-	check (InstanceCO);
 
 	if (!ValidateProvidedAssetPath(FileName,AssetPath,InstanceCO))
 	{
@@ -335,36 +345,10 @@ void BakeCustomizableObjectInstance(
 		return;
 	}
 	
-	// Ensure the CO of the COI is accessible 
-	if (!InstanceCO || InstanceCO->IsLocked())
-	{
-		FCustomizableObjectEditorLogger::CreateLog(
-		LOCTEXT("CustomizableObjectCompilingTryLater", "Please wait until the Customizable Object is compiled"))
-		.Category(ELoggerCategory::COInstanceBaking)
-		.CustomNotification()
-		.Notification(true)
-		.Log();
-
-		return;
-	}
-	
 	// Exit early if the provided instance does not have a skeletal mesh
+	if (!InInstance.HasAnySkeletalMesh())
 	{
-		bool bHasSkeletalMesh = false;
-		
-		for (int32 ComponentIndex = 0; ComponentIndex < InInstance.SkeletalMeshes.Num(); ++ComponentIndex)
-		{
-			if (InInstance.SkeletalMeshes.IsValidIndex(ComponentIndex) && InInstance.SkeletalMeshes[ComponentIndex])
-			{
-				bHasSkeletalMesh = true;
-			}
-		}
-		
-		if (!bHasSkeletalMesh)
-		{
-			UE_LOG(LogMutable, Error, TEXT("The provided Customizable Object Instance has no skeletal meshes."));
-			return;
-		}
+		return;
 	}
 
 	// COI Validation completed : Proceed with the baking operation ----------------------------------------------------------------------------------
@@ -380,17 +364,17 @@ void BakeCustomizableObjectInstance(
 		.Log();
 	}
 	
-	TArray<UPackage*> PackagesToSave;
 	
 	// Set the overriding flag to false wo we ask the user at least once about if he is willing to override old baked data
 	bool bUsedGrantedOverridingRights = false;
 
-	const int32 NumComponents = InInstance.SkeletalMeshes.Num();
+	TArray<UPackage*> PackagesToSave;
+
+	const int32 NumComponents = InInstance.GetNumComponents();
 	for (int32 ComponentIndex = 0; ComponentIndex < NumComponents; ++ComponentIndex)
 	{
-		USkeletalMesh* Mesh = InInstance.SkeletalMeshes.IsValidIndex(ComponentIndex) && InInstance.SkeletalMeshes[ComponentIndex] ?
-			Cast<USkeletalMesh>(InInstance.SkeletalMeshes[ComponentIndex]) : nullptr;
-		
+		USkeletalMesh* Mesh = InInstance.GetSkeletalMesh(ComponentIndex);
+
 		if (!Mesh)
 		{
 			continue;
@@ -401,10 +385,6 @@ void BakeCustomizableObjectInstance(
 		{
 			ObjectName = FileName + "_Component_" + FString::FromInt(ComponentIndex);
 		}
-		else
-		{
-			ObjectName = FileName;
-		}
 
 		TMap<UObject*, UObject*> ReplacementMap;
 		TArray<FString> ArrayCachedElement;
@@ -468,7 +448,7 @@ void BakeCustomizableObjectInstance(
 									continue;
 								}
 
-								if (!ManageBakingAction(AssetPath, ResourceName,bUsedGrantedOverridingRights))
+								if (!ManageBakingAction(AssetPath, ResourceName, bUsedGrantedOverridingRights))
 								{
 									return;
 								}
@@ -483,6 +463,7 @@ void BakeCustomizableObjectInstance(
 								else
 								{
 									UMaterialInstanceConstant* InstConstant = Cast<UMaterialInstanceConstant>(Mesh->GetMaterials()[m].MaterialInterface);
+
 									if (InstConstant != nullptr)
 									{
 										InstConstant->Parent->GetTextureParameterValue(FName(*ParameterNames[i].GetPlainNameString()), OriginalTexture);
@@ -531,7 +512,7 @@ void BakeCustomizableObjectInstance(
 
 							if (ArrayCachedElement.Find(ResourceName) == INDEX_NONE)
 							{
-								if (!ManageBakingAction(AssetPath, ResourceName,bUsedGrantedOverridingRights))
+								if (!ManageBakingAction(AssetPath, ResourceName, bUsedGrantedOverridingRights))
 								{
 									return;
 								}
@@ -568,7 +549,7 @@ void BakeCustomizableObjectInstance(
 						continue;
 					}
 
-					if (!ManageBakingAction(AssetPath, ResourceName,bUsedGrantedOverridingRights))
+					if (!ManageBakingAction(AssetPath, ResourceName, bUsedGrantedOverridingRights))
 					{
 						return;
 					}
@@ -603,7 +584,7 @@ void BakeCustomizableObjectInstance(
 					continue;
 				}
 
-				if (!ManageBakingAction(AssetPath, MatObjName,bUsedGrantedOverridingRights))
+				if (!ManageBakingAction(AssetPath, MatObjName, bUsedGrantedOverridingRights))
 				{
 					return;
 				}
@@ -615,11 +596,11 @@ void BakeCustomizableObjectInstance(
 				ArrayCachedElement.Add(MatObjName);
 				PackagesToSave.Add(DupMat->GetPackage());
 
-				UMaterialInstance* MaterialInstance = Cast<UMaterialInstance>(Interface);
+				UMaterialInstance* Inst = Cast<UMaterialInstance>(Interface);
 
 				// Only need to duplicate the generate textures if the original material is a dynamic instance
 				// If the material has Mutable textures, then it will be a dynamic material instance for sure
-				if (MaterialInstance)
+				if (Inst)
 				{
 					// Duplicate generated textures
 					UMaterialInstanceDynamic* InstDynamic = Cast<UMaterialInstanceDynamic>(DupMat);
@@ -627,17 +608,17 @@ void BakeCustomizableObjectInstance(
 
 					if (InstDynamic || InstConstant)
 					{
-						for (int32 TextureIndex = 0; TextureIndex < MaterialInstance->TextureParameterValues.Num(); ++TextureIndex)
+						for (int32 TextureIndex = 0; TextureIndex < Inst->TextureParameterValues.Num(); ++TextureIndex)
 						{
-							if (MaterialInstance->TextureParameterValues[TextureIndex].ParameterValue)
+							if (Inst->TextureParameterValues[TextureIndex].ParameterValue)
 							{
-								if (MaterialInstance->TextureParameterValues[TextureIndex].ParameterValue->HasAnyFlags(RF_Transient))
+								if (Inst->TextureParameterValues[TextureIndex].ParameterValue->HasAnyFlags(RF_Transient))
 								{
-									UTexture2D* SrcTex = Cast<UTexture2D>(MaterialInstance->TextureParameterValues[TextureIndex].ParameterValue);
+									UTexture2D* SrcTex = Cast<UTexture2D>(Inst->TextureParameterValues[TextureIndex].ParameterValue);
 
 									if (SrcTex)
 									{
-										FString ParameterSanitized = MaterialInstance->TextureParameterValues[TextureIndex].ParameterInfo.Name.ToString();
+										FString ParameterSanitized = Inst->TextureParameterValues[TextureIndex].ParameterInfo.Name.ToString();
 										RemoveRestrictedChars(ParameterSanitized);
 
 										FString TexObjName = ObjectName + "_" + MaterialName + "_" + ParameterSanitized;
@@ -648,17 +629,17 @@ void BakeCustomizableObjectInstance(
 
 											if (InstDynamic)
 											{
-												InstDynamic->SetTextureParameterValue(MaterialInstance->TextureParameterValues[TextureIndex].ParameterInfo.Name, PrevTexture);
+												InstDynamic->SetTextureParameterValue(Inst->TextureParameterValues[TextureIndex].ParameterInfo.Name, PrevTexture);
 											}
 											else if (InstConstant)
 											{
-												InstConstant->SetTextureParameterValueEditorOnly(MaterialInstance->TextureParameterValues[TextureIndex].ParameterInfo.Name, PrevTexture);
+												InstConstant->SetTextureParameterValueEditorOnly(Inst->TextureParameterValues[TextureIndex].ParameterInfo.Name, PrevTexture);
 											}
 											
 											continue;
 										}
 
-										if (!ManageBakingAction(AssetPath, TexObjName,bUsedGrantedOverridingRights))
+										if (!ManageBakingAction(AssetPath, TexObjName, bUsedGrantedOverridingRights))
 										{
 											return;
 										}
@@ -672,11 +653,11 @@ void BakeCustomizableObjectInstance(
 
 										if (InstDynamic)
 										{
-											InstDynamic->SetTextureParameterValue(MaterialInstance->TextureParameterValues[TextureIndex].ParameterInfo.Name, DupTex);
+											InstDynamic->SetTextureParameterValue(Inst->TextureParameterValues[TextureIndex].ParameterInfo.Name, DupTex);
 										}
 										else if(InstConstant)
 										{
-											InstConstant->SetTextureParameterValueEditorOnly(MaterialInstance->TextureParameterValues[TextureIndex].ParameterInfo.Name, DupTex);
+											InstConstant->SetTextureParameterValueEditorOnly(Inst->TextureParameterValues[TextureIndex].ParameterInfo.Name, DupTex);
 										}
 									}
 									else
@@ -690,11 +671,11 @@ void BakeCustomizableObjectInstance(
 									// Just set the original texture
 									if (InstDynamic)
 									{
-										InstDynamic->SetTextureParameterValue(MaterialInstance->TextureParameterValues[TextureIndex].ParameterInfo.Name, MaterialInstance->TextureParameterValues[TextureIndex].ParameterValue);
+										InstDynamic->SetTextureParameterValue(Inst->TextureParameterValues[TextureIndex].ParameterInfo.Name, Inst->TextureParameterValues[TextureIndex].ParameterValue);
 									}
 									else if (InstConstant)
 									{
-										InstConstant->SetTextureParameterValueEditorOnly(MaterialInstance->TextureParameterValues[TextureIndex].ParameterInfo.Name, MaterialInstance->TextureParameterValues[TextureIndex].ParameterValue);
+										InstConstant->SetTextureParameterValueEditorOnly(Inst->TextureParameterValues[TextureIndex].ParameterInfo.Name, Inst->TextureParameterValues[TextureIndex].ParameterValue);
 									}
 								}
 							}
@@ -713,7 +694,7 @@ void BakeCustomizableObjectInstance(
 			if (bTransient || bExportAllResources)
 			{
 				FString SkeletonName = ObjectName + "_Skeleton";
-				if (!ManageBakingAction(AssetPath, SkeletonName,bUsedGrantedOverridingRights))
+				if (!ManageBakingAction(AssetPath, SkeletonName, bUsedGrantedOverridingRights))
 				{
 					return;
 				}
@@ -729,7 +710,7 @@ void BakeCustomizableObjectInstance(
 		}
 
 		// Skeletal Mesh
-		if (!ManageBakingAction(AssetPath, ObjectName,bUsedGrantedOverridingRights))
+		if (!ManageBakingAction(AssetPath, ObjectName, bUsedGrantedOverridingRights))
 		{
 			return;
 		}
@@ -767,7 +748,7 @@ void BakeCustomizableObjectInstance(
 			SkeletalMesh->Build();
 		}
 
-		// Remove duplicated UObjects from Root (previously added to avoid objects from beeing GC in the middle of the bake process)
+		// Remove duplicated UObjects from Root (previously added to avoid objects from being GC in the middle of the bake process)
 		for (UObject* Obj : ArrayCachedObject)
 		{
 			Obj->RemoveFromRoot();
@@ -782,6 +763,5 @@ void BakeCustomizableObjectInstance(
 		FEditorFileUtils::PromptForCheckoutAndSave(PackagesToSave, false, !FApp::IsUnattended());
 	}
 }
-*/
 
 #undef LOCTEXT_NAMESPACE 
diff --git a/Engine/Plugins/Experimental/Mutable/Source/CustomizableObjectEditor/Public/MuCOE/CustomizableObjectInstanceBakingUtils.h b/Engine/Plugins/Experimental/Mutable/Source/CustomizableObjectEditor/Public/MuCOE/CustomizableObjectInstanceBakingUtils.h
--- a/Engine/Plugins/Experimental/Mutable/Source/CustomizableObjectEditor/Public/MuCOE/CustomizableObjectInstanceBakingUtils.h
+++ b/Engine/Plugins/Experimental/Mutable/Source/CustomizableObjectEditor/Public/MuCOE/CustomizableObjectInstanceBakingUtils.h
@@ -13,7 +13,7 @@
 CUSTOMIZABLEOBJECTEDITOR_API
 void UpdateInstanceForBaking(UCustomizableObjectInstance& InInstance, FInstanceUpdateNativeDelegate& InInstanceUpdateDelegate);
 
-// TODO: enable after review -> Required for MTBL-1674
+
 /**
  * Serializes onto disk the resources used by the targeted Customizable Object Instance. The operation can be configured in the FInstanceBakingSettings settings object.
  * @param InInstance The mutable COI instance whose resources we want to bake onto disk
@@ -21,10 +21,7 @@ void UpdateInstanceForBaking(UCustomizableObjectInstance& InInstance, FInstanceU
  * @param AssetPath The path where to save the assets
  * @param bExportAllResources Determines if we want a full or partial export
  * @param bGenerateConstantMaterialInstances Determines if we want to generate constant material instances or not
- * @param bIsAssetRegistryLoaded Status of the Asset Registry
  */
-/*
 CUSTOMIZABLEOBJECTEDITOR_API
 void BakeCustomizableObjectInstance (UCustomizableObjectInstance& InInstance, const FString& FileName, const FString& AssetPath,
-	const bool bExportAllResources, const bool bGenerateConstantMaterialInstances, const bool bIsAssetRegistryLoaded);
-*/
\ No newline at end of file
+	const bool bExportAllResources, const bool bGenerateConstantMaterialInstances);
\ No newline at end of file
diff --git a/Engine/Plugins/Experimental/Mutable/Source/MutableValidation/MutableValidation.Build.cs b/Engine/Plugins/Experimental/Mutable/Source/MutableValidation/MutableValidation.Build.cs
--- a/Engine/Plugins/Experimental/Mutable/Source/MutableValidation/MutableValidation.Build.cs
+++ b/Engine/Plugins/Experimental/Mutable/Source/MutableValidation/MutableValidation.Build.cs
@@ -28,7 +28,9 @@ namespace UnrealBuildTool.Rules
 					"UnrealEd",
 					
 					"DataValidation",
-					"CustomizableObject"
+					"CustomizableObject", 
+					"CustomizableObjectEditor", 
+					"MutableTools",
 				}
 			);
 		}
diff --git a/Engine/Plugins/Experimental/Mutable/Source/MutableValidation/Private/MuV/COIBakingTestCommandlet.cpp b/Engine/Plugins/Experimental/Mutable/Source/MutableValidation/Private/MuV/COIBakingTestCommandlet.cpp
new file mode 100644
--- /dev/null
+++ b/Engine/Plugins/Experimental/Mutable/Source/MutableValidation/Private/MuV/COIBakingTestCommandlet.cpp
@@ -0,0 +1,198 @@
+﻿// Copyright Epic Games, Inc. All Rights Reserved.
+
+#include "../../Public/MuV/COIBakingTestCommandlet.h"
+
+#include "AssetRegistry/AssetRegistryModule.h"
+#include "MuCO/CustomizableObject.h"
+#include "MuCO/CustomizableObjectInstance.h"
+#include "MuCO/CustomizableObjectSystem.h"
+#include "MuCOE/CustomizableObjectInstanceBakingUtils.h"
+#include "HAL/FileManager.h"
+#include "MuT/UnrealPixelFormatOverride.h"
+
+
+/** Flag useful to know if we are currently updating an instance or not */
+bool bIsInstanceBeingUpdated = false;
+
+/** Did the instance update finish with a successful status? */
+bool bWasInstanceUpdateSuccessful = false;
+
+
+void OnInstanceUpdate(const FUpdateContext& Result)
+{
+	const EUpdateResult InstanceUpdateResult = Result.UpdateResult;
+
+	UE_LOG(LogMutable,Display,TEXT("Instance finished update with state : %s."), *UEnum::GetValueAsString(Result.UpdateResult));
+	bWasInstanceUpdateSuccessful = InstanceUpdateResult == EUpdateResult::Success;
+	
+	// Clear update flag so we can exit the update while loop
+	bIsInstanceBeingUpdated = false;
+}
+
+
+int32 UCOIBakingTestCommandlet::Main(const FString& Params)
+{
+	// Ensure we do not show any OK dialog since we are not an user that can interact with them
+	GIsRunningUnattendedScript = true;
+	
+	// Look for the COI to be baked and load it
+	{
+		FString CustomizableObjectInstanceAssetPath = "";
+		if (!FParse::Value(*Params, TEXT("CustomizableObjectInstance="), CustomizableObjectInstanceAssetPath))
+		{
+			UE_LOG(LogMutable,Error,TEXT("Failed to parse Customizable Object Instance package name from provided argument : %s"),*Params)
+			return 1;
+		}
+	
+		// Load the resource
+		UObject* FoundObject = FSoftObjectPath(CustomizableObjectInstanceAssetPath).TryLoad();
+		if (!FoundObject)
+		{
+			UE_LOG(LogMutable,Error,TEXT("Failed to retrieve UObject from path %s"),*CustomizableObjectInstanceAssetPath);
+			return 1;
+		}
+	
+		// Get the CustomizableObjectInstance.
+		TargetInstance = Cast<UCustomizableObjectInstance>(FoundObject);
+		if (!TargetInstance)
+		{
+			UE_LOG(LogMutable,Error,TEXT("Failed to cast found UObject to UCustomizableObjectInstance."));
+			return 1;
+		}
+
+		UE_LOG(LogMutable,Display,TEXT("Succesfully loaded %s Customizable Object Instance!"), *TargetInstance->GetName());
+	}
+
+	// Perform a blocking search to ensure all assets used by mutable are reachable using the AssetRegistry
+	{
+		UE_LOG(LogMutable,Display,TEXT("Searching all assets (this will take some time)..."));
+		FAssetRegistryModule& AssetRegistryModule = FModuleManager::LoadModuleChecked<FAssetRegistryModule>(AssetRegistryConstants::ModuleName);
+		AssetRegistryModule.Get().SearchAllAssets(true /* bSynchronousSearch */);
+	}
+	
+	// Compile it's CO (using current config)
+	UCustomizableObject* InstanceCustomizableObject = TargetInstance->GetCustomizableObject();
+	UCustomizableObjectSystem* System =  UCustomizableObjectSystem::GetInstanceChecked();
+	check (InstanceCustomizableObject);
+	{
+		UE_LOG(LogMutable,Display,TEXT("Compiling Customizable Object..."));
+    	
+		// Request a compiler to be able to locate the root and to compile it
+		const TUniquePtr<FCustomizableObjectCompilerBase> Compiler =
+			TUniquePtr<FCustomizableObjectCompilerBase>(System->GetNewCompiler());
+
+		// Compile (sync)
+		// Compile the CO with the provided compilation options
+		// Run Sync compilation -> Warning : Potentially long operation -------------
+		Compiler->Compile(*InstanceCustomizableObject, InstanceCustomizableObject->CompileOptions, false);
+		// --------------------------------------------------------------------------
+		
+		// Get the compilation result
+		const ECustomizableObjectCompilationState CompilationEndResult = Compiler->GetCompilationState();
+		check(CompilationEndResult != ECustomizableObjectCompilationState::None);
+		check(CompilationEndResult != ECustomizableObjectCompilationState::InProgress);
+
+		UE_LOG(LogMutable, Display, TEXT("CO compilation completed with end status : %s "), *UEnum::GetValueAsString(CompilationEndResult));
+		
+		const bool bWasCompilationSuccessful = CompilationEndResult == ECustomizableObjectCompilationState::Completed;
+		if (!bWasCompilationSuccessful)
+		{
+			UE_LOG(LogMutable,Error,TEXT("Failed to compile the target CO."));
+			return 1;
+		}
+	}
+	
+	// Update the instance
+	{
+		// If this fail something is very wrong
+		check(TargetInstance);
+
+		// Instance update delegate
+		FInstanceUpdateNativeDelegate InstanceUpdateDelegate;
+		InstanceUpdateDelegate.AddStatic(&OnInstanceUpdate);
+		
+		bIsInstanceBeingUpdated = true;
+		UE_LOG(LogMutable,Display,TEXT("Invoking instance update."));
+		UpdateInstanceForBaking(*TargetInstance, InstanceUpdateDelegate);
+		
+		// Now tick the engine so the instance gets updated while running in the commandlet context
+        while (bIsInstanceBeingUpdated)
+        {
+            // Tick the engine
+            CommandletHelpers::TickEngine();
+
+            // Stop if exit was requested
+            if (IsEngineExitRequested())
+            {
+            	break;
+            }
+        }
+		
+		// Check the end status of the instance update
+		if (!bWasInstanceUpdateSuccessful)
+		{
+			UE_LOG(LogMutable,Error,TEXT("Failed to succesfully update the target COI."));
+			return 1;
+		}
+	}
+	
+
+	// Bake the instance
+	{
+		const FString BakedResourcesFileName = "MuBakedInstances";
+		
+		// If this fail something is very wrong
+		check(TargetInstance);
+		const FString InstanceName = TargetInstance.GetName();
+		
+		// Create the actual directory in the filesystem of the host machine
+		const FString GlobalBakingDirectory = FPaths::ConvertRelativePathToFull(FPaths::Combine( FPaths::ProjectContentDir(), BakedResourcesFileName, InstanceName));
+
+		// This will only happen if we did make a partial run before and therefore the directory was not cleansed.
+		// Delete it then but notify the user since this may mean that we have duplicated COIs
+		if (FPaths::DirectoryExists(GlobalBakingDirectory))
+		{
+			UE_LOG(LogMutable,Warning,TEXT("The directory with path  \" %s \" does already exist. This may be produced by an incompelte execution of a previous test. Clearing it out before continuing..."), *GlobalBakingDirectory);
+			if (!IFileManager::Get().DeleteDirectory(*GlobalBakingDirectory,false,true))
+			{
+				UE_LOG(LogMutable,Error,TEXT("Failed to deleta baking directory at path \" %s \"."), *GlobalBakingDirectory);
+				return 1;
+			}
+		}
+		
+		// Compute the local path to the generated directory where to save the baked data
+		const FString LocalBakingDirectory = FPaths::Combine("/","Game", BakedResourcesFileName, InstanceName);
+		
+		// Create a new directory where to save the bake itself
+		if (IFileManager::Get().MakeDirectory(*GlobalBakingDirectory, true))
+		{
+			UE_LOG(LogMutable,Display,TEXT("Starting Instance Baking operation."));
+			{
+				BakeCustomizableObjectInstance(
+					*TargetInstance,
+					FString::Printf(TEXT("%s_Bake"), *TargetInstance->GetName()),
+					LocalBakingDirectory,
+					true,
+					true);
+			}
+			UE_LOG(LogMutable,Display,TEXT("Instance Baking operation has been completed."));
+			
+			// Delete the target directory where we did save the baked instance
+			if (!IFileManager::Get().DeleteDirectory(*GlobalBakingDirectory,true,true))
+			{
+				UE_LOG(LogMutable,Error,TEXT("Failed to deleta baking directory at path \" %s \"."), *GlobalBakingDirectory);
+				return 1;
+			}
+		}
+		else
+		{
+			UE_LOG(LogMutable,Error,TEXT("Failed to create baking directory at path \" %s \"."), *GlobalBakingDirectory);
+			return 1;
+		}
+	}
+	
+	return 0;
+}
+
+
+
diff --git a/Engine/Plugins/Experimental/Mutable/Source/MutableValidation/Public/MuV/COIBakingTestCommandlet.h b/Engine/Plugins/Experimental/Mutable/Source/MutableValidation/Public/MuV/COIBakingTestCommandlet.h
new file mode 100644
--- /dev/null
+++ b/Engine/Plugins/Experimental/Mutable/Source/MutableValidation/Public/MuV/COIBakingTestCommandlet.h
@@ -0,0 +1,31 @@
+﻿// Copyright Epic Games, Inc. All Rights Reserved.
+#pragma once
+
+#include "CoreMinimal.h"
+#include "Commandlets/Commandlet.h"
+
+#include "COIBakingTestCommandlet.generated.h"
+
+class UCustomizableObjectInstance;
+struct FUpdateContext;
+
+/**
+ * Commandlet designed to test the baking of instances.
+ * It expects an instance UAsset being provided to it so it can later compile it's CO, update the instance itself, and then bake it.
+ * EX : -Run=COIBakingTest -EnablePlugins=MutableTesting -CustomizableObjectInstance="/MutableTesting/Cyborg_Character/Character_Inst" -AllowCommandletRendering
+ */
+UCLASS()
+class MUTABLEVALIDATION_API UCOIBakingTestCommandlet : public UCommandlet
+{
+	GENERATED_BODY()
+	
+public:
+	virtual int32 Main(const FString& Params) override;
+	
+private:
+	
+	/** Instance targeted for baking */
+	UPROPERTY()
+	TObjectPtr<UCustomizableObjectInstance> TargetInstance = nullptr;
+};
+
