[mutable] Fixed missconfiguration of Editor Tests
- Forced generation of all lods during instance genration
- Fixed random paramters not being really random

diff --git a/Engine/Plugins/Experimental/Mutable/Source/CustomizableObject/Classes/MuCO/CustomizableObjectInstance.h b/Engine/Plugins/Experimental/Mutable/Source/CustomizableObject/Classes/MuCO/CustomizableObjectInstance.h
--- a/Engine/Plugins/Experimental/Mutable/Source/CustomizableObject/Classes/MuCO/CustomizableObjectInstance.h
+++ b/Engine/Plugins/Experimental/Mutable/Source/CustomizableObject/Classes/MuCO/CustomizableObjectInstance.h
@@ -255,10 +255,7 @@ public:
 
 	/**  Set random values to the parameters using a stream. */
 	UFUNCTION(BlueprintCallable, Category = CustomizableObjectInstance)
-	void SetRandomValuesFromStream(const FRandomStream& Stream);
-
-	/**  Set random values to the parameters using a seed. */
-	void SetRandomValues(const int32 InRandomizationSeed);
+	void SetRandomValuesFromStream(const FRandomStream& InStream);
 
 	/** Returns the AssetUserData that was gathered from all the constituent mesh parts during the last update. 
 	 *  It requires that the CustomizableObject had the bEnableAssetUserDataMerge set to true during compilation. 
diff --git a/Engine/Plugins/Experimental/Mutable/Source/CustomizableObject/Private/MuCO/CustomizableObjectInstance.cpp b/Engine/Plugins/Experimental/Mutable/Source/CustomizableObject/Private/MuCO/CustomizableObjectInstance.cpp
--- a/Engine/Plugins/Experimental/Mutable/Source/CustomizableObject/Private/MuCO/CustomizableObjectInstance.cpp
+++ b/Engine/Plugins/Experimental/Mutable/Source/CustomizableObject/Private/MuCO/CustomizableObjectInstance.cpp
@@ -2522,14 +2522,9 @@ void UCustomizableObjectInstance::SetRandomValues()
 	Descriptor.SetRandomValues();
 }
 
-void UCustomizableObjectInstance::SetRandomValuesFromStream(const FRandomStream& Stream)
+void UCustomizableObjectInstance::SetRandomValuesFromStream(const FRandomStream& InStream)
 {
-	Descriptor.SetRandomValuesFromStream(Stream);
-}
-
-void UCustomizableObjectInstance::SetRandomValues(const int32 InRandomizationSeed)
-{
-	Descriptor.SetRandomValues(InRandomizationSeed);
+	Descriptor.SetRandomValuesFromStream(InStream);
 }
 
 
diff --git a/Engine/Plugins/Experimental/Mutable/Source/CustomizableObject/Private/MuCO/CustomizableObjectInstanceDescriptor.cpp b/Engine/Plugins/Experimental/Mutable/Source/CustomizableObject/Private/MuCO/CustomizableObjectInstanceDescriptor.cpp
--- a/Engine/Plugins/Experimental/Mutable/Source/CustomizableObject/Private/MuCO/CustomizableObjectInstanceDescriptor.cpp
+++ b/Engine/Plugins/Experimental/Mutable/Source/CustomizableObject/Private/MuCO/CustomizableObjectInstanceDescriptor.cpp
@@ -2318,31 +2318,36 @@ void FCustomizableObjectInstanceDescriptor::SetCurrentState(const FString& State
 
 void FCustomizableObjectInstanceDescriptor::SetRandomValues()
 {
-	SetRandomValues(FMath::SRand() * TNumericLimits<int32>::Max());
+	const int32 RandomSeed = FMath::SRand() * TNumericLimits<int32>::Max();
+	FRandomStream RandomStream{RandomSeed};
+	SetRandomValuesFromStream(RandomSeed);
 }
 
 
-void FCustomizableObjectInstanceDescriptor::SetRandomValues(const int32 InRandomizationSeed)
+void FCustomizableObjectInstanceDescriptor::SetRandomValuesFromStream(const FRandomStream& InStream)
 {
 	check(CustomizableObject);
-	RETURN_ON_UNCOMPILED_CO(CustomizableObject, TEXT("Error: Cannot set random values"))
 
-	// Set seed for deterministic behaviour
-	const FRandomStream Random(InRandomizationSeed);
-		
+	if (!CustomizableObject)
+	{
+		return;
+	}
+	
+	RETURN_ON_UNCOMPILED_CO(CustomizableObject, TEXT("Error: Cannot set random values"))
+	
 	for (FCustomizableObjectFloatParameterValue& FloatParameter : FloatParameters)
 	{
-		FloatParameter.ParameterValue = Random.GetFraction();
+		FloatParameter.ParameterValue = InStream.GetFraction();
 
 		for (float& RangeValue : FloatParameter.ParameterRangeValues)
 		{
-			RangeValue = Random.GetFraction();
+			RangeValue = InStream.GetFraction();
 		}
 	}
 
 	for (FCustomizableObjectBoolParameterValue& BoolParameter : BoolParameters)
 	{
-		BoolParameter.ParameterValue = static_cast<bool>(Random.RandRange(0, 1));
+		BoolParameter.ParameterValue = static_cast<bool>(InStream.RandRange(0, 1));
 	}
 	
 	for (FCustomizableObjectIntParameterValue& IntParameter : IntParameters)
@@ -2352,21 +2357,21 @@ void FCustomizableObjectInstanceDescriptor::SetRandomValues(const int32 InRandom
 
 		if (NumValues)
 		{
-			IntParameter.ParameterValueName = CustomizableObject->GetIntParameterAvailableOption(ParamIndex, NumValues * Random.GetFraction());
+			IntParameter.ParameterValueName = CustomizableObject->GetIntParameterAvailableOption(ParamIndex, NumValues * InStream.GetFraction());
 
 			for (FString& RangeValue : IntParameter.ParameterRangeValueNames)
 			{
-				RangeValue = CustomizableObject->GetIntParameterAvailableOption(ParamIndex, NumValues * Random.GetFraction());
+				RangeValue = CustomizableObject->GetIntParameterAvailableOption(ParamIndex, NumValues * InStream.GetFraction());
 			}
 		}		
 	}
 
 	for (FCustomizableObjectVectorParameterValue& VectorParameter : VectorParameters)
 	{
-		VectorParameter.ParameterValue.R = Random.GetFraction();
-		VectorParameter.ParameterValue.G = Random.GetFraction();
-		VectorParameter.ParameterValue.B = Random.GetFraction();
-		VectorParameter.ParameterValue.A = Random.GetFraction();
+		VectorParameter.ParameterValue.R = InStream.GetFraction();
+		VectorParameter.ParameterValue.G = InStream.GetFraction();
+		VectorParameter.ParameterValue.B = InStream.GetFraction();
+		VectorParameter.ParameterValue.A = InStream.GetFraction();
 	}
 
 	const UCustomizableObjectSystemPrivate* SystemPrivate = UCustomizableObjectSystem::GetInstance()->GetPrivate();
@@ -2390,11 +2395,11 @@ void FCustomizableObjectInstanceDescriptor::SetRandomValues(const int32 InRandom
 	{
 		for (FCustomizableObjectTextureParameterValue& TextureParameter : TextureParameters)
 		{
-			TextureParameter.ParameterValue = PossibleValues[NumPossibleValues * Random.GetFraction()];
+			TextureParameter.ParameterValue = PossibleValues[NumPossibleValues * InStream.GetFraction()];
 		
 			for (FName& RangeValue : TextureParameter.ParameterRangeValues)
 			{
-				RangeValue = PossibleValues[NumPossibleValues * Random.GetFraction()];
+				RangeValue = PossibleValues[NumPossibleValues * InStream.GetFraction()];
 			}				
 		}		
 	}
@@ -2403,41 +2408,6 @@ void FCustomizableObjectInstanceDescriptor::SetRandomValues(const int32 InRandom
 }
 
 
-void FCustomizableObjectInstanceDescriptor::SetRandomValuesFromStream(const FRandomStream& Stream)
-{
-	check(CustomizableObject);
-	RETURN_ON_UNCOMPILED_CO(CustomizableObject, TEXT("Error: Cannot set random values"))
-		
-	if (!CustomizableObject) return;
-
-	for (int32 i = 0; i < FloatParameters.Num(); ++i)
-	{
-		FloatParameters[i].ParameterValue = Stream.FRand();
-	}
-
-	for (int32 i = 0; i < BoolParameters.Num(); ++i)
-	{
-		BoolParameters[i].ParameterValue = Stream.GetUnsignedInt() % 2 == 0;
-	}
-
-	for (int32 i = 0; i < IntParameters.Num(); ++i)
-	{
-		int32 ParameterIndexInCO = CustomizableObject->FindParameter(IntParameters[i].ParameterName);
-
-		if (ParameterIndexInCO >= 0)
-		{
-			int32 NumValues = CustomizableObject->GetIntParameterNumOptions(ParameterIndexInCO);
-			if (NumValues > 0)
-			{
-				int32 Index = Stream.GetUnsignedInt() % NumValues;
-				FString Option = CustomizableObject->GetIntParameterAvailableOption(ParameterIndexInCO, Index);
-				SetIntParameterSelectedOption(i, Option);
-			}
-		}
-	}
-}
-
-
 bool FCustomizableObjectInstanceDescriptor::CreateMultiLayerProjector(const FName& ProjectorParamName)
 {
 	if (!FMultilayerProjector::AreDescriptorParametersValid(*this, ProjectorParamName.ToString()))
diff --git a/Engine/Plugins/Experimental/Mutable/Source/CustomizableObject/Public/MuCO/CustomizableObjectInstanceDescriptor.h b/Engine/Plugins/Experimental/Mutable/Source/CustomizableObject/Public/MuCO/CustomizableObjectInstanceDescriptor.h
--- a/Engine/Plugins/Experimental/Mutable/Source/CustomizableObject/Public/MuCO/CustomizableObjectInstanceDescriptor.h
+++ b/Engine/Plugins/Experimental/Mutable/Source/CustomizableObject/Public/MuCO/CustomizableObjectInstanceDescriptor.h
@@ -294,9 +294,7 @@ struct CUSTOMIZABLEOBJECT_API FCustomizableObjectInstanceDescriptor
 
 	void SetRandomValues();
 	
-	void SetRandomValues(int32 InRandomizationSeed);
-
-	void SetRandomValuesFromStream(const FRandomStream& Stream);
+	void SetRandomValuesFromStream(const FRandomStream& InStream);
 
 	// ------------------------------------------------------------
 	// Multilayer Projectors
diff --git a/Engine/Plugins/Experimental/Mutable/Source/MutableValidation/Private/MuV/CustomizableObjectValidationCommandlet.cpp b/Engine/Plugins/Experimental/Mutable/Source/MutableValidation/Private/MuV/CustomizableObjectValidationCommandlet.cpp
--- a/Engine/Plugins/Experimental/Mutable/Source/MutableValidation/Private/MuV/CustomizableObjectValidationCommandlet.cpp
+++ b/Engine/Plugins/Experimental/Mutable/Source/MutableValidation/Private/MuV/CustomizableObjectValidationCommandlet.cpp
@@ -122,7 +122,7 @@ int32 UCustomizableObjectValidationCommandlet::Main(const FString& Params)
 
 		// CO embedded data size ------ //
 		{
-			TArray<uint8> EmbeddedDataBytes;
+			TArray<uint8> EmbeddedDataBytes{};
 			FMemoryWriter SerializationTarget{EmbeddedDataBytes, false};
 			
 			ToTestCustomizableObject->SaveEmbeddedData(SerializationTarget);
@@ -135,6 +135,9 @@ int32 UCustomizableObjectValidationCommandlet::Main(const FString& Params)
 		bool bWasInstancesCreationSuccessful = true;
 		{
 			UE_LOG(LogMutable,Display,TEXT("Generating %i random instances..."), InstancesToGenerate);	
+
+			// Create randomization stream for the parameters of the instance
+			FRandomStream RandomizationStream = FRandomStream(0);
 			
 			// Generate a series of instances to later update
 			for (uint32 CustomizableObjectInstanceIndex = 0; CustomizableObjectInstanceIndex < InstancesToGenerate; CustomizableObjectInstanceIndex++)
@@ -142,8 +145,13 @@ int32 UCustomizableObjectValidationCommandlet::Main(const FString& Params)
 				UCustomizableObjectInstance* GeneratedInstance = ToTestCustomizableObject->CreateInstance();
 				if (GeneratedInstance)
 				{
+					// Force generation of all LODS
+					TArray<uint16> RequestedLodLevels{};
+					RequestedLodLevels.Init(MAX_uint8, GeneratedInstance->GetNumComponents());
+					GeneratedInstance->GetDescriptor().SetRequestedLODLevels(RequestedLodLevels);
+					
 					// Randomize instance values
-					GeneratedInstance->SetRandomValues(0);
+					GeneratedInstance->SetRandomValuesFromStream(RandomizationStream);
 					InstancesToProcess.Push(GeneratedInstance);
 				}
 				else
