[Mutable] Fixed viewport cubemap reflections.

#rnx

diff --git a/Engine/Plugins/Experimental/Mutable/Source/CustomizableObjectEditor/Private/MuCOE/CustomizableObjectEditorViewportClient.cpp b/Engine/Plugins/Experimental/Mutable/Source/CustomizableObjectEditor/Private/MuCOE/CustomizableObjectEditorViewportClient.cpp
--- a/Engine/Plugins/Experimental/Mutable/Source/CustomizableObjectEditor/Private/MuCOE/CustomizableObjectEditorViewportClient.cpp
+++ b/Engine/Plugins/Experimental/Mutable/Source/CustomizableObjectEditor/Private/MuCOE/CustomizableObjectEditorViewportClient.cpp
@@ -114,14 +114,7 @@ FCustomizableObjectEditorViewportClient::FCustomizableObjectEditorViewportClient
 	OverrideNearClipPlane(1.0f);
 
 	SetPreviewComponent(nullptr);
-
-	// add capture component for reflection
-	USphereReflectionCaptureComponent* CaptureComponent = NewObject<USphereReflectionCaptureComponent>();
-
-	const FTransform CaptureTransform(FRotator(0, 0, 0), FVector(0.f, 0.f, 100.f), FVector(1.f));
-	PreviewScene->AddComponent(CaptureComponent, CaptureTransform);
-	CaptureComponent->UpdateReflectionCaptureContents(PreviewScene->GetWorld());
-
+	
 	// now add the ClipMorph plane
 	ClipMorphNode = nullptr;
 	bClipMorphLocalStartOffset = true;
