[Mutable] Fix UnrealMutableDiskStreamer data race.

#rnx

diff --git a/Engine/Plugins/Experimental/Mutable/Source/CustomizableObject/Private/MuCO/CustomizableObjectSystem.cpp b/Engine/Plugins/Experimental/Mutable/Source/CustomizableObject/Private/MuCO/CustomizableObjectSystem.cpp
--- a/Engine/Plugins/Experimental/Mutable/Source/CustomizableObject/Private/MuCO/CustomizableObjectSystem.cpp
+++ b/Engine/Plugins/Experimental/Mutable/Source/CustomizableObject/Private/MuCO/CustomizableObjectSystem.cpp
@@ -108,12 +108,6 @@ TAutoConsoleVariable<bool> CVarEnableMeshCache(
 	TEXT("Enables or disables the reuse of meshes."),
 	ECVF_Scalability);
 
-TAutoConsoleVariable<bool> CVarRollbackFixModelDiskStreamerDataRace(
-	TEXT("mutable.Rollback.FixModelDiskStreamerDataRace"),
-	true,
-	TEXT("When true, use the new code path that fixes a data race in UnrealMutableModelDiskStreamer."));
-
-
 int32 UCustomizableObjectSystemPrivate::SkeletalMeshMinLodQualityLevel = -1;
 
 static void CVarMutableSinkFunction()
@@ -546,27 +540,16 @@ void UCustomizableObjectSystem::BeginDestroy()
 		Private->PendingTasks.Empty();
 
 		// Complete pending taskgraph tasks
-		if (CVarRollbackFixModelDiskStreamerDataRace.GetValueOnAnyThread())
-		{
-			Private->MutableTaskGraph.AllowLaunchingMutableTaskLowPriority(false, false);
-			check(Private->Streamer);
-			Private->MutableTaskGraph.AddMutableThreadTask(TEXT("EndStream"), [Streamer = Private->Streamer]()
+		Private->MutableTaskGraph.AllowLaunchingMutableTaskLowPriority(false, false);
+		check(Private->Streamer);
+		Private->MutableTaskGraph.AddMutableThreadTask(TEXT("EndStream"), [Streamer = Private->Streamer]()
 			{
 				Streamer->EndStreaming();
 			});
-		}
 		Private->MutableTaskGraph.WaitForMutableTasks();
 
 		// Clear the ongoing operation
 		Private->CurrentMutableOperation = nullptr;
-
-		if (!CVarRollbackFixModelDiskStreamerDataRace.GetValueOnAnyThread())
-		{
-			// Deallocate streaming
-			check(Private->Streamer != nullptr);
-			Private->Streamer->EndStreaming();
-		}
-
 		Private->CurrentInstanceBeingUpdated = nullptr;
 
 		UCustomizableObjectSystemPrivate::SSystem = nullptr;
@@ -1349,7 +1332,11 @@ bool UCustomizableObjectSystem::LockObject(const class UCustomizableObject* InOb
 		check(GetPrivate() != nullptr);
 		if (GetPrivate()->Streamer)
 		{
-			GetPrivate()->Streamer->CancelStreamingForObject(InObject);
+			FMutableTaskGraph::TaskType Task = Private->MutableTaskGraph.AddMutableThreadTask(TEXT("EndStream"), [InObject, Streamer = GetPrivate()->Streamer]()
+				{
+					Streamer->CancelStreamingForObject(InObject);
+				});
+			Task.Wait();
 		}
 
 		// Clear the cache for the instance, since we will remake it
@@ -1571,12 +1558,9 @@ namespace impl
 		Operation->UpdateStartBytes = mu::FGlobalMemoryCounter::GetCounter();
 		mu::FGlobalMemoryCounter::Zero();
 
-		if (CVarRollbackFixModelDiskStreamerDataRace.GetValueOnAnyThread())
-		{
-			// Prepare streaming for the current customizable object
-			check(SystemPrivate->Streamer != nullptr);
-			SystemPrivate->Streamer->PrepareStreamingForObject(Operation->Instance->GetCustomizableObject());			
-		}
+		// Prepare streaming for the current customizable object
+		check(SystemPrivate->Streamer != nullptr);
+		SystemPrivate->Streamer->PrepareStreamingForObject(Operation->Instance->GetCustomizableObject());			
 
 		const mu::Ptr<mu::System> MutableSystem = SystemPrivate->MutableSystem;
 
@@ -2705,13 +2689,6 @@ namespace impl
 
 		SystemPrivateData->CurrentInstanceBeingUpdated = CandidateInstance;
 
-		if (!CVarRollbackFixModelDiskStreamerDataRace.GetValueOnAnyThread())
-		{
-			// Prepare streaming for the current customizable object
-			check(SystemPrivateData->Streamer != nullptr);
-			SystemPrivateData->Streamer->PrepareStreamingForObject(CustomizableObject);
-		}
-
 		check(SystemPrivateData->ExtensionDataStreamer != nullptr);
 		SystemPrivateData->ExtensionDataStreamer->SetActiveObject(CustomizableObject);
 
diff --git a/Engine/Plugins/Experimental/Mutable/Source/CustomizableObject/Private/MuCO/UnrealMutableModelDiskStreamer.cpp b/Engine/Plugins/Experimental/Mutable/Source/CustomizableObject/Private/MuCO/UnrealMutableModelDiskStreamer.cpp
--- a/Engine/Plugins/Experimental/Mutable/Source/CustomizableObject/Private/MuCO/UnrealMutableModelDiskStreamer.cpp
+++ b/Engine/Plugins/Experimental/Mutable/Source/CustomizableObject/Private/MuCO/UnrealMutableModelDiskStreamer.cpp
@@ -49,10 +49,6 @@ void UnrealMutableInputStream::Read(void* pData, uint64 size)
 //-------------------------------------------------------------------------------------------------
 FUnrealMutableModelBulkReader::~FUnrealMutableModelBulkReader()
 {
-	if (!CVarRollbackFixModelDiskStreamerDataRace.GetValueOnAnyThread())
-	{
-		EndStreaming();
-	}
 }
 
 
@@ -67,14 +63,8 @@ bool FUnrealMutableModelBulkReader::PrepareStreamingForObject(UCustomizableObjec
 	// See if we can free previuously allocated resources
 	for (int32 ObjectIndex = 0; ObjectIndex < Objects.Num(); )
 	{
-		if (!Objects[ObjectIndex].Model.Pin())
+		if (!Objects[ObjectIndex].Model.Pin() && Objects[ObjectIndex].CurrentReadRequests.IsEmpty())
 		{
-			// The CustomizableObject is gone, so we won't be streaming for it anymore.
-			for (TPair<OPERATION_ID, FReadRequest>& it : Objects[ObjectIndex].CurrentReadRequests)
-			{
-				it.Value.ReadRequest->WaitCompletion();
-			}
-
 			Objects.RemoveAtSwap(ObjectIndex);
 		}
 		else
@@ -148,9 +138,6 @@ bool FUnrealMutableModelBulkReader::PrepareStreamingForObject(UCustomizableObjec
 #if WITH_EDITOR
 void FUnrealMutableModelBulkReader::CancelStreamingForObject(const UCustomizableObject* CustomizableObject)
 {
-	// This happens in the game thread
-	check(IsInGameThread());
-
 	if (!CustomizableObject)
 	{
 		check(false);
@@ -161,10 +148,7 @@ void FUnrealMutableModelBulkReader::CancelStreamingForObject(const UCustomizable
 	{
 		if (Objects[ObjectIndex].Model.Pin() == CustomizableObject->GetModel())
 		{
-			for (TPair<OPERATION_ID, FReadRequest>& it : Objects[ObjectIndex].CurrentReadRequests)
-			{
-				it.Value.ReadRequest->WaitCompletion();
-			}
+			check(Objects[ObjectIndex].CurrentReadRequests.IsEmpty());
 
 			Objects.RemoveAtSwap(ObjectIndex);
 			break;
