[Mutable] Expose WorkingMemory to the public API.

#rnx

diff --git a/Engine/Plugins/Experimental/Mutable/Source/CustomizableObject/Classes/MuCO/CustomizableObjectSystem.h b/Engine/Plugins/Experimental/Mutable/Source/CustomizableObject/Classes/MuCO/CustomizableObjectSystem.h
--- a/Engine/Plugins/Experimental/Mutable/Source/CustomizableObject/Classes/MuCO/CustomizableObjectSystem.h
+++ b/Engine/Plugins/Experimental/Mutable/Source/CustomizableObject/Classes/MuCO/CustomizableObjectSystem.h
@@ -407,6 +407,17 @@ private:
 	void StartNextRecompile();
 	void TickRecompileCustomizableObjects();
 
+public:
+	/** Set Mutable's working memory limit (bytes). Mutable will flush internal caches to try to keep its memory consumption below the WorkingMemory (i.e., it is not a hard limit).
+	 * The working memory limit will especially reduce the memory required to perform Instance Updates and Texture Streaming.
+ 	 * Notice that Mutable does not track all its memory (e.g., UObjects memory is no tracked).
+	 * This value can also be set using "mutable.WorkingMemory" CVar. */
+	void SetWorkingMemory(int32 Bytes);
+
+	/** Get Mutable's working memory limit (bytes). See SetWorkingMemory(int32). */
+	int32 GetWorkingMemory() const;
+	
+private:
 	FCustomizableObjectCompilerBase* RecompileCustomizableObjectsCompiler = nullptr;
 	
 	TArray<FAssetData> ObjectsToRecompile;
diff --git a/Engine/Plugins/Experimental/Mutable/Source/CustomizableObject/Private/MuCO/CustomizableObjectSystem.cpp b/Engine/Plugins/Experimental/Mutable/Source/CustomizableObject/Private/MuCO/CustomizableObjectSystem.cpp
--- a/Engine/Plugins/Experimental/Mutable/Source/CustomizableObject/Private/MuCO/CustomizableObjectSystem.cpp
+++ b/Engine/Plugins/Experimental/Mutable/Source/CustomizableObject/Private/MuCO/CustomizableObjectSystem.cpp
@@ -68,13 +68,16 @@ static FAutoConsoleVariableRef CVarMutableEnabled(
 	TEXT("true/false - Disabling Mutable will turn off CO compilation, mesh generation, and texture streaming and will remove the system ticker. "),
 	FConsoleVariableDelegate::CreateStatic(&UCustomizableObjectSystemPrivate::OnMutableEnabledChanged));
 
-static TAutoConsoleVariable<int32> CVarWorkingMemory(
-	TEXT("mutable.WorkingMemory"),
+int32 WorkingMemory =
 #if !PLATFORM_DESKTOP
-	(10 * 1024),
+(10 * 1024);
 #else
-	(50 * 1024),
+(50 * 1024);
 #endif
+
+static FAutoConsoleVariableRef CVarWorkingMemory(
+	TEXT("mutable.WorkingMemory"),
+	WorkingMemory,
 	TEXT("Limit the amount of memory (in KB) to use as working memory when building characters. More memory reduces the object construction time. 0 means no restriction. Defaults: Desktop = 50,000 KB, Others = 10,000 KB"),
 	ECVF_Scalability);
 
@@ -470,7 +473,7 @@ void UCustomizableObjectSystem::InitSystem()
 	Private->CurrentMutableOperation = nullptr;
 	Private->CurrentInstanceBeingUpdated = nullptr;
 
-	Private->LastWorkingMemoryBytes = CVarWorkingMemory.GetValueOnGameThread() * 1024;
+	Private->LastWorkingMemoryBytes = CVarWorkingMemory->GetInt() * 1024;
 	Private->LastGeneratedResourceCacheSize = CVarGeneratedResourcesCacheSize.GetValueOnGameThread();
 
 	const mu::Ptr<mu::Settings> pSettings = new mu::Settings;
@@ -1489,7 +1492,7 @@ void UCustomizableObjectSystemPrivate::UpdateMemoryLimit()
 	// This must run on game thread, and when the mutable thread is not running
 	check(IsInGameThread());
 
-	const uint64 MemoryBytes = CVarWorkingMemory.GetValueOnGameThread() * 1024;
+	const uint64 MemoryBytes = CVarWorkingMemory->GetInt() * 1024;
 	if (MemoryBytes != LastWorkingMemoryBytes)
 	{
 		LastWorkingMemoryBytes = MemoryBytes;
@@ -3582,6 +3585,19 @@ void UCustomizableObjectSystem::TickRecompileCustomizableObjects()
 }
 
 
+void UCustomizableObjectSystem::SetWorkingMemory(int32 Bytes)
+{
+	WorkingMemory = Bytes;
+	UE_LOG(LogMutable, Log, TEXT("Working Memory set to %i bytes."), Bytes);
+}
+
+
+int32 UCustomizableObjectSystem::GetWorkingMemory() const
+{
+	return WorkingMemory;
+}
+
+
 uint64 UCustomizableObjectSystem::GetMaxChunkSizeForPlatform(const ITargetPlatform* TargetPlatform)
 {
 	const FString& PlatformName = TargetPlatform ? TargetPlatform->IniPlatformName() : FPlatformProperties::IniPlatformName();
