[Mutable] Add missing Undo/Redo support on gizmos.

- ClipMorph Nodes
- ClipMeshWithMesh Nodes
- Lights

#rnx

diff --git a/Engine/Plugins/Experimental/Mutable/Source/CustomizableObjectEditor/Private/MuCOE/CustomizableObjectCustomSettings.cpp b/Engine/Plugins/Experimental/Mutable/Source/CustomizableObjectEditor/Private/MuCOE/CustomizableObjectCustomSettings.cpp
--- a/Engine/Plugins/Experimental/Mutable/Source/CustomizableObjectEditor/Private/MuCOE/CustomizableObjectCustomSettings.cpp
+++ b/Engine/Plugins/Experimental/Mutable/Source/CustomizableObjectEditor/Private/MuCOE/CustomizableObjectCustomSettings.cpp
@@ -626,9 +626,9 @@ FReply SCustomizableObjectCustomSettings::OnPointLightAdded() const
 		return FReply::Handled();	
 	}
 
-	UPointLightComponent* Component = NewObject<UPointLightComponent>();
+	UPointLightComponent* Component = NewObject<UPointLightComponent>(GetTransientPackage(), NAME_None, RF_Transactional);
 	Viewport->GetViewportClient()->AddLightToScene(Component);
-	Editor->ShowGizmoLight(*Component);			
+	Editor->ShowGizmoLight(*Component);
 	
 	return FReply::Handled();
 }
diff --git a/Engine/Plugins/Experimental/Mutable/Source/CustomizableObjectEditor/Private/MuCOE/CustomizableObjectEditorViewportClient.cpp b/Engine/Plugins/Experimental/Mutable/Source/CustomizableObjectEditor/Private/MuCOE/CustomizableObjectEditorViewportClient.cpp
--- a/Engine/Plugins/Experimental/Mutable/Source/CustomizableObjectEditor/Private/MuCOE/CustomizableObjectEditorViewportClient.cpp
+++ b/Engine/Plugins/Experimental/Mutable/Source/CustomizableObjectEditor/Private/MuCOE/CustomizableObjectEditorViewportClient.cpp
@@ -87,7 +87,6 @@ FCustomizableObjectEditorViewportClient::FCustomizableObjectEditorViewportClient
 	UVChannelToDrawInUVs = 0;
 	MaterialToDrawInUVsComponent = 0;
 
-
 	bReferenceMeshMissingWarningMessageVisible = false;
 
 	DrawHelper.bDrawPivot = false;
@@ -745,8 +744,10 @@ void FCustomizableObjectEditorViewportClient::HideGizmoClipMesh()
 {
 	if (WidgetType == EWidgetType::ClipMesh)
 	{
+		ClipMeshComp->SetVisibility(false);
+		
 		SetWidgetType(EWidgetType::Hidden);
-	}	
+	}
 }
 
 
@@ -1244,6 +1245,9 @@ bool FCustomizableObjectEditorViewportClient::HandleBeginTransform()
 	switch (WidgetType)
 	{
 	case EWidgetType::Projector:
+	case EWidgetType::ClipMorph:
+	case EWidgetType::ClipMesh:
+	case EWidgetType::Light:
 		{
 			bManipulating = true;
 
@@ -1251,25 +1255,40 @@ bool FCustomizableObjectEditorViewportClient::HandleBeginTransform()
 				
 			if (WidgetMode == UE::Widget::WM_Translate)
 			{
-				GEditor->BeginTransaction(LOCTEXT("CustomizableObjectEditor_TranslateProjector", "Translate Projector"));
+				GEditor->BeginTransaction(LOCTEXT("CustomizableObjectEditor_Translate", "Translate"));
 			}
 			else if (WidgetMode == UE::Widget::WM_Rotate)
 			{
-				GEditor->BeginTransaction(LOCTEXT("CustomizableObjectEditor_RotateProjector", "Rotate Projector"));
+				GEditor->BeginTransaction(LOCTEXT("CustomizableObjectEditor_Rotate", "Rotate"));
 			}
 			else if (WidgetMode == UE::Widget::WM_Scale)
 			{
-				GEditor->BeginTransaction(LOCTEXT("CustomizableObjectEditor_ScaleProjector", "Scale Projector"));
+				GEditor->BeginTransaction(LOCTEXT("CustomizableObjectEditor_Scale", "Scale"));
 			}
 
-			(void)WidgetTrackingStartedDelegate.ExecuteIfBound();
-			return true;
+			break;
 		}
-
-	// The following cases are missing Undo/Redo functionality MTBL-391.
+		
+	case EWidgetType::Hidden:
+		break;
+	default:
+		unimplemented();
+	}
+	
+	switch (WidgetType)
+	{
+	case EWidgetType::Projector:
+		WidgetTrackingStartedDelegate.ExecuteIfBound();
+		break;
 	case EWidgetType::ClipMorph:
+		ClipMorphNode->Modify();
+		break;
 	case EWidgetType::ClipMesh:
+		ClipMeshNode->Modify();
+		break;
 	case EWidgetType::Light:
+		SelectedLightComponent->Modify();
+		break;
 	case EWidgetType::Hidden:
 		return true;
 		break;
@@ -1284,17 +1303,19 @@ bool FCustomizableObjectEditorViewportClient::HandleEndTransform()
 	switch (WidgetType)
 	{
 	case EWidgetType::Projector:
+	case EWidgetType::ClipMorph:
+	case EWidgetType::ClipMesh:
+	case EWidgetType::Light:
 		if (bManipulating)
 		{
 			bManipulating = false;
 			GEditor->EndTransaction();
 			return true;
 		}
-
+		
+		break;
+		
 	case EWidgetType::Hidden:
-	case EWidgetType::ClipMorph:
-	case EWidgetType::ClipMesh:
-	case EWidgetType::Light:
 		return true;
 		break;
 	default:
