[Mutable] Add undo support to Clip Mesh With Mesh Tags editor.

#rnx

diff --git a/Engine/Plugins/Experimental/Mutable/Source/CustomizableObjectEditor/Private/MuCOE/Nodes/CustomizableObjectNodeMeshClipWithMeshDetails.cpp b/Engine/Plugins/Experimental/Mutable/Source/CustomizableObjectEditor/Private/MuCOE/Nodes/CustomizableObjectNodeMeshClipWithMeshDetails.cpp
--- a/Engine/Plugins/Experimental/Mutable/Source/CustomizableObjectEditor/Private/MuCOE/Nodes/CustomizableObjectNodeMeshClipWithMeshDetails.cpp
+++ b/Engine/Plugins/Experimental/Mutable/Source/CustomizableObjectEditor/Private/MuCOE/Nodes/CustomizableObjectNodeMeshClipWithMeshDetails.cpp
@@ -51,7 +51,10 @@ void FCustomizableObjectNodeMeshClipWithMeshDetails::CustomizeDetails(IDetailLay
 		// Hidding unnecessary properties
 		DetailBuilder.HideProperty("CustomizableObjectToClipWith");
 		DetailBuilder.HideProperty("ArrayMaterialNodeToClipWithID");
-		DetailBuilder.HideProperty("Tags");
+		if (!Node->bUseTags)
+		{
+			DetailBuilder.HideProperty("Tags");			
+		}
 
 		if (const UEdGraphPin* ConnectedPin = FollowInputPin(*Node->ClipMeshPin());
 			ConnectedPin && !Cast<UCustomizableObjectNodeStaticMesh>(ConnectedPin->GetOwningNode()))
@@ -163,48 +166,36 @@ void FCustomizableObjectNodeMeshClipWithMeshDetails::CustomizeDetails(IDetailLay
 		}
 
         ArrayMaterialNodeOption.Sort(CompareNames);
-		
-		// Creating a TagView widget for each tag of the node
-		CreateTagView();
 
 		// Widget that shows the Tags widget or/and the materials ComboBox
-		CustomizableObjectToClipCategory.AddCustomRow(LOCTEXT("FCustomizableObjectNodeMeshClipWithMeshDetails", "Blocks"))
-		[
-			SNew(SVerticalBox)
-			+SVerticalBox::Slot()
-			.AutoHeight()
+		if (Node->bUseMaterials)
+		{
+			CustomizableObjectToClipCategory.AddCustomRow(LOCTEXT("FCustomizableObjectNodeMeshClipWithMeshDetails", "Blocks"))
 			[
-				SAssignNew(MaterialsSelector, SHorizontalBox)
-				+ SHorizontalBox::Slot()
-				.FillWidth(0.15f)
-				.Padding(0.0f, 3.0f, 0.0f, 0.0f)
-				[
-					SNew(STextBlock)
-					.Text(LOCTEXT("MaterialsText", "Material:"))
-				]
-				+ SHorizontalBox::Slot()
-				.FillWidth(0.85f)
-				.Padding(2.0f, 0.0f, 0.0f, 0.0f)
+				SNew(SVerticalBox)
+				+SVerticalBox::Slot()
+				.AutoHeight()
 				[
-					SNew(STextComboBox)
-					.OptionsSource(&ArrayMaterialNodeOption)
-					.InitiallySelectedItem(ItemToSelect)
-					.OnSelectionChanged(this, &FCustomizableObjectNodeMeshClipWithMeshDetails::OnMeshClipWithMeshNodeComboBoxSelectionChanged)
+					SNew(SHorizontalBox)
+					+ SHorizontalBox::Slot()
+					.FillWidth(0.15f)
+					.Padding(0.0f, 3.0f, 0.0f, 0.0f)
+					[
+						SNew(STextBlock)
+						.Text(LOCTEXT("MaterialsText", "Material:"))
+					]
+					+ SHorizontalBox::Slot()
+					.FillWidth(0.85f)
+					.Padding(2.0f, 0.0f, 0.0f, 0.0f)
+					[
+						SNew(STextComboBox)
+						.OptionsSource(&ArrayMaterialNodeOption)
+						.InitiallySelectedItem(ItemToSelect)
+						.OnSelectionChanged(this, &FCustomizableObjectNodeMeshClipWithMeshDetails::OnMeshClipWithMeshNodeComboBoxSelectionChanged)
+					]
 				]
-			]
-
-			+ SVerticalBox::Slot()
-			.AutoHeight()
-			.Padding(1.0f,3.0f,0.0f,0.0f)
-			[
-				TagView.ToSharedRef()
-			]
-		];
-
-		// Hidding the widgets if necessary
-		MaterialsSelector->SetVisibility(Node->bUseMaterials ? EVisibility::Visible : EVisibility::Collapsed);
-		TagView->SetVisibility(Node->bUseTags ? EVisibility::Visible : EVisibility::Collapsed);
-
+			];
+		}
 	}
 	else
 	{
@@ -286,16 +277,12 @@ void FCustomizableObjectNodeMeshClipWithMeshDetails::OnClippingMethodComboBoxSel
 			Node->bUseMaterials = false;
 			Node->bUseTags = true;
 			Node->ArrayMaterialNodeToClipWithID.Empty();
-
-			DetailBuilderPtr->ForceRefreshDetails();
 		}
 		else if(*Selection == "Material")
 		{
 			Node->bUseMaterials = true;
 			Node->bUseTags = false;
 			Node->Tags.Empty();
-
-			DetailBuilderPtr->ForceRefreshDetails();
 		}
 		else
 		{
@@ -303,9 +290,8 @@ void FCustomizableObjectNodeMeshClipWithMeshDetails::OnClippingMethodComboBoxSel
 			Node->bUseTags = true;
 		}
 
-		TagView->SetVisibility(Node->bUseTags ? EVisibility::Visible : EVisibility::Collapsed);
-		MaterialsSelector->SetVisibility(Node->bUseMaterials ? EVisibility::Visible : EVisibility::Collapsed);
-	}
+		DetailBuilderPtr->ForceRefreshDetails();
+	}	
 }
 
 
@@ -344,116 +330,4 @@ FReply FCustomizableObjectNodeMeshClipWithMeshDetails::OnAddTagPressed()
 }
 
 
-void FCustomizableObjectNodeMeshClipWithMeshDetails::CreateTagView()
-{
-	TagView = SNew(SVerticalBox);
-	TagView->AddSlot().AutoHeight()
-	[
-		SNew(SHorizontalBox)
-		+ SHorizontalBox::Slot()
-		.AutoWidth()
-		.Padding(0.0f, 3.0f, 0.0f, 0.0f)
-		[
-			SNew(STextBlock)
-			.Text(LOCTEXT("ClippingMethodTags", "Tags"))
-		]
-		+ SHorizontalBox::Slot()
-		.AutoWidth()
-		.Padding(5.0f, 0.0f, 0.0f, 0.0f)
-		[
-			SNew(SButton)
-			.OnClicked(this, &FCustomizableObjectNodeMeshClipWithMeshDetails::OnAddTagPressed)
-			.ToolTipText(LOCTEXT("AddNewTag", "Add New Tag"))
-			[
-				SNew(SImage)
-				.Image(UE_MUTABLE_GET_BRUSH(TEXT("Plus")))	
-			]
-		]
-	];
-
-	for (int32 i = 0; i < Node->Tags.Num(); ++i)
-	{
-		TagView->AddSlot()
-		.Padding(10.0f,5.0f,0.0f,0.0f)
-		[
-			SNew(STagView)
-			.TagValue(Node->Tags[i])
-			.TagIndex(i)
-			.Node(Node)
-			.DetailBuilderPtr(DetailBuilderPtr)
-		];
-	}
-}
-
-
-// Tag View Widget ---------------------------------------------------------
-
-void STagView::Construct(const FArguments& InArgs)
-{
-	TagValue = InArgs._TagValue;
-	TagIndex = InArgs._TagIndex;
-	Node = InArgs._Node;
-	DetailBuilderPtr = InArgs._DetailBuilderPtr;
-
-	FString TagName = FString::Printf(TEXT("Tag %d"), TagIndex);
-
-	ChildSlot
-	[
-		SNew(SHorizontalBox)
-		+ SHorizontalBox::Slot()
-		.AutoWidth()
-		[
-			SNew(STextBlock)
-			.Text(FText::FromString(TagName))
-		]
-
-		+ SHorizontalBox::Slot()
-		.AutoWidth()
-		.Padding(5.0f,0.0f,0.0f,0.0f)
-		[
-			SNew(SEditableTextBox)
-			.Text(FText::FromString(TagValue))
-			.MinDesiredWidth(50.0f)
-			.HintText(FText::FromString("New Tag"))
-			.OnTextCommitted(this, &STagView::OnTextBoxTextCommitted)
-		]
-		
-		+ SHorizontalBox::Slot()
-		.AutoWidth()
-		.Padding(5.0f, 0.0f, 0.0f, 0.0f)
-		[
-			SNew(SButton)
-			.OnClicked(this, &STagView::DeleteTag)
-			.ToolTipText(LOCTEXT("DeleteTag", "Delete Tag"))
-			[
-				SNew(SImage)
-				.Image(UE_MUTABLE_GET_BRUSH(TEXT("Cross")))
-			]
-		]
-	];
-
-}
-
-
-FReply STagView::DeleteTag()
-{
-	// Remove a tag and redraw the details view
-	Node->Tags.RemoveAt(TagIndex);
-	DetailBuilderPtr->ForceRefreshDetails();
-
-	return FReply::Handled();
-}
-
-
-void STagView::OnTextBoxTextCommitted(const FText& NewText, ETextCommit::Type InTextCommit)
-{
-	// Save the content when the widget loses the focus or we press enter
-	if (InTextCommit == ETextCommit::OnEnter || InTextCommit == ETextCommit::OnUserMovedFocus)
-	{
-		TagValue = NewText.ToString();
-		Node->Tags[TagIndex] = TagValue;
-	}
-}
-
-
 #undef LOCTEXT_NAMESPACE
diff --git a/Engine/Plugins/Experimental/Mutable/Source/CustomizableObjectEditor/Private/MuCOE/Nodes/CustomizableObjectNodeMeshClipWithMeshDetails.h b/Engine/Plugins/Experimental/Mutable/Source/CustomizableObjectEditor/Private/MuCOE/Nodes/CustomizableObjectNodeMeshClipWithMeshDetails.h
--- a/Engine/Plugins/Experimental/Mutable/Source/CustomizableObjectEditor/Private/MuCOE/Nodes/CustomizableObjectNodeMeshClipWithMeshDetails.h
+++ b/Engine/Plugins/Experimental/Mutable/Source/CustomizableObjectEditor/Private/MuCOE/Nodes/CustomizableObjectNodeMeshClipWithMeshDetails.h
@@ -15,41 +15,6 @@ class SVerticalBox;
 class UCustomizableObjectNodeMeshClipWithMesh;
 struct FAssetData;
 
-// Widget that represents the array of tags of the clip with mesh node
-class STagView : public SCompoundWidget
-{
-public:
-
-	SLATE_BEGIN_ARGS(class STagView) {}
-	SLATE_ARGUMENT(FString, TagValue)
-	SLATE_ARGUMENT(int32, TagIndex)
-	SLATE_ARGUMENT(UCustomizableObjectNodeMeshClipWithMesh*, Node)
-	SLATE_ARGUMENT(IDetailLayoutBuilder*, DetailBuilderPtr)
-	SLATE_END_ARGS()
-
-	void Construct(const FArguments& InArgs);
-
-	// Delete a tag of the array
-	FReply DeleteTag();
-
-	// Callback for the text box. Modifies the content of a tag.
-	void OnTextBoxTextCommitted(const FText& NewText, ETextCommit::Type InTextCommit);
-
-public:
-
-	// Pointer to the builder passed by CustomizeDetails method
-	IDetailLayoutBuilder* DetailBuilderPtr = nullptr;
-
-	// Pointer to the clip with mesh node
-	UCustomizableObjectNodeMeshClipWithMesh* Node;
-
-	// Value of the tag
-	FString TagValue;
-
-	// Index of the tag inside the array of tags
-	int32 TagIndex;
-};
-
 
 class FCustomizableObjectNodeMeshClipWithMeshDetails : public IDetailCustomization
 {
@@ -70,9 +35,6 @@ private:
 	/** Callback when new material name is selected */
 	void OnMeshClipWithMeshNodeComboBoxSelectionChanged(TSharedPtr<FString> Selection, ESelectInfo::Type SelectInfo);
 
-	/** Creates the array of widgets tags */
-	void CreateTagView();
-
 	/** Initializes the text of the method combobox and the view to display */
 	TSharedPtr<FString> SetInitialClippingMethod();
 
@@ -94,12 +56,6 @@ private:
 	/** Selected CO in the details tab of the Node node, for which to select a material that will be clipped by Node node */
 	class UCustomizableObject* SelectedCO;
 
-	/** Selector with all the materials of the Customizable Object */
-	TSharedPtr< SHorizontalBox > MaterialsSelector;
-
-	/** ComboBox with all the materials of the Customizable Object */
-	TSharedPtr< SVerticalBox > TagView;
-
 	/** Pointer to a FString to indicate that a method must be selected */
 	TSharedPtr<FString> SelectAMethod;
 
