[mutable] Added Warning category for the Instance update result
- Added method to replace the on success actions in the parts of the source code the successful value was used in the if condition.

diff --git a/Engine/Plugins/Experimental/Mutable/Source/CustomizableObject/Classes/MuCO/CustomizableObjectInstance.h b/Engine/Plugins/Experimental/Mutable/Source/CustomizableObject/Classes/MuCO/CustomizableObjectInstance.h
--- a/Engine/Plugins/Experimental/Mutable/Source/CustomizableObject/Classes/MuCO/CustomizableObjectInstance.h
+++ b/Engine/Plugins/Experimental/Mutable/Source/CustomizableObject/Classes/MuCO/CustomizableObjectInstance.h
@@ -50,7 +50,9 @@ enum class EUpdateRequired : uint8
 UENUM()
 enum class EUpdateResult : uint8
 {
-	Success, // There only exist one Success case. Any other new cases have to be errors.
+	Success, // Update finished without issues.
+	Warning, // Generic warning. Update finished but with warnings.
+	
 	Error, // Generic error.
 	ErrorOptimized, // The update was skipped since its result would have been the same as the current customization.
 	ErrorReplaced, // The update was replaced by a newer update request.
diff --git a/Engine/Plugins/Experimental/Mutable/Source/CustomizableObject/Classes/MuCO/CustomizableObjectSystem.h b/Engine/Plugins/Experimental/Mutable/Source/CustomizableObject/Classes/MuCO/CustomizableObjectSystem.h
--- a/Engine/Plugins/Experimental/Mutable/Source/CustomizableObject/Classes/MuCO/CustomizableObjectSystem.h
+++ b/Engine/Plugins/Experimental/Mutable/Source/CustomizableObject/Classes/MuCO/CustomizableObjectSystem.h
@@ -196,6 +196,12 @@ public:
 	
 	UFUNCTION(BlueprintCallable, BlueprintPure, Category = Status)
 	static UCustomizableObjectSystem* GetInstanceChecked();
+
+	/** Determines if the result of the instance update is valid or not.
+	* @return true if the result is successful or has warnings, false if the result is from the Error category
+	 */
+	UFUNCTION(BlueprintCallable, Category = Status)
+	static bool IsUpdateResultValid(const EUpdateResult UpdateResult);
 	
 	// Return true if the singleton has been created. It is different than GetInstance in that GetInstance will create it if it doesn't exist.
 	static bool IsCreated();
diff --git a/Engine/Plugins/Experimental/Mutable/Source/CustomizableObject/Private/MuCO/CustomizableObjectInstance.cpp b/Engine/Plugins/Experimental/Mutable/Source/CustomizableObject/Private/MuCO/CustomizableObjectInstance.cpp
--- a/Engine/Plugins/Experimental/Mutable/Source/CustomizableObject/Private/MuCO/CustomizableObjectInstance.cpp
+++ b/Engine/Plugins/Experimental/Mutable/Source/CustomizableObject/Private/MuCO/CustomizableObjectInstance.cpp
@@ -5032,14 +5032,16 @@ FGraphEventRef UCustomizableInstancePrivate::LoadAdditionalAssetsAsync(const TSh
 								else
 								{
 									// Two submeshes should not have the same animation slot index
-									FString ErrorMsg = FString::Printf(TEXT("Two submeshes have the same anim slot index [%s] in a Mutable Instance."), *SlotIndex.ToString());
-									UE_LOG(LogMutable, Error, TEXT("%s"), *ErrorMsg);
+									OperationData->UpdateResult = EUpdateResult::Warning;
+									
+									FString WarningMessage = FString::Printf(TEXT("Two submeshes have the same anim slot index [%s] in a Mutable Instance."), *SlotIndex.ToString());
+									UE_LOG(LogMutable, Warning, TEXT("%s"), *WarningMessage);
 #if WITH_EDITOR
 									FMessageLogModule& MessageLogModule = FModuleManager::LoadModuleChecked<FMessageLogModule>("MessageLog");
 									MessageLogModule.RegisterLogListing(FName("Mutable"), FText::FromString(FString("Mutable")));
 									FMessageLog MessageLog("Mutable");
 
-									MessageLog.Notify(FText::FromString(ErrorMsg), EMessageSeverity::Error, true);
+									MessageLog.Notify(FText::FromString(WarningMessage), EMessageSeverity::Warning, true);
 #endif
 								}
 							}
diff --git a/Engine/Plugins/Experimental/Mutable/Source/CustomizableObject/Private/MuCO/CustomizableObjectSystem.cpp b/Engine/Plugins/Experimental/Mutable/Source/CustomizableObject/Private/MuCO/CustomizableObjectSystem.cpp
--- a/Engine/Plugins/Experimental/Mutable/Source/CustomizableObject/Private/MuCO/CustomizableObjectSystem.cpp
+++ b/Engine/Plugins/Experimental/Mutable/Source/CustomizableObject/Private/MuCO/CustomizableObjectSystem.cpp
@@ -330,6 +330,12 @@ UCustomizableObjectSystem* UCustomizableObjectSystem::GetInstanceChecked()
 }
 
 
+bool UCustomizableObjectSystem::IsUpdateResultValid(const EUpdateResult UpdateResult)
+{
+	return UpdateResult == EUpdateResult::Success || UpdateResult == EUpdateResult::Warning;
+}
+
+
 UCustomizableInstanceLODManagementBase* UCustomizableObjectSystem::GetInstanceLODManagement() const
 {
 	return CurrentInstanceLODManagement.Get();
@@ -749,6 +755,7 @@ void FinishUpdateGlobal(const TSharedRef<FUpdateContextPrivate>& Context)
 		switch (Context->UpdateResult)
 		{
 		case EUpdateResult::Success:
+		case EUpdateResult::Warning:
 			PrivateInstance->SkeletalMeshStatus = ESkeletalMeshStatus::Success;
 
 			PrivateInstance->DescriptorHash = Context->InstanceDescriptorHash;
@@ -777,7 +784,7 @@ void FinishUpdateGlobal(const TSharedRef<FUpdateContextPrivate>& Context)
 		}
 	}
 
-	if (Context->UpdateResult == EUpdateResult::Success)
+	if (UCustomizableObjectSystem::IsUpdateResultValid(Context->UpdateResult))
 	{
 		// Call CustomizableObjectInstanceUsages updated callbacks.
 		for (TObjectIterator<UCustomizableObjectInstanceUsage> It; It; ++It) // Since iterating objects is expensive, for now CustomizableObjectInstanceUsage does not have a FinishUpdate function.
diff --git a/Engine/Plugins/Experimental/Mutable/Source/CustomizableObjectEditor/Private/MuCOE/CustomizableObjectEditorViewportClient.cpp b/Engine/Plugins/Experimental/Mutable/Source/CustomizableObjectEditor/Private/MuCOE/CustomizableObjectEditorViewportClient.cpp
--- a/Engine/Plugins/Experimental/Mutable/Source/CustomizableObjectEditor/Private/MuCOE/CustomizableObjectEditorViewportClient.cpp
+++ b/Engine/Plugins/Experimental/Mutable/Source/CustomizableObjectEditor/Private/MuCOE/CustomizableObjectEditorViewportClient.cpp
@@ -1646,7 +1646,7 @@ void FCustomizableObjectEditorViewportClient::OnInstanceForBakingUpdate(const FU
 	}
 
 	// Early exit if update result is not success
-	if (Result.UpdateResult != EUpdateResult::Success)
+	if ( !UCustomizableObjectSystem::IsUpdateResultValid(Result.UpdateResult) )
 	{
 		UE_LOG(LogMutable, Warning ,TEXT("Instance finished update with an error state : %s. Skipping instance baking"), *UEnum::GetValueAsString(Result.UpdateResult));
 		BakeTempInstance = nullptr;
diff --git a/Engine/Plugins/Experimental/Mutable/Source/MutableValidation/Private/MuV/COIBakingTestCommandlet.cpp b/Engine/Plugins/Experimental/Mutable/Source/MutableValidation/Private/MuV/COIBakingTestCommandlet.cpp
--- a/Engine/Plugins/Experimental/Mutable/Source/MutableValidation/Private/MuV/COIBakingTestCommandlet.cpp
+++ b/Engine/Plugins/Experimental/Mutable/Source/MutableValidation/Private/MuV/COIBakingTestCommandlet.cpp
@@ -22,8 +22,8 @@ void OnInstanceUpdate(const FUpdateContext& Result)
 {
 	const EUpdateResult InstanceUpdateResult = Result.UpdateResult;
 
-	UE_LOG(LogMutable,Display,TEXT("Instance finished update with state : %s."), *UEnum::GetValueAsString(Result.UpdateResult));
-	bWasInstanceUpdateSuccessful = InstanceUpdateResult == EUpdateResult::Success;
+	UE_LOG(LogMutable,Display,TEXT("Instance finished update with state : %s."), *UEnum::GetValueAsString(InstanceUpdateResult));
+	bWasInstanceUpdateSuccessful = UCustomizableObjectSystem::IsUpdateResultValid(InstanceUpdateResult);
 	
 	// Clear update flag so we can exit the update while loop
 	bIsInstanceBeingUpdated = false;
diff --git a/Engine/Plugins/Experimental/Mutable/Source/MutableValidation/Private/MuV/CustomizableObjectValidationCommandlet.cpp b/Engine/Plugins/Experimental/Mutable/Source/MutableValidation/Private/MuV/CustomizableObjectValidationCommandlet.cpp
--- a/Engine/Plugins/Experimental/Mutable/Source/MutableValidation/Private/MuV/CustomizableObjectValidationCommandlet.cpp
+++ b/Engine/Plugins/Experimental/Mutable/Source/MutableValidation/Private/MuV/CustomizableObjectValidationCommandlet.cpp
@@ -268,7 +268,7 @@ void UCustomizableObjectValidationCommandlet::OnInstanceUpdate(const FUpdateCont
 {
 	const FString InstanceName = InstanceBeingUpdated->GetName();
 	const EUpdateResult InstanceUpdateResult = Result.UpdateResult;
-	if (InstanceUpdateResult == EUpdateResult::Success)
+	if (UCustomizableObjectSystem::IsUpdateResultValid(InstanceUpdateResult))
 	{
 		UE_LOG(LogMutable,Display,TEXT("Instance %s finished update succesfully."),*InstanceName);
 
