[mutable] Added custom AssetUSerData to Customizable Instance Skeletal Meshes.

#rnx

diff --git a/Engine/Plugins/Experimental/Mutable/Source/CustomizableObject/Private/MuCO/CustomizableObjectInstance.cpp b/Engine/Plugins/Experimental/Mutable/Source/CustomizableObject/Private/MuCO/CustomizableObjectInstance.cpp
--- a/Engine/Plugins/Experimental/Mutable/Source/CustomizableObject/Private/MuCO/CustomizableObjectInstance.cpp
+++ b/Engine/Plugins/Experimental/Mutable/Source/CustomizableObject/Private/MuCO/CustomizableObjectInstance.cpp
@@ -26,6 +26,7 @@
 #include "MuCO/CustomizableObjectMipDataProvider.h"
 #include "MuCO/CustomizableObjectPrivate.h"
 #include "MuCO/CustomizableObjectInstanceUsage.h"
+#include "MuCO/CustomizableObjectInstanceAssetUserData.h"
 #include "MuCO/DefaultImageProvider.h"
 #include "MuCO/ICustomizableObjectModule.h"
 #include "MuCO/UnrealConversionUtils.h"
@@ -3326,6 +3327,14 @@ void UCustomizableInstancePrivate::InitSkeletalMeshData(const TSharedRef<FUpdate
 		{
 			SkeletalMesh->AddAssetUserData(AssetUserData);
 		}
+
+		//Custom Asset User Data
+		if (OperationData->Instance->GetAnimationGameplayTags().Num())
+		{
+			UCustomizableObjectInstanceUserData* InstanceData = NewObject<UCustomizableObjectInstanceUserData>(SkeletalMesh, NAME_None, RF_Public | RF_Transactional);
+			InstanceData->SetAnimationGameplayTags(OperationData->Instance->GetAnimationGameplayTags());
+			SkeletalMesh->AddAssetUserData(InstanceData);
+		}
 	}
 
 	// Allocate resources for rendering and add LOD Info
diff --git a/Engine/Plugins/Experimental/Mutable/Source/CustomizableObject/Public/MuCO/CustomizableObjectInstanceAssetUserData.h b/Engine/Plugins/Experimental/Mutable/Source/CustomizableObject/Public/MuCO/CustomizableObjectInstanceAssetUserData.h
new file mode 100644
--- /dev/null
+++ b/Engine/Plugins/Experimental/Mutable/Source/CustomizableObject/Public/MuCO/CustomizableObjectInstanceAssetUserData.h
@@ -0,0 +1,34 @@
+// Copyright Epic Games, Inc. All Rights Reserved.
+
+#include "Engine/AssetUserData.h"
+#include "GameplayTagContainer.h"
+
+#include "CustomizableObjectInstanceAssetUserData.generated.h"
+
+
+UCLASS(BlueprintType)
+class CUSTOMIZABLEOBJECT_API UCustomizableObjectInstanceUserData : public UAssetUserData
+{
+	GENERATED_BODY()
+
+public:
+
+	/** Return the list of tags for this instance. */
+	UFUNCTION(BlueprintCallable, Category = CustomizableObjectInstance)
+	const FGameplayTagContainer& GetAnimationGameplayTags() const
+	{
+		return AnimBPGameplayTags;
+	};
+
+	/** Sets the list of tags for this instance. */
+	UFUNCTION(BlueprintCallable, Category = CustomizableObjectInstance)
+	void SetAnimationGameplayTags(const FGameplayTagContainer& InstanceTags)
+	{
+		AnimBPGameplayTags = InstanceTags;
+	};
+
+private:
+
+	UPROPERTY()
+	FGameplayTagContainer AnimBPGameplayTags;
+};
diff --git a/Engine/Plugins/Experimental/Mutable/Source/CustomizableObjectEditor/Public/MuCOE/CustomizableObjectInstanceBakingUtils.cpp b/Engine/Plugins/Experimental/Mutable/Source/CustomizableObjectEditor/Public/MuCOE/CustomizableObjectInstanceBakingUtils.cpp
--- a/Engine/Plugins/Experimental/Mutable/Source/CustomizableObjectEditor/Public/MuCOE/CustomizableObjectInstanceBakingUtils.cpp
+++ b/Engine/Plugins/Experimental/Mutable/Source/CustomizableObjectEditor/Public/MuCOE/CustomizableObjectInstanceBakingUtils.cpp
@@ -13,6 +13,7 @@
 #include "Materials/Material.h"
 
 #include "MuCO/CustomizableObject.h"
+#include "MuCO/CustomizableObjectInstanceAssetUserData.h"
 #include "MuCO/CustomizableObjectMipDataProvider.h"
 #include "MuT/UnrealPixelFormatOverride.h"
 #include "Rendering/SkeletalMeshModel.h"
@@ -744,6 +745,16 @@ void BakeCustomizableObjectInstance(
 				}
 			}
 
+			// Add Instance Info in a custom AssetUserData
+			{
+				if (InInstance.GetAnimationGameplayTags().Num())
+				{
+					UCustomizableObjectInstanceUserData* InstanceData = NewObject<UCustomizableObjectInstanceUserData>(SkeletalMesh, NAME_None, RF_Public | RF_Transactional);
+					InstanceData->SetAnimationGameplayTags(InInstance.GetAnimationGameplayTags());
+					SkeletalMesh->AddAssetUserData(InstanceData);
+				}
+			}
+
 			// Generate render data
 			SkeletalMesh->Build();
 		}
