[Mutable] Fix crash when opening a CO Instance before loading the asset registry.

#rnx

diff --git a/Engine/Plugins/Experimental/Mutable/Source/CustomizableObject/Classes/MuCO/CustomizableObject.h b/Engine/Plugins/Experimental/Mutable/Source/CustomizableObject/Classes/MuCO/CustomizableObject.h
--- a/Engine/Plugins/Experimental/Mutable/Source/CustomizableObject/Classes/MuCO/CustomizableObject.h
+++ b/Engine/Plugins/Experimental/Mutable/Source/CustomizableObject/Classes/MuCO/CustomizableObject.h
@@ -1370,7 +1370,7 @@ private:
 	// This is a manual version number for the binary blobs in this asset.
 	// Increasing it invalidates all the previously compiled models.
 	// Warning: If while merging code both versions have changed, take the highest+1.
-	static const int32 CurrentSupportedVersion = 418;
+	static const int32 CurrentSupportedVersion = 419;
 
 	// Compile the object for a specific platform - Compile for Cook Customizable Object
 	void CompileForTargetPlatform(const ITargetPlatform* TargetPlatform);
diff --git a/Engine/Plugins/Experimental/Mutable/Source/CustomizableObject/Classes/MuCO/CustomizableObjectInstance.h b/Engine/Plugins/Experimental/Mutable/Source/CustomizableObject/Classes/MuCO/CustomizableObjectInstance.h
--- a/Engine/Plugins/Experimental/Mutable/Source/CustomizableObject/Classes/MuCO/CustomizableObjectInstance.h
+++ b/Engine/Plugins/Experimental/Mutable/Source/CustomizableObject/Classes/MuCO/CustomizableObjectInstance.h
@@ -682,12 +682,6 @@ public:
 	bool bSelectedProfileDirty = false;
 #endif 
 
-private:
-
-#if WITH_EDITOR
-	void BindPostCompileDelegate(UCustomizableObject* CustomizableObject);
-#endif
-
 private:
 	UPROPERTY()
 	FCustomizableObjectInstanceDescriptor Descriptor;
diff --git a/Engine/Plugins/Experimental/Mutable/Source/CustomizableObject/Internal/MuCO/CustomizableInstancePrivateData.h b/Engine/Plugins/Experimental/Mutable/Source/CustomizableObject/Internal/MuCO/CustomizableInstancePrivateData.h
--- a/Engine/Plugins/Experimental/Mutable/Source/CustomizableObject/Internal/MuCO/CustomizableInstancePrivateData.h
+++ b/Engine/Plugins/Experimental/Mutable/Source/CustomizableObject/Internal/MuCO/CustomizableInstancePrivateData.h
@@ -7,6 +7,7 @@
 #include "Rendering/SkeletalMeshLODModel.h"
 #include "Rendering/SkeletalMeshModel.h"
 #include "MuCO/CustomizableObjectInstance.h"
+#include "MuCO/CustomizableObjectPrivate.h"
 #include "MuR/Instance.h"
 #include "GameplayTagContainer.h"
 #include "MuCO/DescriptorHash.h"
@@ -196,8 +197,10 @@ public:
 
 	virtual void PostDuplicate(bool bDuplicateForPIE) override;
 
-	void OnPostCompile();
+	void BindObjectDelegates(UCustomizableObject* CurrentCustomizableObject, UCustomizableObject* NewCustomizableObject);
 
+	void OnPostCompile();
+	void OnObjectStatusChanged(FCustomizableObjectStatus::EState Previous, FCustomizableObjectStatus::EState Next);
 #endif
 
 	/** Invalidates the previously generated data and retrieves information from the CObject after specific actions.
diff --git a/Engine/Plugins/Experimental/Mutable/Source/CustomizableObject/Private/MuCO/CustomizableObject.cpp b/Engine/Plugins/Experimental/Mutable/Source/CustomizableObject/Private/MuCO/CustomizableObject.cpp
--- a/Engine/Plugins/Experimental/Mutable/Source/CustomizableObject/Private/MuCO/CustomizableObject.cpp
+++ b/Engine/Plugins/Experimental/Mutable/Source/CustomizableObject/Private/MuCO/CustomizableObject.cpp
@@ -155,7 +155,6 @@ void UCustomizableObject::PreSave(FObjectPreSaveContext ObjectSaveContext)
 				*ObjectSaveContext.GetTargetPlatform()->PlatformName());
 			
 			ClearCompiledData();
-			SetModel(nullptr);
 		}
 	}
 #endif
@@ -566,7 +565,7 @@ void UCustomizableObject::SaveCompiledData(FArchive& MemoryWriter, bool bIsCooki
 
 void UCustomizableObject::LoadCompiledData(FArchive& MemoryReader, const ITargetPlatform* InTargetPlatform, bool bIsCooking)
 {
-	Private->SetModel(nullptr, FGuid());
+	TSharedPtr<mu::Model, ESPMode::ThreadSafe> LoadedModel;
 	ClearCompiledData();
 
 	MutableCompiledDataStreamHeader Header;
@@ -721,13 +720,12 @@ void UCustomizableObject::LoadCompiledData(FArchive& MemoryReader, const ITarget
 		{
 			UnrealMutableInputStream stream(MemoryReader);
 			mu::InputArchive arch(&stream);
-			TSharedPtr<mu::Model, ESPMode::ThreadSafe> Model = mu::Model::StaticUnserialise( arch );
-
-			Private->SetModel(Model, GenerateIdentifier(*this));
+			LoadedModel = mu::Model::StaticUnserialise( arch );
 		}
 	}
 	
-	UpdateParameterPropertiesFromModel(GetModel());
+	UpdateParameterPropertiesFromModel(LoadedModel);
+	Private->SetModel(LoadedModel, GenerateIdentifier(*this));
 }
 
 
@@ -837,9 +835,6 @@ void UCustomizableObject::CompileForTargetPlatform(const ITargetPlatform* Target
 		||
 		(Relevancy == ECustomizableObjectRelevancy::ClientOnly && !TargetPlatform->IsServerOnly());
 
-	// Discard any older compilation
-	Private->SetModel(nullptr, FGuid());
-
 	if (bIsRootObject && bIsRelevantForThisTarget)
 	{
 		FCompilationOptions Options;
@@ -911,9 +906,6 @@ bool UCustomizableObject::ConditionalAutoCompile()
 		return false;
 	}
 
-	// Discard any older compilation
-	Private->SetModel(nullptr, FGuid());
-
 	// Sync/Async compilation
 	if (System->IsAutoCompilationSync())
 	{
diff --git a/Engine/Plugins/Experimental/Mutable/Source/CustomizableObject/Private/MuCO/CustomizableObjectInstance.cpp b/Engine/Plugins/Experimental/Mutable/Source/CustomizableObject/Private/MuCO/CustomizableObjectInstance.cpp
--- a/Engine/Plugins/Experimental/Mutable/Source/CustomizableObject/Private/MuCO/CustomizableObjectInstance.cpp
+++ b/Engine/Plugins/Experimental/Mutable/Source/CustomizableObject/Private/MuCO/CustomizableObjectInstance.cpp
@@ -266,7 +266,7 @@ void UCustomizableObjectInstance::SetDescriptor(const FCustomizableObjectInstanc
 
 #if WITH_EDITOR
 	// Bind a lambda to the PostCompileDelegate and unbind from the previous object if any.
-	BindPostCompileDelegate(InCustomizableObject);
+	PrivateData->BindObjectDelegates(GetCustomizableObject(), InCustomizableObject);
 #endif
 
 	Descriptor = InDescriptor;
@@ -328,19 +328,34 @@ void UCustomizableInstancePrivate::OnPostCompile()
 }
 
 
-void UCustomizableObjectInstance::BindPostCompileDelegate(UCustomizableObject* InCustomizableObject)
+void UCustomizableInstancePrivate::OnObjectStatusChanged(FCustomizableObjectStatus::EState Previous, FCustomizableObjectStatus::EState Next)
 {
+	if (Previous != Next && Next == FCustomizableObjectStatus::EState::ModelLoaded)
+	{
+		OnPostCompile();
+	}
+}
+
+
+void UCustomizableInstancePrivate::BindObjectDelegates(UCustomizableObject*  CurrentCustomizableObject, UCustomizableObject* NewCustomizableObject)
+{
+	if (CurrentCustomizableObject == NewCustomizableObject)
+	{
+		return;
+	}
+
 	// Unbind callback from the previous CO
-	if (UCustomizableObject* CurrentObject = GetCustomizableObject())
+	if (CurrentCustomizableObject)
 	{
-		CurrentObject->PostCompileDelegate.RemoveAll(GetPrivate());
+		CurrentCustomizableObject->PostCompileDelegate.RemoveAll(this);
+		CurrentCustomizableObject->GetPrivate()->Status.GetOnStateChangedDelegate().RemoveAll(this);
 	}
 
 	// Bind callback to the new CO
-	if (InCustomizableObject &&
-		!InCustomizableObject->PostCompileDelegate.IsBoundToObject(GetPrivate()))
+	if (NewCustomizableObject)
 	{
-		InCustomizableObject->PostCompileDelegate.AddUObject(GetPrivate(), &UCustomizableInstancePrivate::OnPostCompile);
+		NewCustomizableObject->PostCompileDelegate.AddUObject(this, &UCustomizableInstancePrivate::OnPostCompile);
+		NewCustomizableObject->GetPrivate()->Status.GetOnStateChangedDelegate().AddUObject(this, &UCustomizableInstancePrivate::OnObjectStatusChanged);
 	}
 }
 
@@ -394,11 +409,6 @@ bool UCustomizableObjectInstance::IsEditorOnly() const
 
 void UCustomizableObjectInstance::BeginDestroy()
 {
-#if WITH_EDITOR
-	// Unbind PostCompileDelegate
-	BindPostCompileDelegate(nullptr);
-#endif
-
 	BeginDestroyDelegate.Broadcast(this);
 	BeginDestroyNativeDelegate.Broadcast(this);
 
@@ -407,6 +417,11 @@ void UCustomizableObjectInstance::BeginDestroy()
 
 	if (PrivateData)
 	{
+#if WITH_EDITOR
+		// Unbind Object delegates
+		PrivateData->BindObjectDelegates(GetCustomizableObject(), nullptr);
+#endif
+
 		if (PrivateData->StreamingHandle.IsValid() && PrivateData->StreamingHandle->IsActive())
 		{
 			PrivateData->StreamingHandle->CancelHandle();
@@ -560,7 +575,7 @@ void UCustomizableObjectInstance::PostLoad()
 	Super::PostLoad();
 
 #if WITH_EDITOR
-	BindPostCompileDelegate(GetCustomizableObject());
+	PrivateData->BindObjectDelegates(nullptr, GetCustomizableObject());
 #endif
 
 	// Skip the cost of ReloadParameters in the cook commandlet; it will be reloaded during PreSave. For cooked runtime
@@ -685,7 +700,7 @@ void UCustomizableObjectInstance::SetObject(UCustomizableObject* InObject)
 {
 #if WITH_EDITOR
 	// Bind a lambda to the PostCompileDelegate and unbind from the previous object if any.
-	BindPostCompileDelegate(InObject);
+	PrivateData->BindObjectDelegates(GetCustomizableObject(), InObject);
 #endif
 
 	Descriptor.SetCustomizableObject(*InObject);
