[Mutable] Remove references gathered by editor compilations from the cooked builds. Move data derived from the compilation from the CustomizableObject to the CustomizableObjectPrivate.

#rnx

diff --git a/Engine/Plugins/Experimental/Mutable/Source/CustomizableObject/Classes/MuCO/CustomizableObject.h b/Engine/Plugins/Experimental/Mutable/Source/CustomizableObject/Classes/MuCO/CustomizableObject.h
--- a/Engine/Plugins/Experimental/Mutable/Source/CustomizableObject/Classes/MuCO/CustomizableObject.h
+++ b/Engine/Plugins/Experimental/Mutable/Source/CustomizableObject/Classes/MuCO/CustomizableObject.h
@@ -17,21 +17,13 @@
 #define UE_MUTABLE_MAX_OPTIMIZATION			2
 
 class FReply;
-class FMemoryReaderView;
-class FMemoryWriter64;
 class FObjectPreSaveContext;
-class FObjectPreSaveRootContext;
 class FText;
 class IAsyncReadFileHandle;
 class ITargetPlatform;
-class UAnimInstance;
-class UAssetUserData;
 class UCustomizableObject;
 class UEdGraph;
-class UMaterialInterface;
-class UPhysicsAsset;
 class USkeletalMesh;
-class USkeleton;
 struct FFrame;
 struct FStreamableHandle;
 template <typename FuncType> class TFunctionRef;
@@ -169,7 +161,7 @@ struct FCompilationOptions
 	bool bIsCooking = false;
 
 	// This can be set for additional settings
-	const class ITargetPlatform* TargetPlatform = nullptr;
+	const ITargetPlatform* TargetPlatform = nullptr;
 
 	// Used to prevent the current model from being stored in the editor's streamed data and cache when the compilation finishes
 	bool bDontUpdateStreamedDataAndCache = false;
@@ -268,121 +260,6 @@ struct MutableCompiledDataStreamHeader
 };
 
 
-USTRUCT()
-struct FMutableModelImageProperties
-{
-	GENERATED_USTRUCT_BODY()
-
-	FMutableModelImageProperties()
-		: Filter(TF_Default)
-		, SRGB(0)
-		, FlipGreenChannel(0)
-		, IsPassThrough(0)
-		, LODBias(0)
-		, LODGroup(TEXTUREGROUP_World)
-		, AddressX(TA_Clamp)
-		, AddressY(TA_Clamp)
-	{}
-
-	FMutableModelImageProperties(const FString& InTextureParameterName, TextureFilter InFilter, uint32 InSRGB, 
-		uint32 InFlipGreenChannel, uint32 bInIsPassThrough, int32 InLODBias, TEnumAsByte<enum TextureGroup> InLODGroup,
-		TEnumAsByte<enum TextureAddress> InAddressX, TEnumAsByte<enum TextureAddress> InAddressY)
-		: TextureParameterName(InTextureParameterName)
-		, Filter(InFilter)
-		, SRGB(InSRGB)
-		, FlipGreenChannel(InFlipGreenChannel)
-		, IsPassThrough(bInIsPassThrough)
-		, LODBias(InLODBias)
-		, LODGroup(InLODGroup)
-		, AddressX(InAddressX)
-		, AddressY(InAddressY)
-	{}
-
-	// Name in the material.
-	UPROPERTY()
-	FString TextureParameterName;
-
-	UPROPERTY()
-	TEnumAsByte<enum TextureFilter> Filter;
-
-	UPROPERTY()
-	uint32 SRGB : 1;
-
-	UPROPERTY()
-	uint32 FlipGreenChannel : 1;
-
-	UPROPERTY()
-	uint32 IsPassThrough : 1;
-
-	UPROPERTY()
-	int32 LODBias;
-
-	UPROPERTY()
-	TEnumAsByte<enum TextureGroup> LODGroup;
-
-	UPROPERTY()
-	TEnumAsByte<enum TextureAddress> AddressX;
-
-	UPROPERTY()
-	TEnumAsByte<enum TextureAddress> AddressY;
-
-	bool operator!=(const FMutableModelImageProperties& rhs)
-	{
-		return
-			TextureParameterName != rhs.TextureParameterName ||
-			Filter != rhs.Filter ||
-			SRGB != rhs.SRGB ||
-			FlipGreenChannel != rhs.FlipGreenChannel ||
-			IsPassThrough != rhs.IsPassThrough ||
-			LODBias != rhs.LODBias ||
-			LODGroup != rhs.LODGroup ||
-			AddressX != rhs.AddressX ||
-			AddressY != rhs.AddressY;
-	}
-
-	friend FArchive& operator<<(FArchive& Ar, FMutableModelImageProperties& ImageProps)
-	{
-		Ar << ImageProps.TextureParameterName;
-		Ar << ImageProps.Filter;
-
-		// Bitfields don't serialize automatically with FArchive
-		if (Ar.IsLoading())
-		{
-			int32 Aux = 0;
-			Ar << Aux;
-			ImageProps.SRGB = Aux;
-
-			Aux = 0;
-			Ar << Aux;
-			ImageProps.FlipGreenChannel = Aux;
-
-			Aux = 0;
-			Ar << Aux;
-			ImageProps.IsPassThrough = Aux;
-		}
-		else
-		{
-			int32 Aux = ImageProps.SRGB;
-			Ar << Aux;
-
-			Aux = ImageProps.FlipGreenChannel;
-			Ar << Aux;
-
-			Aux = ImageProps.IsPassThrough;
-			Ar << Aux;
-		}
-
-		Ar << ImageProps.LODBias;
-		Ar << ImageProps.LODGroup;
-
-		Ar << ImageProps.AddressX;
-		Ar << ImageProps.AddressY;
-
-		return Ar;
-	}
-};
-
-
 USTRUCT()
 struct FMutableModelParameterValue
 {
@@ -414,56 +291,6 @@ struct FMutableModelParameterProperties
 };
 
 
-USTRUCT()
-struct CUSTOMIZABLEOBJECT_API FAnimBpOverridePhysicsAssetsInfo
-{
-	GENERATED_BODY()
-
-	UPROPERTY()
-	TSoftClassPtr<UAnimInstance> AnimInstanceClass;
-	
-	UPROPERTY()
-	TSoftObjectPtr<UPhysicsAsset> SourceAsset;
-	
-	UPROPERTY()
-	int32 PropertyIndex = -1;
-
-	friend FArchive& operator<<(FArchive& Ar, FAnimBpOverridePhysicsAssetsInfo& Info)
-	{
-		FString AnimInstanceClassPathString;
-		FString PhysicsAssetPathString;
-
-		if (Ar.IsLoading())
-		{	
-			Ar << AnimInstanceClassPathString;
-			Ar << PhysicsAssetPathString;
-			Ar << Info.PropertyIndex;
-
-			Info.AnimInstanceClass = TSoftClassPtr<UAnimInstance>(AnimInstanceClassPathString);
-			Info.SourceAsset = TSoftObjectPtr<UPhysicsAsset>(PhysicsAssetPathString);
-		}
-
-		if (Ar.IsSaving())
-		{
-			AnimInstanceClassPathString = Info.AnimInstanceClass.ToString();
-			PhysicsAssetPathString = Info.SourceAsset.ToString();
-
-			Ar << AnimInstanceClassPathString;
-			Ar << PhysicsAssetPathString;
-			Ar << Info.PropertyIndex;
-		}
-
-		return Ar;
-	}
-
-	friend bool operator==(const FAnimBpOverridePhysicsAssetsInfo& Lhs, const FAnimBpOverridePhysicsAssetsInfo& Rhs)
-	{
-		return Lhs.AnimInstanceClass == Rhs.AnimInstanceClass && 
-			   Lhs.SourceAsset	     == Rhs.SourceAsset       && 
-		 	   Lhs.PropertyIndex     == Rhs.PropertyIndex;
-	}
-};
-
 USTRUCT()
 struct FMorphTargetInfo
 {
@@ -623,42 +450,6 @@ struct FCustomizableObjectMeshToMeshVertData
 };
 template<> struct TCanBulkSerialize<FCustomizableObjectMeshToMeshVertData> { enum { Value = true }; };
 
-
-USTRUCT()
-struct FMutableSkinWeightProfileInfo
-{
-	GENERATED_USTRUCT_BODY()
-
-	FMutableSkinWeightProfileInfo() {};
-
-	FMutableSkinWeightProfileInfo(FName InName, bool InDefaultProfile, int8 InDefaultProfileFromLODIndex) : Name(InName),
-		DefaultProfile(InDefaultProfile), DefaultProfileFromLODIndex(InDefaultProfileFromLODIndex) {};
-
-	UPROPERTY()
-	FName Name;
-
-	UPROPERTY()
-	bool DefaultProfile = false;
-
-	UPROPERTY(meta = (ClampMin = 0))
-	int8 DefaultProfileFromLODIndex = 0;
-
-	friend FArchive& operator<<(FArchive& Ar, FMutableSkinWeightProfileInfo& Info)
-	{
-		Ar << Info.Name;
-		Ar << Info.DefaultProfile;
-		Ar << Info.DefaultProfileFromLODIndex;
-
-		return Ar;
-	}
-
-	bool operator==(const FMutableSkinWeightProfileInfo& Other) const
-	{
-		return Name == Other.Name;
-	}
-};
-
-
 USTRUCT()
 struct CUSTOMIZABLEOBJECT_API FMutableStreamableBlock
 {
@@ -696,240 +487,6 @@ struct CUSTOMIZABLEOBJECT_API FMutableCachedPlatformData
 };
 
 
-USTRUCT()
-struct FMutableRefLODInfo
-{
-	GENERATED_BODY()
-
-	UPROPERTY()
-	float ScreenSize = 0.f;
-
-	UPROPERTY()
-	float LODHysteresis = 0.f;
-
-	UPROPERTY()
-	bool bSupportUniformlyDistributedSampling = false;
-
-	UPROPERTY()
-	bool bAllowCPUAccess = false;
-
-#if WITH_EDITORONLY_DATA
-	friend FArchive& operator<<(FArchive& Ar, FMutableRefLODInfo& Data)
-	{
-		Ar << Data.ScreenSize;
-		Ar << Data.LODHysteresis;
-		Ar << Data.bSupportUniformlyDistributedSampling;
-		Ar << Data.bAllowCPUAccess;
-
-		return Ar;
-	}
-#endif
-};
-
-
-USTRUCT()
-struct FMutableRefLODRenderData
-{
-	GENERATED_BODY()
-
-	UPROPERTY()
-	bool bIsLODOptional = false;
-
-	UPROPERTY()
-	bool bStreamedDataInlined = false;
-	
-#if WITH_EDITORONLY_DATA
-	friend FArchive& operator<<(FArchive& Ar, FMutableRefLODRenderData& Data)
-	{
-		Ar << Data.bIsLODOptional;
-		Ar << Data.bStreamedDataInlined;
-
-		return Ar;
-	}
-#endif
-};
-
-
-USTRUCT()
-struct FMutableRefLODData
-{
-	GENERATED_BODY()
-
-	UPROPERTY()
-	FMutableRefLODInfo LODInfo;
-
-	UPROPERTY()
-	FMutableRefLODRenderData RenderData;
-	
-#if WITH_EDITORONLY_DATA
-	friend FArchive& operator<<(FArchive& Ar, FMutableRefLODData& Data)
-	{
-		Ar << Data.LODInfo;
-		Ar << Data.RenderData;
-		
-		return Ar;
-	}
-#endif
-};
-
-
-USTRUCT()
-struct FMutableRefSocket
-{
-	GENERATED_BODY()
-
-	UPROPERTY()
-	FName SocketName;
-	UPROPERTY()
-	FName BoneName;
-
-	UPROPERTY()
-	FVector RelativeLocation = FVector::ZeroVector;
-	UPROPERTY()
-	FRotator RelativeRotation = FRotator::ZeroRotator;
-	UPROPERTY()
-	FVector RelativeScale = FVector::ZeroVector;;
-
-	UPROPERTY()
-	bool bForceAlwaysAnimated = false;
-
-	// When two sockets have the same name, the one with higher priority will be picked and the other discarded
-	UPROPERTY()
-	int32 Priority = -1;
-
-	bool operator ==(const FMutableRefSocket& Other) const
-	{
-		if (
-			SocketName == Other.SocketName &&
-			BoneName == Other.BoneName &&
-			RelativeLocation == Other.RelativeLocation &&
-			RelativeRotation == Other.RelativeRotation &&
-			RelativeScale == Other.RelativeScale &&
-			bForceAlwaysAnimated == Other.bForceAlwaysAnimated &&
-			Priority == Other.Priority)
-		{
-			return true;
-		}
-
-		return false;
-	}
-	
-#if WITH_EDITORONLY_DATA
-	friend FArchive& operator<<(FArchive& Ar, FMutableRefSocket& Data)
-	{
-		Ar << Data.SocketName;
-		Ar << Data.BoneName;
-		Ar << Data.RelativeLocation;
-		Ar << Data.RelativeRotation;
-		Ar << Data.RelativeScale;
-		Ar << Data.bForceAlwaysAnimated;
-		Ar << Data.Priority;
-
-		return Ar;
-	}
-#endif
-};
-
-
-USTRUCT()
-struct FMutableRefSkeletalMeshSettings
-{
-	GENERATED_BODY()
-
-	UPROPERTY()
-	bool bEnablePerPolyCollision = false;
-
-	UPROPERTY()
-	float DefaultUVChannelDensity = 0.f;
-
-#if WITH_EDITORONLY_DATA
-	friend FArchive& operator<<(FArchive& Ar, FMutableRefSkeletalMeshSettings& Data)
-	{
-		Ar << Data.bEnablePerPolyCollision;
-		Ar << Data.DefaultUVChannelDensity;
-
-		return Ar;
-	}
-#endif
-};
-
-
-USTRUCT()
-struct FMutableRefAssetUserData
-{
-	GENERATED_BODY()
-
-	UPROPERTY()
-	TObjectPtr<UAssetUserData> AssetUserData;
-
-#if WITH_EDITORONLY_DATA
-	FString ClassPath;
-	TArray<uint8> Bytes;
-
-	friend FArchive& operator<<(FArchive& Ar, FMutableRefAssetUserData& Data);
-
-	void InitResources(UCustomizableObject* InOuter);
-#endif
-
-};
-
-
-USTRUCT()
-struct FMutableRefSkeletalMeshData
-{
-	GENERATED_BODY()
-
-	// Reference Skeletal Mesh
-	UPROPERTY()
-	TObjectPtr<USkeletalMesh> SkeletalMesh;
-
-	// Path to load the ReferenceSkeletalMesh
-	UPROPERTY()
-	FSoftObjectPath SkeletalMeshAssetPath;
-
-	// LOD info
-	UPROPERTY()
-	TArray<FMutableRefLODData> LODData;
-	
-	// Sockets
-	UPROPERTY()
-	TArray<FMutableRefSocket> Sockets;
-
-	// Bounding Box
-	UPROPERTY()
-	FBoxSphereBounds Bounds = FBoxSphereBounds(ForceInitToZero);
-
-	// Settings
-	UPROPERTY()
-	FMutableRefSkeletalMeshSettings Settings;
-
-	// Skeleton, must be stored in the ReferencedSkeletons too
-	UPROPERTY()
-	TSoftObjectPtr<USkeleton> Skeleton;
-	
-	// PhysicsAsset, must be stored in the PhysicsAssetMap too
-	UPROPERTY()
-	TSoftObjectPtr<UPhysicsAsset> PhysicsAsset;
-	
-	// Post Processing AnimBP, must be stored in the AnimBPAssetsMap too
-	UPROPERTY() 
-	TSoftClassPtr<UAnimInstance> PostProcessAnimInst;
-	
-	// Shadow PhysicsAsset, must be stored in the PhysicsAssetMap too
-	UPROPERTY()
-	TSoftObjectPtr<UPhysicsAsset> ShadowPhysicsAsset;
-
-	// Asset user data
-	UPROPERTY()
-	TArray<FMutableRefAssetUserData> AssetUserData;
-
-#if WITH_EDITORONLY_DATA
-	friend FArchive& operator<<(FArchive& Ar, FMutableRefSkeletalMeshData& Data);
-
-	void InitResources(UCustomizableObject* InOuter, const ITargetPlatform* InTargetPlatform);
-#endif
-
-};
 
 USTRUCT()
 struct FMutableLODSettings
@@ -1049,7 +606,7 @@ public:
 	* For more information on this topic read the Basic Concepts at work.anticto.com.
 	*/
 	UPROPERTY()
-	TObjectPtr<class USkeletalMesh> ReferenceSkeletalMesh_DEPRECATED;
+	TObjectPtr<USkeletalMesh> ReferenceSkeletalMesh_DEPRECATED;
 
 	/** All the SkeletalMeshes generated for this CustomizableObject instances will use the Reference Skeletal Mesh
 	* properties for everything that Mutable doesn't create or modify. This includes data like LOD distances, Physics
@@ -1062,47 +619,16 @@ public:
 	*
 	* For more information on this topic read the Basic Concepts at work.anticto.com.
 	*/
-	UPROPERTY(EditAnywhere, Category=CustomizableObject)
-	TArray< TObjectPtr<class USkeletalMesh> > ReferenceSkeletalMeshes;
+	UPROPERTY(EditAnywhere, Category = CustomizableObject)
+	TArray<TObjectPtr<USkeletalMesh>> ReferenceSkeletalMeshes;
 #endif
 
-	/** All the SkeletalMeshes generated for this CustomizableObject instances will use the Reference Skeletal Mesh
-	 * properties for everything that Mutable doesn't create or modify. This struct stores the information used from
-	 * the Reference Skeletal Meshes to avoid having them loaded at all times. This includes data like LOD distances,
-	 * LOD render data settings, Mesh sockets, Bounding volumes, etc.
-	 */
-	UPROPERTY()
-	TArray<FMutableRefSkeletalMeshData> ReferenceSkeletalMeshesData;
-
-	/** List of Materials referenced by this or any child customizable object. */
-	UPROPERTY()
-	TArray<TSoftObjectPtr<UMaterialInterface>> ReferencedMaterials;
-
-	/** List of Material slot names for the materials referenced by this or any child customizable object. */
-	UPROPERTY()
-	TArray<FName> ReferencedMaterialSlotNames;
-
-	/** List of skeletons referenced by any of the parts of this customizable object. 
-	 * The position in this array is used as skeleton ID passed to Mutable bones.
-	 */
-	UPROPERTY()
-	TArray<TSoftObjectPtr<USkeleton>> ReferencedSkeletons;
-
-	UPROPERTY()
-	TArray<TSoftObjectPtr<UTexture>> ReferencedPassThroughTextures;
-
 	UPROPERTY(EditAnywhere, Category = CustomizableObject, meta = (DisplayName = "LOD Settings"))
 	FMutableLODSettings LODSettings;
 
-	UPROPERTY(VisibleAnywhere, Category = CustomizableObject)
-	TArray<FMutableModelImageProperties> ImageProperties;
-
 	/** This is a non-user-controlled flag to disable streaming (set at object compilation time, depending on optimization). */
 	bool bDisableTextureStreaming = false;
 
-	/** If the object is compiled, this flag is true if it was compiled with maximum optimizations. If the object is not compiled, its value is meaningless. */
-	bool bIsCompiledWithOptimization = true;
-
 	UPROPERTY(Transient)
 	TArray<FMorphTargetInfo> ContributingMorphTargetsInfo;
 	
@@ -1118,9 +644,6 @@ public:
 	UPROPERTY(Transient)
 	TArray<FCustomizableObjectMeshToMeshVertData> ClothMeshToMeshVertData;
 
-	UPROPERTY()
-	TArray<FMutableSkinWeightProfileInfo> SkinWeightProfilesInfo;
-
 	// mu::ExtensionData::Index is an index into this array when mu::ExtensionData::Origin is ConstantAlwaysLoaded
 	UPROPERTY()
 	TArray<FCustomizableObjectResourceData> AlwaysLoadedExtensionData;
@@ -1264,11 +787,6 @@ public:
 	//
 	USkeletalMesh* GetRefSkeletalMesh(int32 ComponentIndex = 0) const;
 	
-	//
-	FMutableRefSkeletalMeshData* GetRefSkeletalMeshData(int32 ComponentIndex = 0);
-	
-	TSoftObjectPtr<UMaterialInterface> GetReferencedMaterialAssetPtr(uint32 Index);
-
 private:
 
 	
@@ -1370,13 +888,7 @@ private:
 	// This is a manual version number for the binary blobs in this asset.
 	// Increasing it invalidates all the previously compiled models.
 	// Warning: If while merging code both versions have changed, take the highest+1.
-	static const int32 CurrentSupportedVersion = 419;
-
-	// Compile the object for a specific platform - Compile for Cook Customizable Object
-	void CompileForTargetPlatform(const ITargetPlatform* TargetPlatform);
-
-	// Unless we are packaging there is no need for keeping all the data generated during compilation, this information is stored in the derived data.
-	void ClearCompiledData();
+	static const int32 CurrentSupportedVersion = 420;
 
 public:
 
@@ -1384,6 +896,12 @@ public:
 
 #if WITH_EDITOR
 	
+	// Compile the object for a specific platform - Compile for Cook Customizable Object
+	void CompileForTargetPlatform(const ITargetPlatform* TargetPlatform);
+
+	// Unless we are packaging there is no need for keeping all the data generated during compilation, this information is stored in the derived data.
+	void ClearCompiledData(bool bIsCooking);
+
 	/** Compile the object if Automatic Compilation is enabled and has not been already compiled.
 	 * Automatic compilation can be enabled/disabled in the Mutable's Plugin Settings.
 	 * @return true if compiled */
@@ -1497,39 +1015,14 @@ public:
 	UFUNCTION(BlueprintCallable, Category = CustomizableObject)
 	FParameterUIData GetParameterUIMetadataFromIndex(int32 ParamIndex) const;
 
-	TSoftObjectPtr<USkeleton> GetReferencedSkeletonAssetPtr( uint32 Index );
-
 	/** See bPreserveUserLODsOnFirstGeneration. */
 	bool IsPreserveUserLODsOnFirstGeneration() const;
 
-	/** Stores all the parameter UI metadata information for all the dependencies of this Customizable Object. */
-	UPROPERTY()
-	TMap<FString, FParameterUIData> ParameterUIDataMap;
-
 	/** Textures marked as low priority will generate defaulted resident mips (if texture streaming is enabled).
 	  * Generating defaulted resident mips greatly reduce initial generation times. */
 	UPROPERTY(EditAnywhere, Category = CustomizableObject)
 	TArray<FName> LowPriorityTextures;
 
-	/** Stores all the state UI metadata information for all the dependencies of this Customizable Object */
-	UPROPERTY()
-	TMap<FString, FParameterUIData> StateUIDataMap;
-
-	/** Stores the physics assets gathered from the SkeletalMesh nodes during compilation, to be used in mesh generation in-game */
-	UPROPERTY()
-	TMap<FString, TSoftObjectPtr<class UPhysicsAsset>> PhysicsAssetsMap;
-
-	/** Stores the UAnimBlueprint assets gathered from the SkeletalMesh nodes during compilation, to be used in mesh generation in-game */
-	UPROPERTY()
-	TMap<FString, TSoftClassPtr<UAnimInstance>> AnimBPAssetsMap;
-
-	UPROPERTY()
-	TArray<FAnimBpOverridePhysicsAssetsInfo> AnimBpOverridePhysiscAssetsInfo;
-
-	UPROPERTY()
-	/** Stores the sockets provided by the part skeletal meshes, to be merged in the generated meshes */
-	TArray<FMutableRefSocket> SocketArray;
-
 	/** Map of Hash to Streaming blocks, used to stream a block of data representing a resource from the BulkData */
 	UPROPERTY()
 	TMap<uint64, FMutableStreamableBlock> HashToStreamableBlock;
@@ -1558,6 +1051,9 @@ public:
 	UPROPERTY()
 	bool bIsChildObject = false;
 
+	/** If the object is compiled, this flag is true if it was compiled with maximum optimizations. If the object is not compiled, its value is meaningless. */
+	bool bIsCompiledWithOptimization = true;
+
 	ECustomizableObjectCompilationState CompilationState = ECustomizableObjectCompilationState::None;
 
 	FPostCompileDelegate PostCompileDelegate;
@@ -1586,8 +1082,6 @@ public:
 
 #if WITH_EDITOR
 	void SetModel(TSharedPtr<mu::Model, ESPMode::ThreadSafe> Model);
-
-	void SetBoneNamesArray(const TArray<FName>& BoneNames);
 #endif
 
 	int32 GetNumLODs() const;
@@ -1598,8 +1092,6 @@ public:
 	/** Return the names used by mutable to identify which mu::Image should be considered of LowPriority. */
 	void GetLowPriorityTextureNames(TArray<FString>& OutTextureNames);
 
-	const TArray<FName>& GetBoneNamesArray() const;
-
 	/** Return the MinLOD index to generate based on the active LODSettings (PerPlatformMinLOD or PerQualityLevelMinLOD) */
 	int32 GetMinLODIndex() const;
 
@@ -1610,10 +1102,6 @@ public:
 	bool IsMeshCacheEnabled() const;
 
 private:
-	
-	/** Stores the bone names of all the bones that can possibly use the generated meshes */
-	UPROPERTY()
-	TArray<FName> BoneNames;
 
 	/** BulkData that stores all in-game resources used by Mutable when generating instances.
 	  * Only valid in packaged builds */
diff --git a/Engine/Plugins/Experimental/Mutable/Source/CustomizableObject/Classes/MuCO/CustomizableObjectInstance.h b/Engine/Plugins/Experimental/Mutable/Source/CustomizableObject/Classes/MuCO/CustomizableObjectInstance.h
--- a/Engine/Plugins/Experimental/Mutable/Source/CustomizableObject/Classes/MuCO/CustomizableObjectInstance.h
+++ b/Engine/Plugins/Experimental/Mutable/Source/CustomizableObject/Classes/MuCO/CustomizableObjectInstance.h
@@ -31,6 +31,7 @@ struct FFrame;
 struct FGameplayTagContainer;
 struct FPropertyChangedEvent;
 struct FTexturePlatformData;
+struct FMutableModelImageProperties;
 
 
 // Priority for the mutable update queue, Low is the normal distance-based priority, High is normally used for discards and Mid for LOD downgrades
diff --git a/Engine/Plugins/Experimental/Mutable/Source/CustomizableObject/Internal/MuCO/CustomizableInstancePrivateData.h b/Engine/Plugins/Experimental/Mutable/Source/CustomizableObject/Internal/MuCO/CustomizableInstancePrivateData.h
--- a/Engine/Plugins/Experimental/Mutable/Source/CustomizableObject/Internal/MuCO/CustomizableInstancePrivateData.h
+++ b/Engine/Plugins/Experimental/Mutable/Source/CustomizableObject/Internal/MuCO/CustomizableInstancePrivateData.h
@@ -22,6 +22,7 @@ namespace mu
 	typedef uint64 FResourceID;
 }
 
+struct FModelResources;
 struct FMutableModelImageProperties;
 struct FMutableRefSkeletalMeshData;
 struct FMutableImageCacheKey;
@@ -63,7 +64,7 @@ struct FReferencedPhysicsAssets
 {
 	GENERATED_USTRUCT_BODY();
 	
-	TArray<FString> PhysicsAssetToLoad;
+	TArray<int32> PhysicsAssetToLoad;
 	
 	UPROPERTY(Transient)
 	TArray< TObjectPtr<UPhysicsAsset> > PhysicsAssetsToMerge;
@@ -122,7 +123,7 @@ struct FCustomizableInstanceComponentData
 	FReferencedPhysicsAssets PhysicsAssets;
 
 	/** Clothing PhysicsAssets required by the current generated instance. PhysicsAssets to be loaded and merged.*/
-	TArray<TPair<int32, FString>> ClothingPhysicsAssetsToStream;
+	TArray<TPair<int32, int32>> ClothingPhysicsAssetsToStream;
 
 	/** Array of generated MeshIds per each LOD, used to decide if the mesh should be updated or not.
 	 *  Size == NumLODsAvailable
@@ -258,10 +259,10 @@ public:
 
 private:
 
-	void InitSkeletalMeshData(const TSharedRef<FUpdateContextPrivate>& OperationData, USkeletalMesh* SkeletalMesh, const FMutableRefSkeletalMeshData* RefSkeletalMeshData, const UCustomizableObject& CustomizableObject, int32 ComponentIndex);
+	void InitSkeletalMeshData(const TSharedRef<FUpdateContextPrivate>& OperationData, USkeletalMesh* SkeletalMesh, const FMutableRefSkeletalMeshData& RefSkeletalMeshData, const UCustomizableObject& CustomizableObject, int32 ComponentIndex);
 
 	bool BuildSkeletonData(const TSharedRef<FUpdateContextPrivate>& OperationData, USkeletalMesh& SkeletalMesh, const FMutableRefSkeletalMeshData& RefSkeletalMeshData, UCustomizableObject& CustomizableObject, int32 ComponentIndex);
-	void BuildMeshSockets(const TSharedRef<FUpdateContextPrivate>& OperationData, USkeletalMesh* SkeletalMesh, const FMutableRefSkeletalMeshData* RefSkeletalMeshData, UCustomizableObjectInstance* CustomizableObjectInstance, mu::Ptr<const mu::Mesh> MutableMesh);
+	void BuildMeshSockets(const TSharedRef<FUpdateContextPrivate>& OperationData, USkeletalMesh* SkeletalMesh, const FModelResources& ModelResources, const FMutableRefSkeletalMeshData& RefSkeletalMeshData, mu::Ptr<const mu::Mesh> MutableMesh);
 	void BuildOrCopyElementData(const TSharedRef<FUpdateContextPrivate>& OperationData, USkeletalMesh* SkeletalMesh, UCustomizableObjectInstance* CustomizableObjectInstance, int32 ComponentIndex);
 	void BuildOrCopyMorphTargetsData(const TSharedRef<FUpdateContextPrivate>& OperationData, USkeletalMesh* SkeletalMesh, const USkeletalMesh* SrcSkeletalMesh, UCustomizableObjectInstance* CustomizableObjectInstance, int32 ComponentIndex);
 	bool BuildOrCopyRenderData(const TSharedRef<FUpdateContextPrivate>& OperationData, USkeletalMesh* SkeletalMesh, const USkeletalMesh* SrcSkeletalMesh, UCustomizableObjectInstance* CustomizableObjectInstance, int32 ComponentIndex);
diff --git a/Engine/Plugins/Experimental/Mutable/Source/CustomizableObject/Internal/MuCO/CustomizableObjectPrivate.h b/Engine/Plugins/Experimental/Mutable/Source/CustomizableObject/Internal/MuCO/CustomizableObjectPrivate.h
--- a/Engine/Plugins/Experimental/Mutable/Source/CustomizableObject/Internal/MuCO/CustomizableObjectPrivate.h
+++ b/Engine/Plugins/Experimental/Mutable/Source/CustomizableObject/Internal/MuCO/CustomizableObjectPrivate.h
@@ -3,8 +3,10 @@
 #pragma once
 
 #include "MuCO/StateMachine.h"
+#include "MuCO/CustomizableObjectUIData.h"
 #include "Templates/SharedPointer.h"
 #include "UObject/WeakObjectPtr.h"
+#include "UObject/SoftObjectPtr.h"
 #include "MuR/Types.h"
 
 #if WITH_EDITOR
@@ -17,7 +19,11 @@ namespace mu { class Model; }
 class UCustomizableObject;
 class USkeletalMesh;
 class USkeleton;
-
+class UPhysicsAsset;
+class UMaterialInterface;
+class UTexture;
+class UAnimInstance;
+class UAssetUserData;
 
 class FMeshCache
 {
@@ -30,6 +36,19 @@ private:
 	TMap<TArray<mu::FResourceID>, TWeakObjectPtr<USkeletalMesh>> GeneratedMeshes;
 };
 
+
+class FSkeletonCache
+{
+public:
+	USkeleton* Get(const TArray<uint16>& Key);
+
+	void Add(const TArray<uint16>& Key, USkeleton* Value);
+
+private:
+	TMap<TArray<uint16>, TWeakObjectPtr<USkeleton>> MergedSkeletons;
+};
+
+
 struct FCustomizableObjectStatusTypes
 {
 	enum class EState : uint8
@@ -54,19 +73,371 @@ struct FCustomizableObjectStatusTypes
 	};
 };
 
-
 using FCustomizableObjectStatus = FStateMachine<FCustomizableObjectStatusTypes>;
 
 
-class FSkeletonCache
+USTRUCT()
+struct CUSTOMIZABLEOBJECT_API FMutableModelImageProperties
 {
-public:
-	USkeleton* Get(const TArray<uint16>& Key);
+	GENERATED_USTRUCT_BODY()
 
-	void Add(const TArray<uint16>& Key, USkeleton* Value);
+	FMutableModelImageProperties()
+		: Filter(TF_Default)
+		, SRGB(0)
+		, FlipGreenChannel(0)
+		, IsPassThrough(0)
+		, LODBias(0)
+		, LODGroup(TEXTUREGROUP_World)
+		, AddressX(TA_Clamp)
+		, AddressY(TA_Clamp)
+	{}
 
-private:
-	TMap<TArray<uint16>, TWeakObjectPtr<USkeleton>> MergedSkeletons;
+	FMutableModelImageProperties(const FString& InTextureParameterName, TextureFilter InFilter, uint32 InSRGB,
+		uint32 InFlipGreenChannel, uint32 bInIsPassThrough, int32 InLODBias, TEnumAsByte<enum TextureGroup> InLODGroup,
+		TEnumAsByte<enum TextureAddress> InAddressX, TEnumAsByte<enum TextureAddress> InAddressY)
+		: TextureParameterName(InTextureParameterName)
+		, Filter(InFilter)
+		, SRGB(InSRGB)
+		, FlipGreenChannel(InFlipGreenChannel)
+		, IsPassThrough(bInIsPassThrough)
+		, LODBias(InLODBias)
+		, LODGroup(InLODGroup)
+		, AddressX(InAddressX)
+		, AddressY(InAddressY)
+	{}
+
+	// Name in the material.
+	UPROPERTY()
+	FString TextureParameterName;
+
+	UPROPERTY()
+	TEnumAsByte<enum TextureFilter> Filter;
+
+	UPROPERTY()
+	uint32 SRGB : 1;
+
+	UPROPERTY()
+	uint32 FlipGreenChannel : 1;
+
+	UPROPERTY()
+	uint32 IsPassThrough : 1;
+
+	UPROPERTY()
+	int32 LODBias;
+
+	UPROPERTY()
+	TEnumAsByte<enum TextureGroup> LODGroup;
+
+	UPROPERTY()
+	TEnumAsByte<enum TextureAddress> AddressX;
+
+	UPROPERTY()
+	TEnumAsByte<enum TextureAddress> AddressY;
+
+	bool operator!=(const FMutableModelImageProperties& rhs) const;
+
+#if WITH_EDITORONLY_DATA
+	friend FArchive& operator<<(FArchive& Ar, FMutableModelImageProperties& ImageProps);
+#endif
+};
+
+
+USTRUCT()
+struct CUSTOMIZABLEOBJECT_API FMutableRefSocket
+{
+	GENERATED_BODY()
+
+	UPROPERTY()
+	FName SocketName;
+	UPROPERTY()
+	FName BoneName;
+
+	UPROPERTY()
+	FVector RelativeLocation = FVector::ZeroVector;
+	UPROPERTY()
+	FRotator RelativeRotation = FRotator::ZeroRotator;
+	UPROPERTY()
+	FVector RelativeScale = FVector::ZeroVector;
+
+	UPROPERTY()
+	bool bForceAlwaysAnimated = false;
+
+	// When two sockets have the same name, the one with higher priority will be picked and the other discarded
+	UPROPERTY()
+	int32 Priority = -1;
+
+	bool operator ==(const FMutableRefSocket& Other) const;
+
+#if WITH_EDITORONLY_DATA
+	friend FArchive& operator<<(FArchive& Ar, FMutableRefSocket& Data);
+#endif
+};
+
+
+USTRUCT()
+struct FMutableRefLODInfo
+{
+	GENERATED_BODY()
+
+	UPROPERTY()
+	float ScreenSize = 0.f;
+
+	UPROPERTY()
+	float LODHysteresis = 0.f;
+
+	UPROPERTY()
+	bool bSupportUniformlyDistributedSampling = false;
+
+	UPROPERTY()
+	bool bAllowCPUAccess = false;
+
+#if WITH_EDITORONLY_DATA
+	friend FArchive& operator<<(FArchive& Ar, FMutableRefLODInfo& Data);
+#endif
+};
+
+
+USTRUCT()
+struct FMutableRefLODRenderData
+{
+	GENERATED_BODY()
+
+	UPROPERTY()
+	bool bIsLODOptional = false;
+
+	UPROPERTY()
+	bool bStreamedDataInlined = false;
+
+#if WITH_EDITORONLY_DATA
+	friend FArchive& operator<<(FArchive& Ar, FMutableRefLODRenderData& Data);
+#endif
+};
+
+
+USTRUCT()
+struct FMutableRefLODData
+{
+	GENERATED_BODY()
+
+	UPROPERTY()
+	FMutableRefLODInfo LODInfo;
+
+	UPROPERTY()
+	FMutableRefLODRenderData RenderData;
+
+#if WITH_EDITORONLY_DATA
+	friend FArchive& operator<<(FArchive& Ar, FMutableRefLODData& Data);
+#endif
+};
+
+
+USTRUCT()
+struct FMutableRefSkeletalMeshSettings
+{
+	GENERATED_BODY()
+
+	UPROPERTY()
+	bool bEnablePerPolyCollision = false;
+
+	UPROPERTY()
+	float DefaultUVChannelDensity = 0.f;
+
+#if WITH_EDITORONLY_DATA
+	friend FArchive& operator<<(FArchive& Ar, FMutableRefSkeletalMeshSettings& Data);
+#endif
+};
+
+
+USTRUCT()
+struct FMutableRefAssetUserData
+{
+	GENERATED_BODY()
+
+	UPROPERTY()
+	TObjectPtr<UAssetUserData> AssetUserData;
+
+#if WITH_EDITORONLY_DATA
+	FString ClassPath;
+	TArray<uint8> Bytes;
+
+	friend FArchive& operator<<(FArchive& Ar, FMutableRefAssetUserData& Data);
+
+	void InitResources(UCustomizableObject* InOuter);
+#endif
+
+};
+
+
+USTRUCT()
+struct CUSTOMIZABLEOBJECT_API FMutableRefSkeletalMeshData
+{
+	GENERATED_BODY()
+
+	// Reference Skeletal Mesh
+	UPROPERTY()
+	TObjectPtr<USkeletalMesh> SkeletalMesh;
+
+	// Path to load the ReferenceSkeletalMesh
+	UPROPERTY()
+	TSoftObjectPtr<USkeletalMesh> SoftSkeletalMesh;
+
+	// LOD info
+	UPROPERTY()
+	TArray<FMutableRefLODData> LODData;
+
+	// Sockets
+	UPROPERTY()
+	TArray<FMutableRefSocket> Sockets;
+
+	// Bounding Box
+	UPROPERTY()
+	FBoxSphereBounds Bounds = FBoxSphereBounds(ForceInitToZero);
+
+	// Settings
+	UPROPERTY()
+	FMutableRefSkeletalMeshSettings Settings;
+
+	// Skeleton
+	UPROPERTY()
+	TObjectPtr<USkeleton> Skeleton;
+
+	// PhysicsAsset
+	UPROPERTY()
+	TObjectPtr<UPhysicsAsset> PhysicsAsset;
+
+	// Post Processing AnimBP
+	UPROPERTY()
+	TSoftClassPtr<UAnimInstance> PostProcessAnimInst;
+
+	// Shadow PhysicsAsset
+	UPROPERTY()
+	TObjectPtr<UPhysicsAsset> ShadowPhysicsAsset;
+
+	// Asset user data
+	UPROPERTY()
+	TArray<FMutableRefAssetUserData> AssetUserData;
+
+#if WITH_EDITORONLY_DATA
+	friend FArchive& operator<<(FArchive& Ar, FMutableRefSkeletalMeshData& Data);
+
+	void InitResources(UCustomizableObject* InOuter, const ITargetPlatform* InTargetPlatform);
+#endif
+
+};
+
+
+USTRUCT()
+struct CUSTOMIZABLEOBJECT_API FAnimBpOverridePhysicsAssetsInfo
+{
+	GENERATED_BODY()
+
+	UPROPERTY()
+	TSoftClassPtr<UAnimInstance> AnimInstanceClass;
+
+	UPROPERTY()
+	TSoftObjectPtr<UPhysicsAsset> SourceAsset;
+
+	UPROPERTY()
+	int32 PropertyIndex = -1;
+
+	bool operator==(const FAnimBpOverridePhysicsAssetsInfo& Rhs) const;
+
+#if WITH_EDITORONLY_DATA
+	friend FArchive& operator<<(FArchive& Ar, FAnimBpOverridePhysicsAssetsInfo& Info);
+#endif
+};
+
+
+USTRUCT()
+struct CUSTOMIZABLEOBJECT_API FMutableSkinWeightProfileInfo
+{
+	GENERATED_USTRUCT_BODY()
+
+	FMutableSkinWeightProfileInfo() {};
+
+	FMutableSkinWeightProfileInfo(FName InName, bool InDefaultProfile, int8 InDefaultProfileFromLODIndex) : Name(InName),
+	DefaultProfile(InDefaultProfile), DefaultProfileFromLODIndex(InDefaultProfileFromLODIndex) {};
+
+	UPROPERTY()
+	FName Name;
+
+	UPROPERTY()
+	bool DefaultProfile = false;
+
+	UPROPERTY(meta = (ClampMin = 0))
+	int8 DefaultProfileFromLODIndex = 0;
+
+	bool operator==(const FMutableSkinWeightProfileInfo& Other) const;
+
+#if WITH_EDITORONLY_DATA
+	friend FArchive& operator<<(FArchive& Ar, FMutableSkinWeightProfileInfo& Info);
+#endif
+};
+
+
+// Referenced materials, skeletons, passthrough textures...
+USTRUCT()
+struct FModelResources
+{
+	GENERATED_BODY()
+
+	/** All the SkeletalMeshes generated for this CustomizableObject instances will use the Reference Skeletal Mesh
+	 * properties for everything that Mutable doesn't create or modify. This struct stores the information used from
+	 * the Reference Skeletal Meshes to avoid having them loaded at all times. This includes data like LOD distances,
+	 * LOD render data settings, Mesh sockets, Bounding volumes, etc.
+	 */
+	UPROPERTY()
+	TArray<FMutableRefSkeletalMeshData> ReferenceSkeletalMeshesData;
+
+	/** Skeletons used by the compiled mu::Model. */
+	UPROPERTY()
+	TArray<TSoftObjectPtr<USkeleton>> Skeletons;
+
+	/** Materials used by the compiled mu::Model */
+	UPROPERTY()
+	TArray<TSoftObjectPtr<UMaterialInterface>> Materials;
+
+	/** PassThrough Textures used by the mu::Model. */
+	UPROPERTY()
+	TArray<TSoftObjectPtr<UTexture>> PassThroughTextures;
+
+	/** Physics assets gathered from the SkeletalMeshes, to be used in mesh generation in-game */
+	UPROPERTY()
+	TArray<TSoftObjectPtr<UPhysicsAsset>> PhysicsAssets;
+
+	/** UAnimBlueprint assets gathered from the SkeletalMesh, to be used in mesh generation in-game */
+	UPROPERTY()
+	TArray<TSoftClassPtr<UAnimInstance>> AnimBPs;
+
+	/** */
+	UPROPERTY()
+	TArray<FAnimBpOverridePhysicsAssetsInfo> AnimBpOverridePhysiscAssetsInfo;
+
+	/** Material slot names for the materials referenced by the surfaces. */
+	UPROPERTY()
+	TArray<FName> MaterialSlotNames;
+
+	/** Bone names of all the bones that can possibly use the generated meshes */
+	UPROPERTY()
+	TArray<FName> BoneNames;
+
+	/** Mesh sockets provided by the part skeletal meshes, to be merged in the generated meshes */
+	UPROPERTY()
+	TArray<FMutableRefSocket> SocketArray;
+
+	UPROPERTY()
+	TArray<FMutableSkinWeightProfileInfo> SkinWeightProfilesInfo;
+
+	UPROPERTY()
+	TArray<FMutableModelImageProperties> ImageProperties;
+
+	/** Parameter UI metadata information for all the dependencies of this Customizable Object. */
+	UPROPERTY()
+	TMap<FString, FParameterUIData> ParameterUIDataMap;
+
+	/** State UI metadata information for all the dependencies of this Customizable Object */
+	UPROPERTY()
+	TMap<FString, FParameterUIData> StateUIDataMap;
 };
 
 
@@ -74,25 +445,48 @@ UCLASS()
 class UCustomizableObjectPrivate : public UObject
 {
 	GENERATED_BODY()
-	
+
+private:
+
 	TSharedPtr<mu::Model, ESPMode::ThreadSafe> MutableModel;
 
+	/** Stores resources to be used by MutableModel In-Game. Cooked resources. */
+	UPROPERTY()
+	FModelResources ModelResources;
+
+#if WITH_EDITORONLY_DATA
+	/** Stores resources to be used by MutableModel in the Editor. Editor resources.
+	 * Editor-Only to avoid packaging assets referenced by editor compilations. */
+	UPROPERTY(Transient)
+	FModelResources ModelResourcesEditor;
+#endif
+
 public:
 	void SetModel(const TSharedPtr<mu::Model, ESPMode::ThreadSafe>& Model, const FGuid Identifier);
 	const TSharedPtr<mu::Model, ESPMode::ThreadSafe>& GetModel();
 	TSharedPtr<const mu::Model, ESPMode::ThreadSafe> GetModel() const;
 
+	const FModelResources& GetModelResources() const;
+
 #if WITH_EDITORONLY_DATA
-	/** See UCustomizableObject::ParticipatingObjects. */
-	CUSTOMIZABLEOBJECT_API TMap<TObjectPtr<const UObject>, FGuid>& GetParticipatingObjects(UCustomizableObject& Public);
+	CUSTOMIZABLEOBJECT_API FModelResources& GetModelResources(bool bIsCooking);
 #endif
-	
+
+	/** Cache of generated SkeletalMeshes */
+	FMeshCache MeshCache;
+
+	/** Cache of merged Skeletons */
+	FSkeletonCache SkeletonCache;
+
 	// See UCustomizableObjectSystem::LockObject. Must only be modified from the game thread
 	bool bLocked = false;
 
-#if WITH_EDITOR
+#if WITH_EDITORONLY_DATA
+
+	/** Unique Identifier - Deterministic. Used to locate Model and Streamable data on disk. Should not be modified. */
 	FGuid Identifier;
 
+	/** Cooked platform names. */
 	TArray<FString> CachedPlatformNames;
 
 	/** List of external packages that if changed, a compilation is required.
@@ -102,12 +496,6 @@ public:
 	TMap<FName, FGuid> ParticipatingObjects;
 #endif
 
-	/** Cache of generated SkeletalMeshes */
-	FMeshCache MeshCache;
-
-	/** Cache of merged Skeletons */
-	FSkeletonCache SkeletonCache;
-
 	FCustomizableObjectStatus Status;
 };
 
diff --git a/Engine/Plugins/Experimental/Mutable/Source/CustomizableObject/Private/MuCO/CustomizableObject.cpp b/Engine/Plugins/Experimental/Mutable/Source/CustomizableObject/Private/MuCO/CustomizableObject.cpp
--- a/Engine/Plugins/Experimental/Mutable/Source/CustomizableObject/Private/MuCO/CustomizableObject.cpp
+++ b/Engine/Plugins/Experimental/Mutable/Source/CustomizableObject/Private/MuCO/CustomizableObject.cpp
@@ -8,6 +8,7 @@
 #include "EdGraph/EdGraph.h"
 #include "Engine/Engine.h"
 #include "Engine/SkeletalMesh.h"
+#include "Animation/AnimInstance.h"
 #include "Animation/Skeleton.h"
 #include "Engine/AssetUserData.h"
 #include "HAL/FileManager.h"
@@ -21,6 +22,7 @@
 #include "MuCO/CustomizableObjectInstance.h"
 #include "MuCO/CustomizableObjectPrivate.h"
 #include "MuCO/CustomizableObjectSystem.h"
+#include "MuCO/CustomizableObjectUIData.h"
 #include "MuCO/ICustomizableObjectModule.h"
 #include "MuCO/MutableProjectorTypeUtils.h"
 #include "MuCO/UnrealMutableModelDiskStreamer.h"
@@ -154,7 +156,7 @@ void UCustomizableObject::PreSave(FObjectPreSaveContext ObjectSaveContext)
 			UE_LOG(LogMutable, Warning, TEXT("Cook: Customizable Object [%s] is missing [%s] platform data."), *GetName(),
 				*ObjectSaveContext.GetTargetPlatform()->PlatformName());
 			
-			ClearCompiledData();
+			ClearCompiledData(true);
 		}
 	}
 #endif
@@ -218,15 +220,6 @@ void UCustomizableObject::PostLoad()
 
 	const int32 CustomizableObjectCustomVersion = GetLinkerCustomVersion(FCustomizableObjectCustomVersion::GUID);
 
-	// Update state never-stream flag from deprecated enum
-	if (CustomizableObjectCustomVersion < FCustomizableObjectCustomVersion::CustomizableObjectStateHasSeparateNeverStreamFlag)
-	{
-		for (TPair<FString, FParameterUIData>& s : StateUIDataMap)
-		{
-			s.Value.bDisableTextureStreaming = s.Value.TextureCompressionStrategy != ETextureCompressionStrategy::None;
-		}
-	}
-
 	if (!IsRunningCookCommandlet())
 	{
 		ITargetPlatformManagerModule& TargetPlatformManager = GetTargetPlatformManagerRef();
@@ -331,29 +324,20 @@ FGuid GenerateIdentifier(const UCustomizableObject& CustomizableObject)
 }
 
 
-void UCustomizableObject::ClearCompiledData()
+void UCustomizableObject::ClearCompiledData(bool bIsCooking)
 {
-	ReferenceSkeletalMeshesData.Empty();
-	ReferencedMaterials.Empty();
-	ReferencedMaterialSlotNames.Empty();
-	ReferencedPassThroughTextures.Empty();
-	ReferencedSkeletons.Empty();
-	ImageProperties.Empty();
-	ParameterUIDataMap.Empty();
-	StateUIDataMap.Empty();
-	PhysicsAssetsMap.Empty();
+	Private->GetModelResources(bIsCooking) = FModelResources();
+
 	ContributingMorphTargetsInfo.Empty();
 	MorphTargetReconstructionData.Empty();
 	ClothMeshToMeshVertData.Empty();
 	ContributingClothingAssetsData.Empty();
 	ClothSharedConfigsData.Empty();
-	SkinWeightProfilesInfo.Empty();
-	AnimBpOverridePhysiscAssetsInfo.Empty();
-	BoneNames.Empty();
 
 #if WITH_EDITORONLY_DATA
 	CustomizableObjectPathMap.Empty();
 	GroupNodeMap.Empty();
+	Private->ParticipatingObjects.Empty();
 #endif
 
 	HashToStreamableBlock.Empty();
@@ -383,7 +367,7 @@ void SerializeStreamedResources(FArchive& Ar, UObject* Object, TArray<FCustomiza
 
 				if (AssetUserData->AssetUserDataEditor)
 				{
-					AssetUserDataPath = TSoftObjectPtr<UAssetUserData>(AssetUserData->AssetUserDataEditor).ToSoftObjectPath().ToString();
+					AssetUserDataPath = TSoftObjectPtr<UAssetUserData>(AssetUserData->AssetUserDataEditor).ToString();
 				}
 
 				Ar << AssetUserDataPath;
@@ -478,82 +462,78 @@ void UCustomizableObject::SaveCompiledData(FArchive& MemoryWriter, bool bIsCooki
 	MutableCompiledDataStreamHeader Header(InternalVersion, VersionId);
 	MemoryWriter << Header;
 
-	MemoryWriter << ReferenceSkeletalMeshesData;
+	FModelResources& ModelResources = Private->GetModelResources(false);
+
+	MemoryWriter << ModelResources.ReferenceSkeletalMeshesData;
 
 	SerializeStreamedResources(MemoryWriter, this, StreamedResourceData, bIsCooking);
-	
-	int32 NumReferencedMaterials = ReferencedMaterials.Num();
+
+	int32 NumReferencedMaterials = ModelResources.Materials.Num();
 	MemoryWriter << NumReferencedMaterials;
 
-	for (const TSoftObjectPtr<UMaterialInterface>& Material : ReferencedMaterials)
+	for (const TSoftObjectPtr<UMaterialInterface>& Material : ModelResources.Materials)
 	{
-		FString StringRef = Material.ToSoftObjectPath().ToString();
+		FString StringRef = Material.ToString();
 		MemoryWriter << StringRef;
 	}
 
-	int32 NumReferencedMaterialSlotNames = ReferencedMaterialSlotNames.Num();
-	MemoryWriter << NumReferencedMaterialSlotNames;
+	int32 NumReferencedSkeletons = ModelResources.Skeletons.Num();
+	MemoryWriter << NumReferencedSkeletons;
 
-	for (const FName& MaterialSlotName : ReferencedMaterialSlotNames)
+	for (const TSoftObjectPtr<USkeleton>& Skeleton : ModelResources.Skeletons)
 	{
-		FString StringRef = MaterialSlotName.ToString();
+		FString StringRef = Skeleton.ToString();
 		MemoryWriter << StringRef;
 	}
 
-	int32 NumReferencedSkeletons = ReferencedSkeletons.Num();
-	MemoryWriter << NumReferencedSkeletons;
+	int32 NumPassthroughTextures = ModelResources.PassThroughTextures.Num();
+	MemoryWriter << NumPassthroughTextures;
 
-	for (const TSoftObjectPtr<USkeleton>& Skeleton : ReferencedSkeletons)
+	for (const TSoftObjectPtr<UTexture>& PassthroughTexture : ModelResources.PassThroughTextures)
 	{
-		FString StringRef = Skeleton.ToSoftObjectPath().ToString();
+		FString StringRef = PassthroughTexture.ToString();
 		MemoryWriter << StringRef;
 	}
 
-	int32 NumPassthroughTextures = ReferencedPassThroughTextures.Num();
-	MemoryWriter << NumPassthroughTextures;
+	int32 NumPhysicsAssets = ModelResources.PhysicsAssets.Num();
+	MemoryWriter << NumPhysicsAssets;
 
-	for (const TSoftObjectPtr<UTexture>& PassthroughTexture : ReferencedPassThroughTextures)
+	for (const TSoftObjectPtr<UPhysicsAsset>& PhysicsAsset : ModelResources.PhysicsAssets)
 	{
-		FString StringRef = PassthroughTexture.ToSoftObjectPath().ToString();
+		FString StringRef = PhysicsAsset.ToString();
 		MemoryWriter << StringRef;
 	}
 
-	MemoryWriter << ImageProperties;
-	MemoryWriter << ParameterUIDataMap;
-	MemoryWriter << StateUIDataMap;
+	int32 NumAnimBps = ModelResources.AnimBPs.Num();
+	MemoryWriter << NumAnimBps;
 
-	int32 NumPhysicsAssets = PhysicsAssetsMap.Num();
-	MemoryWriter << NumPhysicsAssets;
-
-	for (const TPair<FString, TSoftObjectPtr<class UPhysicsAsset>>& PhysicsAsset : PhysicsAssetsMap)
+	for (const TSoftClassPtr<UAnimInstance>& AnimBp : ModelResources.AnimBPs)
 	{
-		FString StringRef = PhysicsAsset.Value.ToSoftObjectPath().ToString();
+		FString StringRef = AnimBp.ToString();
 		MemoryWriter << StringRef;
 	}
 
+	MemoryWriter << ModelResources.AnimBpOverridePhysiscAssetsInfo;
+
+	MemoryWriter << ModelResources.MaterialSlotNames;
+	MemoryWriter << ModelResources.BoneNames;
+	MemoryWriter << ModelResources.SocketArray;
+
+	MemoryWriter << ModelResources.SkinWeightProfilesInfo;
+
+	MemoryWriter << ModelResources.ImageProperties;
+	MemoryWriter << ModelResources.ParameterUIDataMap;
+	MemoryWriter << ModelResources.StateUIDataMap;
+
 	MemoryWriter << ContributingMorphTargetsInfo;
 	MemoryWriter << MorphTargetReconstructionData;
 	
 	MemoryWriter << ClothMeshToMeshVertData;
 	MemoryWriter << ContributingClothingAssetsData;
 	MemoryWriter << ClothSharedConfigsData; 
-	
-	MemoryWriter << SkinWeightProfilesInfo;
-
-	MemoryWriter << AnimBpOverridePhysiscAssetsInfo;
-
-	MemoryWriter << BoneNames;
 
 	MemoryWriter << HashToStreamableBlock;
 
-	// All Editor Only data must be serialized here
-	if (!bIsCooking)
-	{
-		MemoryWriter << CustomizableObjectPathMap;
-		MemoryWriter << GroupNodeMap;
-		MemoryWriter << GetPrivate()->ParticipatingObjects;
-	}
-
 	MemoryWriter << LODSettings.NumLODsInRoot;
 	MemoryWriter << NumMeshComponentsInRoot;
 
@@ -561,12 +541,18 @@ void UCustomizableObject::SaveCompiledData(FArchive& MemoryWriter, bool bIsCooki
 
 	MemoryWriter << LODSettings.NumLODsToStream;
 	MemoryWriter << LODSettings.bLODStreamingEnabled;
+
+	// Editor Only data
+	MemoryWriter << bIsCompiledWithOptimization;
+	MemoryWriter << CustomizableObjectPathMap;
+	MemoryWriter << GroupNodeMap;
+	MemoryWriter << GetPrivate()->ParticipatingObjects;
 }
 
 void UCustomizableObject::LoadCompiledData(FArchive& MemoryReader, const ITargetPlatform* InTargetPlatform, bool bIsCooking)
 {
 	TSharedPtr<mu::Model, ESPMode::ThreadSafe> LoadedModel;
-	ClearCompiledData();
+	ClearCompiledData(bIsCooking);
 
 	MutableCompiledDataStreamHeader Header;
 	MemoryReader << Header;
@@ -576,16 +562,20 @@ void UCustomizableObject::LoadCompiledData(FArchive& MemoryReader, const ITarget
 		// Make sure mutable has been initialised.
 		UCustomizableObjectSystem::GetInstance();
 
-		MemoryReader << ReferenceSkeletalMeshesData;
+		FModelResources& ModelResources = Private->GetModelResources(bIsCooking);
+		ModelResources = FModelResources();
+
+		MemoryReader << ModelResources.ReferenceSkeletalMeshesData;
 
 		// Initialize resources. 
-		for(FMutableRefSkeletalMeshData& ReferenceSkeletalMeshData : ReferenceSkeletalMeshesData)
+		for(FMutableRefSkeletalMeshData& ReferenceSkeletalMeshData : ModelResources.ReferenceSkeletalMeshesData)
 		{
 			ReferenceSkeletalMeshData.InitResources(this, InTargetPlatform);
 		}
 
 		SerializeStreamedResources(MemoryReader, this, StreamedResourceData, bIsCooking);
 
+
 		int32 NumReferencedMaterials = 0;
 		MemoryReader << NumReferencedMaterials;
 
@@ -594,18 +584,7 @@ void UCustomizableObject::LoadCompiledData(FArchive& MemoryReader, const ITarget
 			FString StringRef;
 			MemoryReader << StringRef;
 
-			ReferencedMaterials.Add(TSoftObjectPtr<UMaterialInterface>(FSoftObjectPath(StringRef)));
-		}
-
-		int32 NumReferencedMaterialSlotNames = 0;
-		MemoryReader << NumReferencedMaterialSlotNames;
-
-		for (int32 i = 0; i < NumReferencedMaterialSlotNames; ++i)
-		{
-			FString StringRef;
-			MemoryReader << StringRef;
-
-			ReferencedMaterialSlotNames.Add(FName(*StringRef));
+			ModelResources.Materials.Add(TSoftObjectPtr<UMaterialInterface>(FSoftObjectPath(StringRef)));
 		}
 
 		int32 NumReferencedSkeletons = 0;
@@ -616,7 +595,7 @@ void UCustomizableObject::LoadCompiledData(FArchive& MemoryReader, const ITarget
 			FString StringRef;
 			MemoryReader << StringRef;
 
-			ReferencedSkeletons.Add(TSoftObjectPtr<USkeleton>(FSoftObjectPath(StringRef)));
+			ModelResources.Skeletons.Add(TSoftObjectPtr<USkeleton>(FSoftObjectPath(StringRef)));
 		}
 
 		int32 NumPassthroughTextures = 0;
@@ -627,13 +606,9 @@ void UCustomizableObject::LoadCompiledData(FArchive& MemoryReader, const ITarget
 			FString StringRef;
 			MemoryReader << StringRef;
 
-			ReferencedPassThroughTextures.Add(TSoftObjectPtr<UTexture>(FSoftObjectPath(StringRef)));
+			ModelResources.PassThroughTextures.Add(TSoftObjectPtr<UTexture>(FSoftObjectPath(StringRef)));
 		}
 
-		MemoryReader << ImageProperties;
-		MemoryReader << ParameterUIDataMap;
-		MemoryReader << StateUIDataMap;
-
 		int32 NumPhysicsAssets = 0;
 		MemoryReader << NumPhysicsAssets;
 
@@ -642,9 +617,33 @@ void UCustomizableObject::LoadCompiledData(FArchive& MemoryReader, const ITarget
 			FString StringRef;
 			MemoryReader << StringRef;
 
-			PhysicsAssetsMap.Add(StringRef, TSoftObjectPtr<UPhysicsAsset>(FSoftObjectPath(StringRef)));
+			ModelResources.PhysicsAssets.Add(TSoftObjectPtr<UPhysicsAsset>(FSoftObjectPath(StringRef)));
 		}
 
+		int32 NumAnimBps = 0;
+		MemoryReader << NumAnimBps;
+
+		for (int32 Index = 0; Index < NumAnimBps; ++Index)
+		{
+			FString StringRef;
+			MemoryReader << StringRef;
+
+			ModelResources.AnimBPs.Add(TSoftClassPtr<UAnimInstance>(StringRef));
+		}
+
+		MemoryReader << ModelResources.AnimBpOverridePhysiscAssetsInfo;
+
+		MemoryReader << ModelResources.MaterialSlotNames;
+		MemoryReader << ModelResources.BoneNames;
+		MemoryReader << ModelResources.SocketArray;
+
+		MemoryReader << ModelResources.SkinWeightProfilesInfo;
+
+		MemoryReader << ModelResources.ImageProperties;
+		MemoryReader << ModelResources.ParameterUIDataMap;
+		MemoryReader << ModelResources.StateUIDataMap;
+
+
 		MemoryReader << ContributingMorphTargetsInfo;
 		MemoryReader << MorphTargetReconstructionData;
 
@@ -652,67 +651,57 @@ void UCustomizableObject::LoadCompiledData(FArchive& MemoryReader, const ITarget
 		MemoryReader << ContributingClothingAssetsData;
 		MemoryReader << ClothSharedConfigsData;
 
-		MemoryReader << SkinWeightProfilesInfo;
-
-		MemoryReader << AnimBpOverridePhysiscAssetsInfo;
+		MemoryReader << HashToStreamableBlock;
 
-		TArray<FString> StringBoneNames;
-		MemoryReader << StringBoneNames;
+		MemoryReader << LODSettings.NumLODsInRoot;
+		MemoryReader << NumMeshComponentsInRoot;
 
-		BoneNames.Reserve(StringBoneNames.Num());
-		for (const FString& BoneName : StringBoneNames)
-		{
-			BoneNames.Add(FName(*BoneName));
-		}
+		MemoryReader << LODSettings.FirstLODAvailable;
 
-		MemoryReader << HashToStreamableBlock;
+		MemoryReader << LODSettings.NumLODsToStream;
+		MemoryReader << LODSettings.bLODStreamingEnabled;
 
 		bool bForceRecompilation = false;
 
-		// All Editor Only data must be loaded here
-		if (!bIsCooking)
+		// Editor Only data
 		{
+			MemoryReader << bIsCompiledWithOptimization;
 			MemoryReader << CustomizableObjectPathMap;
 			MemoryReader << GroupNodeMap;
 
 			TMap<FName, FGuid>& ParticipatingObjects = GetPrivate()->ParticipatingObjects;
 			MemoryReader << ParticipatingObjects;
 
-			const FAssetRegistryModule& AssetRegistryModule = FModuleManager::LoadModuleChecked<FAssetRegistryModule>("AssetRegistry");
-			
-			for (TTuple<FName, FGuid>& ParticipatingObject : ParticipatingObjects)
-			{				
-				FAssetPackageData AssetPackageData;
-				const UE::AssetRegistry::EExists Result = AssetRegistryModule.Get().TryGetAssetPackageData(ParticipatingObject.Key, AssetPackageData);
-				if (Result == UE::AssetRegistry::EExists::Exists)
-				{
-					PRAGMA_DISABLE_DEPRECATION_WARNINGS
-					const FGuid PackageGuid = AssetPackageData.PackageGuid;
-					PRAGMA_ENABLE_DEPRECATION_WARNINGS
-		
-					bForceRecompilation = PackageGuid != ParticipatingObject.Value;
-				}
-				else
+			if (!bIsCooking)
+			{
+				const FAssetRegistryModule& AssetRegistryModule = FModuleManager::LoadModuleChecked<FAssetRegistryModule>("AssetRegistry");
+
+				for (TTuple<FName, FGuid>& ParticipatingObject : ParticipatingObjects)
 				{
-					bForceRecompilation = true;
+					FAssetPackageData AssetPackageData;
+					const UE::AssetRegistry::EExists Result = AssetRegistryModule.Get().TryGetAssetPackageData(ParticipatingObject.Key, AssetPackageData);
+					if (Result == UE::AssetRegistry::EExists::Exists)
+					{
+						PRAGMA_DISABLE_DEPRECATION_WARNINGS
+							const FGuid PackageGuid = AssetPackageData.PackageGuid;
+						PRAGMA_ENABLE_DEPRECATION_WARNINGS
+
+							bForceRecompilation = PackageGuid != ParticipatingObject.Value;
+					}
+					else
+					{
+						bForceRecompilation = true;
+					}
+
+					if (bForceRecompilation)
+					{
+						UE_LOG(LogMutable, Display, TEXT("Forcing recompilation due to changes in %s."), *ParticipatingObject.Key.ToString());
+						break;
+					}
 				}
-		
-				if (bForceRecompilation)
-                {
-                	UE_LOG(LogMutable, Display, TEXT("Forcing recompilation due to changes in %s."), *ParticipatingObject.Key.ToString());
-					break;
-                }
 			}
 		}
 
-		MemoryReader << LODSettings.NumLODsInRoot;
-		MemoryReader << NumMeshComponentsInRoot;
-
-		MemoryReader << LODSettings.FirstLODAvailable;
-
-		MemoryReader << LODSettings.NumLODsToStream;
-		MemoryReader << LODSettings.bLODStreamingEnabled;
-
 		bool bModelSerialized = false;
 		MemoryReader << bModelSerialized;
 
@@ -994,7 +983,6 @@ void UCustomizableObject::SaveEmbeddedData(FArchive& Ar)
 	{
 		// General derived flags
 		Ar << bDisableTextureStreaming;
-		Ar << bIsCompiledWithOptimization;
 
 		// Serialize morph data
 		{
@@ -1038,7 +1026,6 @@ void UCustomizableObject::LoadEmbeddedData(FArchive& Ar)
 	{
 		// General derived flags
 		Ar << bDisableTextureStreaming;
-		Ar << bIsCompiledWithOptimization;
 
 		// Load morph data
 		{
@@ -1101,59 +1088,17 @@ USkeletalMesh* UCustomizableObject::GetRefSkeletalMesh(int32 ComponentIndex) con
 		return ReferenceSkeletalMeshes[ComponentIndex];
 	}
 #else
-	if (ReferenceSkeletalMeshesData.IsValidIndex(ComponentIndex))
+	const FModelResources& ModelResources = Private->GetModelResources();
+	if (ModelResources.ReferenceSkeletalMeshesData.IsValidIndex(ComponentIndex))
 	{
 		// Can be nullptr if RefSkeletalMeshes are not loaded yet.
-		return ReferenceSkeletalMeshesData[ComponentIndex].SkeletalMesh;
+		return ModelResources.ReferenceSkeletalMeshesData[ComponentIndex].SkeletalMesh;
 	}
 #endif
 	return nullptr;
 }
 
 
-FMutableRefSkeletalMeshData* UCustomizableObject::GetRefSkeletalMeshData(int32 ComponentIndex)
-{
-	if(ReferenceSkeletalMeshesData.IsValidIndex(ComponentIndex))
-	{
-		return &ReferenceSkeletalMeshesData[ComponentIndex];
-	}
-
-	UE_LOG(LogMutable, Warning, TEXT("UCustomizableObject::GetRefSkeletalMeshData with an invalid Index. "
-			"Reference SkeletalMesh data and compiled model may be out of sync. "
-			"Try recompiling the CustomizableObject [%s]."), *GetName() );
-	
-	return nullptr;
-}
-
-
-TSoftObjectPtr<UMaterialInterface> UCustomizableObject::GetReferencedMaterialAssetPtr( uint32 Index )
-{
-	if (ReferencedMaterials.IsValidIndex(Index))
-	{
-		return ReferencedMaterials[Index];
-	}
-
-	UE_LOG(LogMutable, Warning, TEXT("UCustomizableObject::GetReferencedMaterial with an invalid Index. "
-		"Source material data and CustomizableObject data may be out of sync. "
-		"Try recompiling and saving the CustomizableObject asset [%s]."), *GetName() );
-	return nullptr;
-}
-
-
-TSoftObjectPtr<USkeleton> UCustomizableObject::GetReferencedSkeletonAssetPtr( uint32 Index )
-{
-	if (ReferencedSkeletons.IsValidIndex(Index))
-	{
-		return ReferencedSkeletons[Index];
-	}
-
-	UE_LOG(LogMutable, Warning, TEXT("UCustomizableObject::GetReferencedSkeleton with an invalid Index. "
-		"Skeleton data and CustomizableObject data may be out of sync. "
-		"Try recompiling and saving the CustomizableObject asset [%s]."), *GetName());
-	return nullptr;
-}
-
-
 int32 UCustomizableObject::FindState( const FString& Name ) const
 {
 	int32 Result = -1;
@@ -1291,12 +1236,6 @@ void UCustomizableObject::SetModel(TSharedPtr<mu::Model, ESPMode::ThreadSafe> Mo
 
 	Private->SetModel(Model, GenerateIdentifier(*this));
 }
-
-void UCustomizableObject::SetBoneNamesArray(const TArray<FName>& InBoneNames)
-{
-	BoneNames = InBoneNames;
-}
-
 #endif // End WITH_EDITOR
 
 
@@ -1570,7 +1509,7 @@ FString UCustomizableObject::FindIntParameterValueName(int32 ParamIndex, int32 P
 
 FParameterUIData UCustomizableObject::GetParameterUIMetadata(const FString& ParamName) const
 {
-	const FParameterUIData* ParameterUIData = ParameterUIDataMap.Find(ParamName);
+	const FParameterUIData* ParameterUIData = Private->GetModelResources().ParameterUIDataMap.Find(ParamName);
 
 	return ParameterUIData ? *ParameterUIData : FParameterUIData();
 }
@@ -1584,7 +1523,7 @@ FParameterUIData UCustomizableObject::GetParameterUIMetadataFromIndex(int32 Para
 
 FParameterUIData UCustomizableObject::GetStateUIMetadata(const FString& StateName) const
 {
-	const FParameterUIData* StateUIData = StateUIDataMap.Find(StateName);
+	const FParameterUIData* StateUIData = Private->GetModelResources().StateUIDataMap.Find(StateName);
 
 	return StateUIData ? *StateUIData : FParameterUIData();
 }
@@ -1817,10 +1756,11 @@ void UCustomizableObject::GetLowPriorityTextureNames(TArray<FString>& OutTexture
 
 	if (!LowPriorityTextures.IsEmpty())
 	{
-		const int32 ImageCount = ImageProperties.Num();
+		const FModelResources& ModelResources = Private->GetModelResources();
+		const int32 ImageCount = ModelResources.ImageProperties.Num();
 		for (int32 ImageIndex = 0; ImageIndex < ImageCount; ++ImageIndex)
 		{
-			if (LowPriorityTextures.Find(FName(ImageProperties[ImageIndex].TextureParameterName)) != INDEX_NONE)
+			if (LowPriorityTextures.Find(FName(ModelResources.ImageProperties[ImageIndex].TextureParameterName)) != INDEX_NONE)
 			{
 				OutTextureNames.Add(FString::FromInt(ImageIndex));
 			}
@@ -1829,12 +1769,6 @@ void UCustomizableObject::GetLowPriorityTextureNames(TArray<FString>& OutTexture
 }
 
 
-const TArray<FName>& UCustomizableObject::GetBoneNamesArray() const
-{
-	return BoneNames;
-}
-
-
 int32 UCustomizableObject::GetMinLODIndex() const
 {
 	int32 MinLODIdx = 0;
@@ -1963,6 +1897,23 @@ TSharedPtr<const mu::Model, ESPMode::ThreadSafe> UCustomizableObjectPrivate::Get
 	return MutableModel;
 }
 
+const FModelResources& UCustomizableObjectPrivate::GetModelResources() const
+{
+#if WITH_EDITORONLY_DATA
+	return ModelResourcesEditor;
+#else
+	return ModelResources;
+#endif
+}
+
+
+#if WITH_EDITORONLY_DATA
+FModelResources& UCustomizableObjectPrivate::GetModelResources(bool bIsCooking)
+{
+	return bIsCooking ? ModelResources : ModelResourcesEditor;
+}
+#endif
+
 
 //-------------------------------------------------------------------------------------------------
 //-------------------------------------------------------------------------------------------------
@@ -2105,26 +2056,179 @@ void UCustomizableObjectBulk::PrepareBulkData(UCustomizableObject* InOuter, cons
 	}
 }
 
-#if WITH_EDITORONLY_DATA
+#endif // WITH_EDITOR
 
 //-------------------------------------------------------------------------------------------------
 //-------------------------------------------------------------------------------------------------
 //-------------------------------------------------------------------------------------------------
-void FMutableRefAssetUserData::InitResources(UCustomizableObject* InOuter)
+bool FAnimBpOverridePhysicsAssetsInfo::operator==(const FAnimBpOverridePhysicsAssetsInfo& Rhs) const
 {
-	UClass* AssetUserDataClass = FindObject<UClass>(nullptr, *ClassPath);
-	if (AssetUserDataClass)
+	return AnimInstanceClass == Rhs.AnimInstanceClass &&
+		SourceAsset == Rhs.SourceAsset &&
+		PropertyIndex == Rhs.PropertyIndex;
+}
+
+
+bool FMutableModelImageProperties::operator!=(const FMutableModelImageProperties& Other) const
+{
+	return
+		TextureParameterName != Other.TextureParameterName ||
+		Filter != Other.Filter ||
+		SRGB != Other.SRGB ||
+		FlipGreenChannel != Other.FlipGreenChannel ||
+		IsPassThrough != Other.IsPassThrough ||
+		LODBias != Other.LODBias ||
+		LODGroup != Other.LODGroup ||
+		AddressX != Other.AddressX ||
+		AddressY != Other.AddressY;
+}
+
+
+bool FMutableRefSocket::operator==(const FMutableRefSocket& Other) const
+{
+	if (
+		SocketName == Other.SocketName &&
+		BoneName == Other.BoneName &&
+		RelativeLocation == Other.RelativeLocation &&
+		RelativeRotation == Other.RelativeRotation &&
+		RelativeScale == Other.RelativeScale &&
+		bForceAlwaysAnimated == Other.bForceAlwaysAnimated &&
+		Priority == Other.Priority)
 	{
-		AssetUserData = UE::Mutable::Private::MoveOldObjectAndCreateNew<UAssetUserData>(AssetUserDataClass, InOuter);
-		if (AssetUserData)
-		{
-			FMemoryReaderView MemoryReader(Bytes);
-			AssetUserData->Serialize(MemoryReader);
+		return true;
+	}
 
-			ClassPath.Empty();
-			Bytes.Empty();
-		}
+	return false;
+}
+
+
+bool FMutableSkinWeightProfileInfo::operator==(const FMutableSkinWeightProfileInfo& Other) const
+{
+	return Name == Other.Name;
+}
+
+
+#if WITH_EDITORONLY_DATA
+FArchive& operator<<(FArchive& Ar, FMutableModelImageProperties& ImageProps)
+{
+	Ar << ImageProps.TextureParameterName;
+	Ar << ImageProps.Filter;
+
+	// Bitfields don't serialize automatically with FArchive
+	if (Ar.IsLoading())
+	{
+		int32 Aux = 0;
+		Ar << Aux;
+		ImageProps.SRGB = Aux;
+
+		Aux = 0;
+		Ar << Aux;
+		ImageProps.FlipGreenChannel = Aux;
+
+		Aux = 0;
+		Ar << Aux;
+		ImageProps.IsPassThrough = Aux;
+	}
+	else
+	{
+		int32 Aux = ImageProps.SRGB;
+		Ar << Aux;
+
+		Aux = ImageProps.FlipGreenChannel;
+		Ar << Aux;
+
+		Aux = ImageProps.IsPassThrough;
+		Ar << Aux;
+	}
+
+	Ar << ImageProps.LODBias;
+	Ar << ImageProps.LODGroup;
+
+	Ar << ImageProps.AddressX;
+	Ar << ImageProps.AddressY;
+
+	return Ar;
+}
+
+
+FArchive& operator<<(FArchive& Ar, FAnimBpOverridePhysicsAssetsInfo& Info)
+{
+	FString AnimInstanceClassPathString;
+	FString PhysicsAssetPathString;
+
+	if (Ar.IsLoading())
+	{
+		Ar << AnimInstanceClassPathString;
+		Ar << PhysicsAssetPathString;
+		Ar << Info.PropertyIndex;
+
+		Info.AnimInstanceClass = TSoftClassPtr<UAnimInstance>(AnimInstanceClassPathString);
+		Info.SourceAsset = TSoftObjectPtr<UPhysicsAsset>(PhysicsAssetPathString);
 	}
+
+	if (Ar.IsSaving())
+	{
+		AnimInstanceClassPathString = Info.AnimInstanceClass.ToString();
+		PhysicsAssetPathString = Info.SourceAsset.ToString();
+
+		Ar << AnimInstanceClassPathString;
+		Ar << PhysicsAssetPathString;
+		Ar << Info.PropertyIndex;
+	}
+
+	return Ar;
+}
+
+
+FArchive& operator<<(FArchive& Ar, FMutableRefSocket& Data)
+{
+	Ar << Data.SocketName;
+	Ar << Data.BoneName;
+	Ar << Data.RelativeLocation;
+	Ar << Data.RelativeRotation;
+	Ar << Data.RelativeScale;
+	Ar << Data.bForceAlwaysAnimated;
+	Ar << Data.Priority;
+
+	return Ar;
+}
+
+
+FArchive& operator<<(FArchive& Ar, FMutableRefLODRenderData& Data)
+{
+	Ar << Data.bIsLODOptional;
+	Ar << Data.bStreamedDataInlined;
+
+	return Ar;
+}
+
+
+FArchive& operator<<(FArchive& Ar, FMutableRefLODInfo& Data)
+{
+	Ar << Data.ScreenSize;
+	Ar << Data.LODHysteresis;
+	Ar << Data.bSupportUniformlyDistributedSampling;
+	Ar << Data.bAllowCPUAccess;
+
+	return Ar;
+}
+
+
+FArchive& operator<<(FArchive& Ar, FMutableRefLODData& Data)
+{
+	Ar << Data.LODInfo;
+	Ar << Data.RenderData;
+
+	return Ar;
+}
+
+
+FArchive& operator<<(FArchive& Ar, FMutableRefSkeletalMeshSettings& Data)
+{
+	Ar << Data.bEnablePerPolyCollision;
+	Ar << Data.DefaultUVChannelDensity;
+
+	return Ar;
 }
 
 
@@ -2156,6 +2260,58 @@ FArchive& operator<<(FArchive& Ar, FMutableRefAssetUserData& Data)
 }
 
 
+FArchive& operator<<(FArchive& Ar, FMutableSkinWeightProfileInfo& Info)
+{
+	Ar << Info.Name;
+	Ar << Info.DefaultProfile;
+	Ar << Info.DefaultProfileFromLODIndex;
+
+	return Ar;
+}
+
+
+//-------------------------------------------------------------------------------------------------
+//-------------------------------------------------------------------------------------------------
+//-------------------------------------------------------------------------------------------------
+void FMutableRefAssetUserData::InitResources(UCustomizableObject* InOuter)
+{
+	UClass* AssetUserDataClass = FindObject<UClass>(nullptr, *ClassPath);
+	if (AssetUserDataClass)
+	{
+		AssetUserData = UE::Mutable::Private::MoveOldObjectAndCreateNew<UAssetUserData>(AssetUserDataClass, InOuter);
+		if (AssetUserData)
+		{
+			FMemoryReaderView MemoryReader(Bytes);
+			AssetUserData->Serialize(MemoryReader);
+
+			ClassPath.Empty();
+			Bytes.Empty();
+		}
+	}
+}
+
+
+//-------------------------------------------------------------------------------------------------
+//-------------------------------------------------------------------------------------------------
+//-------------------------------------------------------------------------------------------------
+void FMutableRefSkeletalMeshData::InitResources(UCustomizableObject* InOuter, const ITargetPlatform* InTargetPlatform)
+{
+	check(InOuter);
+
+	const bool bHasServer = InTargetPlatform ? !InTargetPlatform->IsClientOnly() : false;	
+	if (InOuter->IsEnableUseRefSkeletalMeshAsPlaceholder() || bHasServer)
+	{
+		SkeletalMesh = TSoftObjectPtr<USkeletalMesh>(SoftSkeletalMesh).LoadSynchronous();
+	}
+
+	// Initialize AssetUserData
+	for (FMutableRefAssetUserData& Data : AssetUserData)
+	{
+		Data.InitResources(InOuter);
+	}
+}
+
+
 FArchive& operator<<(FArchive& Ar, FMutableRefSkeletalMeshData& Data)
 {
 	Ar << Data.LODData;
@@ -2165,19 +2321,19 @@ FArchive& operator<<(FArchive& Ar, FMutableRefSkeletalMeshData& Data)
 
 	if (Ar.IsSaving())
 	{
-		FString AssetPath = Data.SkeletalMeshAssetPath.ToString();
+		FString AssetPath = Data.SoftSkeletalMesh.ToString();
 		Ar << AssetPath;
 
-		AssetPath = Data.Skeleton.ToSoftObjectPath().ToString();
+		AssetPath = TSoftObjectPtr<USkeleton>(Data.Skeleton).ToString();
 		Ar << AssetPath;
 
-		AssetPath = Data.PhysicsAsset.ToSoftObjectPath().ToString();
+		AssetPath = TSoftObjectPtr<UPhysicsAsset>(Data.PhysicsAsset).ToString();
 		Ar << AssetPath;
 
-		AssetPath = Data.PostProcessAnimInst.ToSoftObjectPath().ToString();
+		AssetPath = Data.PostProcessAnimInst.ToString();
 		Ar << AssetPath;
 
-		AssetPath = Data.ShadowPhysicsAsset.ToSoftObjectPath().ToString();
+		AssetPath = TSoftObjectPtr<UPhysicsAsset>(Data.ShadowPhysicsAsset).ToString();
 		Ar << AssetPath;
 
 	}
@@ -2185,23 +2341,23 @@ FArchive& operator<<(FArchive& Ar, FMutableRefSkeletalMeshData& Data)
 	{
 		FString SkeletalMeshAssetPath;
 		Ar << SkeletalMeshAssetPath;
-		Data.SkeletalMeshAssetPath = SkeletalMeshAssetPath;
+		Data.SoftSkeletalMesh = SkeletalMeshAssetPath;
 
 		FString SkeletonAssetPath;
 		Ar << SkeletonAssetPath;
-		Data.Skeleton = TSoftObjectPtr<USkeleton>(FSoftObjectPath(SkeletonAssetPath));
+		Data.Skeleton = TSoftObjectPtr<USkeleton>(SkeletonAssetPath).LoadSynchronous();
 
 		FString PhysicsAssetPath;
 		Ar << PhysicsAssetPath;
-		Data.PhysicsAsset = TSoftObjectPtr<UPhysicsAsset>(FSoftObjectPath(PhysicsAssetPath));
+		Data.PhysicsAsset = TSoftObjectPtr<UPhysicsAsset>(PhysicsAssetPath).LoadSynchronous();
 
 		FString PostProcessAnimInstAssetPath;
 		Ar << PostProcessAnimInstAssetPath;
-		Data.PostProcessAnimInst = TSoftClassPtr<UAnimInstance>(FSoftObjectPath(PostProcessAnimInstAssetPath));
+		Data.PostProcessAnimInst = TSoftClassPtr<UAnimInstance>(PostProcessAnimInstAssetPath).LoadSynchronous();
 
 		FString ShadowPhysicsAssetPath;
 		Ar << ShadowPhysicsAssetPath;
-		Data.ShadowPhysicsAsset = TSoftObjectPtr<UPhysicsAsset>(FSoftObjectPath(ShadowPhysicsAssetPath));
+		Data.ShadowPhysicsAsset = TSoftObjectPtr<UPhysicsAsset>(ShadowPhysicsAssetPath).LoadSynchronous();
 	}
 
 	Ar << Data.AssetUserData;
@@ -2209,27 +2365,6 @@ FArchive& operator<<(FArchive& Ar, FMutableRefSkeletalMeshData& Data)
 	return Ar;
 }
 
-
-void FMutableRefSkeletalMeshData::InitResources(UCustomizableObject* InOuter, const ITargetPlatform* InTargetPlatform)
-{
-	check(InOuter);
-
-	const bool bHasServer = InTargetPlatform ? !InTargetPlatform->IsClientOnly() : false;	
-	if (InOuter->IsEnableUseRefSkeletalMeshAsPlaceholder() || bHasServer)
-	{
-		SkeletalMesh = TSoftObjectPtr<USkeletalMesh>(SkeletalMeshAssetPath).LoadSynchronous();
-	}
-
-	// Initialize AssetUserData
-	for (FMutableRefAssetUserData& Data : AssetUserData)
-	{
-		Data.InitResources(InOuter);
-	}
-}
-
-
-#endif
-
-#endif
+#endif // WITH_EDITORONLY_DATA
 
 #undef LOCTEXT_NAMESPACE
diff --git a/Engine/Plugins/Experimental/Mutable/Source/CustomizableObject/Private/MuCO/CustomizableObjectInstance.cpp b/Engine/Plugins/Experimental/Mutable/Source/CustomizableObject/Private/MuCO/CustomizableObjectInstance.cpp
--- a/Engine/Plugins/Experimental/Mutable/Source/CustomizableObject/Private/MuCO/CustomizableObjectInstance.cpp
+++ b/Engine/Plugins/Experimental/Mutable/Source/CustomizableObject/Private/MuCO/CustomizableObjectInstance.cpp
@@ -1294,7 +1294,7 @@ namespace
 
 		Result->bNotForDedicatedServer = TemplateAsset->bNotForDedicatedServer;
 
-		const TArray<FName>& BoneNames = CustomizableObject.GetBoneNamesArray();
+		const FModelResources& ModelResources = CustomizableObject.GetPrivate()->GetModelResources();
 		TMap<FName, int32> BonesInUse;
 
 		const int32 MutablePhysicsBodyCount = MutablePhysics->GetBodyCount();
@@ -1302,9 +1302,9 @@ namespace
 		for ( int32 I = 0; I < MutablePhysicsBodyCount; ++I )
 		{
 			const uint16 BoneNameId = MutablePhysics->GetBodyBoneId(I);
-			if (BoneNames.IsValidIndex(BoneNameId))
+			if (ModelResources.BoneNames.IsValidIndex(BoneNameId))
 			{
-				FName BoneName = BoneNames[BoneNameId];
+				FName BoneName = ModelResources.BoneNames[BoneNameId];
 				BonesInUse.Add(BoneName, I);
 			}
 		}
@@ -1462,7 +1462,7 @@ UPhysicsAsset* UCustomizableInstancePrivate::GetOrBuildMainPhysicsAsset(
 
 	Result->bNotForDedicatedServer = TemplateAsset->bNotForDedicatedServer;
 
-	const TArray<FName>& BoneNames = CustomizableObject.GetBoneNamesArray();
+	const FModelResources& ModelResources = CustomizableObject.GetPrivate()->GetModelResources();
 	TMap<FName, int32> BonesInUse;
 
 	const int32 MutablePhysicsBodyCount = MutablePhysics->GetBodyCount();
@@ -1470,9 +1470,9 @@ UPhysicsAsset* UCustomizableInstancePrivate::GetOrBuildMainPhysicsAsset(
 	for ( int32 I = 0; I < MutablePhysicsBodyCount; ++I )
 	{
 		const uint16 BoneNameId = MutablePhysics->GetBodyBoneId(I);
-		if (BoneNames.IsValidIndex(BoneNameId))
+		if (ModelResources.BoneNames.IsValidIndex(BoneNameId))
 		{
-			FName BoneName = BoneNames[BoneNameId];
+			FName BoneName = ModelResources.BoneNames[BoneNameId];
 			BonesInUse.Add(BoneName, I);
 		}
 	}
@@ -1852,6 +1852,8 @@ bool UCustomizableInstancePrivate::UpdateSkeletalMesh_PostBeginUpdate0(UCustomiz
 	UCustomizableObject* CustomizableObject = Public->GetCustomizableObject();
 	check(CustomizableObject);
 
+	const FModelResources& ModelResources = CustomizableObject->GetPrivate()->GetModelResources();
+
 	// Initialize the maximum number of SkeletalMeshes we could possibly have. 
 	SkeletalMeshes.Init(nullptr, CustomizableObject->GetComponentCount());
 
@@ -1901,13 +1903,14 @@ bool UCustomizableInstancePrivate::UpdateSkeletalMesh_PostBeginUpdate0(UCustomiz
 			continue;
 		}
 
-		const FMutableRefSkeletalMeshData* RefSkeletalMeshData = CustomizableObject->GetRefSkeletalMeshData(Component.Id);
-		if (!RefSkeletalMeshData)
+		if (!ModelResources.ReferenceSkeletalMeshesData.IsValidIndex(Component.Id))
 		{
 			bSuccess = false;
 			break;
 		}
 
+		const FMutableRefSkeletalMeshData& RefSkeletalMeshData = ModelResources.ReferenceSkeletalMeshesData[Component.Id];
+
 		// Create and initialize the SkeletalMesh for this component
 		MUTABLE_CPUPROFILER_SCOPE(ConstructMesh);
 
@@ -1922,7 +1925,7 @@ bool UCustomizableInstancePrivate::UpdateSkeletalMesh_PostBeginUpdate0(UCustomiz
 		if (Component.Mesh)
 		{
 			// Construct a new skeleton, fix up ActiveBones and Bonemap arrays and recompute the RefInvMatrices
-			bSuccess = BuildSkeletonData(OperationData, *SkeletalMesh, *RefSkeletalMeshData, *CustomizableObject, Component.Id);
+			bSuccess = BuildSkeletonData(OperationData, *SkeletalMesh, RefSkeletalMeshData, *CustomizableObject, Component.Id);
 
 			if (!bSuccess)
 			{
@@ -1934,7 +1937,7 @@ bool UCustomizableInstancePrivate::UpdateSkeletalMesh_PostBeginUpdate0(UCustomiz
 			{
 				constexpr bool bDisallowCollisionBetweenAssets = true;
 				UPhysicsAsset* PhysicsAssetResult = GetOrBuildMainPhysicsAsset(
-					RefSkeletalMeshData->PhysicsAsset.Get(), MutablePhysics.get(), *CustomizableObject, Component.Id, bDisallowCollisionBetweenAssets);
+					RefSkeletalMeshData.PhysicsAsset, MutablePhysics.get(), *CustomizableObject, Component.Id, bDisallowCollisionBetweenAssets);
 
 				SkeletalMesh->SetPhysicsAsset(PhysicsAssetResult);
 
@@ -1963,7 +1966,7 @@ bool UCustomizableInstancePrivate::UpdateSkeletalMesh_PostBeginUpdate0(UCustomiz
 
 				const int32 PhysicsBodyExternalId = Component.Mesh->AdditionalPhysicsBodies[I]->CustomId;
 				
-				const FAnimBpOverridePhysicsAssetsInfo& Info = CustomizableObject->AnimBpOverridePhysiscAssetsInfo[PhysicsBodyExternalId];
+				const FAnimBpOverridePhysicsAssetsInfo& Info = ModelResources.AnimBpOverridePhysiscAssetsInfo[PhysicsBodyExternalId];
 
 				// Make sure the AnimInstance class is loaded. It is expected to be already loaded at this point though. 
 				UClass* AnimInstanceClassLoaded = Info.AnimInstanceClass.LoadSynchronous();
@@ -1987,7 +1990,7 @@ bool UCustomizableInstancePrivate::UpdateSkeletalMesh_PostBeginUpdate0(UCustomiz
 			}
 
 			// Add sockets from the SkeletalMesh of reference and from the MutableMesh
-			BuildMeshSockets(OperationData, SkeletalMesh, RefSkeletalMeshData, Public, Component.Mesh);
+			BuildMeshSockets(OperationData, SkeletalMesh, ModelResources, RefSkeletalMeshData, Component.Mesh);
 		}
 		else
 		{
@@ -2815,6 +2818,8 @@ void UCustomizableInstancePrivate::SetDefaultSkeletalMesh(bool bSetEmptyMesh) co
 	}
 
 	UCustomizableObject* CustomizableObject = Instance->GetCustomizableObject();
+	const FModelResources& ModelResources = CustomizableObject->GetPrivate()->GetModelResources();
+
 	const int32 NumComponents = CustomizableObject->GetComponentCount();
 
 	for (TObjectIterator<UCustomizableObjectInstanceUsage> It; It; ++It)
@@ -2831,14 +2836,18 @@ void UCustomizableInstancePrivate::SetDefaultSkeletalMesh(bool bSetEmptyMesh) co
 			continue;
 		}
 #endif
+		const int32 ComponentIndex = CustomizableObjectInstanceUsage->GetComponentIndex();
+		if (!ModelResources.ReferenceSkeletalMeshesData.IsValidIndex(ComponentIndex))
+		{
+			continue;
+		}
 
 		USkeletalMesh* SkeletalMesh = nullptr;
 		
-		FMutableRefSkeletalMeshData* RefSkeletalMeshData = CustomizableObject->GetRefSkeletalMeshData(CustomizableObjectInstanceUsage->GetComponentIndex());
-		if (!bSetEmptyMesh && RefSkeletalMeshData)
+		if (!bSetEmptyMesh)
 		{
 			// Force load the reference mesh if necessary. 
-			SkeletalMesh = TSoftObjectPtr<USkeletalMesh>(RefSkeletalMeshData->SkeletalMeshAssetPath).LoadSynchronous();
+			SkeletalMesh = ModelResources.ReferenceSkeletalMeshesData[ComponentIndex].SoftSkeletalMesh.LoadSynchronous();
 		}
 		
 		CustomizableObjectInstanceUsage->SetSkeletalMesh(SkeletalMesh);
@@ -3302,19 +3311,18 @@ void ConvertImage(UTexture2D* Texture, mu::Ptr<const mu::Image> MutableImage, co
 }
 
 
-void UCustomizableInstancePrivate::InitSkeletalMeshData(const TSharedRef<FUpdateContextPrivate>& OperationData, USkeletalMesh* SkeletalMesh, const FMutableRefSkeletalMeshData* RefSkeletalMeshData, const UCustomizableObject& CustomizableObject, int32 ComponentIndex)
+void UCustomizableInstancePrivate::InitSkeletalMeshData(const TSharedRef<FUpdateContextPrivate>& OperationData, USkeletalMesh* SkeletalMesh, const FMutableRefSkeletalMeshData& RefSkeletalMeshData, const UCustomizableObject& CustomizableObject, int32 ComponentIndex)
 {
 	MUTABLE_CPUPROFILER_SCOPE(UCustomizableInstancePrivate::InitSkeletalMesh);
 
 	check(SkeletalMesh);
-	check(RefSkeletalMeshData);
 
 	// Mutable skeletal meshes are generated dynamically in-game and cannot be streamed from disk
 	SkeletalMesh->NeverStream = 1;
 
-	SkeletalMesh->SetImportedBounds(RefSkeletalMeshData->Bounds);
-	SkeletalMesh->SetPostProcessAnimBlueprint(RefSkeletalMeshData->PostProcessAnimInst.Get());
-	SkeletalMesh->SetShadowPhysicsAsset(RefSkeletalMeshData->ShadowPhysicsAsset.Get());
+	SkeletalMesh->SetImportedBounds(RefSkeletalMeshData.Bounds);
+	SkeletalMesh->SetPostProcessAnimBlueprint(RefSkeletalMeshData.PostProcessAnimInst.Get());
+	SkeletalMesh->SetShadowPhysicsAsset(RefSkeletalMeshData.ShadowPhysicsAsset.Get());
 
 	// Set Min LOD
 	SkeletalMesh->SetMinLod(FMath::Max(CustomizableObject.LODSettings.MinLOD.GetDefault(), (int32)FirstLODAvailable));
@@ -3324,12 +3332,12 @@ void UCustomizableInstancePrivate::InitSkeletalMeshData(const TSharedRef<FUpdate
 	SkeletalMesh->SetHasVertexColors(false);
 
 	// Set the default Physics Assets
-	SkeletalMesh->SetPhysicsAsset(RefSkeletalMeshData->PhysicsAsset.Get());
-	SkeletalMesh->SetEnablePerPolyCollision(RefSkeletalMeshData->Settings.bEnablePerPolyCollision);
+	SkeletalMesh->SetPhysicsAsset(RefSkeletalMeshData.PhysicsAsset.Get());
+	SkeletalMesh->SetEnablePerPolyCollision(RefSkeletalMeshData.Settings.bEnablePerPolyCollision);
 
 	// Asset User Data
 	{
-		for (const FMutableRefAssetUserData& MutAssetUserData : RefSkeletalMeshData->AssetUserData)
+		for (const FMutableRefAssetUserData& MutAssetUserData : RefSkeletalMeshData.AssetUserData)
 		{
 			if (MutAssetUserData.AssetUserData)
 			{
@@ -3362,7 +3370,7 @@ void UCustomizableInstancePrivate::InitSkeletalMeshData(const TSharedRef<FUpdate
 		{
 			RenderData->LODRenderData.Add(new FSkeletalMeshLODRenderData());
 
-			const FMutableRefLODData& LODData = RefSkeletalMeshData->LODData[LODIndex];
+			const FMutableRefLODData& LODData = RefSkeletalMeshData.LODData[LODIndex];
 
 			FSkeletalMeshLODInfo& LODInfo = SkeletalMesh->AddLODInfo();
 			LODInfo.ScreenSize = LODData.LODInfo.ScreenSize;
@@ -3420,15 +3428,15 @@ bool UCustomizableInstancePrivate::BuildSkeletonData(const TSharedRef<FUpdateCon
 	{
 		MUTABLE_CPUPROFILER_SCOPE(BuildSkeletonData_EnsureBonesExist);
 
-		const TArray<FName>& BoneNames = CustomizableObject.GetBoneNamesArray();
+		const FModelResources& ModelResources = CustomizableObject.GetPrivate()->GetModelResources();
 
 		// Ensure all the required bones are present in the skeleton
 		for (int32 BoneIndex = 0; BoneIndex < MutBoneCount; ++BoneIndex)
 		{
 			const uint16 BoneId = MutSkeletonData.BoneIds[BoneIndex];
-			check(BoneNames.IsValidIndex(BoneId));
+			check(ModelResources.BoneNames.IsValidIndex(BoneId));
 
-			const FName BoneName = BoneNames[BoneId]; 
+			const FName BoneName = ModelResources.BoneNames[BoneId];
 			check(BoneName != NAME_None);
 
 			const int32 SourceBoneIndex = ReferenceSkeleton.FindRawBoneIndex(BoneName);
@@ -3524,16 +3532,13 @@ bool UCustomizableInstancePrivate::BuildSkeletonData(const TSharedRef<FUpdateCon
 }
 
 
-void UCustomizableInstancePrivate::BuildMeshSockets(const TSharedRef<FUpdateContextPrivate>& OperationData, USkeletalMesh* SkeletalMesh, const FMutableRefSkeletalMeshData* RefSkeletalMeshData, UCustomizableObjectInstance* Public, mu::MeshPtrConst MutableMesh)
+void UCustomizableInstancePrivate::BuildMeshSockets(const TSharedRef<FUpdateContextPrivate>& OperationData, USkeletalMesh* SkeletalMesh, const FModelResources& ModelResources, const FMutableRefSkeletalMeshData& RefSkeletalMeshData, mu::MeshPtrConst MutableMesh)
 {
 	// Build mesh sockets.
 	MUTABLE_CPUPROFILER_SCOPE(UCustomizableInstancePrivate::BuildMeshSockets);
-
 	check(SkeletalMesh);
-	check(RefSkeletalMeshData);
-	check(Public);
 
-	const uint32 SocketCount = RefSkeletalMeshData->Sockets.Num();
+	const uint32 SocketCount = RefSkeletalMeshData.Sockets.Num();
 
 	TArray<TObjectPtr<USkeletalMeshSocket>>& Sockets = SkeletalMesh->GetMeshOnlySocketList();
 	Sockets.Empty(SocketCount);
@@ -3545,7 +3550,7 @@ void UCustomizableInstancePrivate::BuildMeshSockets(const TSharedRef<FUpdateCont
 	
 		for (uint32 SocketIndex = 0; SocketIndex < SocketCount; ++SocketIndex)
 		{
-			const FMutableRefSocket& RefSocket = RefSkeletalMeshData->Sockets[SocketIndex];
+			const FMutableRefSocket& RefSocket = RefSkeletalMeshData.Sockets[SocketIndex];
 
 			USkeletalMeshSocket* Socket = NewObject<USkeletalMeshSocket>(SkeletalMesh);
 
@@ -3568,9 +3573,6 @@ void UCustomizableInstancePrivate::BuildMeshSockets(const TSharedRef<FUpdateCont
 	{
 		MUTABLE_CPUPROFILER_SCOPE(BuildMeshSockets_MutableSockets);
 
-		const UCustomizableObject* CustomizableObject = Public->GetCustomizableObject();
-		check(CustomizableObject);
-
 		for (int32 TagIndex = 0; TagIndex < MutableMesh->GetTagCount(); ++TagIndex)
 		{
 			FString Tag = MutableMesh->GetTag(TagIndex);
@@ -3580,9 +3582,9 @@ void UCustomizableInstancePrivate::BuildMeshSockets(const TSharedRef<FUpdateCont
 				check(Tag.IsNumeric());
 				const int32 MutableSocketIndex = FCString::Atoi(*Tag);
 
-				if (CustomizableObject->SocketArray.IsValidIndex(MutableSocketIndex))
+				if (ModelResources.SocketArray.IsValidIndex(MutableSocketIndex))
 				{
-					const FMutableRefSocket& MutableSocket = CustomizableObject->SocketArray[MutableSocketIndex];
+					const FMutableRefSocket& MutableSocket = ModelResources.SocketArray[MutableSocketIndex];
 					int32 IndexToWriteSocket = -1;
 
 					if (TTuple<int32, int32>* FoundSocket = SocketMap.Find(MutableSocket.SocketName))
@@ -4679,6 +4681,7 @@ bool UCustomizableInstancePrivate::BuildOrCopyRenderData(const TSharedRef<FUpdat
 
 	UCustomizableObject* CustomizableObject = Public->GetCustomizableObject();
 	check(CustomizableObject);
+	const FModelResources& ModelResources = CustomizableObject->GetPrivate()->GetModelResources();
 
 	int32 LastValidLODIndex = OperationData->CurrentMaxLOD;
 	for (int32 LODIndex = OperationData->CurrentMaxLOD; LODIndex >= FirstLODAvailable; --LODIndex)
@@ -4766,14 +4769,13 @@ bool UCustomizableInstancePrivate::BuildOrCopyRenderData(const TSharedRef<FUpdat
 		LODModel.ActiveBoneIndices.Append(Component->ActiveBones);
 		LODModel.RequiredBones.Append(Component->ActiveBones);
 
-		const TArray<FMutableSkinWeightProfileInfo>& SkinWeightProfilesInfo = CustomizableObject->SkinWeightProfilesInfo;
-		if (SkinWeightProfilesInfo.Num())
+		if (!ModelResources.SkinWeightProfilesInfo.IsEmpty())
 		{
 			bool bHasSkinWeightProfiles = false;
 
 			const mu::FMeshBufferSet& MutableMeshVertexBuffers = Component->Mesh->GetVertexBuffers();
-
-			const int32 SkinWeightProfilesCount = CustomizableObject->SkinWeightProfilesInfo.Num();
+			
+			const int32 SkinWeightProfilesCount = ModelResources.SkinWeightProfilesInfo.Num();
 			for (int32 ProfileIndex = 0; ProfileIndex < SkinWeightProfilesCount; ++ProfileIndex)
 			{
 				const int32 ProfileSemanticsIndex = ProfileIndex + 10;
@@ -4794,7 +4796,7 @@ bool UCustomizableInstancePrivate::BuildOrCopyRenderData(const TSharedRef<FUpdat
 					bHasSkinWeightProfiles = true;
 				}
 
-				const FMutableSkinWeightProfileInfo& Profile = CustomizableObject->SkinWeightProfilesInfo[ProfileIndex];
+				const FMutableSkinWeightProfileInfo& Profile = ModelResources.SkinWeightProfilesInfo[ProfileIndex];
 
 				const FSkinWeightProfileInfo* ExistingProfile = SkeletalMesh->GetSkinWeightProfiles().FindByPredicate(
 					[&Profile](const FSkinWeightProfileInfo& P) { return P.Name == Profile.Name; });
@@ -4820,7 +4822,7 @@ bool UCustomizableInstancePrivate::BuildOrCopyRenderData(const TSharedRef<FUpdat
 		}
 
 		// Update LOD and streaming data
-		const FMutableRefLODRenderData& RefLODRenderData = CustomizableObject->GetRefSkeletalMeshData(Component->Id)->LODData[LODIndex].RenderData;
+		const FMutableRefLODRenderData& RefLODRenderData = ModelResources.ReferenceSkeletalMeshesData[Component->Id].LODData[LODIndex].RenderData;
 		LODModel.bIsLODOptional = RefLODRenderData.bIsLODOptional;
 		LODModel.bStreamedDataInlined = RefLODRenderData.bStreamedDataInlined;
 
@@ -4849,6 +4851,8 @@ FGraphEventRef UCustomizableInstancePrivate::LoadAdditionalAssetsAsync(const TSh
 
 	UCustomizableObject* CustomizableObject = Public->GetCustomizableObject();
 
+	const FModelResources& ModelResources = CustomizableObject->GetPrivate()->GetModelResources();
+
 	FGraphEventRef Result = nullptr;
 
 	TArray<FSoftObjectPath> AssetsToStream;
@@ -4869,26 +4873,21 @@ FGraphEventRef UCustomizableInstancePrivate::LoadAdditionalAssetsAsync(const TSh
 
 	for (const FInstanceUpdateData::FSurface& Surface : OperationData->InstanceUpdateData.Surfaces)
 	{
-		uint32 MaterialIndex = Surface.MaterialIndex;
-
-		if (!ObjectToInstanceIndexMap.Contains(MaterialIndex))
+		const uint32 MaterialIndex = Surface.MaterialIndex;
+		if (ObjectToInstanceIndexMap.Contains(MaterialIndex))
 		{
-			UMaterialInterface* MatInterface = CustomizableObject->GetReferencedMaterialAssetPtr(MaterialIndex).Get();
-
-			ObjectToInstanceIndexMap.Add(MaterialIndex, ReferencedMaterials.Num());
-			ReferencedMaterials.Add(MatInterface);
+			continue;
+		}
 
-			check(ObjectToInstanceIndexMap.Num() == ReferencedMaterials.Num());
+		TSoftObjectPtr<UMaterialInterface> AssetPtr = ModelResources.Materials.IsValidIndex(MaterialIndex) ? ModelResources.Materials[MaterialIndex] : nullptr;
+		UMaterialInterface* LoadedMaterial = AssetPtr.Get();
 
-			if (!MatInterface)
-			{
-				FSoftObjectPath MaterialPath = CustomizableObject->GetReferencedMaterialAssetPtr(MaterialIndex).ToSoftObjectPath();
+		const int32 ReferencedMaterialsIndex = ReferencedMaterials.Add(LoadedMaterial);
+		ObjectToInstanceIndexMap.Add(MaterialIndex, ReferencedMaterialsIndex);
 
-				if (MaterialPath.IsValid())
-				{
-					AssetsToStream.Add(MaterialPath);
-				}
-			}
+		if (!LoadedMaterial && !AssetPtr.IsNull())
+		{
+			AssetsToStream.Add(AssetPtr.ToSoftObjectPath());
 		}
 	}
 
@@ -4915,8 +4914,7 @@ FGraphEventRef UCustomizableInstancePrivate::LoadAdditionalAssetsAsync(const TSh
 		// Add Skeletons to merge
 		for (const uint32 SkeletonId : SkeletonData.SkeletonIds)
 		{
-			TSoftObjectPtr<USkeleton> AssetPtr = CustomizableObject->GetReferencedSkeletonAssetPtr(SkeletonId);
-
+			TSoftObjectPtr<USkeleton> AssetPtr = ModelResources.Skeletons.IsValidIndex(SkeletonId) ? ModelResources.Skeletons[SkeletonId] : nullptr;
 			if (AssetPtr.IsNull())
 			{
 				continue;
@@ -4978,95 +4976,100 @@ FGraphEventRef UCustomizableInstancePrivate::LoadAdditionalAssetsAsync(const TSh
 			for (int32 TagIndex = 0; TagIndex < MutableMesh->GetTagCount(); ++TagIndex)
 			{
 				FString Tag = MutableMesh->GetTag(TagIndex);
-				if (Tag.RemoveFromStart("__PhysicsAsset:"))
+				if (Tag.RemoveFromStart("__PA:"))
 				{
-					TSoftObjectPtr<UPhysicsAsset>* PhysicsAsset = CustomizableObject->PhysicsAssetsMap.Find(Tag);
+					const int32 AssetIndex = FCString::Atoi(*Tag);
+					const TSoftObjectPtr<UPhysicsAsset>& PhysicsAsset = ModelResources.PhysicsAssets.IsValidIndex(AssetIndex) ? ModelResources.PhysicsAssets[AssetIndex] : nullptr;
 
-					if (!PhysicsAsset->IsNull())
+					if (!PhysicsAsset.IsNull())
 					{
-						if (PhysicsAsset->Get())
+						if (PhysicsAsset.Get())
 						{
-							ComponentData->PhysicsAssets.PhysicsAssetsToMerge.Add(PhysicsAsset->Get());
+							ComponentData->PhysicsAssets.PhysicsAssetsToMerge.Add(PhysicsAsset.Get());
 						}
 						else
 						{
-							ComponentData->PhysicsAssets.PhysicsAssetToLoad.Add(PhysicsAsset->ToSoftObjectPath().ToString());
-							AssetsToStream.Add(PhysicsAsset->ToSoftObjectPath());
+							ComponentData->PhysicsAssets.PhysicsAssetToLoad.Add(AssetIndex);
+							AssetsToStream.Add(PhysicsAsset.ToSoftObjectPath());
 						}
 					}
 				}
 				else if (Tag.RemoveFromStart("__ClothPhysicsAsset:"))
 				{
-					FString AssetIndexString, AssetPath;
+					FString AssetIndexString, PhysicsAssetIndexString;
 
-					if (Tag.Split(TEXT("_AssetIdx_"), &AssetIndexString, &AssetPath) && AssetIndexString.IsNumeric())
+					if (Tag.Split(TEXT("_AssetIdx_"), &AssetIndexString, &PhysicsAssetIndexString) && AssetIndexString.IsNumeric())
 					{
+						const int32 AssetIndex = FCString::Atoi(*AssetIndexString);
+						const int32 PhysicsAssetIndex = FCString::Atoi(*PhysicsAssetIndexString);
 
-						int32 AssetIndex = FCString::Atoi(*AssetIndexString);
-
-						TSoftObjectPtr<UPhysicsAsset>* PhysicsAssetPtr = CustomizableObject->PhysicsAssetsMap.Find(AssetPath);
+						const TSoftObjectPtr<UPhysicsAsset>& PhysicsAsset = ModelResources.PhysicsAssets.IsValidIndex(PhysicsAssetIndex) ? ModelResources.PhysicsAssets[PhysicsAssetIndex] : nullptr;
 
 						// The entry should always be in the map
-						check(PhysicsAssetPtr);
-						if (!PhysicsAssetPtr->IsNull())
+						if (!PhysicsAsset.IsNull())
 						{
-							if (PhysicsAssetPtr->Get())
+							if (PhysicsAsset.Get())
 							{
 								if (ClothingPhysicsAssets.IsValidIndex(AssetIndex))
 								{
-									ClothingPhysicsAssets[AssetIndex] = PhysicsAssetPtr->Get();
+									ClothingPhysicsAssets[AssetIndex] = PhysicsAsset.Get();
 								}
 							}
 							else
 							{
-								ComponentData->ClothingPhysicsAssetsToStream.Emplace(AssetIndex, PhysicsAssetPtr->ToSoftObjectPath().ToString());
-								AssetsToStream.Add(PhysicsAssetPtr->ToSoftObjectPath());
+								ComponentData->ClothingPhysicsAssetsToStream.Emplace(AssetIndex, PhysicsAssetIndex);
+								AssetsToStream.Add(PhysicsAsset.ToSoftObjectPath());
 							}
 						}
 					}
 				}
 				if (Tag.RemoveFromStart("__AnimBP:"))
 				{
-					FString SlotIndexString, AssetPath;
+					FString SlotIndexString, AnimBpIndexString;
 
-					if (Tag.Split(TEXT("_Slot_"), &SlotIndexString, &AssetPath))
+					if (Tag.Split(TEXT("_Slot_"), &SlotIndexString, &AnimBpIndexString))
 					{
-						if (!SlotIndexString.IsEmpty())
+						if (SlotIndexString.IsEmpty() || AnimBpIndexString.IsEmpty())
 						{
-							FName SlotIndex = *SlotIndexString;
+							continue;
+						}
 
-							TSoftClassPtr<UAnimInstance>* AnimBPAsset = CustomizableObject->AnimBPAssetsMap.Find(AssetPath);
+						const int32 AnimBpIndex = FCString::Atoi(*AnimBpIndexString);
+						if (!ModelResources.AnimBPs.IsValidIndex(AnimBpIndex))
+						{
+							continue;
+						}
+
+						FName SlotIndex = *SlotIndexString;
 
-							if (AnimBPAsset && !AnimBPAsset->IsNull())
+						const TSoftClassPtr<UAnimInstance>& AnimBPAsset = ModelResources.AnimBPs[AnimBpIndex];
+
+						if (!AnimBPAsset.IsNull())
+						{
+							if (!ComponentData->AnimSlotToBP.Contains(SlotIndex))
 							{
-								if (!ComponentData->AnimSlotToBP.Contains(SlotIndex))
-								{
-									ComponentData->AnimSlotToBP.Add(SlotIndex, *AnimBPAsset);
+								ComponentData->AnimSlotToBP.Add(SlotIndex, AnimBPAsset);
 
-									if (AnimBPAsset->Get())
-									{
-										GatheredAnimBPs.Add(AnimBPAsset->Get());
-									}
-									else
-									{
-										AssetsToStream.Add(AnimBPAsset->ToSoftObjectPath());
-									}
+								if (AnimBPAsset.Get())
+								{
+									GatheredAnimBPs.Add(AnimBPAsset.Get());
 								}
 								else
 								{
-									// Two submeshes should not have the same animation slot index
-									OperationData->UpdateResult = EUpdateResult::Warning;
-									
-									FString WarningMessage = FString::Printf(TEXT("Two submeshes have the same anim slot index [%s] in a Mutable Instance."), *SlotIndex.ToString());
-									UE_LOG(LogMutable, Warning, TEXT("%s"), *WarningMessage);
-#if WITH_EDITOR
-									FMessageLogModule& MessageLogModule = FModuleManager::LoadModuleChecked<FMessageLogModule>("MessageLog");
-									MessageLogModule.RegisterLogListing(FName("Mutable"), FText::FromString(FString("Mutable")));
-									FMessageLog MessageLog("Mutable");
+									AssetsToStream.Add(AnimBPAsset.ToSoftObjectPath());
+								}
+							}
+							else
+							{
+								// Two submeshes should not have the same animation slot index
+								OperationData->UpdateResult = EUpdateResult::Warning;
 
-									MessageLog.Notify(FText::FromString(WarningMessage), EMessageSeverity::Warning, true);
+								FString WarningMessage = FString::Printf(TEXT("Two submeshes have the same anim slot index [%s] in a Mutable Instance."), *SlotIndex.ToString());
+								UE_LOG(LogMutable, Warning, TEXT("%s"), *WarningMessage);
+#if WITH_EDITOR
+								FMessageLog MessageLog("Mutable");
+								MessageLog.Notify(FText::FromString(WarningMessage), EMessageSeverity::Warning, true);
 #endif
-								}
 							}
 						}
 					}
@@ -5089,7 +5092,7 @@ FGraphEventRef UCustomizableInstancePrivate::LoadAdditionalAssetsAsync(const TSh
 				const int32 ExternalId = MutableMesh->AdditionalPhysicsBodies[I]->CustomId;
 				
 				ComponentData->PhysicsAssets.AdditionalPhysicsAssetsToLoad.Add(ExternalId);
-				AssetsToStream.Add(CustomizableObject->AnimBpOverridePhysiscAssetsInfo[ExternalId].SourceAsset.ToSoftObjectPath());
+				AssetsToStream.Add(ModelResources.AnimBpOverridePhysiscAssetsInfo[ExternalId].SourceAsset.ToSoftObjectPath());
 			}
 		}
 	}
@@ -5135,30 +5138,25 @@ void UCustomizableInstancePrivate::AdditionalAssetsAsyncLoaded(UCustomizableObje
 
 	UCustomizableObject* CustomizableObject = Public->GetCustomizableObject();
 
+	const FModelResources& ModelResources = CustomizableObject->GetPrivate()->GetModelResources();
+
 	// Loaded Materials
 	check(ObjectToInstanceIndexMap.Num() == ReferencedMaterials.Num());
 
 	for (TPair<uint32, uint32> Pair : ObjectToInstanceIndexMap)
 	{
-		ReferencedMaterials[Pair.Value] = CustomizableObject->GetReferencedMaterialAssetPtr(Pair.Key).Get();
-	}
+		const TSoftObjectPtr<UMaterialInterface>& AssetPtr = ModelResources.Materials.IsValidIndex(Pair.Key) ? ModelResources.Materials[Pair.Key] : nullptr;
+		ReferencedMaterials[Pair.Value] = AssetPtr.Get();
 
 #if WITH_EDITOR
-	for (int32 i = 0; i < ReferencedMaterials.Num(); ++i)
-	{
-		if (!ReferencedMaterials[i])
+		if (!ReferencedMaterials[Pair.Value])
 		{
-			const uint32* Key = ObjectToInstanceIndexMap.FindKey(i);
-			TSoftObjectPtr<UMaterialInterface> ErrorMaterial = Key != nullptr ? CustomizableObject->GetReferencedMaterialAssetPtr(*Key) : nullptr;
-
-			if (!ErrorMaterial.IsNull())
+			if (!AssetPtr.IsNull())
 			{
-				FMessageLogModule& MessageLogModule = FModuleManager::LoadModuleChecked<FMessageLogModule>("MessageLog");
-				MessageLogModule.RegisterLogListing(FName("Mutable"), FText::FromString(FString("Mutable")));
-				FMessageLog MessageLog("Mutable");
-
-				FString ErrorMsg = FString::Printf(TEXT("Mutable couldn't load the material [%s] and won't be rendered. If it has been deleted or renamed, please recompile all the mutable objects that use it."), *ErrorMaterial.GetAssetName());
+				FString ErrorMsg = FString::Printf(TEXT("Mutable couldn't load the material [%s] and won't be rendered. If it has been deleted or renamed, please recompile all the mutable objects that use it."), *AssetPtr.GetAssetName());
 				UE_LOG(LogMutable, Error, TEXT("%s"), *ErrorMsg);
+
+				FMessageLog MessageLog("Mutable");
 				MessageLog.Notify(FText::FromString(ErrorMsg), EMessageSeverity::Error, true);
 			}
 			else
@@ -5166,8 +5164,8 @@ void UCustomizableInstancePrivate::AdditionalAssetsAsyncLoaded(UCustomizableObje
 				ensure(false); // Couldn't load the material, and we don't know which material
 			}
 		}
-	}
 #endif
+	}
 
 	
 	TArray<FCustomizableObjectStreamedResourceData>& StreamedResources = CustomizableObject->StreamedResourceData;
@@ -5204,27 +5202,27 @@ void UCustomizableInstancePrivate::AdditionalAssetsAsyncLoaded(UCustomizableObje
 		FReferencedSkeletons& Skeletons = ComponentData.Skeletons;
 		for (int32 SkeletonIndex : Skeletons.SkeletonIds)
 		{
-			Skeletons.SkeletonsToMerge.AddUnique(CustomizableObject->GetReferencedSkeletonAssetPtr(SkeletonIndex).Get());
+			const TSoftObjectPtr<USkeleton>& AssetPtr = ModelResources.Skeletons.IsValidIndex(SkeletonIndex) ? ModelResources.Skeletons[SkeletonIndex] : nullptr;
+			Skeletons.SkeletonsToMerge.AddUnique(AssetPtr.Get());
 		}
 
 		// Loaded PhysicsAssets
 		FReferencedPhysicsAssets& PhysicsAssets = ComponentData.PhysicsAssets;
-		for(const FString& Path : PhysicsAssets.PhysicsAssetToLoad)
+		for(const int32 PhysicsAssetIndex : PhysicsAssets.PhysicsAssetToLoad)
 		{
-			const TSoftObjectPtr<UPhysicsAsset>* PhysicsAssetPtr = CustomizableObject->PhysicsAssetsMap.Find(Path);
-			PhysicsAssets.PhysicsAssetsToMerge.Add(PhysicsAssetPtr->Get());
+			check(ModelResources.PhysicsAssets.IsValidIndex(PhysicsAssetIndex));
+			const TSoftObjectPtr<UPhysicsAsset>& PhysicsAsset = ModelResources.PhysicsAssets[PhysicsAssetIndex];
+			PhysicsAssets.PhysicsAssetsToMerge.Add(PhysicsAsset.Get());
 
 #if WITH_EDITOR
-			if (!PhysicsAssetPtr->Get())
+			if (!PhysicsAsset.Get())
 			{
-				if (!PhysicsAssetPtr->IsNull())
+				if (!PhysicsAsset.IsNull())
 				{
-					FMessageLogModule& MessageLogModule = FModuleManager::LoadModuleChecked<FMessageLogModule>("MessageLog");
-					MessageLogModule.RegisterLogListing(FName("Mutable"), FText::FromString(FString("Mutable")));
-					FMessageLog MessageLog("Mutable");
-
-					FString ErrorMsg = FString::Printf(TEXT("Mutable couldn't load the PhysicsAsset [%s] and won't be merged. If it has been deleted or renamed, please recompile all the mutable objects that use it."), *PhysicsAssetPtr->GetAssetName());
+					FString ErrorMsg = FString::Printf(TEXT("Mutable couldn't load the PhysicsAsset [%s] and won't be merged. If it has been deleted or renamed, please recompile all the mutable objects that use it."), *PhysicsAsset.GetAssetName());
 					UE_LOG(LogMutable, Error, TEXT("%s"), *ErrorMsg);
+
+					FMessageLog MessageLog("Mutable");
 					MessageLog.Notify(FText::FromString(ErrorMsg), EMessageSeverity::Error, true);
 				}
 				else
@@ -5237,16 +5235,14 @@ void UCustomizableInstancePrivate::AdditionalAssetsAsyncLoaded(UCustomizableObje
 		PhysicsAssets.PhysicsAssetToLoad.Empty();
 		
 		// Loaded Clothing PhysicsAssets 
-		for ( TPair<int32, FString>& AssetToStream : ComponentData.ClothingPhysicsAssetsToStream )
+		for ( TPair<int32, int32>& AssetToStream : ComponentData.ClothingPhysicsAssetsToStream )
 		{
 			const int32 AssetIndex = AssetToStream.Key;
 
-			if (ClothingPhysicsAssets.IsValidIndex(AssetIndex))
+			if (ClothingPhysicsAssets.IsValidIndex(AssetIndex) && ModelResources.PhysicsAssets.IsValidIndex(AssetToStream.Value))
 			{
-				const TSoftObjectPtr<UPhysicsAsset>* PhysicsAssetPtr = CustomizableObject->PhysicsAssetsMap.Find(AssetToStream.Value);
-				check(PhysicsAssetPtr) // Should always be found.
-			
-				ClothingPhysicsAssets[AssetIndex] = PhysicsAssetPtr->Get();
+				const TSoftObjectPtr<UPhysicsAsset>& PhysicsAssetPtr = ModelResources.PhysicsAssets[AssetToStream.Value];
+				ClothingPhysicsAssets[AssetIndex] = PhysicsAssetPtr.Get();
 			}
 		}
 		ComponentData.ClothingPhysicsAssetsToStream.Empty();
@@ -5264,12 +5260,10 @@ void UCustomizableInstancePrivate::AdditionalAssetsAsyncLoaded(UCustomizableObje
 #if WITH_EDITOR
 			else
 			{
-				FMessageLogModule& MessageLogModule = FModuleManager::LoadModuleChecked<FMessageLogModule>("MessageLog");
-				MessageLogModule.RegisterLogListing(FName("Mutable"), FText::FromString(FString("Mutable")));
-				FMessageLog MessageLog("Mutable");
-
 				FString ErrorMsg = FString::Printf(TEXT("Mutable couldn't load the AnimBlueprint [%s]. If it has been deleted or renamed, please recompile all the mutable objects that use it."), *SlotAnimBP.Value.GetAssetName());
 				UE_LOG(LogMutable, Error, TEXT("%s"), *ErrorMsg);
+
+				FMessageLog MessageLog("Mutable");
 				MessageLog.Notify(FText::FromString(ErrorMsg), EMessageSeverity::Error, true);
 			}
 #endif
@@ -5282,7 +5276,7 @@ void UCustomizableInstancePrivate::AdditionalAssetsAsyncLoaded(UCustomizableObje
 			// Make the loaded assets references strong.
 			const int32 AnimBpPhysicsOverrideIndex = ComponentData.PhysicsAssets.AdditionalPhysicsAssetsToLoad[I];
 			ComponentData.PhysicsAssets.AdditionalPhysicsAssets.Add( 
-					CustomizableObject->AnimBpOverridePhysiscAssetsInfo[AnimBpPhysicsOverrideIndex].SourceAsset.Get());
+				ModelResources.AnimBpOverridePhysiscAssetsInfo[AnimBpPhysicsOverrideIndex].SourceAsset.Get());
 		}
 		ComponentData.PhysicsAssets.AdditionalPhysicsAssetsToLoad.Empty();
 	}
@@ -5405,6 +5399,8 @@ void UCustomizableInstancePrivate::BuildMaterials(const TSharedRef<FUpdateContex
 
 	UCustomizableObject* CustomizableObject = Public->GetCustomizableObject();
 
+	const FModelResources& ModelResources = CustomizableObject->GetPrivate()->GetModelResources();
+
 	// Find skipped LODs. The following valid LOD will be copied into them. 
 	TArray<bool> LODsSkipped;
 	LODsSkipped.SetNum(OperationData->NumLODsAvailable);
@@ -5479,8 +5475,7 @@ void UCustomizableInstancePrivate::BuildMaterials(const TSharedRef<FUpdateContex
 
 			check(Component.bGenerated);
 
-			const FMutableRefSkeletalMeshData* RefSkeletalMeshData = CustomizableObject->GetRefSkeletalMeshData(Component.Id);
-			check(RefSkeletalMeshData);
+			const FMutableRefSkeletalMeshData& RefSkeletalMeshData = ModelResources.ReferenceSkeletalMeshesData[Component.Id];
 
 			for (int32 SurfaceIndex = 0; SurfaceIndex < Component.SurfaceCount; ++SurfaceIndex)
 			{
@@ -5511,8 +5506,13 @@ void UCustomizableInstancePrivate::BuildMaterials(const TSharedRef<FUpdateContex
 				const int32 MaterialSlotIndex = Materials.Num();
 				FSkeletalMaterial& MaterialSlot = Materials.AddDefaulted_GetRef();
 				MaterialSlot.MaterialInterface = MaterialTemplate;
-				MaterialSlot.MaterialSlotName = CustomizableObject->ReferencedMaterialSlotNames.IsValidIndex(Surface.MaterialIndex) ? CustomizableObject->ReferencedMaterialSlotNames[Surface.MaterialIndex] : NAME_None;
-				SetMeshUVChannelDensity(MaterialSlot.UVChannelData, RefSkeletalMeshData->Settings.DefaultUVChannelDensity);
+
+				if (ModelResources.MaterialSlotNames.IsValidIndex(Surface.MaterialIndex))
+				{
+					MaterialSlot.MaterialSlotName = ModelResources.MaterialSlotNames[Surface.MaterialIndex];
+				}
+
+				SetMeshUVChannelDensity(MaterialSlot.UVChannelData, RefSkeletalMeshData.Settings.DefaultUVChannelDensity);
 
 				const int32 LODMaterialIndex = SkeletalMesh->GetLODInfoArray()[LODIndex].LODMaterialMap.Add(MaterialSlotIndex);
 				SkeletalMesh->GetResourceForRendering()->LODRenderData[LODIndex].RenderSections[SurfaceIndex].MaterialIndex = LODMaterialIndex;
@@ -5633,9 +5633,9 @@ void UCustomizableInstancePrivate::BuildMaterials(const TSharedRef<FUpdateContex
 							check(MutableImage->IsReference());
 
 							uint32 ReferenceID = MutableImage->GetReferencedTexture();
-							if (CustomizableObject->ReferencedPassThroughTextures.IsValidIndex(ReferenceID))
+							if (ModelResources.PassThroughTextures.IsValidIndex(ReferenceID))
 							{
-								TSoftObjectPtr<UTexture> Ref = CustomizableObject->ReferencedPassThroughTextures[ReferenceID];
+								TSoftObjectPtr<UTexture> Ref = ModelResources.PassThroughTextures[ReferenceID];
 
 								// The texture should have been loaded by now by LoadAdditionalAssetsAsync()
 								PassThroughTexture = Ref.Get();
@@ -5661,9 +5661,9 @@ void UCustomizableInstancePrivate::BuildMaterials(const TSharedRef<FUpdateContex
 
 						// Find the additional information for this image
 						int32 ImageKey = FCString::Atoi(*KeyName);
-						if (ImageKey >= 0 && ImageKey < CustomizableObject->ImageProperties.Num())
+						if (ImageKey >= 0 && ImageKey < ModelResources.ImageProperties.Num())
 						{
-							const FMutableModelImageProperties& Props = CustomizableObject->ImageProperties[ImageKey];
+							const FMutableModelImageProperties& Props = ModelResources.ImageProperties[ImageKey];
 
 							if (!MutableTexture && !PassThroughTexture && MutableImage)
 							{
diff --git a/Engine/Plugins/Experimental/Mutable/Source/CustomizableObject/Private/MuCO/CustomizableObjectSystem.cpp b/Engine/Plugins/Experimental/Mutable/Source/CustomizableObject/Private/MuCO/CustomizableObjectSystem.cpp
--- a/Engine/Plugins/Experimental/Mutable/Source/CustomizableObject/Private/MuCO/CustomizableObjectSystem.cpp
+++ b/Engine/Plugins/Experimental/Mutable/Source/CustomizableObject/Private/MuCO/CustomizableObjectSystem.cpp
@@ -14,6 +14,7 @@
 #include "MuCO/CustomizableInstanceLODManagement.h"
 #include "MuCO/CustomizableInstancePrivateData.h"
 #include "MuCO/CustomizableObjectPrivate.h"
+#include "MuCO/CustomizableObjectUIData.h"
 #include "MuCO/DefaultImageProvider.h"
 #include "MuCO/CustomizableObjectInstanceUsage.h"
 #include "MuCO/ICustomizableObjectModule.h"
@@ -907,7 +908,7 @@ void UCustomizableObjectSystemPrivate::GetMipStreamingConfig(const UCustomizable
 
 	// From user-controlled per-state flag?
 	const FString CurrentState = Instance.GetCurrentState();
-	const FParameterUIData* State = Instance.GetCustomizableObject()->StateUIDataMap.Find(CurrentState);	
+	const FParameterUIData* State = Instance.GetCustomizableObject()->GetPrivate()->GetModelResources().StateUIDataMap.Find(CurrentState);	
 	if (State)
 	{
 		bOutNeverStream = State->bDisableTextureStreaming;
@@ -1647,6 +1648,8 @@ namespace impl
 		check(System != nullptr);
 
 		const UCustomizableObject* CustomizableObject = OperationData->Instance->GetCustomizableObject();
+		const FModelResources& ModelResources = CustomizableObject->GetPrivate()->GetModelResources();
+
 		UCustomizableInstancePrivate* CustomizableObjectInstancePrivateData = OperationData->Instance->GetPrivate();
 
 		CustomizableObjectInstancePrivateData->PassThroughTexturesToLoad.Empty();
@@ -1777,9 +1780,9 @@ namespace impl
 							FString KeyName = Image.Name.ToString();
 							int32 ImageKey = FCString::Atoi(*KeyName);
 
-							if (ImageKey >= 0 && ImageKey < CustomizableObject->ImageProperties.Num())
+							if (ImageKey >= 0 && ImageKey < ModelResources.ImageProperties.Num())
 							{
-								const FMutableModelImageProperties& Props = CustomizableObject->ImageProperties[ImageKey];
+								const FMutableModelImageProperties& Props = ModelResources.ImageProperties[ImageKey];
 
 								if (Props.IsPassThrough)
 								{
@@ -1791,9 +1794,9 @@ namespace impl
 
 									uint32 ReferenceID = Image.Image->GetReferencedTexture();
 
-									if (CustomizableObject->ReferencedPassThroughTextures.IsValidIndex(ReferenceID))
+									if (ModelResources.PassThroughTextures.IsValidIndex(ReferenceID))
 									{
-										TSoftObjectPtr<UTexture> Ref = CustomizableObject->ReferencedPassThroughTextures[ReferenceID];
+										TSoftObjectPtr<UTexture> Ref = ModelResources.PassThroughTextures[ReferenceID];
 										CustomizableObjectInstancePrivateData->PassThroughTexturesToLoad.Add(Ref);
 									}
 									else
@@ -2655,19 +2658,13 @@ namespace impl
 
 		bool bCancel = false;
 
-		// If the object is locked (for instance, compiling) we skip any instance update.
 		TObjectPtr<UCustomizableObject> CustomizableObject = CandidateInstance->GetCustomizableObject();
-		if (!CustomizableObject)
+
+		// If the object is locked (for instance, compiling) we skip any instance update.
+		if (!CustomizableObject || CustomizableObject->GetPrivate()->bLocked)
 		{
 			bCancel = true;
 		}
-		else
-		{
-			if (CustomizableObject->GetPrivate()->bLocked)
-			{
-				bCancel = true;
-			}
-		}
 
 		// Only update resources if the instance is in range (it could have got far from the player since the task was queued)
 		check(System->CurrentInstanceLODManagement != nullptr);
@@ -2705,7 +2702,7 @@ namespace impl
 		CandidateInstance->CommitMinMaxLOD();
 
 		FString StateName = CandidateInstance->GetCustomizableObject()->GetStateName(CandidateInstance->GetState());
-		const FParameterUIData* StateData = CandidateInstance->GetCustomizableObject()->StateUIDataMap.Find(StateName);
+		const FParameterUIData* StateData = CandidateInstance->GetCustomizableObject()->GetPrivate()->GetModelResources().StateUIDataMap.Find(StateName);
 
 		Operation->bLiveUpdateMode = false;
 
diff --git a/Engine/Plugins/Experimental/Mutable/Source/CustomizableObject/Private/MuCO/UnrealMutableImageProvider.cpp b/Engine/Plugins/Experimental/Mutable/Source/CustomizableObject/Private/MuCO/UnrealMutableImageProvider.cpp
--- a/Engine/Plugins/Experimental/Mutable/Source/CustomizableObject/Private/MuCO/UnrealMutableImageProvider.cpp
+++ b/Engine/Plugins/Experimental/Mutable/Source/CustomizableObject/Private/MuCO/UnrealMutableImageProvider.cpp
@@ -4,6 +4,7 @@
 
 #include "MuCO/CustomizableObjectSystem.h"
 #include "MuCO/CustomizableObject.h"
+#include "MuCO/CustomizableObjectPrivate.h"
 #include "MuCO/UnrealToMutableTextureConversionUtils.h"
 #include "TextureResource.h"
 #include "MuR/Parameters.h"
@@ -123,7 +124,8 @@ bool FUnrealMutableImageProvider::Tick(float DeltaTime)
 		return true;
 	}
 
-	if (!CO->ReferencedPassThroughTextures.IsValidIndex(Request->Id))
+	const FModelResources& ModelResources = CO->GetPrivate()->GetModelResources();
+	if (!ModelResources.PassThroughTextures.IsValidIndex(Request->Id))
 	{
 		// The id is not valid for this CO
 		check(false);
@@ -131,7 +133,7 @@ bool FUnrealMutableImageProvider::Tick(float DeltaTime)
 	}
 
 	// Find the texture id
-	TSoftObjectPtr<UTexture> TexturePtr = CO->ReferencedPassThroughTextures[Request->Id];
+	TSoftObjectPtr<UTexture> TexturePtr = ModelResources.PassThroughTextures[Request->Id];
 
 	// This can cause a stall because of loading the asset.
 	UTexture2D* Texture = Cast<UTexture2D>( TexturePtr.LoadSynchronous() );
diff --git a/Engine/Plugins/Experimental/Mutable/Source/CustomizableObjectEditor/Private/MuCOE/CustomizableObjectCompiler.cpp b/Engine/Plugins/Experimental/Mutable/Source/CustomizableObjectEditor/Private/MuCOE/CustomizableObjectCompiler.cpp
--- a/Engine/Plugins/Experimental/Mutable/Source/CustomizableObjectEditor/Private/MuCOE/CustomizableObjectCompiler.cpp
+++ b/Engine/Plugins/Experimental/Mutable/Source/CustomizableObjectEditor/Private/MuCOE/CustomizableObjectCompiler.cpp
@@ -9,6 +9,7 @@
 #include "Interfaces/ITargetPlatformManagerModule.h"
 #include "Materials/MaterialInterface.h"
 #include "MessageLogModule.h"
+#include "MuCO/CustomizableObjectPrivate.h"
 #include "MuCO/CustomizableObjectInstance.h"
 #include "MuCO/CustomizableObjectSystem.h"
 #include "MuCO/ICustomizableObjectModule.h"
@@ -929,13 +930,65 @@ void FCustomizableObjectCompiler::CompileInternal(UCustomizableObject* Object, c
 	}
 	else
 	{
+		FModelResources& ModelResources = Object->GetPrivate()->GetModelResources(false);
+		ModelResources = FModelResources();
+		
+		ModelResources.ReferenceSkeletalMeshesData = MoveTemp(GenerationContext.ReferenceSkeletalMeshesData);
+
+		ModelResources.Skeletons.Reserve(GenerationContext.ReferencedSkeletons.Num());
+		for (const USkeleton* Skeleton : GenerationContext.ReferencedSkeletons)
+		{
+			ModelResources.Skeletons.Emplace(Skeleton);
+		}
+		
+		ModelResources.Materials.Reserve(GenerationContext.ReferencedMaterials.Num());
+		for (const UMaterialInterface* Material : GenerationContext.ReferencedMaterials)
+		{
+			ModelResources.Materials.Emplace(Material);
+		}
+
+		for (const TPair<TSoftObjectPtr<UTexture>, FMutableGraphGenerationContext::FGeneratedReferencedTexture>& Pair : GenerationContext.RuntimeReferencedTextureMap)
+		{
+			check(Pair.Value.ID == ModelResources.PassThroughTextures.Num());
+			ModelResources.PassThroughTextures.Add(Pair.Key);
+		}
+
+		ModelResources.PhysicsAssets = MoveTemp(GenerationContext.PhysicsAssets);
+
+		ModelResources.AnimBPs = MoveTemp(GenerationContext.AnimBPAssets);
+		ModelResources.AnimBpOverridePhysiscAssetsInfo = MoveTemp(GenerationContext.AnimBpOverridePhysicsAssetsInfo);
+
+		ModelResources.MaterialSlotNames = MoveTemp(GenerationContext.ReferencedMaterialSlotNames);
+		ModelResources.BoneNames = MoveTemp(GenerationContext.BoneNames);
+		ModelResources.SocketArray = MoveTemp(GenerationContext.SocketArray);
+
+		ModelResources.SkinWeightProfilesInfo = MoveTemp(GenerationContext.SkinWeightProfilesInfo);
+
+		TArray<FGeneratedImageProperties> ImageProperties;
+		GenerationContext.ImageProperties.GenerateValueArray(ImageProperties);
+
+		ModelResources.ImageProperties.Empty(ImageProperties.Num());
+
+		for (const FGeneratedImageProperties& ImageProp : ImageProperties)
+		{
+			ModelResources.ImageProperties.Add({ ImageProp.TextureParameterName,
+										ImageProp.Filter,
+										ImageProp.SRGB,
+										ImageProp.bFlipGreenChannel,
+										ImageProp.bIsPassThrough,
+										ImageProp.LODBias,
+										ImageProp.LODGroup,
+										ImageProp.AddressX, ImageProp.AddressY });
+		}
+
+		ModelResources.ParameterUIDataMap = MoveTemp(GenerationContext.ParameterUIDataMap);
+		ModelResources.StateUIDataMap = MoveTemp(GenerationContext.StateUIDataMap);
+
 		// Morph target generated data does not need extra processing so move semantics can be used
 		// to avoid a possibly expensive copy.
 		Object->ContributingMorphTargetsInfo = MoveTemp(GenerationContext.ContributingMorphTargetsInfo);
 		Object->MorphTargetReconstructionData = MoveTemp(GenerationContext.MorphTargetReconstructionData);
 		
-		Object->SkinWeightProfilesInfo = MoveTemp(GenerationContext.SkinWeightProfilesInfo);
-		Object->AnimBpOverridePhysiscAssetsInfo = MoveTemp(GenerationContext.AnimBpOverridePhysicsAssetsInfo);
 		
 		// Clothing	
 		Object->ClothMeshToMeshVertData = MoveTemp(GenerationContext.ClothMeshToMeshVertData);
@@ -980,63 +1033,8 @@ void FCustomizableObjectCompiler::CompileInternal(UCustomizableObject* Object, c
 			 ClothingAssetData.ConfigsData.RemoveAllSwap(IsSharedConfigData);
 		}
 
-		// Mark the object as modified, used to avoid missing assets in packages. 
-		if (!ParamNamesToSelectedOptions.Num()) // Don't mark the objects as modified because of a partial compilation
-		{
-			if (Object->ReferencedMaterials.Num() != GenerationContext.ReferencedMaterials.Num())
-			{
-				Object->MarkPackageDirty();
-			}
-			else
-			{
-				for (int32 i = 0; i < Object->ReferencedMaterials.Num(); ++i)
-				{
-					if (Object->ReferencedMaterials[i] != GenerationContext.ReferencedMaterials[i])
-					{
-						Object->MarkPackageDirty();
-						break;
-					}
-				}
-			}
-		}
-
-		Object->ReferenceSkeletalMeshesData = GenerationContext.ReferenceSkeletalMeshesData;
-		
-		Object->ReferencedMaterials.Empty(GenerationContext.ReferencedMaterials.Num());
-
-		for (const UMaterialInterface* Material : GenerationContext.ReferencedMaterials)
-		{
-			Object->ReferencedMaterials.Add(Material);
-		}
-
-		Object->ReferencedMaterialSlotNames.Empty(GenerationContext.ReferencedMaterialSlotNames.Num());
-
-		for (const FName& MaterialSlotName : GenerationContext.ReferencedMaterialSlotNames)
-		{
-			Object->ReferencedMaterialSlotNames.Add(MaterialSlotName);
-		}
-
-		Object->ImageProperties.Empty(GenerationContext.ImageProperties.Num());
-
-		TArray<FGeneratedImageProperties> ImageProperties;
-		GenerationContext.ImageProperties.GenerateValueArray(ImageProperties);
-
-		for (const FGeneratedImageProperties& ImageProp : ImageProperties)
-		{
-			Object->ImageProperties.Add({ ImageProp.TextureParameterName,
-										ImageProp.Filter,
-										ImageProp.SRGB,
-										ImageProp.bFlipGreenChannel,
-										ImageProp.bIsPassThrough,
-										ImageProp.LODBias,
-										ImageProp.LODGroup,
-										ImageProp.AddressX, ImageProp.AddressY });
-		}
-
 		Object->GroupNodeMap = GenerationContext.GroupNodeMap;
-		Object->ParameterUIDataMap = GenerationContext.ParameterUIDataMap;
 
-		Object->StateUIDataMap = GenerationContext.StateUIDataMap;
 		if (GenerationContext.Options.OptimizationLevel == 0)
 		{
 			// If the optimization level is "none" disable texture streaming, because textures are all referenced
@@ -1067,70 +1065,8 @@ void FCustomizableObjectCompiler::CompileInternal(UCustomizableObject* Object, c
 		Object->LODSettings.bLODStreamingEnabled = GenerationContext.bEnableLODStreaming;
 		Object->LODSettings.NumLODsToStream = GenerationContext.NumMaxLODsToStream;
 
-		// Mark the object as modified, used to avoid missing assets in packages. 
-		if (!Object->PhysicsAssetsMap.OrderIndependentCompareEqual(GenerationContext.PhysicsAssetMap))
-		{
-			if (!ParamNamesToSelectedOptions.Num()) // Don't mark the objects as modified because of a partial compilation
-			{
-				Object->MarkPackageDirty();
-			}
-		}
-
-		Object->PhysicsAssetsMap = GenerationContext.PhysicsAssetMap;
-
-		// Mark the object as modified, used to avoid missing assets in packages. 
-		if (!Object->AnimBPAssetsMap.OrderIndependentCompareEqual(GenerationContext.AnimBPAssetsMap))
-		{
-			if (!ParamNamesToSelectedOptions.Num()) // Don't mark the objects as modified because of a partial compilation
-			{
-				Object->MarkPackageDirty();
-			}
-		}
-	
-		Object->AnimBPAssetsMap = GenerationContext.AnimBPAssetsMap;
-
 		Object->StreamedResourceData = MoveTemp(GenerationContext.StreamedResourceData);
 
-		// Mark the object as modified, used to avoid missing assets in packages.
-		if (Object->SocketArray != GenerationContext.SocketArray)
-		{
-			if (!ParamNamesToSelectedOptions.Num()) // Don't mark the objects as modified because of a partial compilation
-			{
-				Object->MarkPackageDirty();
-			}
-		}
-
-		Object->SocketArray = GenerationContext.SocketArray;
-
-		// 
-		if (!ParamNamesToSelectedOptions.Num()) // Don't mark the objects as modified because of a partial compilation
-		{
-			if (Object->ReferencedSkeletons.Num() != GenerationContext.ReferencedSkeletons.Num())
-			{
-				Object->MarkPackageDirty();
-			}
-			else
-			{
-				for (int32 SkeletonIndex = 0; SkeletonIndex < Object->ReferencedSkeletons.Num(); ++SkeletonIndex)
-				{
-					if (Object->ReferencedSkeletons[SkeletonIndex] != GenerationContext.ReferencedSkeletons[SkeletonIndex])
-					{
-						Object->MarkPackageDirty();
-						break;
-					}
-				}
-			}
-		}
-
-		Object->ReferencedSkeletons.Empty(GenerationContext.ReferencedSkeletons.Num());
-
-		for (const USkeleton* Skeleton : GenerationContext.ReferencedSkeletons)
-		{
-			Object->ReferencedSkeletons.Add(Skeleton);
-		}
-
-		Object->SetBoneNamesArray(GenerationContext.BoneNames);
-
 		// Pass-through textures
 		TArray<TSoftObjectPtr<UTexture>> NewCompileTimeReferencedTextures;
 		for (const TPair<TSoftObjectPtr<UTexture>, FMutableGraphGenerationContext::FGeneratedReferencedTexture>& Pair : GenerationContext.CompileTimeTextureMap)
@@ -1139,24 +1075,6 @@ void FCustomizableObjectCompiler::CompileInternal(UCustomizableObject* Object, c
 			NewCompileTimeReferencedTextures.Add(Pair.Key);
 		}
 
-		TArray<TSoftObjectPtr<UTexture>> NewRuntimeReferencedTextures;
-		for (const TPair<TSoftObjectPtr<UTexture>, FMutableGraphGenerationContext::FGeneratedReferencedTexture>& Pair : GenerationContext.RuntimeReferencedTextureMap)
-		{
-			check(Pair.Value.ID == NewRuntimeReferencedTextures.Num());
-			NewRuntimeReferencedTextures.Add(Pair.Key);
-		}
-
-		// Mark the object as modified, used to avoid missing assets in packages.
-		if (Object->ReferencedPassThroughTextures != NewRuntimeReferencedTextures)
-		{
-			Object->ReferencedPassThroughTextures = NewRuntimeReferencedTextures;
-
-			if (!ParamNamesToSelectedOptions.Num()) // Don't mark the objects as modified because of a partial compilation
-			{
-				Object->MarkPackageDirty();
-			}
-		}
-
 		if (!ParamNamesToSelectedOptions.Num())
 		{
 			const FName PackageName = Object->GetPackage()->GetFName();
diff --git a/Engine/Plugins/Experimental/Mutable/Source/CustomizableObjectEditor/Private/MuCOE/GenerateMutableSource/GenerateMutableSource.cpp b/Engine/Plugins/Experimental/Mutable/Source/CustomizableObjectEditor/Private/MuCOE/GenerateMutableSource/GenerateMutableSource.cpp
--- a/Engine/Plugins/Experimental/Mutable/Source/CustomizableObjectEditor/Private/MuCOE/GenerateMutableSource/GenerateMutableSource.cpp
+++ b/Engine/Plugins/Experimental/Mutable/Source/CustomizableObjectEditor/Private/MuCOE/GenerateMutableSource/GenerateMutableSource.cpp
@@ -6,7 +6,9 @@
 #include "Engine/SkeletalMeshSocket.h"
 #include "Engine/TextureLODSettings.h"
 #include "Interfaces/ITargetPlatform.h"
+#include "MuCO/CustomizableObjectPrivate.h"
 #include "MuCO/CustomizableObjectExtension.h"
+#include "MuCO/CustomizableObjectUIData.h"
 #include "MuCO/ICustomizableObjectModule.h"
 #include "MuCOE/CustomizableObjectCompiler.h"
 #include "MuCOE/EdGraphSchema_CustomizableObject.h"
@@ -1476,9 +1478,9 @@ int32 AddTagToMutableMeshUnique(mu::Mesh& MutableMesh, const FString& Tag)
 	return TagCount;
 }
 
-FString GenerateAnimationInstanceTag(const FString& AnimInstance, const FName& SlotIndex)
+FString GenerateAnimationInstanceTag(const int32 AnimBpIndex, const FName& SlotIndex)
 {
-	return FString("__AnimBP:") + FString::Printf(TEXT("%s_Slot_"), *SlotIndex.ToString()) + AnimInstance;
+	return FString("__AnimBP:") + FString::Printf(TEXT("%s_Slot_"), *SlotIndex.ToString()) + FString::FromInt(AnimBpIndex);
 }
 
 
@@ -1505,7 +1507,7 @@ void PopulateReferenceSkeletalMeshesData(FMutableGraphGenerationContext& Generat
 
 		// Set the RefSkeletalMesh
 		Data.SkeletalMesh = TObjectPtr<USkeletalMesh>(RefSkeletalMesh);
-		Data.SkeletalMeshAssetPath = FSoftObjectPath(RefSkeletalMesh);
+		Data.SoftSkeletalMesh = RefSkeletalMesh;
 
 		// Gather LODData, this may include per LOD settings such as render data config or LODDataInfoArray
 		Data.LODData.AddDefaulted(GenerationContext.NumLODsInRoot);
@@ -1553,6 +1555,7 @@ void PopulateReferenceSkeletalMeshesData(FMutableGraphGenerationContext& Generat
 			Socket.RelativeLocation = RefSocket->RelativeLocation;
 			Socket.RelativeRotation = RefSocket->RelativeRotation;
 			Socket.RelativeScale = RefSocket->RelativeScale;
+			Socket.bForceAlwaysAnimated = RefSocket->bForceAlwaysAnimated;
 		}
 
 		// TODO: Generate Bounds?
@@ -1575,37 +1578,30 @@ void PopulateReferenceSkeletalMeshesData(FMutableGraphGenerationContext& Generat
 		}
 
 		// Skeleton
-		if(const USkeleton* Skeleton = RefSkeletalMesh->GetSkeleton())
+		if(USkeleton* Skeleton = RefSkeletalMesh->GetSkeleton())
 		{
-			Data.Skeleton = RefSkeletalMesh->GetSkeleton();
-			GenerationContext.ReferencedSkeletons.AddUnique(Skeleton);
+			Data.Skeleton = Skeleton;
 		}
 		
 		// Physics Asset
-		if (const UPhysicsAsset* PhysicsAsset = RefSkeletalMesh->GetPhysicsAsset())
+		if (UPhysicsAsset* PhysicsAsset = RefSkeletalMesh->GetPhysicsAsset())
 		{
 			GenerationContext.AddParticipatingObject(*PhysicsAsset);
-
 			Data.PhysicsAsset = PhysicsAsset;
-			GenerationContext.PhysicsAssetMap.FindOrAdd(Data.PhysicsAsset.ToString(), Data.PhysicsAsset);
 		}
 
 		// Post ProcessAnimInstance
 		if(const TSubclassOf<UAnimInstance> PostProcessAnimInstance = RefSkeletalMesh->GetPostProcessAnimBlueprint())
 		{
 			GenerationContext.AddParticipatingObject(*PostProcessAnimInstance.Get());
-
 			Data.PostProcessAnimInst = PostProcessAnimInstance;
-			GenerationContext.AnimBPAssetsMap.FindOrAdd(Data.PostProcessAnimInst.ToString(), Data.PostProcessAnimInst);
 		}
 		
 		// Shadow Physics Asset
-		if (const UPhysicsAsset* PhysicsAsset = RefSkeletalMesh->GetShadowPhysicsAsset())
+		if (UPhysicsAsset* PhysicsAsset = RefSkeletalMesh->GetShadowPhysicsAsset())
 		{
 			GenerationContext.AddParticipatingObject(*PhysicsAsset);
-
 			Data.ShadowPhysicsAsset = PhysicsAsset;
-			GenerationContext.PhysicsAssetMap.FindOrAdd(Data.PhysicsAsset.ToString(), Data.PhysicsAsset);
 		}
 
 		// Asset User Data
diff --git a/Engine/Plugins/Experimental/Mutable/Source/CustomizableObjectEditor/Private/MuCOE/GenerateMutableSource/GenerateMutableSource.h b/Engine/Plugins/Experimental/Mutable/Source/CustomizableObjectEditor/Private/MuCOE/GenerateMutableSource/GenerateMutableSource.h
--- a/Engine/Plugins/Experimental/Mutable/Source/CustomizableObjectEditor/Private/MuCOE/GenerateMutableSource/GenerateMutableSource.h
+++ b/Engine/Plugins/Experimental/Mutable/Source/CustomizableObjectEditor/Private/MuCOE/GenerateMutableSource/GenerateMutableSource.h
@@ -41,9 +41,12 @@ class UObject;
 class UPhysicsAsset;
 class UTexture2D;
 struct FCustomizableObjectIdPair;
+struct FAnimBpOverridePhysicsAssetsInfo;
 struct FMutableGraphGenerationContext;
 struct FParameterUIData;
-
+struct FMutableRefSkeletalMeshData;
+struct FMutableRefSocket;
+struct FMutableSkinWeightProfileInfo;
 
 struct FGeneratedImageProperties
 {
@@ -751,10 +754,10 @@ struct FMutableGraphGenerationContext
 	const UCustomizableObject* CustomizableObjectWithCycle = nullptr;
 
 	/** Stores the physics assets gathered from the SkeletalMesh nodes during compilation, to be used in mesh generation in-game */
-	TMap<FString, TSoftObjectPtr<UPhysicsAsset>> PhysicsAssetMap;
+	TArray<TSoftObjectPtr<UPhysicsAsset>> PhysicsAssets;
 
 	/** Stores the anim BP assets gathered from the SkeletalMesh nodes during compilation, to be used in mesh generation in-game */
-	TMap<FString, TSoftClassPtr<UAnimInstance>> AnimBPAssetsMap;
+	TArray<TSoftClassPtr<UAnimInstance>> AnimBPAssets;
 
 	/** Stores the sockets provided by the part skeletal meshes, to be merged in the generated meshes */
 	TArray<FMutableRefSocket> SocketArray;
@@ -859,7 +862,7 @@ int32 AddTagToMutableMeshUnique(mu::Mesh& MutableMesh, const FString& Tag);
 void AddSocketTagsToMesh(const USkeletalMesh* SourceMesh, mu::Ptr<mu::Mesh> MutableMesh, FMutableGraphGenerationContext& GenerationContext);
 
 // Generates the tag for an animation instance
-FString GenerateAnimationInstanceTag(const FString& AnimInstance, const FName& SlotIndex);
+FString GenerateAnimationInstanceTag(const int32 AnimInstanceIndex, const FName& SlotIndex);
 
 
 FString GenerateGameplayTag(const FString& GameplayTag);
diff --git a/Engine/Plugins/Experimental/Mutable/Source/CustomizableObjectEditor/Private/MuCOE/GenerateMutableSource/GenerateMutableSourceMesh.cpp b/Engine/Plugins/Experimental/Mutable/Source/CustomizableObjectEditor/Private/MuCOE/GenerateMutableSource/GenerateMutableSourceMesh.cpp
--- a/Engine/Plugins/Experimental/Mutable/Source/CustomizableObjectEditor/Private/MuCOE/GenerateMutableSource/GenerateMutableSourceMesh.cpp
+++ b/Engine/Plugins/Experimental/Mutable/Source/CustomizableObjectEditor/Private/MuCOE/GenerateMutableSource/GenerateMutableSourceMesh.cpp
@@ -11,6 +11,7 @@
 #include "ClothingAsset.h"
 #include "MeshUtilities.h"
 #include "Modules/ModuleManager.h"
+#include "MuCO/CustomizableObjectPrivate.h"
 #include "MuCO/CustomizableObjectInstance.h"
 #include "MuCO/CustomizableObjectSystem.h"
 #include "MuCO/UnrealConversionUtils.h"
@@ -2866,9 +2867,9 @@ mu::NodeMeshPtr GenerateMutableSourceMesh(const UEdGraphPin* Pin,
 				}
 
 				FName SlotIndex = TypedNodeSkel->AnimBlueprintSlotName;
-				GenerationContext.AnimBPAssetsMap.Add(TypedNodeSkel->AnimInstance.ToString(), TypedNodeSkel->AnimInstance);
+				const int32 AnimInstanceIndex = GenerationContext.AnimBPAssets.AddUnique(TypedNodeSkel->AnimInstance);
 
-				AnimBPAssetTag = GenerateAnimationInstanceTag(TypedNodeSkel->AnimInstance.ToString(), SlotIndex);
+				AnimBPAssetTag = GenerateAnimationInstanceTag(AnimInstanceIndex, SlotIndex);
 				MeshUniqueTags += AnimBPAssetTag;
 			}
 
@@ -2919,8 +2920,8 @@ mu::NodeMeshPtr GenerateMutableSourceMesh(const UEdGraphPin* Pin,
 
 					GenerationContext.AddParticipatingObject(*PhysicsAsset);
 
-					GenerationContext.PhysicsAssetMap.Add(PhysicsAsset.ToString(), PhysicsAsset);
-					FString PhysicsAssetTag = FString("__PhysicsAsset:") + PhysicsAsset.ToString();
+					const int32 AssetIndex = GenerationContext.PhysicsAssets.AddUnique(PhysicsAsset);
+					FString PhysicsAssetTag = FString("__PA:") + FString::FromInt(AssetIndex);
 
 					AddTagToMutableMeshUnique(*MutableMesh, PhysicsAssetTag);
 				}
@@ -2946,7 +2947,7 @@ mu::NodeMeshPtr GenerateMutableSourceMesh(const UEdGraphPin* Pin,
 
 						GenerationContext.AddParticipatingObject(*ClothingAssetCommon->PhysicsAsset);
 						
-						GenerationContext.PhysicsAssetMap.Add(PhysicsAsset.ToString(), ClothingAssetCommon->PhysicsAsset.Get());
+						GenerationContext.PhysicsAssets.AddUnique(ClothingAssetCommon->PhysicsAsset);
 
 						AddTagToMutableMeshUnique(*MutableMesh, ClothPhysicsAssetTag);
 					}
diff --git a/Engine/Plugins/Experimental/Mutable/Source/CustomizableObjectEditor/Private/MuCOE/GenerateMutableSource/GenerateMutableSourceTable.cpp b/Engine/Plugins/Experimental/Mutable/Source/CustomizableObjectEditor/Private/MuCOE/GenerateMutableSource/GenerateMutableSourceTable.cpp
--- a/Engine/Plugins/Experimental/Mutable/Source/CustomizableObjectEditor/Private/MuCOE/GenerateMutableSource/GenerateMutableSourceTable.cpp
+++ b/Engine/Plugins/Experimental/Mutable/Source/CustomizableObjectEditor/Private/MuCOE/GenerateMutableSource/GenerateMutableSourceTable.cpp
@@ -13,6 +13,7 @@
 #include "Materials/Material.h"
 #include "Materials/MaterialInstance.h"
 #include "MuCO/CustomizableObjectSystem.h"
+#include "MuCO/CustomizableObjectUIData.h"
 #include "MuCO/UnrealPortabilityHelpers.h"
 #include "MuCOE/CustomizableObjectCompiler.h"
 #include "MuCOE/GenerateMutableSource/GenerateMutableSourceMesh.h"
@@ -120,9 +121,9 @@ bool FillTableColumn(const UCustomizableObjectNodeTable* TableNode,	mu::TablePtr
 													GenerationContext.AddParticipatingObject(*Anim);													
 												}
 
-												GenerationContext.AnimBPAssetsMap.Add(AnimInstance.ToString(), AnimInstance);
+												const int32 AnimInstanceIndex = GenerationContext.AnimBPAssets.AddUnique(AnimInstance);
 
-												AnimBPAssetTag = GenerateAnimationInstanceTag(AnimInstance.ToString(), SlotIndex);
+												AnimBPAssetTag = GenerateAnimationInstanceTag(AnimInstanceIndex, SlotIndex);
 											}
 										}
 									}
@@ -280,8 +281,8 @@ bool FillTableColumn(const UCustomizableObjectNodeTable* TableNode,	mu::TablePtr
 
 					GenerationContext.AddParticipatingObject(*PhysicsAsset); 						
 
-					GenerationContext.PhysicsAssetMap.Add(PhysicsAsset.ToString(), PhysicsAsset);
-					FString PhysicsAssetTag = FString("__PhysicsAsset:") + PhysicsAsset.ToString();
+					const int32 AssetIndex = GenerationContext.PhysicsAssets.AddUnique(PhysicsAsset);
+					FString PhysicsAssetTag = FString("__PA:") + FString::FromInt(AssetIndex);
 
 					AddTagToMutableMeshUnique(*MutableMesh, PhysicsAssetTag);
 				}
diff --git a/Engine/Plugins/Experimental/Mutable/Source/CustomizableObjectEditor/Private/MuCOE/SMutableImageViewer.cpp b/Engine/Plugins/Experimental/Mutable/Source/CustomizableObjectEditor/Private/MuCOE/SMutableImageViewer.cpp
--- a/Engine/Plugins/Experimental/Mutable/Source/CustomizableObjectEditor/Private/MuCOE/SMutableImageViewer.cpp
+++ b/Engine/Plugins/Experimental/Mutable/Source/CustomizableObjectEditor/Private/MuCOE/SMutableImageViewer.cpp
@@ -6,6 +6,7 @@
 #include "MuT/TypeInfo.h"
 #include "TextureResource.h"
 #include "MuCO/CustomizableObject.h"
+#include "MuCO/CustomizableObjectPrivate.h"
 #include "MuCO/CustomizableObjectInstance.h"
 #include "Widgets/Input/SNumericEntryBox.h"
 #include "Widgets/Input/SSegmentedControl.h"
