[Mutable] Add missing include in CustomizableObjectTableNode.

#rnx

diff --git a/Engine/Plugins/Experimental/Mutable/Source/CustomizableObjectEditor/Classes/MuCOE/Nodes/CustomizableObjectNodeTable.h b/Engine/Plugins/Experimental/Mutable/Source/CustomizableObjectEditor/Classes/MuCOE/Nodes/CustomizableObjectNodeTable.h
--- a/Engine/Plugins/Experimental/Mutable/Source/CustomizableObjectEditor/Classes/MuCOE/Nodes/CustomizableObjectNodeTable.h
+++ b/Engine/Plugins/Experimental/Mutable/Source/CustomizableObjectEditor/Classes/MuCOE/Nodes/CustomizableObjectNodeTable.h
@@ -19,6 +19,7 @@ class UObject;
 class USkeletalMesh;
 class UTexture2D;
 class UTexture2DArray;
+class UAnimInstance;
 struct FGuid;
 
 
diff --git a/Engine/Plugins/Experimental/Mutable/Source/CustomizableObjectEditor/Private/MuCOE/Nodes/CustomizableObjectNodeTable.cpp b/Engine/Plugins/Experimental/Mutable/Source/CustomizableObjectEditor/Private/MuCOE/Nodes/CustomizableObjectNodeTable.cpp
--- a/Engine/Plugins/Experimental/Mutable/Source/CustomizableObjectEditor/Private/MuCOE/Nodes/CustomizableObjectNodeTable.cpp
+++ b/Engine/Plugins/Experimental/Mutable/Source/CustomizableObjectEditor/Private/MuCOE/Nodes/CustomizableObjectNodeTable.cpp
@@ -20,6 +20,7 @@
 #include "MuCOE/UnrealEditorPortabilityHelpers.h"
 #include "Rendering/SkeletalMeshLODModel.h"
 #include "Rendering/SkeletalMeshModel.h"
+#include "Animation/AnimInstance.h"
 
 class ICustomizableObjectEditor;
 class UCustomizableObjectNodeRemapPins;
