[mutable] Add support for alpha channel to Color Parameters.
#rnx
#tests PIE

diff --git a/Engine/Plugins/Experimental/Mutable/Source/CustomizableObject/Classes/MuCO/CustomizableObject.h b/Engine/Plugins/Experimental/Mutable/Source/CustomizableObject/Classes/MuCO/CustomizableObject.h
--- a/Engine/Plugins/Experimental/Mutable/Source/CustomizableObject/Classes/MuCO/CustomizableObject.h
+++ b/Engine/Plugins/Experimental/Mutable/Source/CustomizableObject/Classes/MuCO/CustomizableObject.h
@@ -888,7 +888,7 @@ private:
 	// This is a manual version number for the binary blobs in this asset.
 	// Increasing it invalidates all the previously compiled models.
 	// Warning: If while merging code both versions have changed, take the highest+1.
-	static const int32 CurrentSupportedVersion = 420;
+	static const int32 CurrentSupportedVersion = 421;
 
 public:
 
diff --git a/Engine/Plugins/Experimental/Mutable/Source/CustomizableObject/Private/MuCO/CustomizableObject.cpp b/Engine/Plugins/Experimental/Mutable/Source/CustomizableObject/Private/MuCO/CustomizableObject.cpp
--- a/Engine/Plugins/Experimental/Mutable/Source/CustomizableObject/Private/MuCO/CustomizableObject.cpp
+++ b/Engine/Plugins/Experimental/Mutable/Source/CustomizableObject/Private/MuCO/CustomizableObject.cpp
@@ -1612,7 +1612,7 @@ FLinearColor UCustomizableObject::GetColorParameterDefaultValue(const FString& I
 	}
 	
 	FLinearColor Value;
-	Model->GetColourDefaultValue(ParameterIndex, &Value.R, &Value.G, &Value.B);
+	Model->GetColourDefaultValue(ParameterIndex, &Value.R, &Value.G, &Value.B, &Value.A);
 
 	return Value;
 }
diff --git a/Engine/Plugins/Experimental/Mutable/Source/CustomizableObject/Private/MuCO/CustomizableObjectInstanceDescriptor.cpp b/Engine/Plugins/Experimental/Mutable/Source/CustomizableObject/Private/MuCO/CustomizableObjectInstanceDescriptor.cpp
--- a/Engine/Plugins/Experimental/Mutable/Source/CustomizableObject/Private/MuCO/CustomizableObjectInstanceDescriptor.cpp
+++ b/Engine/Plugins/Experimental/Mutable/Source/CustomizableObject/Private/MuCO/CustomizableObjectInstanceDescriptor.cpp
@@ -547,7 +547,8 @@ mu::ParametersPtr FCustomizableObjectInstanceDescriptor::GetParameters() const
 			{
 				if (VectorParameter.ParameterName == Name || (Uid.IsValid() && VectorParameter.Id == Uid))
 				{
-					MutableParameters->SetColourValue(ParamIndex, VectorParameter.ParameterValue.R, VectorParameter.ParameterValue.G, VectorParameter.ParameterValue.B);
+					MutableParameters->SetColourValue(ParamIndex, VectorParameter.ParameterValue.R, VectorParameter.ParameterValue.G, 
+																  VectorParameter.ParameterValue.B, VectorParameter.ParameterValue.A);
 
 					break;
 				}
@@ -980,8 +981,8 @@ void FCustomizableObjectInstanceDescriptor::ReloadParameters()
 			}
 			else // Not found in Instance Parameters. Use Mutable Parameters.
 			{
-				MutableParameters->GetColourValue(ParamIndex, &Param.ParameterValue.R, &Param.ParameterValue.G, &Param.ParameterValue.B);
-				Param.ParameterValue.A = 1.0f;
+				MutableParameters->GetColourValue(ParamIndex, &Param.ParameterValue.R, &Param.ParameterValue.G,
+															  &Param.ParameterValue.B, &Param.ParameterValue.A);
 			}
 
 			VectorParameters.Add(Param);
diff --git a/Engine/Plugins/Experimental/Mutable/Source/CustomizableObjectEditor/Private/MuCOE/CustomizableInstanceDetails.cpp b/Engine/Plugins/Experimental/Mutable/Source/CustomizableObjectEditor/Private/MuCOE/CustomizableInstanceDetails.cpp
--- a/Engine/Plugins/Experimental/Mutable/Source/CustomizableObjectEditor/Private/MuCOE/CustomizableInstanceDetails.cpp
+++ b/Engine/Plugins/Experimental/Mutable/Source/CustomizableObjectEditor/Private/MuCOE/CustomizableInstanceDetails.cpp
@@ -1134,16 +1134,13 @@ FReply FCustomizableInstanceDetails::OnColorBlockMouseButtonDown(const FGeometry
 		return FReply::Unhandled();
 	}
 
-	FLinearColor col = GetColorParameterValue(ParamName);
-
-	TArray<FLinearColor*> LinearColorArray;
-	LinearColorArray.Add(&col);
+	FLinearColor Col = GetColorParameterValue(ParamName);
 
 	FColorPickerArgs args;
 	args.bIsModal = true;
-	args.bUseAlpha = false;
+	args.bUseAlpha = true;
 	args.bOnlyRefreshOnMouseUp = false;
-	args.InitialColor = col;
+	args.InitialColor = Col;
 	args.OnColorCommitted = FOnLinearColorValueChanged::CreateSP(this, &FCustomizableInstanceDetails::OnSetColorFromColorPicker, ParamName);
 	OpenColorPicker(args);
 
diff --git a/Engine/Plugins/Experimental/Mutable/Source/CustomizableObjectEditor/Private/MuCOE/GenerateMutableSource/GenerateMutableSourceColor.cpp b/Engine/Plugins/Experimental/Mutable/Source/CustomizableObjectEditor/Private/MuCOE/GenerateMutableSource/GenerateMutableSourceColor.cpp
--- a/Engine/Plugins/Experimental/Mutable/Source/CustomizableObjectEditor/Private/MuCOE/GenerateMutableSource/GenerateMutableSourceColor.cpp
+++ b/Engine/Plugins/Experimental/Mutable/Source/CustomizableObjectEditor/Private/MuCOE/GenerateMutableSource/GenerateMutableSourceColor.cpp
@@ -71,7 +71,7 @@ mu::NodeColourPtr GenerateMutableSourceColor(const UEdGraphPin* Pin, FMutableGra
 
 		ColorNode->SetName(TypedNodeColorParam->ParameterName);
 		ColorNode->SetUid(GenerationContext.GetNodeIdUnique(Node).ToString());
-		ColorNode->SetDefaultValue(TypedNodeColorParam->DefaultValue.R, TypedNodeColorParam->DefaultValue.G, TypedNodeColorParam->DefaultValue.B);
+		ColorNode->SetDefaultValue(TypedNodeColorParam->DefaultValue.R, TypedNodeColorParam->DefaultValue.G, TypedNodeColorParam->DefaultValue.B, TypedNodeColorParam->DefaultValue.A);
 
 		GenerationContext.ParameterUIDataMap.Add(TypedNodeColorParam->ParameterName, FParameterUIData(
 			TypedNodeColorParam->ParameterName,
diff --git a/Engine/Plugins/Experimental/Mutable/Source/CustomizableObjectEditor/Private/MuCOE/SCustomizableInstanceProperties.cpp b/Engine/Plugins/Experimental/Mutable/Source/CustomizableObjectEditor/Private/MuCOE/SCustomizableInstanceProperties.cpp
--- a/Engine/Plugins/Experimental/Mutable/Source/CustomizableObjectEditor/Private/MuCOE/SCustomizableInstanceProperties.cpp
+++ b/Engine/Plugins/Experimental/Mutable/Source/CustomizableObjectEditor/Private/MuCOE/SCustomizableInstanceProperties.cpp
@@ -1649,16 +1649,13 @@ FReply SCustomizableInstanceProperties::OnColorBlockMouseButtonDown(const FGeome
 		return FReply::Unhandled();
 	}
 
-	FLinearColor col = GetColorParameterValue(ParamName);
-
-	TArray<FLinearColor*> LinearColorArray;
-	LinearColorArray.Add(&col);
+	FLinearColor Col = GetColorParameterValue(ParamName);
 
 	FColorPickerArgs args;
 	args.bIsModal = true;
-	args.bUseAlpha = false;
+	args.bUseAlpha = true;
 	args.bOnlyRefreshOnMouseUp = false;
-	args.InitialColor = col;
+	args.InitialColor = Col;
 	args.OnColorCommitted = FOnLinearColorValueChanged::CreateSP(this, &SCustomizableInstanceProperties::OnSetColorFromColorPicker, ParamName);
 	OpenColorPicker(args);
 
diff --git a/Engine/Plugins/Experimental/Mutable/Source/CustomizableObjectEditor/Private/MuCOE/SMutableParametersWidget.cpp b/Engine/Plugins/Experimental/Mutable/Source/CustomizableObjectEditor/Private/MuCOE/SMutableParametersWidget.cpp
--- a/Engine/Plugins/Experimental/Mutable/Source/CustomizableObjectEditor/Private/MuCOE/SMutableParametersWidget.cpp
+++ b/Engine/Plugins/Experimental/Mutable/Source/CustomizableObjectEditor/Private/MuCOE/SMutableParametersWidget.cpp
@@ -574,8 +574,7 @@ FLinearColor SMutableParametersWidget::GetColorParameterValue(int32 ParamIndex,
 	}
 
 	FLinearColor Result;
-	MutableParameters->GetColourValue(ParamIndex, &Result.R, &Result.G, &Result.B, RangeIndex);
-	Result.A = 1.0f;
+	MutableParameters->GetColourValue(ParamIndex, &Result.R, &Result.G, &Result.B, &Result.A, RangeIndex);
 
 	return Result;
 }
@@ -588,16 +587,13 @@ FReply SMutableParametersWidget::OnColorBlockMouseButtonDown(const FGeometry& My
 		return FReply::Unhandled();
 	}
 
-	FLinearColor col = GetColorParameterValue(ParamIndex,RangeIndex);
-
-	TArray<FLinearColor*> LinearColorArray;
-	LinearColorArray.Add(&col);
+	FLinearColor Col = GetColorParameterValue(ParamIndex,RangeIndex);
 
 	FColorPickerArgs args;
 	args.bIsModal = true;
-	args.bUseAlpha = false;
+	args.bUseAlpha = true;
 	args.bOnlyRefreshOnMouseUp = false;
-	args.InitialColor = col;
+	args.InitialColor = Col;
 	args.OnColorCommitted = FOnLinearColorValueChanged::CreateSP(this, &SMutableParametersWidget::OnSetColorFromColorPicker, ParamIndex, RangeIndex);
 	OpenColorPicker(args);
 
@@ -614,7 +610,7 @@ void SMutableParametersWidget::OnSetColorFromColorPicker(FLinearColor NewColor,
 		return;
 	}
 	
-	MutableParameters->SetColourValue(ParamIndex, NewColor.R, NewColor.G, NewColor.B, RangeIndex);
+	MutableParameters->SetColourValue(ParamIndex, NewColor.R, NewColor.G, NewColor.B, NewColor.A, RangeIndex);
 
 	OnParametersValueChanged.ExecuteIfBound(ParamIndex);
 }
diff --git a/Engine/Plugins/Experimental/Mutable/Source/MutableRuntime/Internal/MuR/ParametersPrivate.h b/Engine/Plugins/Experimental/Mutable/Source/MutableRuntime/Internal/MuR/ParametersPrivate.h
--- a/Engine/Plugins/Experimental/Mutable/Source/MutableRuntime/Internal/MuR/ParametersPrivate.h
+++ b/Engine/Plugins/Experimental/Mutable/Source/MutableRuntime/Internal/MuR/ParametersPrivate.h
@@ -373,7 +373,7 @@ namespace Private
 	using ParamBoolType = bool;
 	using ParamIntType = int32;
 	using ParamFloatType = float;
-	using ParamColorType = FVector3f;
+	using ParamColorType = FVector4f;
 	using ParamProjectorType = Private::TIndirectObject<FProjector>;
 	using ParamImageType = FName;
 	using ParamStringType = Private::TIndirectObject<FString>;
@@ -382,7 +382,7 @@ namespace Private
             ParamBoolType, ParamIntType, ParamFloatType, ParamColorType, ParamProjectorType, ParamImageType, ParamStringType>;
 
     // static_assert to track PARAMETER_VALUE size changes. It is ok to change if needed.
-    static_assert(sizeof(PARAMETER_VALUE) == 8*3, "PARAMETER_VALUE size has changed.");
+    static_assert(sizeof(PARAMETER_VALUE) == 8*4, "PARAMETER_VALUE size has changed.");
 
 	// TVariant currently does not support this operator. Once supported remove it.
 	inline bool operator==(const PARAMETER_VALUE& ValueA, const PARAMETER_VALUE& ValueB)
@@ -462,7 +462,7 @@ namespace Private
         //!
         void Serialise( OutputArchive& arch ) const
         {
-            const int32 ver = 9;
+            const int32 ver = 10;
             arch << ver;
 
 			arch << m_name;
@@ -478,12 +478,12 @@ namespace Private
         {
             int32 ver;
             arch >> ver;
-            check( ver==9 );
+            check( ver == 10 );
 
 			arch >> m_name;
 			arch >> m_uid;
             arch >> m_type;
-            arch >> m_defaultValue;
+			arch >> m_defaultValue;
             arch >> m_ranges;
 			arch >> m_possibleValues;
         }
diff --git a/Engine/Plugins/Experimental/Mutable/Source/MutableRuntime/Internal/MuR/SerialisationPrivate.h b/Engine/Plugins/Experimental/Mutable/Source/MutableRuntime/Internal/MuR/SerialisationPrivate.h
--- a/Engine/Plugins/Experimental/Mutable/Source/MutableRuntime/Internal/MuR/SerialisationPrivate.h
+++ b/Engine/Plugins/Experimental/Mutable/Source/MutableRuntime/Internal/MuR/SerialisationPrivate.h
@@ -336,6 +336,7 @@ namespace mu
 
 	MUTABLE_DEFINE_POD_SERIALISABLE(vec2f);
 	MUTABLE_DEFINE_POD_SERIALISABLE(vec3f);
+	MUTABLE_DEFINE_POD_SERIALISABLE(vec4f);
 	MUTABLE_DEFINE_POD_SERIALISABLE(mat3f);
 	MUTABLE_DEFINE_POD_SERIALISABLE(mat4f);
 	MUTABLE_DEFINE_POD_SERIALISABLE(vec2<int>);
@@ -360,6 +361,7 @@ namespace mu
 
 	MUTABLE_IMPLEMENT_POD_VECTOR_SERIALISABLE(vec2f);
 	MUTABLE_IMPLEMENT_POD_VECTOR_SERIALISABLE(vec3f);
+	MUTABLE_IMPLEMENT_POD_VECTOR_SERIALISABLE(vec4f);
 	MUTABLE_IMPLEMENT_POD_VECTOR_SERIALISABLE(mat3f);
 	MUTABLE_IMPLEMENT_POD_VECTOR_SERIALISABLE(mat4f);
 	MUTABLE_IMPLEMENT_POD_VECTOR_SERIALISABLE(vec2<int>);
diff --git a/Engine/Plugins/Experimental/Mutable/Source/MutableRuntime/Private/MuR/CodeRunner.cpp b/Engine/Plugins/Experimental/Mutable/Source/MutableRuntime/Private/MuR/CodeRunner.cpp
--- a/Engine/Plugins/Experimental/Mutable/Source/MutableRuntime/Private/MuR/CodeRunner.cpp
+++ b/Engine/Plugins/Experimental/Mutable/Source/MutableRuntime/Private/MuR/CodeRunner.cpp
@@ -5553,11 +5553,12 @@ namespace mu
         {
 			OP::ParameterArgs args = Program.GetOpArgs<OP::ParameterArgs>(item.At);
 			Ptr<RangeIndex> index = BuildCurrentOpRangeIndex( item, pParams, pModel, args.variable );
-            float r=0.0f;
-            float g=0.0f;
-            float b=0.0f;            
-            pParams->GetColourValue( args.variable, &r, &g, &b, index );
-            StoreColor( item, FVector4f(r,g,b,1.0f) );
+            float R = 0.f;
+            float G = 0.f;
+			float B = 0.f;
+			float A = 0.f;
+            pParams->GetColourValue( args.variable, &R, &G, &B, &A, index );
+            StoreColor( item, FVector4f(R, G, B, A) );
             break;
         }
 
diff --git a/Engine/Plugins/Experimental/Mutable/Source/MutableRuntime/Private/MuR/Model.cpp b/Engine/Plugins/Experimental/Mutable/Source/MutableRuntime/Private/MuR/Model.cpp
--- a/Engine/Plugins/Experimental/Mutable/Source/MutableRuntime/Private/MuR/Model.cpp
+++ b/Engine/Plugins/Experimental/Mutable/Source/MutableRuntime/Private/MuR/Model.cpp
@@ -337,7 +337,7 @@ namespace mu
 	}
 
 
-	void Model::GetColourDefaultValue(int32 Index, float* R, float* G, float* B) const
+	void Model::GetColourDefaultValue(int32 Index, float* R, float* G, float* B, float* A) const
     {
     	check(m_pD->m_program.m_parameters.IsValidIndex(Index));
 		check(m_pD->m_program.m_parameters[Index].m_type == PARAMETER_TYPE::T_COLOUR);
@@ -350,9 +350,26 @@ namespace mu
         }
 
         ParamColorType& Color = m_pD->m_program.m_parameters[Index].m_defaultValue.Get<ParamColorType>();
-        if (R) *R = Color[0];
-    	if (G) *G = Color[1];
-    	if (B) *B = Color[2];
+		
+		if (R)
+		{
+			*R = Color[0];
+		}
+
+		if (G)
+		{
+			*G = Color[1];
+		}
+
+		if (B)
+		{
+			*B = Color[2];
+		}
+
+		if (A)
+		{
+			*A = Color[3];
+		}
     }
 
 
@@ -460,9 +477,9 @@ namespace mu
 
                         case PARAMETER_TYPE::T_COLOUR:
                         {
-                            float r,g,b;
-                            pOld->GetColourValue( p, &r, &g, &b );
-                            pRes->SetColourValue( thisP, r, g, b );
+                            float R, G, B, A;
+                            pOld->GetColourValue( p, &R, &G, &B, &A );
+                            pRes->SetColourValue( thisP, R, G, B, A );
                             break;
                         }
 
@@ -738,13 +755,13 @@ namespace mu
 							for (int32 RangePosition = 0; RangePosition < m_pD->DefaultRangeDimension; ++RangePosition)
 							{
 								RangeIndex->SetPosition(Dimensions, RangePosition);
-								res->SetColourValue(i, RandomGenerator(), RandomGenerator(), RandomGenerator(), RangeIndex);
+								res->SetColourValue(i, RandomGenerator(), RandomGenerator(), RandomGenerator(), RandomGenerator(), RangeIndex);
 							}
 						}
 					}
 					else
 					{
-						res->SetColourValue(i, RandomGenerator(), RandomGenerator(), RandomGenerator());
+						res->SetColourValue(i, RandomGenerator(), RandomGenerator(), RandomGenerator(), RandomGenerator());
 					}
                 }
                 break;
@@ -909,13 +926,14 @@ namespace mu
 					{
 						for (int32 RangePosition = 0; RangePosition < m_pD->DefaultRangeDimension; ++RangePosition)
 						{
-							res->SetColourValue(i, randomGenerator(), randomGenerator(), randomGenerator(),RangeIndex);
+							res->SetColourValue(i, randomGenerator(), randomGenerator(), randomGenerator(), randomGenerator(), 
+												RangeIndex);
 						}
 					}
 				}
 				else
 				{
-					res->SetColourValue(i, randomGenerator(), randomGenerator(), randomGenerator());
+					res->SetColourValue(i, randomGenerator(), randomGenerator(), randomGenerator(), randomGenerator());
 				}
                 break;
 
diff --git a/Engine/Plugins/Experimental/Mutable/Source/MutableRuntime/Private/MuR/Parameters.cpp b/Engine/Plugins/Experimental/Mutable/Source/MutableRuntime/Private/MuR/Parameters.cpp
--- a/Engine/Plugins/Experimental/Mutable/Source/MutableRuntime/Private/MuR/Parameters.cpp
+++ b/Engine/Plugins/Experimental/Mutable/Source/MutableRuntime/Private/MuR/Parameters.cpp
@@ -638,7 +638,7 @@ namespace mu
 
 
     //---------------------------------------------------------------------------------------------
-    void Parameters::GetColourValue( int index, float* pR, float* pG, float* pB,
+    void Parameters::GetColourValue( int index, float* pR, float* pG, float* pB, float* pA,
                                      const Ptr<const RangeIndex>& pos ) const
     {
         check( index>=0 && index<(int)m_pD->m_values.Num() );
@@ -658,9 +658,26 @@ namespace mu
         if (!pos)
         {
             // Return the single value
-            if (pR) *pR = m_pD->m_values[index].Get<ParamColorType>()[0];
-            if (pG) *pG = m_pD->m_values[index].Get<ParamColorType>()[1];
-            if (pB) *pB = m_pD->m_values[index].Get<ParamColorType>()[2];
+			if (pR)
+			{
+				*pR = m_pD->m_values[index].Get<ParamColorType>()[0];
+			}
+
+			if (pG)
+			{
+				*pG = m_pD->m_values[index].Get<ParamColorType>()[1];
+			}
+
+			if (pB)
+			{
+				*pB = m_pD->m_values[index].Get<ParamColorType>()[2];
+			}
+
+			if (pA)
+			{
+				*pA = m_pD->m_values[index].Get<ParamColorType>()[3];
+			}
+
             return;
         }
 
@@ -673,23 +690,57 @@ namespace mu
 			const PARAMETER_VALUE* it = m.Find(pos->m_pD->m_values);
             if (it)
             {
-                if (pR) *pR = it->Get<ParamColorType>()[0];
-                if (pG) *pG = it->Get<ParamColorType>()[1];
-                if (pB) *pB = it->Get<ParamColorType>()[2];
+				if (pR)
+				{
+					*pR = it->Get<ParamColorType>()[0];
+				}
+
+				if (pG)
+				{
+					*pG = it->Get<ParamColorType>()[1];
+				}
+
+				if (pB)
+				{
+					*pB = it->Get<ParamColorType>()[2];
+				}
+
+				if (pA)
+				{
+					*pA = it->Get<ParamColorType>()[3];
+				}
+
                 return;
             }
         }
 
         // Multivalue parameter, but no multivalue set. Return single value.
-        if (pR) *pR = m_pD->m_values[index].Get<ParamColorType>()[0];
-        if (pG) *pG = m_pD->m_values[index].Get<ParamColorType>()[1];
-        if (pB) *pB = m_pD->m_values[index].Get<ParamColorType>()[2];
+		if (pR)
+		{
+			*pR = m_pD->m_values[index].Get<ParamColorType>()[0];
+		}
+
+		if (pG)
+		{
+			*pG = m_pD->m_values[index].Get<ParamColorType>()[1];
+		}
+
+		if (pB)
+		{
+			*pB = m_pD->m_values[index].Get<ParamColorType>()[2];
+		}
+
+		if (pA)
+		{
+			*pA = m_pD->m_values[index].Get<ParamColorType>()[3];
+		}
+
         return;
     }
 
 
     //---------------------------------------------------------------------------------------------
-    void Parameters::SetColourValue( int index, float r, float g, float b,
+    void Parameters::SetColourValue( int index, float R, float G, float B, float A,
                                      const Ptr<const RangeIndex>& pos )
     {
         check( index>=0 && index<(int)m_pD->m_values.Num() );
@@ -715,9 +766,10 @@ namespace mu
             }
 
         	ParamColorType Value;
-            Value[0] = r;
-            Value[1] = g;
-            Value[2] = b;
+            Value[0] = R;
+            Value[1] = G;
+            Value[2] = B;
+			Value[3] = A;
         	
         	m_pD->m_values[index].Set<ParamColorType>(Value);
         }
@@ -736,9 +788,10 @@ namespace mu
 			PARAMETER_VALUE& it = m.FindOrAdd(pos->m_pD->m_values);
 
         	ParamColorType Value;
-            Value[0] = r;
-            Value[1] = g;
-            Value[2] = b;
+            Value[0] = R;
+            Value[1] = G;
+			Value[2] = B;
+			Value[3] = A;
 
         	it.Set<ParamColorType>(Value);
         }
diff --git a/Engine/Plugins/Experimental/Mutable/Source/MutableRuntime/Private/MuR/Serialisation.cpp b/Engine/Plugins/Experimental/Mutable/Source/MutableRuntime/Private/MuR/Serialisation.cpp
--- a/Engine/Plugins/Experimental/Mutable/Source/MutableRuntime/Private/MuR/Serialisation.cpp
+++ b/Engine/Plugins/Experimental/Mutable/Source/MutableRuntime/Private/MuR/Serialisation.cpp
@@ -31,7 +31,8 @@ namespace mu
 	MUTABLE_IMPLEMENT_POD_SERIALISABLE(FGuid);
 
 	MUTABLE_IMPLEMENT_POD_SERIALISABLE(vec2f);                                        
-	MUTABLE_IMPLEMENT_POD_SERIALISABLE(vec3f);                                        
+	MUTABLE_IMPLEMENT_POD_SERIALISABLE(vec3f);
+	MUTABLE_IMPLEMENT_POD_SERIALISABLE(vec4f);
 	MUTABLE_IMPLEMENT_POD_SERIALISABLE(mat3f);                                        
 	MUTABLE_IMPLEMENT_POD_SERIALISABLE(mat4f);                                        
 	MUTABLE_IMPLEMENT_POD_SERIALISABLE(vec2<int>);                                    
diff --git a/Engine/Plugins/Experimental/Mutable/Source/MutableRuntime/Public/MuR/Model.h b/Engine/Plugins/Experimental/Mutable/Source/MutableRuntime/Public/MuR/Model.h
--- a/Engine/Plugins/Experimental/Mutable/Source/MutableRuntime/Public/MuR/Model.h
+++ b/Engine/Plugins/Experimental/Mutable/Source/MutableRuntime/Public/MuR/Model.h
@@ -13,8 +13,8 @@
 //! data structures. Compiled models are not necessarily compatible when the runtime is updated,
 //! so this version number can be used externally to verify this. It is not used internally, and
 //! serializing models from different versions than this runtime will probably result in a crash.
-#define MUTABLE_COMPILED_MODEL_CODE_VERSION		uint32( 82 )
-#define MUTABLE_PARAMETERS_VERSION              uint32( 2 )
+#define MUTABLE_COMPILED_MODEL_CODE_VERSION		uint32( 83 )
+#define MUTABLE_PARAMETERS_VERSION              uint32( 3 )
 
 
 namespace mu
@@ -126,7 +126,7 @@ namespace mu
 		//! \pre The parameter specified by index is a T_FLOAT.
         //! \param Index Index of the parameter from 0 to GetCount()-1
         //! \param R,G,B Pointers to values where every resulting colour channel will be stored
-    	void GetColourDefaultValue(int32 Index, float* R, float* G, float* B) const;
+    	void GetColourDefaultValue(int32 Index, float* R, float* G, float* B, float* A) const;
 
     	//! Return the default value of a projector parameter, as a 4x4 matrix. The matrix is supposed to be
 		//! a linear transform in column-major.
diff --git a/Engine/Plugins/Experimental/Mutable/Source/MutableRuntime/Public/MuR/MutableMath.h b/Engine/Plugins/Experimental/Mutable/Source/MutableRuntime/Public/MuR/MutableMath.h
--- a/Engine/Plugins/Experimental/Mutable/Source/MutableRuntime/Public/MuR/MutableMath.h
+++ b/Engine/Plugins/Experimental/Mutable/Source/MutableRuntime/Public/MuR/MutableMath.h
@@ -612,6 +612,138 @@ namespace mu
 	//!
 	typedef vec3<float> vec3f;
 
+	//!
+	template <class SCALAR>
+	class vec4 : public vec<SCALAR, 4>
+	{
+	public:
+
+		//-----------------------------------------------------------------------------------------
+		inline vec4()
+		{
+			// Elements are cleared in the parent constructor
+		}
+
+		//-----------------------------------------------------------------------------------------
+		inline explicit vec4(SCALAR v)
+		{
+			this->m[0] = v;
+			this->m[1] = v;
+			this->m[2] = v;
+			this->m[3] = v;
+		}
+
+		//-----------------------------------------------------------------------------------------
+		inline vec4(SCALAR x, SCALAR y, SCALAR z, SCALAR w)
+		{
+			this->m[0] = x;
+			this->m[1] = y;
+			this->m[2] = z;
+			this->m[3] = w;
+		}
+
+		//-----------------------------------------------------------------------------------------
+		inline vec4(const vec<SCALAR, 4>& other)
+		{
+			this->m[0] = other[0];
+			this->m[1] = other[1];
+			this->m[2] = other[2];
+			this->m[3] = other[3];
+		}
+
+		//-----------------------------------------------------------------------------------------
+		inline vec4(const FVector4f& other)
+		{
+			this->m[0] = other[0];
+			this->m[1] = other[1];
+			this->m[2] = other[2];
+			this->m[3] = other[3];
+		}
+
+		//-----------------------------------------------------------------------------------------
+		inline vec4(const vec<SCALAR, 3>& other)
+		{
+			this->m[0] = other[0];
+			this->m[1] = other[1];
+			this->m[2] = other[2];
+			this->m[3] = 1;
+		}
+
+		//-----------------------------------------------------------------------------------------
+		inline void Set(SCALAR x, SCALAR y, SCALAR z, SCALAR w)
+		{
+			this->m[0] = x;
+			this->m[1] = y;
+			this->m[2] = z;
+			this->m[3] = w;
+		}
+
+
+		//-----------------------------------------------------------------------------------------
+		inline SCALAR& x()
+		{
+			return this->m[0];
+		}
+
+		//-----------------------------------------------------------------------------------------
+		inline SCALAR& y()
+		{
+			return this->m[1];
+		}
+
+		//-----------------------------------------------------------------------------------------
+		inline SCALAR& z()
+		{
+			return this->m[2];
+		}
+
+		//-----------------------------------------------------------------------------------------
+		inline SCALAR& w()
+		{
+			return this->m[3];
+		}
+
+		//-----------------------------------------------------------------------------------------
+		inline const SCALAR& x() const
+		{
+			return this->m[0];
+		}
+
+		//-----------------------------------------------------------------------------------------
+		inline const SCALAR& y() const
+		{
+			return this->m[1];
+		}
+
+		//-----------------------------------------------------------------------------------------
+		inline const SCALAR& z() const
+		{
+			return this->m[2];
+		}
+
+		//-----------------------------------------------------------------------------------------
+		inline const SCALAR& w() const
+		{
+			return this->m[3];
+		}
+
+		//-----------------------------------------------------------------------------------------
+		inline vec2<SCALAR> xy() const
+		{
+			return vec2<SCALAR>(this->m[0], this->m[1]);
+		}
+
+		//-----------------------------------------------------------------------------------------
+		inline vec3<SCALAR> xyz() const
+		{
+			return vec3<SCALAR>(this->m[0], this->m[1], this->m[2]);
+		}
+
+	};
+
+	//!
+	typedef vec4<float> vec4f;
+
 
 	//---------------------------------------------------------------------------------------------
 	//! Cross product of 3D vectors
diff --git a/Engine/Plugins/Experimental/Mutable/Source/MutableRuntime/Public/MuR/Parameters.h b/Engine/Plugins/Experimental/Mutable/Source/MutableRuntime/Public/MuR/Parameters.h
--- a/Engine/Plugins/Experimental/Mutable/Source/MutableRuntime/Public/MuR/Parameters.h
+++ b/Engine/Plugins/Experimental/Mutable/Source/MutableRuntime/Public/MuR/Parameters.h
@@ -283,14 +283,14 @@ namespace mu
         //! \param index Index of the parameter from 0 to GetCount()-1
         //! \param pR,pG,pB Pointers to values where every resulting colour channel will be stored
         //! \param pos Only for multidimensional parameters: relevant position to get in the ranges
-        void GetColourValue( int index, float* pR, float* pG, float* pB,
+        void GetColourValue( int index, float* pR, float* pG, float* pB, float* pA,
                              const Ptr<const RangeIndex>& pos=nullptr ) const;
 
 		//! If the parameter is of the colour type, set its value.
 		//! \param index Index of the parameter from 0 to GetCount()-1
         //! \param r,g,b new value of the parameter
         //! \param pos Only for multidimensional parameters: relevant position to set in the ranges
-        void SetColourValue( int index, float r, float g, float b,
+        void SetColourValue( int index, float R, float G, float B, float A,
                              const Ptr<const RangeIndex>& pos=nullptr );
 
 		//! Return the value of a projector parameter, as a 4x4 matrix. The matrix is supposed to be
diff --git a/Engine/Plugins/Experimental/Mutable/Source/MutableTools/Private/MuT/CodeGenerator_Color.cpp b/Engine/Plugins/Experimental/Mutable/Source/MutableTools/Private/MuT/CodeGenerator_Color.cpp
--- a/Engine/Plugins/Experimental/Mutable/Source/MutableTools/Private/MuT/CodeGenerator_Color.cpp
+++ b/Engine/Plugins/Experimental/Mutable/Source/MutableTools/Private/MuT/CodeGenerator_Color.cpp
@@ -161,6 +161,7 @@ namespace mu
 			Value[0] =  node.m_defaultValue[0];
 			Value[1] = node.m_defaultValue[1];
 			Value[2] = node.m_defaultValue[2];
+			Value[3] = node.m_defaultValue[3];
 
 			param.m_defaultValue.Set<ParamColorType>(Value);
 			
diff --git a/Engine/Plugins/Experimental/Mutable/Source/MutableTools/Private/MuT/NodeColourParameter.cpp b/Engine/Plugins/Experimental/Mutable/Source/MutableTools/Private/MuT/NodeColourParameter.cpp
--- a/Engine/Plugins/Experimental/Mutable/Source/MutableTools/Private/MuT/NodeColourParameter.cpp
+++ b/Engine/Plugins/Experimental/Mutable/Source/MutableTools/Private/MuT/NodeColourParameter.cpp
@@ -44,9 +44,9 @@ namespace mu
 
 
 	//---------------------------------------------------------------------------------------------
-	void NodeColourParameter::SetDefaultValue( float r, float g, float b )
+	void NodeColourParameter::SetDefaultValue( float R, float G, float B, float A )
 	{
-		m_pD->m_defaultValue = vec3<float>( r, g, b );
+		m_pD->m_defaultValue = vec4<float>( R, G, B, A );
 	}
 
 
diff --git a/Engine/Plugins/Experimental/Mutable/Source/MutableTools/Private/MuT/NodeColourParameterPrivate.h b/Engine/Plugins/Experimental/Mutable/Source/MutableTools/Private/MuT/NodeColourParameterPrivate.h
--- a/Engine/Plugins/Experimental/Mutable/Source/MutableTools/Private/MuT/NodeColourParameterPrivate.h
+++ b/Engine/Plugins/Experimental/Mutable/Source/MutableTools/Private/MuT/NodeColourParameterPrivate.h
@@ -18,7 +18,7 @@ namespace mu
 
 		static FNodeType s_type;
 
-		vec3<float> m_defaultValue;
+		vec4<float> m_defaultValue;
 		FString m_name;
 		FString m_uid;
 
@@ -27,7 +27,7 @@ namespace mu
 		//!
 		void Serialise( OutputArchive& arch ) const
 		{
-            uint32_t ver = 3;
+            uint32_t ver = 4;
 			arch << ver;
 
 			arch << m_defaultValue;
@@ -41,9 +41,20 @@ namespace mu
 		{
             uint32_t ver;
 			arch >> ver;
-            check(ver>=2&&ver<=3);
+            check(ver>=2&&ver<=4);
 
-			arch >> m_defaultValue;
+			if (ver <= 3)
+			{
+				vec3<float> aux_defaultValue;
+				arch >> aux_defaultValue;
+
+				m_defaultValue = aux_defaultValue;
+				m_defaultValue.w() = 1.f;
+			}
+			else
+			{
+				arch >> m_defaultValue;
+			}
 			
 			if (ver <= 2)
 			{
diff --git a/Engine/Plugins/Experimental/Mutable/Source/MutableTools/Public/MuT/NodeColourParameter.h b/Engine/Plugins/Experimental/Mutable/Source/MutableTools/Public/MuT/NodeColourParameter.h
--- a/Engine/Plugins/Experimental/Mutable/Source/MutableTools/Public/MuT/NodeColourParameter.h
+++ b/Engine/Plugins/Experimental/Mutable/Source/MutableTools/Public/MuT/NodeColourParameter.h
@@ -54,7 +54,7 @@ namespace mu
 		void SetUid( const FString& );
 
 		//! Set the default value of the parameter.
-		void SetDefaultValue( float r, float g, float b );
+		void SetDefaultValue( float R, float G, float B, float A );
 
         //! Set the number of ranges (dimensions) for this parameter.
         //! By default a parameter has 0 ranges, meaning it only has one value.
