[mutable] Fixed Mutable Texture size logging so that it's not always zero, but actually what's resident on the GPU.
#rnx
#tests PIE

diff --git a/Engine/Plugins/Experimental/Mutable/Source/CustomizableObject/Private/MuCO/CustomizableObjectSystem.cpp b/Engine/Plugins/Experimental/Mutable/Source/CustomizableObject/Private/MuCO/CustomizableObjectSystem.cpp
--- a/Engine/Plugins/Experimental/Mutable/Source/CustomizableObject/Private/MuCO/CustomizableObjectSystem.cpp
+++ b/Engine/Plugins/Experimental/Mutable/Source/CustomizableObject/Private/MuCO/CustomizableObjectSystem.cpp
@@ -3302,7 +3302,7 @@ int32 UCustomizableObjectSystem::GetTotalInstances() const
 
 int32 UCustomizableObjectSystem::GetTextureMemoryUsed() const
 {
-	return 0; // Currently disabled since the reported value was incorrect due to MIP streaming.
+	return GetPrivateChecked()->LogBenchmarkUtil.TextureGPUSize.GetValue();
 }
 
 int32 UCustomizableObjectSystem::GetAverageBuildTime() const
diff --git a/Engine/Plugins/Experimental/Mutable/Source/CustomizableObject/Private/MuCO/LogBenchmarkUtil.cpp b/Engine/Plugins/Experimental/Mutable/Source/CustomizableObject/Private/MuCO/LogBenchmarkUtil.cpp
--- a/Engine/Plugins/Experimental/Mutable/Source/CustomizableObject/Private/MuCO/LogBenchmarkUtil.cpp
+++ b/Engine/Plugins/Experimental/Mutable/Source/CustomizableObject/Private/MuCO/LogBenchmarkUtil.cpp
@@ -155,6 +155,8 @@ void FLogBenchmarkUtil::UpdateStats()
 
 	// Textures
 	uint32 LocalNumAllocatedTextures = 0;
+	uint32 LocalTextureGPUSize = 0;
+
 	for (auto Iterator = TextureTrackerArray.CreateIterator(); Iterator; ++Iterator)
 	{
 		if (Iterator->IsStale())
@@ -164,10 +166,13 @@ void FLogBenchmarkUtil::UpdateStats()
 		else if (Iterator->IsValid() && (*Iterator)->GetResource())
 		{
 			++LocalNumAllocatedTextures;
+
+			LocalTextureGPUSize += (*Iterator)->CalcTextureMemorySizeEnum(ETextureMipCount::TMC_ResidentMips);
 		}
 	}
 
 	NumAllocatedTextures = LocalNumAllocatedTextures;
+	TextureGPUSize = LocalTextureGPUSize;
 
 #if WITH_EDITORONLY_DATA
 	if (UCustomizableObjectSystem::GetInstance()->IsMutableAnimInfoDebuggingEnabled())
diff --git a/Engine/Plugins/Experimental/Mutable/Source/CustomizableObject/Private/MuCO/LogBenchmarkUtil.h b/Engine/Plugins/Experimental/Mutable/Source/CustomizableObject/Private/MuCO/LogBenchmarkUtil.h
--- a/Engine/Plugins/Experimental/Mutable/Source/CustomizableObject/Private/MuCO/LogBenchmarkUtil.h
+++ b/Engine/Plugins/Experimental/Mutable/Source/CustomizableObject/Private/MuCO/LogBenchmarkUtil.h
@@ -56,6 +56,7 @@ DECLARE_STATS_GROUP(TEXT("Mutable"), STATGROUP_Mutable, STATCAT_Advanced);
 
 DECLARE_BENCHMARK_INSIGHTS(NumAllocatedSkeletalMeshes, TEXT("Num Allocated Mutable Skeletal Meshes"));
 DECLARE_BENCHMARK_INSIGHTS(NumAllocatedTextures, TEXT("Num Allocated Mutable Textures"));
+DECLARE_BENCHMARK_INSIGHTS(TextureGPUSize, TEXT("Size of the Mutable Texture mips that are resident on the GPU"));
 DECLARE_BENCHMARK_INSIGHTS(NumInstances, TEXT("Num Instances"));
 DECLARE_BENCHMARK_INSIGHTS(NumInstancesLOD0, TEXT("Num Instances at LOD 0"));
 DECLARE_BENCHMARK_INSIGHTS(NumInstancesLOD1, TEXT("Num Instances at LOD 1"));
@@ -89,6 +90,8 @@ public:
 	void FinishUpdateImage(const FString& CustomizableObjectPathName, const FString& InstancePathName, double TaskUpdateImageTime, const int64 TaskUpdateImageMemoryPeak, const int64 TaskUpdateImageRealMemoryPeak) const;
 	
 	DECLARE_BENCHMARK_STAT(NumAllocatedTextures, uint32);
+	DECLARE_BENCHMARK_STAT(TextureGPUSize, uint32);
+
 	DECLARE_BENCHMARK_STAT(NumAllocatedSkeletalMeshes, int32);
 
 	DECLARE_BENCHMARK_STAT(NumInstances, int32);
