[mutable] Fix the mipmaps being incorrect for very small images, whose compose block size is smaller than the texture compression block size.
#rnx

diff --git a/Engine/Plugins/Experimental/Mutable/Source/MutableTools/Private/MuT/CodeGenerator.cpp b/Engine/Plugins/Experimental/Mutable/Source/MutableTools/Private/MuT/CodeGenerator.cpp
--- a/Engine/Plugins/Experimental/Mutable/Source/MutableTools/Private/MuT/CodeGenerator.cpp
+++ b/Engine/Plugins/Experimental/Mutable/Source/MutableTools/Private/MuT/CodeGenerator.cpp
@@ -1525,6 +1525,13 @@ namespace mu
 								int32 mipsX = FMath::CeilLogTwo(BlockPixelsX / finfo.PixelsPerBlockX);
 								int32 mipsY = FMath::CeilLogTwo(BlockPixelsY / finfo.PixelsPerBlockY);
 								mop->BlockLevels = (uint8)FMath::Max(mipsX, mipsY);
+								
+								if (BlockPixelsX < finfo.PixelsPerBlockX || BlockPixelsY < finfo.PixelsPerBlockY)
+								{
+									// In this case, the mipmap will never be useful for blocks, so we indicate that
+									// it should make the mips at the root of the expression.
+									mop->bOnlyTail = true;
+								}
 
 								mop->AddressMode = mipmapNode->GetPrivate()->m_settings.m_addressMode;
 								mop->FilterType = mipmapNode->GetPrivate()->m_settings.m_filterType;
@@ -1553,6 +1560,13 @@ namespace mu
 								int mipsY = (int)ceilf(logf((float)BlockPixelsY / finfo.PixelsPerBlockY) / logf(2.0f));
 								mop->BlockLevels = (uint8_t)FMath::Max(mipsX, mipsY);
 
+								if (BlockPixelsX < finfo.PixelsPerBlockX || BlockPixelsY < finfo.PixelsPerBlockY)
+								{
+									// In this case, the mipmap will never be useful for blocks, so we indicate that
+									// it should make the mips at the root of the expression.
+									mop->bOnlyTail = true;
+								}
+
 								// Not important for the end of the mip tail?
 								mop->AddressMode = EAddressMode::ClampToEdge;
 								mop->FilterType = EMipmapFilterType::MFT_SimpleAverage;
