[Mutable] Fix subimage decompression for AstcRGBL formats.

#rnx

diff --git a/Engine/Plugins/Experimental/Mutable/Source/MutableRuntime/Private/MuR/BlockCompression/Miro/Miro.cpp b/Engine/Plugins/Experimental/Mutable/Source/MutableRuntime/Private/MuR/BlockCompression/Miro/Miro.cpp
--- a/Engine/Plugins/Experimental/Mutable/Source/MutableRuntime/Private/MuR/BlockCompression/Miro/Miro.cpp
+++ b/Engine/Plugins/Experimental/Mutable/Source/MutableRuntime/Private/MuR/BlockCompression/Miro/Miro.cpp
@@ -9240,31 +9240,39 @@ namespace miro::SubImageDecompression
 	{
 		init_astc_decompress();
 
-		for (uint32 y = 0; y < SubSize.Y; y += BlockSize)
+		const uint32 NumBlocksX = FMath::DivideAndRoundUp(uint32(FromSize.X), BlockSize);
+		const uint32 NumSubBlocksX = FMath::DivideAndRoundUp(uint32(SubSize.X), BlockSize);
+		const uint32 NumSubBlocksY = FMath::DivideAndRoundUp(uint32(SubSize.Y), BlockSize);
+
+		if (SubSize.X <= 0 || SubSize.Y <= 0)
+		{
+			return;
+		}
+
+		for (uint32 BlockY = 0; BlockY < NumSubBlocksY; ++BlockY)
 		{
-			for (uint32 x = 0; x < SubSize.X; x += BlockSize)
+			for (uint32 BlockX = 0; BlockX < NumSubBlocksX; ++BlockX)
 			{
 				uint8 Block[BlockSize * BlockSize * 4];
-				astcdec::decompressSRGB<BlockSize>(Block, From);
+
+				constexpr int32 CompressedBlockSize = 16;
+				const uint8* SrcBlockPtr = From + (BlockY * NumBlocksX + BlockX) * CompressedBlockSize;
+
+				astcdec::decompressSRGB<BlockSize>(Block, SrcBlockPtr);
 
 				for (uint32 py = 0; py < BlockSize; py++)
 				{
 					for (uint32 px = 0; px < BlockSize; px++)
 					{
-						uint32 xi = x + px;
-						uint32 yi = y + py;
+						uint32 xi = FMath::Min(BlockX * BlockSize + px, uint32(SubSize.X) - 1);
+						uint32 yi = FMath::Min(BlockY * BlockSize + py, uint32(SubSize.Y) - 1);
 
-						if (LIKELY((xi < ToSize.X) & (yi < ToSize.Y)))
-						{
-							uint8* ToPixel = To + (yi * ToSize.X + xi) * 3;
-							ToPixel[0] = Block[py * BlockSize * 4 + px * 4 + 0];
-							ToPixel[1] = Block[py * BlockSize * 4 + px * 4 + 1];
-							ToPixel[2] = Block[py * BlockSize * 4 + px * 4 + 2];
-						}
+						uint8* ToPixel = To + (yi * ToSize.X + xi) * 3;
+						ToPixel[0] = Block[py * BlockSize * 4 + px * 4 + 0];
+						ToPixel[1] = Block[py * BlockSize * 4 + px * 4 + 1];
+						ToPixel[2] = Block[py * BlockSize * 4 + px * 4 + 2];
 					}
 				}
-
-				From += 16;
 			}
 		}
 
@@ -9275,32 +9283,40 @@ namespace miro::SubImageDecompression
 	{
 		init_astc_decompress();
 
-		for (uint32 y = 0; y < SubSize.Y; y += BlockSize)
+		const uint32 NumBlocksX = FMath::DivideAndRoundUp(uint32(FromSize.X), BlockSize);
+		const uint32 NumSubBlocksX = FMath::DivideAndRoundUp(uint32(SubSize.X), BlockSize);
+		const uint32 NumSubBlocksY = FMath::DivideAndRoundUp(uint32(SubSize.Y), BlockSize);
+
+		if (SubSize.X <= 0 || SubSize.Y <= 0)
+		{
+			return;
+		}
+
+		for (uint32 BlockY = 0; BlockY < NumSubBlocksY; ++BlockY)
 		{
-			for (uint32 x = 0; x < SubSize.X; x += BlockSize)
+			for (uint32 BlockX = 0; BlockX < NumSubBlocksX; ++BlockX)
 			{
 				uint8 Block[BlockSize * BlockSize * 4];
-				astcdec::decompressSRGB<BlockSize>(Block, From);
+
+				constexpr int32 CompressedBlockSize = 16;
+				const uint8* SrcBlockPtr = From + (BlockY * NumBlocksX + BlockX) * CompressedBlockSize;
+
+				astcdec::decompressSRGB<BlockSize>(Block, SrcBlockPtr);
 
 				for (uint32 py = 0; py < BlockSize; py++)
 				{
 					for (uint32 px = 0; px < BlockSize; px++)
 					{
-						uint32 xi = x + px;
-						uint32 yi = y + py;
+						uint32 xi = FMath::Min(BlockX * BlockSize + px, uint32(SubSize.X) - 1);
+						uint32 yi = FMath::Min(BlockY * BlockSize + py, uint32(SubSize.Y) - 1);
 
-						if (LIKELY((xi < ToSize.X) & (yi < ToSize.Y)))
-						{
-							uint8* ToPixel = To + (yi * ToSize.X + xi) * 4;
-							ToPixel[0] = Block[py * BlockSize * 4 + px * 4 + 0];
-							ToPixel[1] = Block[py * BlockSize * 4 + px * 4 + 1];
-							ToPixel[2] = Block[py * BlockSize * 4 + px * 4 + 2];
-							ToPixel[3] = 255;
-						}
+						uint8* ToPixel = To + (yi * ToSize.X + xi) * 4;
+						ToPixel[0] = Block[py * BlockSize * 4 + px * 4 + 0];
+						ToPixel[1] = Block[py * BlockSize * 4 + px * 4 + 1];
+						ToPixel[2] = Block[py * BlockSize * 4 + px * 4 + 2];
+						ToPixel[3] = 255;
 					}
 				}
-
-				From += 16;
 			}
 		}
 	}
