[mutable] Optimizations to object compilation process, including use of concurrency and memory reduction in the case of large source data objects.
#rnx

diff --git a/Engine/Plugins/Experimental/Mutable/Source/CustomizableObjectEditor/Private/MuCOE/CustomizableObjectCompileRunnable.cpp b/Engine/Plugins/Experimental/Mutable/Source/CustomizableObjectEditor/Private/MuCOE/CustomizableObjectCompileRunnable.cpp
--- a/Engine/Plugins/Experimental/Mutable/Source/CustomizableObjectEditor/Private/MuCOE/CustomizableObjectCompileRunnable.cpp
+++ b/Engine/Plugins/Experimental/Mutable/Source/CustomizableObjectEditor/Private/MuCOE/CustomizableObjectCompileRunnable.cpp
@@ -11,6 +11,7 @@
 #include "MuT/UnrealPixelFormatOverride.h"
 #include "Serialization/MemoryWriter.h"
 #include "Async/Async.h"
+#include "Containers/Ticker.h"
 #include "Trace/Trace.inl"
 
 class ITargetPlatform;
@@ -28,6 +29,42 @@ FCustomizableObjectCompileRunnable::FCustomizableObjectCompileRunnable(mu::Ptr<m
 }
 
 
+mu::Ptr<mu::Image> FCustomizableObjectCompileRunnable::LoadResourceReferenced(int32 ID)
+{
+	check(IsInGameThread());
+
+	MUTABLE_CPUPROFILER_SCOPE(LoadResourceReferenced);
+
+	mu::Ptr<mu::Image> Image;
+	if (!ReferencedTextures.IsValidIndex(ID))
+	{
+		// The id is not valid for this CO
+		check(false);
+		return Image;
+	}
+
+	// Find the texture id
+	TSoftObjectPtr<UTexture> TexturePtr = ReferencedTextures[ID];
+
+	// This can cause a stall because of loading the asset.
+	UTexture2D* Texture = Cast<UTexture2D>(TexturePtr.LoadSynchronous());
+	if (!Texture)
+	{
+		// Failed to load the texture
+		check(false);
+		return Image;
+	}
+
+	// In the editor the src data can be directly accessed
+	Image = new mu::Image();
+	int32 MipmapsToSkip = 0;
+	bool bIsNormalComposite = false; // TODO?
+	EUnrealToMutableConversionError Error = ConvertTextureUnrealSourceToMutable(Image.get(), Texture, bIsNormalComposite, MipmapsToSkip);
+	check(Error == EUnrealToMutableConversionError::Success);
+	return Image;
+}
+
+
 uint32 FCustomizableObjectCompileRunnable::Run()
 {
 	TRACE_BEGIN_REGION(UE_MUTABLE_CORE_REGION);
@@ -81,68 +118,58 @@ uint32 FCustomizableObjectCompileRunnable::Run()
 		CompilerOptions->SetImagePixelFormatOverride( UnrealPixelFormatFunc );
 	}
 
-	CompilerOptions->SetReferencedResourceCallback([this](int32 ID) 
+	CompilerOptions->SetReferencedResourceCallback([this](int32 ID, TSharedPtr<mu::Ptr<mu::Image>> ResolvedImage)
 		{
-			mu::Ptr<mu::Image> Image;
+			// This runs in a random thread
+			UE::Tasks::FTaskEvent CompletionEvent(TEXT("ReferencedResourceCallbackCompletion"));			
 
-			auto LoadFunc = [this,ID]() 
+			if (IsInGameThread())
 			{
-				check(IsInGameThread());
+				// Do everything now
+				mu::Ptr<mu::Image> Result = LoadResourceReferenced(ID);
+				*ResolvedImage = Result;
+				CompletionEvent.Trigger();
+			}
+			else
+			{
+				PendingResourceReferenceRequests.Enqueue(FReferenceResourceRequest(ID, ResolvedImage, MakeShared<UE::Tasks::FTaskEvent>(CompletionEvent)));
+			}
 
-				mu::Ptr<mu::Image> Image;
-				if (!ReferencedTextures.IsValidIndex(ID))
-				{
-					// The id is not valid for this CO
-					check(false);
-					return Image;
-				}
+			return CompletionEvent;
+		});
 
-				// Find the texture id
-				TSoftObjectPtr<UTexture> TexturePtr = ReferencedTextures[ID];
+	// Register the tick function that will process the game-thread image reference resolve requests.
+	FTSTicker::FDelegateHandle ResolveReferenceResourcesTickerHandle = FTSTicker::GetCoreTicker().AddTicker(
+		FTickerDelegate::CreateLambda([this](float)
+			{
+				check(IsInGameThread());
 
-				// This can cause a stall because of loading the asset.
-				UTexture2D* Texture = Cast<UTexture2D>(TexturePtr.LoadSynchronous());
-				if (!Texture)
-				{
-					// Failed to load the texture
-					check(false);
-					return Image;
-				}
+				constexpr double MaxSecondsPerFrame = 0.4;
 
-				// In the editor the src data can be directly accessed
-				Image = new mu::Image();
-				int32 MipmapsToSkip = 0;
-				bool bIsNormalComposite = false; // TODO?
-				EUnrealToMutableConversionError Error = ConvertTextureUnrealSourceToMutable(Image.get(), Texture, bIsNormalComposite, MipmapsToSkip);
-				check(Error == EUnrealToMutableConversionError::Success);
-				return Image;
-			};
+				double MaxTime = FPlatformTime::Seconds() + MaxSecondsPerFrame;
 
-			// This runs in a random thread
-			if (IsInGameThread())
-			{
-				Image = LoadFunc();
-			}
-			else
-			{
-				UE::Tasks::FTaskEvent Completion(TEXT("SetReferencedResourceCallback"));
+				FReferenceResourceRequest Request;
+				while (PendingResourceReferenceRequests.Dequeue(Request))
+				{
+					*Request.ResolvedImage = LoadResourceReferenced(Request.ID);
+					Request.CompletionEvent->Trigger();
 
-				AsyncTask(ENamedThreads::GameThread, [&Image, &Completion, LoadFunc]() 
+					// Simple time limit enforcement to avoid blocking the game thread if there are many requests.
+					double CurrentTime = FPlatformTime::Seconds();
+					if (CurrentTime >= MaxTime)
 					{
-						Image = LoadFunc();
-						Completion.Trigger();
-					});
+						break;
+					}
+				}
 
-				Completion.Wait();
-			}
-			return Image;
-		});
+				return true;
+			}));
 
 	// Minimum resident mip count.
-	const int MinResidentMips = UTexture::GetStaticMinTextureResidentMipCount();
+	const int32 MinResidentMips = UTexture::GetStaticMinTextureResidentMipCount();
 	// Data smaller than this will always be loaded, as part of the customizable object compiled model.
-	const int MinRomSize = 128;
-	CompilerOptions->SetDataPackingStrategy(MinRomSize, MinResidentMips);
+	const int32 MinRomSizeBytes = 128;
+	CompilerOptions->SetDataPackingStrategy(MinRomSizeBytes, MinResidentMips);
 
 	// At object compilation time we don't know if we will want progressive images or not. Assume we will. 
 	// TODO: Per-state setting?
@@ -194,6 +221,8 @@ uint32 FCustomizableObjectCompileRunnable::Run()
 
 	bThreadCompleted = true;
 
+	FTSTicker::GetCoreTicker().RemoveTicker( ResolveReferenceResourcesTickerHandle );
+
 	UE_LOG(LogMutable, Verbose, TEXT("PROFILE: [ %16.8f ] FCustomizableObjectCompileRunnable::Run end."), FPlatformTime::Seconds());
 
 	TRACE_END_REGION(UE_MUTABLE_CORE_REGION);
diff --git a/Engine/Plugins/Experimental/Mutable/Source/CustomizableObjectEditor/Private/MuCOE/CustomizableObjectCompileRunnable.h b/Engine/Plugins/Experimental/Mutable/Source/CustomizableObjectEditor/Private/MuCOE/CustomizableObjectCompileRunnable.h
--- a/Engine/Plugins/Experimental/Mutable/Source/CustomizableObjectEditor/Private/MuCOE/CustomizableObjectCompileRunnable.h
+++ b/Engine/Plugins/Experimental/Mutable/Source/CustomizableObjectEditor/Private/MuCOE/CustomizableObjectCompileRunnable.h
@@ -38,6 +38,18 @@ private:
 	mu::Ptr<mu::Node> MutableRoot;
 	TArray<FError> ArrayErrors;
 
+	/** */
+	struct FReferenceResourceRequest
+	{
+		int32 ID = -1;
+		TSharedPtr<mu::Ptr<mu::Image>> ResolvedImage;
+		TSharedPtr<UE::Tasks::FTaskEvent> CompletionEvent;
+	};
+	TQueue<FReferenceResourceRequest, EQueueMode::Mpsc> PendingResourceReferenceRequests;
+
+	mu::Ptr<mu::Image> LoadResourceReferenced(int32 ID);
+
+
 public:
 
 	FCustomizableObjectCompileRunnable(mu::Ptr<mu::Node> Root);
diff --git a/Engine/Plugins/Experimental/Mutable/Source/MutableRuntime/Private/MuR/CodeRunnerTasked.cpp b/Engine/Plugins/Experimental/Mutable/Source/MutableRuntime/Private/MuR/CodeRunnerTasked.cpp
--- a/Engine/Plugins/Experimental/Mutable/Source/MutableRuntime/Private/MuR/CodeRunnerTasked.cpp
+++ b/Engine/Plugins/Experimental/Mutable/Source/MutableRuntime/Private/MuR/CodeRunnerTasked.cpp
@@ -719,7 +719,7 @@ namespace mu
 
 		}
 
-		if (Args.mask)
+		if (Mask)
 		{
 			if (Base->GetSize() != Mask->GetSize())
 			{
@@ -796,7 +796,7 @@ namespace mu
 		bool bUseBlendSourceFromBlendAlpha = (Args.flags & OP::ImageLayerArgs::F_BLENDED_RGB_FROM_ALPHA) != 0;
 		bool bUseMaskFromBlendAlpha = (Args.flags & OP::ImageLayerArgs::F_USE_MASK_FROM_BLENDED);
 
-		if (!bDone && Args.mask)
+		if (!bDone && Mask)
 		{
 			// Not implemented yet
 			check(!bUseBlendSourceFromBlendAlpha);
diff --git a/Engine/Plugins/Experimental/Mutable/Source/MutableTools/Private/MuT/AST.cpp b/Engine/Plugins/Experimental/Mutable/Source/MutableTools/Private/MuT/AST.cpp
--- a/Engine/Plugins/Experimental/Mutable/Source/MutableTools/Private/MuT/AST.cpp
+++ b/Engine/Plugins/Experimental/Mutable/Source/MutableTools/Private/MuT/AST.cpp
@@ -190,10 +190,10 @@ bool ASTOp::operator==( const ASTOp& other ) const
 
 
 //---------------------------------------------------------------------------------------------
-void ASTOp::FullAssert( const TArray<Ptr<ASTOp>>& roots )
+void ASTOp::FullAssert( const TArray<Ptr<ASTOp>>& Roots)
 {
     MUTABLE_CPUPROFILER_SCOPE(AST_FullAssert);
-    Traverse_TopDown_Unique_Imprecise( roots, [](const Ptr<ASTOp>& n)
+    Traverse_TopDown_Unique_Imprecise(Roots, [](const Ptr<ASTOp>& n)
     {
         n->Assert();
         return true;
@@ -201,16 +201,16 @@ void ASTOp::FullAssert( const TArray<Ptr<ASTOp>>& roots )
 }
 
 //-------------------------------------------------------------------------------------------------
-size_t ASTOp::CountNodes( const TArray<Ptr<ASTOp>>& roots )
+int32 ASTOp::CountNodes( const TArray<Ptr<ASTOp>>& Roots )
 {
     MUTABLE_CPUPROFILER_SCOPE(AST_CountNodes);
-    size_t count=0;
-    Traverse_TopRandom_Unique_NonReentrant( roots, [&](const Ptr<ASTOp>&)
+	int32 Count=0;
+    Traverse_TopRandom_Unique_NonReentrant(Roots, [&](const Ptr<ASTOp>&)
     {
-        ++count;
+        ++Count;
         return true;
     });
-    return count;
+    return Count;
 }
 
 
@@ -339,8 +339,7 @@ void ASTOp::LogHistogram( ASTOpList& roots )
 
 
 //-------------------------------------------------------------------------------------------------
-void ASTOp::Traverse_TopDown_Unique( const TArray<Ptr<ASTOp>>& roots,
-                                     TFunctionRef<bool(Ptr<ASTOp>&)> f )
+void ASTOp::Traverse_TopDown_Unique( const TArray<Ptr<ASTOp>>& roots, TFunctionRef<bool(Ptr<ASTOp>&)> f )
 {
     TQueue<Ptr<ASTOp>> pending;
 	for (const Ptr<ASTOp>& r : roots)
@@ -414,8 +413,7 @@ void ASTOp::Traverse_TopDown_Unique( const TArray<Ptr<ASTOp>>& roots,
 
 
 //-------------------------------------------------------------------------------------------------
-void ASTOp::Traverse_TopDown_Unique_Imprecise( const TArray<Ptr<ASTOp>>& roots,
-                                     TFunctionRef<bool(Ptr<ASTOp>&)> f )
+void ASTOp::Traverse_TopDown_Unique_Imprecise( const TArray<Ptr<ASTOp>>& roots, TFunctionRef<bool(Ptr<ASTOp>&)> f )
 {
     TQueue<Ptr<ASTOp>> pending;
 	for (const Ptr<ASTOp>& r : roots)
@@ -622,39 +620,6 @@ void Visitor_TopDown_Unique_Cloning::Process()
 }
 
 
-//-------------------------------------------------------------------------------------------------
-//-------------------------------------------------------------------------------------------------
-//-------------------------------------------------------------------------------------------------
-void ASTOp::Traverse_TopDown_Repeat( const TArray<Ptr<ASTOp>>& roots,
-                                     TFunctionRef<bool(Ptr<ASTOp>& node)> f )
-{
-    ASTOpList pending = roots;
-
-    while (pending.Num())
-    {
-        Ptr<ASTOp> pCurrent = pending.Pop();
-
-        if (pCurrent)
-        {
-            // Process
-            bool recurse = f(pCurrent);
-
-            // Recurse children
-            if (recurse)
-            {
-                pCurrent->ForEachChild([&]( ASTChild& c )
-                {
-                    if (c)
-                    {
-                        pending.Add( c.m_child );
-                    }
-                });
-            }
-        }
-    }
-}
-
-
 //-------------------------------------------------------------------------------------------------
 //-------------------------------------------------------------------------------------------------
 //-------------------------------------------------------------------------------------------------
diff --git a/Engine/Plugins/Experimental/Mutable/Source/MutableTools/Private/MuT/AST.h b/Engine/Plugins/Experimental/Mutable/Source/MutableTools/Private/MuT/AST.h
--- a/Engine/Plugins/Experimental/Mutable/Source/MutableTools/Private/MuT/AST.h
+++ b/Engine/Plugins/Experimental/Mutable/Source/MutableTools/Private/MuT/AST.h
@@ -3,6 +3,7 @@
 #pragma once
 
 #include "MuT/Platform.h"
+#include "MuT/Compiler.h"
 #include "MuR/Image.h"
 #include "MuR/Mesh.h"
 #include "MuR/ModelPrivate.h"
@@ -21,6 +22,7 @@
 #include "Misc/AssertionMacros.h"
 #include "Templates/TypeHash.h"
 #include "Templates/Function.h"
+#include "Hash/CityHash.h"
 
 #include <array>
 #include <atomic>
@@ -30,7 +32,6 @@
 #include <utility>
 #include <set>
 
-#include "MuT/Compiler.h"
 
 namespace std
 {
@@ -260,11 +261,8 @@ namespace mu
 								::GetTypeHash(Data->m_size[1]))
 				);
 
-				if (Data->m_data.Num()>=4)
-				{
-					uint32 FirstData = *(reinterpret_cast<const uint32*>(Data->m_data.GetData()));
-					Hash = HashCombineFast( Hash, ::GetTypeHash(FirstData) );
-				}
+				uint64 DataHash = CityHash64(reinterpret_cast<const char*>(Data->m_data.GetData()), Data->m_data.Num());
+				Hash = HashCombineFast( Hash, ::GetTypeHash(DataHash) );
 
 				return Hash;
 			}
@@ -460,10 +458,8 @@ namespace mu
 
         inline ASTOp()
         {
-            m_traversedTopDownAdapative = 0;
-            m_constantSubtree = false;
-            m_hasSpecialOpInSubtree = false;
-
+			bIsConstantSubgraph = false;
+			bHasSpecialOpInSubgraph = false;
         }
 
         virtual ~ASTOp() {}
@@ -506,11 +502,30 @@ namespace mu
         //---------------------------------------------------------------------------------------------
         static void FullAssert( const TArray<Ptr<ASTOp>>& roots );
 
-        static size_t CountNodes( const TArray<Ptr<ASTOp>>& roots );
+        static int32 CountNodes( const TArray<Ptr<ASTOp>>& roots );
+
+		inline bool IsConstantOp() const
+		{
+			OP_TYPE Type = GetOpType();
+			return Type == OP_TYPE::BO_CONSTANT
+				|| Type == OP_TYPE::NU_CONSTANT
+				|| Type == OP_TYPE::SC_CONSTANT
+				|| Type == OP_TYPE::CO_CONSTANT
+				|| Type == OP_TYPE::IM_CONSTANT
+				|| Type == OP_TYPE::ME_CONSTANT
+				|| Type == OP_TYPE::LA_CONSTANT
+				|| Type == OP_TYPE::PR_CONSTANT
+				|| Type == OP_TYPE::ST_CONSTANT
+				|| Type == OP_TYPE::ED_CONSTANT
+				;
+		}
 
         //! Deep clone. New node will have no parents and reference new children
         static Ptr<ASTOp> DeepClone( const Ptr<ASTOp>& );
 
+		//!
+		static void LogHistogram(ASTOpList& roots);
+
         // Code optimisation methods
         //---------------------------------------------------------------------------------------------
 
@@ -534,9 +549,6 @@ namespace mu
         //!
         static void ClearLinkData( Ptr<ASTOp>& root );
 
-        //!
-        static void LogHistogram( ASTOpList& roots );
-
     private:
 
         //!
@@ -672,9 +684,6 @@ namespace mu
         }
 
 
-        //! Kind of top-down, but really not.
-        static void Traverse_TopDown_Repeat( const TArray<Ptr<ASTOp>>& roots, TFunctionRef<bool(Ptr<ASTOp>& node)> f );
-
         //! This version is slighlty faster, but doesn't support recursive traversals so
         //! use it only in controlled cases.
         static void Traverse_BottomUp_Unique_NonReentrant( ASTOpList& roots, TFunctionRef<void(Ptr<ASTOp>&)> f );
@@ -715,14 +724,11 @@ namespace mu
         //!
         int8 linkedRange = -1;
 
-        //! special flag for Traverse_TopDown_Adaptative traversal.
-        uint8 m_traversedTopDownAdapative : 1;
-
-        //! Embedded node data for the constant subtree detection. This flag is only valid if the
-        //! constant detection process has been executed and no relevant AST transformations have
-        //! happened.
-        uint8 m_constantSubtree : 1;
-        uint8 m_hasSpecialOpInSubtree : 1;
+        /** Embedded node data for the constant subtree detection.This flag is only valid if the
+         * constant detection process has been executed and no relevant AST transformations have
+         * happened. */
+        uint8 bIsConstantSubgraph : 1;
+        uint8 bHasSpecialOpInSubgraph : 1;
 
     private:
         friend class ASTChild;
@@ -732,7 +738,8 @@ namespace mu
 
     public:
 
-        int GetParentCount() const;
+		/** Get the number of ASTOps that have this one as child, in any existing AST graph. */
+        int32 GetParentCount() const;
 
         //! Make all parents of this node point at the other node instead.
         static void Replace( const Ptr<ASTOp>& node, const Ptr<ASTOp>& other );
@@ -755,10 +762,8 @@ namespace mu
         using FBlockLayoutSizeCache=TMap< const TPair<ASTOp*,int>, TPair<int,int>>;
 
         //! Return the size in layout blocks of a particular block given by absolute index
-        virtual void GetBlockLayoutSize( int blockIndex, int* pBlockX, int* pBlockY,
-                                         FBlockLayoutSizeCache* cache );
-        void GetBlockLayoutSizeCached( int blockIndex, int* pBlockX, int* pBlockY,
-                                       FBlockLayoutSizeCache* cache );
+        virtual void GetBlockLayoutSize( int blockIndex, int* pBlockX, int* pBlockY, FBlockLayoutSizeCache* cache );
+        void GetBlockLayoutSizeCached( int blockIndex, int* pBlockX, int* pBlockY, FBlockLayoutSizeCache* cache );
 
 
         //! Return the size in pixels of the layout grid block for the image operation
diff --git a/Engine/Plugins/Experimental/Mutable/Source/MutableTools/Private/MuT/ASTOpConstantResource.cpp b/Engine/Plugins/Experimental/Mutable/Source/MutableTools/Private/MuT/ASTOpConstantResource.cpp
--- a/Engine/Plugins/Experimental/Mutable/Source/MutableTools/Private/MuT/ASTOpConstantResource.cpp
+++ b/Engine/Plugins/Experimental/Mutable/Source/MutableTools/Private/MuT/ASTOpConstantResource.cpp
@@ -86,7 +86,6 @@ namespace mu
 			if (bCannotBeScaled)
 			{
 				// Store only the mips that we have already calculated. We assume we have calculated them correctly.
-				//MipsToStore = 1;
 				MipsToStore = pImage->GetLODCount();
 			}
 
diff --git a/Engine/Plugins/Experimental/Mutable/Source/MutableTools/Private/MuT/ASTOpReferenceResource.cpp b/Engine/Plugins/Experimental/Mutable/Source/MutableTools/Private/MuT/ASTOpReferenceResource.cpp
--- a/Engine/Plugins/Experimental/Mutable/Source/MutableTools/Private/MuT/ASTOpReferenceResource.cpp
+++ b/Engine/Plugins/Experimental/Mutable/Source/MutableTools/Private/MuT/ASTOpReferenceResource.cpp
@@ -55,11 +55,11 @@ namespace mu
 	void ASTOpReferenceResource::Link(FProgram& program, FLinkerOptions* Options)
 	{
 		if (!linkedAddress)
-		{			
+		{
 			OP::ResourceReferenceArgs Args;
 			FMemory::Memset(&Args, 0, sizeof(Args));
 			Args.ID = ID;
-			Args.ForceLoad = bForceLoad?1:0;
+			Args.ForceLoad = bForceLoad ? 1 : 0;
 			Args.ImageDesc = ImageDesc;
 
 			linkedAddress = (OP::ADDRESS)program.m_opAddress.Num();
@@ -137,29 +137,32 @@ namespace mu
 	{
 		mu::Ptr<ASTOp> NewOp;
 
-		switch (type)
-		{
+		//switch (type)
+		//{
 
-		case OP_TYPE::IM_REFERENCE:
-		{
-			// If we are in reference resolution stage
-			if (Pass>=2 && bForceLoad)
-			{
-				check(options.ReferencedResourceProvider);
-
-				Ptr<Image> ConstantImage = options.ReferencedResourceProvider(ID);
-
-				Ptr<ASTOpConstantResource> ConstantOp = new ASTOpConstantResource;
-				ConstantOp->type = OP_TYPE::IM_CONSTANT;
-				ConstantOp->SetValue( ConstantImage.get(), options.bUseDiskCache );
-				NewOp = ConstantOp;
-			}
-			break;
-		}
+		//case OP_TYPE::IM_REFERENCE:
+		//{
+		//	// If we are in reference resolution stage
+		//	if (Pass>=2 && bForceLoad)
+		//	{
+		//		MUTABLE_CPUPROFILER_SCOPE(ResolveReference);
 
-		default:
-			checkf(false, TEXT("Instruction not supported"));
-		}
+		//		check(options.ReferencedResourceProvider);
+
+		//		TFuture<Ptr<Image>> ConstantImageFuture = options.ReferencedResourceProvider(ID);
+		//		Ptr<Image> ConstantImage = ConstantImageFuture.Get();
+
+		//		Ptr<ASTOpConstantResource> ConstantOp = new ASTOpConstantResource;
+		//		ConstantOp->type = OP_TYPE::IM_CONSTANT;
+		//		ConstantOp->SetValue( ConstantImage.get(), options.bUseDiskCache );
+		//		NewOp = ConstantOp;
+		//	}
+		//	break;
+		//}
+
+		//default:
+		//	checkf(false, TEXT("Instruction not supported"));
+		//}
 
 		return NewOp;
 	}
diff --git a/Engine/Plugins/Experimental/Mutable/Source/MutableTools/Private/MuT/CodeOptimiser.cpp b/Engine/Plugins/Experimental/Mutable/Source/MutableTools/Private/MuT/CodeOptimiser.cpp
--- a/Engine/Plugins/Experimental/Mutable/Source/MutableTools/Private/MuT/CodeOptimiser.cpp
+++ b/Engine/Plugins/Experimental/Mutable/Source/MutableTools/Private/MuT/CodeOptimiser.cpp
@@ -1,15 +1,10 @@
 // Copyright Epic Games, Inc. All Rights Reserved.
 
 #include "MuT/CodeOptimiser.h"
+
 #include "MuT/ErrorLogPrivate.h"
 #include "MuT/AST.h"
 #include "MuT/StreamsPrivate.h"
-
-#include "MuR/ModelPrivate.h"
-#include "MuR/SystemPrivate.h"
-#include "MuR/Operations.h"
-#include "MuR/OpMeshMerge.h"
-
 #include "MuT/ASTOpInstanceAdd.h"
 #include "MuT/ASTOpConditional.h"
 #include "MuT/ASTOpSwitch.h"
@@ -20,10 +15,19 @@
 #include "MuT/ASTOpMeshClipMorphPlane.h"
 #include "MuT/ASTOpMeshApplyPose.h"
 #include "MuT/ASTOpImageRasterMesh.h"
+#include "MuT/ASTOpReferenceResource.h"
+
+#include "MuR/ModelPrivate.h"
+#include "MuR/SystemPrivate.h"
+#include "MuR/Operations.h"
+#include "MuR/OpMeshMerge.h"
+#include "MuR/MutableRuntimeModule.h"
+
+#include "Tasks/Task.h"
+#include "Tasks/Pipe.h"
 
 #include <unordered_set>
 
-#include "MuR/MutableRuntimeModule.h"
 
 namespace mu
 {
@@ -449,24 +453,45 @@ namespace mu
 	{
 	private:
 
-		//
-		mu::Ptr<ASTOp> m_result;
-
-	public:
 		// input
-		mu::Ptr<ASTOp> m_source;
+		Ptr<ASTOp> Source;
 		bool bUseDiskCache = false;
-		int ImageCompressionQuality = 0;
+		int32 ImageCompressionQuality = 0;
+		int32 OptimizationPass = 0;
+		FReferencedResourceFunc ReferencedResourceProvider;
+
+		// Intermediate
+		Ptr<ASTOp> SourceCloned;
+
+		// Result
+		Ptr<ASTOp> Result;
+
+	public:
 
+		ConstantTask( const Ptr<ASTOp>& InSource, const CompilerOptions::Private* InOptions, int32 InOptimizationPass )
+		{
+			OptimizationPass = InOptimizationPass;
+			Source = InSource;
+			bUseDiskCache = InOptions->OptimisationOptions.bUseDiskCache;
+			ImageCompressionQuality = InOptions->ImageCompressionQuality;
+			ReferencedResourceProvider = InOptions->OptimisationOptions.ReferencedResourceProvider;
+		}
+
+		void Prepare()
+		{
+			// This runs in the AST managing thread.
+
+			// We need the clone because linking modifies ASTOp state and also to be safe for concurrency.
+			SourceCloned = ASTOp::DeepClone(Source);
+		}
 
 		void Run(FImageOperator ImOp)
 		{
 			MUTABLE_CPUPROFILER_SCOPE(ConstantTask_Run);
 
-			// We need the clone because linking modifies ASTOp state
-			mu::Ptr<ASTOp> cloned = ASTOp::DeepClone( m_source );
+			// This runs in a worker thread
 
-			OP_TYPE type = cloned->GetOpType();
+			OP_TYPE type = SourceCloned->GetOpType();
 			DATATYPE dtype = GetOpDataType(type);
 
 			Ptr<Settings> pSettings = new Settings;
@@ -476,19 +501,19 @@ namespace mu
 
 			pSystem->GetPrivate()->ImagePixelFormatOverride = ImOp.FormatImageOverride;
 
-			// Don't generate mips suring linking here.
+			// Don't generate mips during linking here.
 			FLinkerOptions LinkerOptions(ImOp);
 			LinkerOptions.MinTextureResidentMipCount = 255;
 
 			TSharedPtr<const Model> model = MakeShared<Model>();
-			ASTOp::FullLink( cloned, model->GetPrivate()->m_program, &LinkerOptions);
-			OP::ADDRESS at = cloned->linkedAddress;
+			ASTOp::FullLink(SourceCloned, model->GetPrivate()->m_program, &LinkerOptions);
+			OP::ADDRESS at = SourceCloned->linkedAddress;
 
 			FProgram::FState state;
 			state.m_root = at;
 			model->GetPrivate()->m_program.m_states.Add(state);
 
-			ParametersPtr localParams = Model::NewParameters(model);
+			Ptr<Parameters> LocalParams = Model::NewParameters(model);
 			pSystem->GetPrivate()->BeginBuild( model );
 
 			// Calculate the value and replace this op by a constant
@@ -498,14 +523,14 @@ namespace mu
 			{
 				MUTABLE_CPUPROFILER_SCOPE(ConstantMesh);
 
-				mu::Ptr<const Mesh> pMesh = pSystem->GetPrivate()->BuildMesh( model, localParams.get(), at );
+				mu::Ptr<const Mesh> pMesh = pSystem->GetPrivate()->BuildMesh( model, LocalParams.get(), at );
 
 				if (pMesh)
 				{
 					mu::Ptr<ASTOpConstantResource> constantOp = new ASTOpConstantResource();
 					constantOp->type = OP_TYPE::ME_CONSTANT;
 					constantOp->SetValue( pMesh, bUseDiskCache );
-					m_result = constantOp;
+					Result = constantOp;
 				  }
 				break;
 			}
@@ -514,14 +539,14 @@ namespace mu
 			{
 				MUTABLE_CPUPROFILER_SCOPE(ConstantImage);
 
-				mu::Ptr<const Image> pImage = pSystem->GetPrivate()->BuildImage( model, localParams.get(), at, 0, 0 );
+				mu::Ptr<const Image> pImage = pSystem->GetPrivate()->BuildImage( model, LocalParams.get(), at, 0, 0 );
 
 				if (pImage)
 				{
 					mu::Ptr<ASTOpConstantResource> constantOp = new ASTOpConstantResource();
 					constantOp->type = OP_TYPE::IM_CONSTANT;
 					constantOp->SetValue( pImage, bUseDiskCache );
-					m_result = constantOp;
+					Result = constantOp;
 				}
 				break;
 			}
@@ -530,14 +555,14 @@ namespace mu
 			{
 				MUTABLE_CPUPROFILER_SCOPE(ConstantLayout);
 
-				mu::Ptr<const Layout> pLayout = pSystem->GetPrivate()->BuildLayout( model, localParams.get(), at );
+				mu::Ptr<const Layout> pLayout = pSystem->GetPrivate()->BuildLayout( model, LocalParams.get(), at );
 
 				if (pLayout)
 				{
 					mu::Ptr<ASTOpConstantResource> constantOp = new ASTOpConstantResource();
 					constantOp->type = OP_TYPE::LA_CONSTANT;
 					constantOp->SetValue( pLayout, bUseDiskCache );
-					m_result = constantOp;
+					Result = constantOp;
 				}
 				break;
 			}
@@ -546,12 +571,12 @@ namespace mu
 			{
 				MUTABLE_CPUPROFILER_SCOPE(ConstantBool);
 
-				bool value = pSystem->GetPrivate()->BuildBool( model, localParams.get(), at );
+				bool value = pSystem->GetPrivate()->BuildBool( model, LocalParams.get(), at );
 
 				{
 					mu::Ptr<ASTOpConstantBool> constantOp = new ASTOpConstantBool();
 					constantOp->value = value;
-					m_result = constantOp;
+					Result = constantOp;
 				}
 				break;
 			}
@@ -560,17 +585,17 @@ namespace mu
 			{
 				MUTABLE_CPUPROFILER_SCOPE(ConstantBool);
 
-				FVector4f Result(0, 0, 0, 0);
-				Result = pSystem->GetPrivate()->BuildColour( model, localParams.get(), at );
+				FVector4f ResultColor(0, 0, 0, 0);
+				ResultColor = pSystem->GetPrivate()->BuildColour( model, LocalParams.get(), at );
 
 				{
 					mu::Ptr<ASTOpFixed> constantOp = new ASTOpFixed();
 					constantOp->op.type = OP_TYPE::CO_CONSTANT;
-					constantOp->op.args.ColourConstant.value[0] = Result[0];
-					constantOp->op.args.ColourConstant.value[1] = Result[1];
-					constantOp->op.args.ColourConstant.value[2] = Result[2];
-					constantOp->op.args.ColourConstant.value[3] = Result[3];
-					m_result = constantOp;
+					constantOp->op.args.ColourConstant.value[0] = ResultColor[0];
+					constantOp->op.args.ColourConstant.value[1] = ResultColor[1];
+					constantOp->op.args.ColourConstant.value[2] = ResultColor[2];
+					constantOp->op.args.ColourConstant.value[3] = ResultColor[3];
+					Result = constantOp;
 				}
 				break;
 			}
@@ -588,227 +613,282 @@ namespace mu
 
 			pSystem->GetPrivate()->EndBuild();
 
+			SourceCloned = nullptr;
 		}
 
 		void Complete()
 		{
-			// This runs in the managing thread
-			ASTOp::Replace( m_source, m_result );
+			// This runs in the AST managing thread.
+			ASTOp::Replace(Source, Result);
+			Source = nullptr;
+			Result = nullptr;
 		}
 	};
 
 
 	//---------------------------------------------------------------------------------------------
-	bool ConstantGeneratorAST( const CompilerOptions::Private* options, Ptr<ASTOp>& root )
+	bool ConstantGeneratorAST( const CompilerOptions::Private* InOptions, Ptr<ASTOp>& Root, int32 Pass )
 	{
 		MUTABLE_CPUPROFILER_SCOPE(ConstantGenerator);
 
-		bool modified = false;
+		UE::Tasks::FPipe ASTPipe(TEXT("ASTPipe"));
 
 		// don't do this if constant optimization has been disabled, usually for debugging.
-		if (!options->OptimisationOptions.bConstReduction)
+		if (!InOptions->OptimisationOptions.bConstReduction)
 		{
 			return false;
 		}
 
-
-		// Calculate constant-subtree and special-op flags
+		// Gather the roots of all constant operations
+		struct FConstantSubgraph
 		{
-			MUTABLE_CPUPROFILER_SCOPE(ConstantGenerator_CalculateFlags);
+			Ptr<ASTOp> Root;
+			UE::Tasks::FTaskEvent CompletedEvent;
+		};
+		TArray< FConstantSubgraph > ConstantSubgraphs;
+		ConstantSubgraphs.Reserve(256);
+		{
+			MUTABLE_CPUPROFILER_SCOPE(ConstantGenerator_GenerateTasks);
 
-			ASTOp::Traverse_BottomUp_Unique( root, [&](Ptr<ASTOp>& n)
-			{
-				bool constantSubtree = true;
-				switch (n->GetOpType())
+			ASTOp::Traverse_BottomUp_Unique(Root,
+				[&ConstantSubgraphs, Pass]
+				(Ptr<ASTOp>& SubgraphRoot)
 				{
-				case OP_TYPE::BO_PARAMETER:
-				case OP_TYPE::NU_PARAMETER:
-				case OP_TYPE::SC_PARAMETER:
-				case OP_TYPE::CO_PARAMETER:
-				case OP_TYPE::PR_PARAMETER:
-				case OP_TYPE::IM_PARAMETER:
-					constantSubtree = false;
-					break;
-				default:
-					// Propagate from children
-					n->ForEachChild( [&constantSubtree](ASTChild& c)
+					OP_TYPE SubgraphType = SubgraphRoot->GetOpType();
+
+					bool bIsConstantSubgraph = true;
+					switch (SubgraphType)
 					{
-						if (c)
-						{
-							constantSubtree = constantSubtree && c->m_constantSubtree;
-						}
-					});
-					break;
-				}
-				n->m_constantSubtree = constantSubtree;
-
-				// We avoid generating constants for these operations, to avoid the memory
-				// explosion.
-				// TODO: Make compiler options for some of them
-				// TODO: Some of them are worth if the code below them is unique.
-				bool hasSpecialOpInSubtree = false;
-				switch (n->GetOpType())
-				{
-				case OP_TYPE::IM_BLANKLAYOUT:
-				case OP_TYPE::IM_REFERENCE:
-				case OP_TYPE::IM_COMPOSE:
-				//case OP_TYPE::IM_RASTERMESH:            // TODO review this one
-				case OP_TYPE::ME_MERGE:
-				case OP_TYPE::ME_CLIPWITHMESH:
-				case OP_TYPE::ME_CLIPMORPHPLANE:
-				case OP_TYPE::ME_APPLYPOSE:
-				case OP_TYPE::ME_REMOVEMASK:
-				case OP_TYPE::ME_ADDTAGS:
-				case OP_TYPE::IM_PLAINCOLOUR:
-					hasSpecialOpInSubtree = true;
-					break;
+					case OP_TYPE::BO_PARAMETER:
+					case OP_TYPE::NU_PARAMETER:
+					case OP_TYPE::SC_PARAMETER:
+					case OP_TYPE::CO_PARAMETER:
+					case OP_TYPE::PR_PARAMETER:
+					case OP_TYPE::IM_PARAMETER:
+						bIsConstantSubgraph = false;
+						break;
+					default:
+						// Propagate from children
+						SubgraphRoot->ForEachChild([&bIsConstantSubgraph](ASTChild& c)
+							{
+								if (c)
+								{
+									bIsConstantSubgraph = bIsConstantSubgraph && c->bIsConstantSubgraph;
+								}
+							});
+						break;
+					}
+					SubgraphRoot->bIsConstantSubgraph = bIsConstantSubgraph;
 
-				default:
-					// Propagate from children
-					n->ForEachChild( [&](ASTChild& c)
+					// We avoid generating constants for these operations, to avoid the memory explosion.
+					// TODO: Make compiler options for some of them
+					// TODO: Some of them are worth if the code below them is unique.
+					bool bHasSpecialOpInSubgraph = false;
+					switch (SubgraphType)
 					{
-						if (c)
+					case OP_TYPE::IM_BLANKLAYOUT:
+					case OP_TYPE::IM_COMPOSE:
+					case OP_TYPE::ME_MERGE:
+					case OP_TYPE::ME_CLIPWITHMESH:
+					case OP_TYPE::ME_CLIPMORPHPLANE:
+					case OP_TYPE::ME_APPLYPOSE:
+					case OP_TYPE::ME_REMOVEMASK:
+					case OP_TYPE::ME_ADDTAGS:
+					case OP_TYPE::IM_PLAINCOLOUR:
+						bHasSpecialOpInSubgraph = true;
+						break;
+
+					case OP_TYPE::IM_RASTERMESH:
+					{
+						const ASTOpImageRasterMesh* Raster = static_cast<const ASTOpImageRasterMesh*>(SubgraphRoot.get());
+						// If this operation is only rastering the mesh UVs, reduce it to constant. Otherwise avoid reducing it
+						// for the case of a constant projector of a large set of possible images. We don't want to generate all the
+						// projected version of the images beforehand. TODO: Make it a comptile-time option?
+						bHasSpecialOpInSubgraph = Raster->image.child().get() != nullptr;
+						break;
+					}
+
+					case OP_TYPE::IM_REFERENCE:
+						// If we are in a reference-resolution optimization phase, then the ops are not special.
+						if (Pass < 2)
 						{
-							hasSpecialOpInSubtree = hasSpecialOpInSubtree || c->m_hasSpecialOpInSubtree;
+							bHasSpecialOpInSubgraph = true;
 						}
-					});
-					break;
-				}
-				n->m_hasSpecialOpInSubtree = hasSpecialOpInSubtree;
-			});
-
-		}
+						else
+						{
+							const ASTOpReferenceResource* Typed = static_cast<const ASTOpReferenceResource*>(SubgraphRoot.get());
+							bHasSpecialOpInSubgraph = !Typed->bForceLoad;
+						}
+						break;
 
-		TArray< Ptr<ASTOp> > roots;
-		roots.Add(root);
+					default:
+						// Propagate from children
+						SubgraphRoot->ForEachChild([&](ASTChild& c)
+							{
+								if (c)
+								{
+									bHasSpecialOpInSubgraph = bHasSpecialOpInSubgraph || c->bHasSpecialOpInSubgraph;
+								}
+							});
+						break;
+					}
+					SubgraphRoot->bHasSpecialOpInSubgraph = bHasSpecialOpInSubgraph;
 
-		// Generate constant operations
-		ASTOp::Traverse_TopDown_Unique_Imprecise( roots, [&](Ptr<ASTOp>& n)
-		{
-			bool recurse = true;
+					bool bIsDataTypeThanCanTurnIntoConst = false;
+					DATATYPE dtype = GetOpDataType(SubgraphType);
+					switch (dtype)
+					{
+					case DT_MESH:
+					case DT_IMAGE:
+					case DT_LAYOUT:
+					case DT_BOOL:
+					case DT_COLOUR:
+						bIsDataTypeThanCanTurnIntoConst = true;
+						break;
+					default:
+						break;
+					}
 
-			OP_TYPE type = n->GetOpType();
-			DATATYPE dtype = GetOpDataType(type);
+					// See if it is worth generating this as constant
+					// ---------------------------------------------
+					if (SubgraphRoot->bIsConstantSubgraph
+						&& !SubgraphRoot->bHasSpecialOpInSubgraph
+						&& !SubgraphRoot->IsConstantOp()
+						&& bIsDataTypeThanCanTurnIntoConst
+						)
+					{
+						ConstantSubgraphs.Add({ SubgraphRoot, UE::Tasks::FTaskEvent(TEXT("MutableConstantSubgraph")) });
+					}
+				});
+		}
 
-			bool constant = false;
-			if (dtype!=DT_INSTANCE)
+		// Launch the tasks. Dot it from a task in the ASTPipe to make sure we "lock" modifications of the AST.
+		UE::Tasks::FTaskEvent EverythingComplete(TEXT("ConstantGenerator_EverythingComplete"));
+		UE::Tasks::FTask LaunchTask = ASTPipe.Launch(TEXT("ConstantGeneratorLaunchTasks"),
+			[&ConstantSubgraphs, &ASTPipe, &EverythingComplete, Pass, InOptions]()
 			{
-				constant = n->m_constantSubtree;
-			}
+				MUTABLE_CPUPROFILER_SCOPE(ConstantGenerator_LaunchTasks);
 
-			// See if it is a case of instructions we want to avoid, but with special parameters that make
-			// them ok
-			bool specialCase = false;
-			if (constant)
-			{
-				//const OP& op = program.m_code[at];
-				// TODO
-	//                if (n->GetOpType()==OP_TYPE::IM_COMPOSE)
-	//                {
-	//                    // Get the layout
-	//                    // In case we fill all of the target image anyway
-	//                    ParametersPtr pParams = m_pModel->NewParameters();
-
-	//                    const auto& args = op.args.ImageCompose;
-	//                    LayoutPtrConst pLayout = m_pSystem->GetPrivate()->BuildLayout
-	//                        ( m_pModel.get(), pParams.get(), args.layout );
-
-	//                    if ( pLayout &&
-	//                         pLayout->GetPrivate()->IsSingleBlockAndFull() )
-	//                    {
-	//                        specialCase = true;
-	//                    }
-	//                }
-
-	//                else
-				if (n->GetOpType()==OP_TYPE::IM_RASTERMESH)
+				FImageOperator ImOp = FImageOperator::GetDefault(InOptions->ImageFormatFunc);
+
+				// Traverse list of contants to generate in reverse order, for a more top-down approach that is better for memory usage.
+				for (int32 Index = ConstantSubgraphs.Num() - 1; Index >= 0; --Index)
 				{
-					const ASTOpImageRasterMesh* Raster = static_cast<const ASTOpImageRasterMesh*>(n.get());
-					if ( !Raster->image )
+					Ptr<ASTOp> SubgraphRoot = ConstantSubgraphs[Index].Root;
+					UE::Tasks::FTaskEvent& SubgraphCompletionEvent = ConstantSubgraphs[Index].CompletedEvent;
+
+					// Referenced images are resolved in its own task to prevent requesting them twice (and loading them twice)
+					bool bIsCompileTimeReferenceImage = (SubgraphRoot->GetOpType() == OP_TYPE::IM_REFERENCE);
+
+					// Launch the task with its dependencies
+					if (bIsCompileTimeReferenceImage)
 					{
-						specialCase = true;
+						// Instead of generating the constant we resolve the reference, which also replaces the ASTOp.
+						TSharedPtr< Ptr<Image> > ResolveImage = MakeShared<Ptr<Image>>();
+
+						const ASTOpReferenceResource* Typed = static_cast<const ASTOpReferenceResource*>(SubgraphRoot.get());
+						UE::Tasks::FTaskEvent ReferenceCompletionEvent = InOptions->OptimisationOptions.ReferencedResourceProvider(Typed->ID, ResolveImage);
+
+						UE::Tasks::FTask CompleteTask = ASTPipe.Launch(TEXT("MutableResolveComplete"),
+							[SubgraphRoot, InOptions, ResolveImage]()
+							{
+								Ptr<ASTOpConstantResource> ConstantOp = new ASTOpConstantResource;
+								ConstantOp->type = OP_TYPE::IM_CONSTANT;
+								ConstantOp->SetValue(ResolveImage->get(), InOptions->OptimisationOptions.bUseDiskCache);
+								ASTOp::Replace(SubgraphRoot, ConstantOp);
+							},
+							ReferenceCompletionEvent,
+							LowLevelTasks::ETaskPriority::BackgroundNormal);
+
+						SubgraphCompletionEvent.AddPrerequisites(CompleteTask);
 					}
-				}
-			}
 
-			// See if the subtree contains operations that we don't want to collapse even if
-			// they are constant
-			if ( constant && !specialCase)
-			{
-				if ( n->m_hasSpecialOpInSubtree )
-				{
-					constant = false;
-				}
-			}
+					else
+					{
+						// Scan for requisites
+						TArray< UE::Tasks::FTask, TInlineAllocator<8> > Requisites;
+						TArray< Ptr<ASTOp> > ScanRoots;
+						ScanRoots.Add(SubgraphRoot);
+						ASTOp::Traverse_TopDown_Unique_Imprecise(ScanRoots, [&SubgraphRoot, &Requisites, &ConstantSubgraphs](Ptr<ASTOp>& ChildNode)
+							{
+								bool bRecurse = true;
 
-			// If children are constant
-			if ( constant
-				 && type!=OP_TYPE::BO_CONSTANT
-				 && type!=OP_TYPE::NU_CONSTANT
-				 && type!=OP_TYPE::SC_CONSTANT
-				 && type!=OP_TYPE::CO_CONSTANT
-				 && type!=OP_TYPE::IM_CONSTANT
-				 && type!=OP_TYPE::ME_CONSTANT
-				 && type!=OP_TYPE::LA_CONSTANT
-				 && type!=OP_TYPE::PR_CONSTANT
-				 && type!=OP_TYPE::NONE
-				)
-			{
-				switch( dtype )
-				{
-				case DT_MESH:
-				case DT_IMAGE:
-				case DT_LAYOUT:
-				case DT_BOOL:
-				case DT_COLOUR:
-				{
-					recurse = false;
-					modified = true;
+								// Subgraph root?
+								if (SubgraphRoot == ChildNode)
+								{
+									return bRecurse;
+								}
 
-					ConstantTask* constantTask = new ConstantTask;
-					constantTask->bUseDiskCache = options->OptimisationOptions.bUseDiskCache;
-					constantTask->m_source = n;
-					constantTask->ImageCompressionQuality = options->ImageCompressionQuality;
+								FConstantSubgraph* DependencyFound = ConstantSubgraphs.FindByPredicate([&ChildNode](const FConstantSubgraph& Candidate) { return Candidate.Root == ChildNode; });
+								if (DependencyFound)
+								{
+									bRecurse = false;
+									Requisites.Add(DependencyFound->CompletedEvent);
+								}
 
-					FImageOperator ImOp = FImageOperator::GetDefault(options->ImageFormatFunc);
-					constantTask->Run(ImOp);
-					constantTask->Complete();
-					delete constantTask;
+								return bRecurse;
+							});
 
-					break;
-				}
+						TUniquePtr<ConstantTask> Task(new ConstantTask(SubgraphRoot, InOptions, Pass));
+						ConstantTask* TaskPtr = Task.Get();
 
-				case DT_INT:
-				case DT_SCALAR:
-				case DT_STRING:
-				case DT_PROJECTOR:
-					// TODO
-					break;
+						// Launch the preparation on the AST-modification pipe
+						UE::Tasks::FTask PrepareTask = ASTPipe.Launch(TEXT("MutableConstantPrepare"), [TaskPtr]()
+							{
+								TaskPtr->Prepare();
+							},
+							Requisites,
+							LowLevelTasks::ETaskPriority::BackgroundHigh);
 
-				case DT_INSTANCE:
-					// nothing to do
-					break;
+						// Launch constant generation on any thread
+						UE::Tasks::FTask RunTask = UE::Tasks::Launch(TEXT("MutableConstantGeneration"), [TaskPtr, ImOp]()
+							{
+								TaskPtr->Run(ImOp);
+							},
+							PrepareTask,
+							LowLevelTasks::ETaskPriority::BackgroundHigh );
 
-				default:
-					break;
+						// Launch the completion on the AST-modification pipe
+						UE::Tasks::FTask CompleteTask = ASTPipe.Launch(TEXT("MutableConstantComplete"), [TaskPtr = MoveTemp(Task)]()
+							{
+								TaskPtr->Complete();
+							},
+							RunTask,
+							LowLevelTasks::ETaskPriority::BackgroundHigh);
+
+						SubgraphCompletionEvent.AddPrerequisites(CompleteTask);
+					}
+
+					ConstantSubgraphs[Index].Root = nullptr;
+					SubgraphCompletionEvent.Trigger();
+					EverythingComplete.AddPrerequisites(SubgraphCompletionEvent);
 				}
 
-			}
+			});
 
-			return recurse;
-		});
+		EverythingComplete.AddPrerequisites(LaunchTask);
+
+		// Wait for pending tasks
+		{
+			MUTABLE_CPUPROFILER_SCOPE(ConstantGenerator_WaitPending);
+			EverythingComplete.Trigger();
 
-		return modified;
+			// Wait without executing tasks, since this is an independent thread.
+			EverythingComplete.Wait();
+		}
+
+		bool bSomethingModified = ConstantSubgraphs.Num() > 0;
+		return bSomethingModified;
 	}
 
 
 	//-------------------------------------------------------------------------------------------------
 	//-------------------------------------------------------------------------------------------------
 	//-------------------------------------------------------------------------------------------------
-	CodeOptimiser::CodeOptimiser(Ptr<CompilerOptions> options, TArray<FStateCompilationData>& states )
-		: m_states( states )
+	CodeOptimiser::CodeOptimiser(Ptr<CompilerOptions> InOptions, TArray<FStateCompilationData>& InStates )
+		: m_states( InStates )
 	{
-		m_options = options;
+		m_options = InOptions;
 	}
 
 
@@ -870,13 +950,13 @@ namespace mu
 		ASTOp::LogHistogram(roots);
 
 		// Generate constants
-		for ( auto& r: roots )
+		for ( Ptr<ASTOp>& Root: roots )
 		{
 			UE_LOG(LogMutableCore, Verbose, TEXT("(int) %s : %ld"), TEXT("ast size"), int64(ASTOp::CountNodes(roots)));
 			UE_LOG(LogMutableCore, Verbose, TEXT(" - constant generator"));
 
 			// Constant subtree generation
-			modified = ConstantGeneratorAST( m_options->GetPrivate(), r );
+			modified = ConstantGeneratorAST( m_options->GetPrivate(), Root, Pass );
 		}
 
 		UE_LOG(LogMutableCore, Verbose, TEXT("(int) %s : %ld"), TEXT("ast size"), int64(ASTOp::CountNodes(roots)));
@@ -1238,7 +1318,19 @@ namespace mu
 			// Constant resolution stage: resolve referenced assets.
 			{
 				MUTABLE_CPUPROFILER_SCOPE(ReferenceResolution);
-				FullOptimiseAST(roots, 2);
+				
+				constexpr int32 Pass = 2;
+
+				//FullOptimiseAST(roots, 2);
+
+				// Generate constants
+				for (Ptr<ASTOp>& Root : roots)
+				{
+					// Constant subtree generation
+					modified = ConstantGeneratorAST(m_options->GetPrivate(), Root, Pass);
+				}
+
+				DuplicatedDataRemoverAST(roots);
 			}
 
 			// Main optimisation stage again for data-aware optimizations
@@ -1310,8 +1402,10 @@ namespace mu
 
 			for ( int32 s=0;  s<m_states.Num(); ++s )
 			{
+				constexpr int32 Pass = 1;
+
 				UE_LOG(LogMutableCore, Verbose, TEXT(" - constant generator"));
-				ConstantGeneratorAST( m_options->GetPrivate(), m_states[s].root );
+				ConstantGeneratorAST( m_options->GetPrivate(), m_states[s].root, Pass );
 				//AXE_INT_VALUE("Mutable", Verbose, "ast size", (int64_t)ASTOp::CountNodes(roots));
 			}
 
diff --git a/Engine/Plugins/Experimental/Mutable/Source/MutableTools/Private/MuT/CodeOptimiser.h b/Engine/Plugins/Experimental/Mutable/Source/MutableTools/Private/MuT/CodeOptimiser.h
--- a/Engine/Plugins/Experimental/Mutable/Source/MutableTools/Private/MuT/CodeOptimiser.h
+++ b/Engine/Plugins/Experimental/Mutable/Source/MutableTools/Private/MuT/CodeOptimiser.h
@@ -98,7 +98,7 @@ namespace mu
     //! ConstantGenerator replaces constant subtrees of operations with an equivalent single
 	//! constant value operation. 
     //---------------------------------------------------------------------------------------------
-    extern bool ConstantGeneratorAST( const CompilerOptions::Private* options, Ptr<ASTOp>& root );
+    extern bool ConstantGeneratorAST( const CompilerOptions::Private* options, Ptr<ASTOp>& root, int32 Pass );
 
     //---------------------------------------------------------------------------------------------
     //! \TODO: shapes, projectors, others? but not switches (they must be unique)
diff --git a/Engine/Plugins/Experimental/Mutable/Source/MutableTools/Private/MuT/CodeOptimiser_Parameters.cpp b/Engine/Plugins/Experimental/Mutable/Source/MutableTools/Private/MuT/CodeOptimiser_Parameters.cpp
--- a/Engine/Plugins/Experimental/Mutable/Source/MutableTools/Private/MuT/CodeOptimiser_Parameters.cpp
+++ b/Engine/Plugins/Experimental/Mutable/Source/MutableTools/Private/MuT/CodeOptimiser_Parameters.cpp
@@ -2132,7 +2132,8 @@ namespace mu
         {
             MUTABLE_CPUPROFILER_SCOPE(Reoptimise);
             bool modified = true;
-            int numIterations = 0;
+            int32 numIterations = 0;
+			int32 Pass = 1;
             while (modified && (!m_optimizeIterationsMax || m_optimizeIterationsLeft>0 || !numIterations ))
             {
                 TArray<Ptr<ASTOp>> roots;
@@ -2150,7 +2151,7 @@ namespace mu
 
                 UE_LOG(LogMutableCore, Verbose, TEXT(" - semantic optimiser"));
                 modified |=
-                    SemanticOptimiserAST( roots, m_options->GetPrivate()->OptimisationOptions, 1 );
+                    SemanticOptimiserAST( roots, m_options->GetPrivate()->OptimisationOptions, Pass );
 				UE_LOG(LogMutableCore, Verbose, TEXT("(int) %s : %ld"), TEXT("ast size"), int64(ASTOp::CountNodes(roots)));
 
                 // Image size operations are treated separately
@@ -2162,8 +2163,8 @@ namespace mu
             for(FStateCompilationData& s:m_states)
             {
                 UE_LOG(LogMutableCore, Verbose, TEXT(" - constant optimiser"));
-                ConstantGeneratorAST( m_options->GetPrivate(), s.root );
-            }
+				modified = ConstantGeneratorAST( m_options->GetPrivate(), s.root, Pass );
+			}
 
             TArray<Ptr<ASTOp>> roots;
             for(const FStateCompilationData& s:m_states)
diff --git a/Engine/Plugins/Experimental/Mutable/Source/MutableTools/Public/MuT/Compiler.h b/Engine/Plugins/Experimental/Mutable/Source/MutableTools/Public/MuT/Compiler.h
--- a/Engine/Plugins/Experimental/Mutable/Source/MutableTools/Public/MuT/Compiler.h
+++ b/Engine/Plugins/Experimental/Mutable/Source/MutableTools/Public/MuT/Compiler.h
@@ -7,6 +7,7 @@
 #include "MuR/Image.h"
 #include "Templates/SharedPointer.h"
 #include "HAL/PlatformMath.h"
+#include "Tasks/Task.h"
 
 namespace mu
 {
@@ -30,7 +31,7 @@ namespace mu
     class Node;
     class NodeTransformedObject;
 
-	typedef TFunction<Ptr<Image>(int32)> FReferencedResourceFunc;
+	typedef TFunction<UE::Tasks::FTaskEvent(int32, TSharedPtr<Ptr<Image>>)> FReferencedResourceFunc;
 
     //! \brief Options used to compile the models with a compiler.
     class MUTABLETOOLS_API CompilerOptions : public RefCounted
