[Mutable] Fix "Parameter not found" error when selecting group node parameters.

Parameters can be optimized due to groups nodes having the all options group type.

#rnx

diff --git a/Engine/Plugins/Experimental/Mutable/Source/CustomizableObjectEditor/Private/MuCOE/CustomizableInstanceDetails.cpp b/Engine/Plugins/Experimental/Mutable/Source/CustomizableObjectEditor/Private/MuCOE/CustomizableInstanceDetails.cpp
--- a/Engine/Plugins/Experimental/Mutable/Source/CustomizableObjectEditor/Private/MuCOE/CustomizableInstanceDetails.cpp
+++ b/Engine/Plugins/Experimental/Mutable/Source/CustomizableObjectEditor/Private/MuCOE/CustomizableInstanceDetails.cpp
@@ -744,7 +744,7 @@ bool FCustomizableInstanceDetails::IsVisible(int32 ParamIndexInObject)
 	{
 		FString* Value = UIData.ParamUIMetadata.ExtraInformation.Find(FString("__DisplayWhenParentValueEquals"));
 
-		if (Value && CustomInstance->GetIntParameterSelectedOption(*ParentName) != *Value)
+		if (Value && CustomizableObject->FindParameter(*ParentName) != INDEX_NONE && CustomInstance->GetIntParameterSelectedOption(*ParentName) != *Value)
 		{
 			return false;
 		}
