[mutable] Fix compilation in some platforms.
#rnx

diff --git a/Engine/Plugins/Experimental/Mutable/Source/CustomizableObjectEditor/Private/MuCOE/CustomizableObjectCompileRunnable.cpp b/Engine/Plugins/Experimental/Mutable/Source/CustomizableObjectEditor/Private/MuCOE/CustomizableObjectCompileRunnable.cpp
--- a/Engine/Plugins/Experimental/Mutable/Source/CustomizableObjectEditor/Private/MuCOE/CustomizableObjectCompileRunnable.cpp
+++ b/Engine/Plugins/Experimental/Mutable/Source/CustomizableObjectEditor/Private/MuCOE/CustomizableObjectCompileRunnable.cpp
@@ -132,7 +132,7 @@ uint32 FCustomizableObjectCompileRunnable::Run()
 			}
 			else
 			{
-				PendingResourceReferenceRequests.Enqueue(FReferenceResourceRequest(ID, ResolvedImage, MakeShared<UE::Tasks::FTaskEvent>(CompletionEvent)));
+				PendingResourceReferenceRequests.Enqueue({ ID, ResolvedImage, MakeShared<UE::Tasks::FTaskEvent>(CompletionEvent) });
 			}
 
 			return CompletionEvent;
