[Mutable] Fix first compilation not showing the generated Instance in the Viewport.

#rnx

diff --git a/Engine/Plugins/Experimental/Mutable/Source/CustomizableObject/Internal/MuCO/CustomizableObjectPrivate.h b/Engine/Plugins/Experimental/Mutable/Source/CustomizableObject/Internal/MuCO/CustomizableObjectPrivate.h
--- a/Engine/Plugins/Experimental/Mutable/Source/CustomizableObject/Internal/MuCO/CustomizableObjectPrivate.h
+++ b/Engine/Plugins/Experimental/Mutable/Source/CustomizableObject/Internal/MuCO/CustomizableObjectPrivate.h
@@ -61,7 +61,7 @@ struct FCustomizableObjectStatusTypes
 		Count,
 	};
 	
-	static constexpr EState StartState = EState::Loading;
+	static constexpr EState StartState = EState::NoModel;
 
 	static constexpr bool ValidTransitions[3][3] =
 	{
@@ -69,7 +69,7 @@ struct FCustomizableObjectStatusTypes
 		// Loading, ModelLoaded, NoModel // FROM
 		{false,   true,        true},  // Loading
 		{false,   true,        true},  // ModelLoaded
-		{false,   true,        true},  // NoModel
+		{true,    true,        true},  // NoModel
 	};
 };
 
diff --git a/Engine/Plugins/Experimental/Mutable/Source/CustomizableObject/Private/MuCO/CustomizableObject.cpp b/Engine/Plugins/Experimental/Mutable/Source/CustomizableObject/Private/MuCO/CustomizableObject.cpp
--- a/Engine/Plugins/Experimental/Mutable/Source/CustomizableObject/Private/MuCO/CustomizableObject.cpp
+++ b/Engine/Plugins/Experimental/Mutable/Source/CustomizableObject/Private/MuCO/CustomizableObject.cpp
@@ -222,6 +222,8 @@ void UCustomizableObject::PostLoad()
 
 	if (!IsRunningCookCommandlet())
 	{
+		GetPrivate()->Status.NextState(FCustomizableObjectStatusTypes::EState::Loading);
+
 		ITargetPlatformManagerModule& TargetPlatformManager = GetTargetPlatformManagerRef();
 		const ITargetPlatform* RunningPlatform = TargetPlatformManager.GetRunningTargetPlatform();
 
