[Mutable] Allow MuCOE Surface Edit Nodes to affect siblings.

Extend Material, Edit Material, Remove Mesh Blocks...

#rnx

diff --git a/Engine/Plugins/Experimental/Mutable/Source/CustomizableObjectEditor/Private/MuCOE/Nodes/CustomizableObjectNodeParentedMaterial.cpp b/Engine/Plugins/Experimental/Mutable/Source/CustomizableObjectEditor/Private/MuCOE/Nodes/CustomizableObjectNodeParentedMaterial.cpp
--- a/Engine/Plugins/Experimental/Mutable/Source/CustomizableObjectEditor/Private/MuCOE/Nodes/CustomizableObjectNodeParentedMaterial.cpp
+++ b/Engine/Plugins/Experimental/Mutable/Source/CustomizableObjectEditor/Private/MuCOE/Nodes/CustomizableObjectNodeParentedMaterial.cpp
@@ -35,7 +35,7 @@ TArray<UCustomizableObjectNodeMaterial*> FCustomizableObjectNodeParentedMaterial
 	ECustomizableObjectAutomaticLODStrategy LODStrategy = ECustomizableObjectAutomaticLODStrategy::Inherited;
 
 	// Iterate backwards, from the Root CO to the parent CO, to propagate the LODStrategy. 
-	for (int32 Index = ParentObjectNodes.Num() - 1; Index > 0; --Index) // Skip first Object node (Index == 0) since we do not support directly material siblings.
+	for (int32 Index = ParentObjectNodes.Num() - 1; Index > -1; --Index)
 	{
 		const UCustomizableObjectNodeObject* ParentObjectNode = ParentObjectNodes[Index];
 		
