[mutable] Remove some of the unnecessary standard library header includes from core modules.
#rnx

diff --git a/Engine/Plugins/Experimental/Mutable/Source/CustomizableObject/Classes/MuCO/CustomizableObject.h b/Engine/Plugins/Experimental/Mutable/Source/CustomizableObject/Classes/MuCO/CustomizableObject.h
--- a/Engine/Plugins/Experimental/Mutable/Source/CustomizableObject/Classes/MuCO/CustomizableObject.h
+++ b/Engine/Plugins/Experimental/Mutable/Source/CustomizableObject/Classes/MuCO/CustomizableObject.h
@@ -141,7 +141,9 @@ struct FCompilationOptions
 {
 	GENERATED_USTRUCT_BODY()
 	
-	// Enum to know what texture compression should be used
+	/** Enum to know what texture compression should be used. This compression is used on ly in manual compiles in editor. 
+	 *  When packaging, ECustomizableObjectTextureCompression::HighQuality is always used.
+	 */
 	UPROPERTY()
 	ECustomizableObjectTextureCompression TextureCompression = ECustomizableObjectTextureCompression::Fast;
 
diff --git a/Engine/Plugins/Experimental/Mutable/Source/CustomizableObjectEditor/Private/MuCOE/Nodes/CustomizableObjectNodeStaticMesh.cpp b/Engine/Plugins/Experimental/Mutable/Source/CustomizableObjectEditor/Private/MuCOE/Nodes/CustomizableObjectNodeStaticMesh.cpp
--- a/Engine/Plugins/Experimental/Mutable/Source/CustomizableObjectEditor/Private/MuCOE/Nodes/CustomizableObjectNodeStaticMesh.cpp
+++ b/Engine/Plugins/Experimental/Mutable/Source/CustomizableObjectEditor/Private/MuCOE/Nodes/CustomizableObjectNodeStaticMesh.cpp
@@ -15,8 +15,6 @@
 #include "MuCOE/RemapPins/CustomizableObjectNodeRemapPinsByNameDefaultPin.h"
 #include "MuCOE/UnrealEditorPortabilityHelpers.h"
 #include "Widgets/Input/SCheckBox.h"
-#include <memory>
-#include <stdexcept>
 
 class UCustomizableObjectNodeRemapPinsByName;
 class UObject;
diff --git a/Engine/Plugins/Experimental/Mutable/Source/MutableRuntime/Private/MuR/Mesh.cpp b/Engine/Plugins/Experimental/Mutable/Source/MutableRuntime/Private/MuR/Mesh.cpp
--- a/Engine/Plugins/Experimental/Mutable/Source/MutableRuntime/Private/MuR/Mesh.cpp
+++ b/Engine/Plugins/Experimental/Mutable/Source/MutableRuntime/Private/MuR/Mesh.cpp
@@ -10,8 +10,6 @@
 #include "MuR/MeshPrivate.h"
 #include "MuR/MutableTrace.h"
 
-#include <limits>
-
 namespace mu
 {
 
@@ -926,8 +924,8 @@ void Mesh::GetVertexMap
 #define MUTABLE_NUM_BUCKETS 256
 #define MUTABLE_BUCKET_CHANNEL 0
 
-    float rangeMin = std::numeric_limits<float>::max();
-    float rangeMax = -std::numeric_limits<float>::max();
+    float rangeMin = TNumericLimits<float>::Max();
+    float rangeMax = -TNumericLimits<float>::Max();
     MeshBufferIteratorConst< MBF_FLOAT32, float, 3 >  itop = itopBegin;
     for ( int ov=0; ov<otherVertexCount; ++ov )
     {
diff --git a/Engine/Plugins/Experimental/Mutable/Source/MutableRuntime/Private/MuR/Serialisation.cpp b/Engine/Plugins/Experimental/Mutable/Source/MutableRuntime/Private/MuR/Serialisation.cpp
--- a/Engine/Plugins/Experimental/Mutable/Source/MutableRuntime/Private/MuR/Serialisation.cpp
+++ b/Engine/Plugins/Experimental/Mutable/Source/MutableRuntime/Private/MuR/Serialisation.cpp
@@ -10,8 +10,6 @@
 #include "MuR/Image.h"
 #include "MuR/SerialisationPrivate.h"
 
-#include <limits>
-
 
 namespace mu
 {
@@ -100,11 +98,8 @@ namespace mu
         {
             check( m_pD->m_pos + size <= m_pD->m_size );
 
-            // This could happen in 32-bit platforms
-            check( size<std::numeric_limits<size_t>::max() );
-
-            const uint8_t* pSource = ((const uint8_t*)(m_pD->m_pBuffer))+m_pD->m_pos;
-            FMemory::Memcpy( pData, pSource, (size_t)size );
+            const uint8* pSource = ((const uint8*)(m_pD->m_pBuffer))+m_pD->m_pos;
+            FMemory::Memcpy( pData, pSource, (SIZE_T)size );
             m_pD->m_pos += size;
         }
     }
@@ -118,9 +113,7 @@ namespace mu
         m_pD = new Private();
         if (reserve)
         {
-            // This could happen in 32-bit platforms
-            check( reserve<std::numeric_limits<size_t>::max() );
-            m_pD->m_buffer.Reserve( reserve );
+             m_pD->m_buffer.Reserve( reserve );
         }
     }
 
@@ -139,9 +132,6 @@ namespace mu
     {
         if (size)
         {
-            // This could happen in 32-bit platforms
-            check( size<std::numeric_limits<size_t>::max() );
-
             uint64 pos = m_pD->m_buffer.Num();
             m_pD->m_buffer.SetNum( pos + size, EAllowShrinking::No );
 			FMemory::Memcpy( &m_pD->m_buffer[pos], pData, size );
@@ -305,7 +295,7 @@ namespace mu
             }
             else
             {
-                if ( (std::size_t)id < m_pD->m_proxyHistory.Num() )
+                if ( id < m_pD->m_proxyHistory.Num() )
                 {
                     p = m_pD->m_proxyHistory[id];
 
diff --git a/Engine/Plugins/Experimental/Mutable/Source/MutableTools/Private/MuT/AST.h b/Engine/Plugins/Experimental/Mutable/Source/MutableTools/Private/MuT/AST.h
--- a/Engine/Plugins/Experimental/Mutable/Source/MutableTools/Private/MuT/AST.h
+++ b/Engine/Plugins/Experimental/Mutable/Source/MutableTools/Private/MuT/AST.h
@@ -26,10 +26,8 @@
 
 #include <array>
 #include <atomic>
-#include <memory>
 #include <unordered_map>
 #include <unordered_set>
-#include <utility>
 #include <set>
 
 
diff --git a/Engine/Plugins/Experimental/Mutable/Source/MutableTools/Private/MuT/ASTOpImageMultiLayer.cpp b/Engine/Plugins/Experimental/Mutable/Source/MutableTools/Private/MuT/ASTOpImageMultiLayer.cpp
--- a/Engine/Plugins/Experimental/Mutable/Source/MutableTools/Private/MuT/ASTOpImageMultiLayer.cpp
+++ b/Engine/Plugins/Experimental/Mutable/Source/MutableTools/Private/MuT/ASTOpImageMultiLayer.cpp
@@ -9,9 +9,6 @@
 #include "MuR/Types.h"
 #include "MuT/StreamsPrivate.h"
 
-#include <memory>
-#include <utility>
-
 
 namespace mu
 {
diff --git a/Engine/Plugins/Experimental/Mutable/Source/MutableTools/Private/MuT/ASTOpImageSwizzle.cpp b/Engine/Plugins/Experimental/Mutable/Source/MutableTools/Private/MuT/ASTOpImageSwizzle.cpp
--- a/Engine/Plugins/Experimental/Mutable/Source/MutableTools/Private/MuT/ASTOpImageSwizzle.cpp
+++ b/Engine/Plugins/Experimental/Mutable/Source/MutableTools/Private/MuT/ASTOpImageSwizzle.cpp
@@ -23,9 +23,6 @@
 #include "MuT/ASTOpSwitch.h"
 #include "MuT/StreamsPrivate.h"
 
-#include <memory>
-#include <utility>
-
 namespace mu
 {
 
diff --git a/Engine/Plugins/Experimental/Mutable/Source/MutableTools/Private/MuT/ASTOpImageTransform.cpp b/Engine/Plugins/Experimental/Mutable/Source/MutableTools/Private/MuT/ASTOpImageTransform.cpp
--- a/Engine/Plugins/Experimental/Mutable/Source/MutableTools/Private/MuT/ASTOpImageTransform.cpp
+++ b/Engine/Plugins/Experimental/Mutable/Source/MutableTools/Private/MuT/ASTOpImageTransform.cpp
@@ -9,9 +9,6 @@
 #include "MuR/Types.h"
 #include "MuT/StreamsPrivate.h"
 
-#include <memory>
-#include <utility>
-
 
 namespace mu
 {
diff --git a/Engine/Plugins/Experimental/Mutable/Source/MutableTools/Private/MuT/ASTOpInstanceAdd.cpp b/Engine/Plugins/Experimental/Mutable/Source/MutableTools/Private/MuT/ASTOpInstanceAdd.cpp
--- a/Engine/Plugins/Experimental/Mutable/Source/MutableTools/Private/MuT/ASTOpInstanceAdd.cpp
+++ b/Engine/Plugins/Experimental/Mutable/Source/MutableTools/Private/MuT/ASTOpInstanceAdd.cpp
@@ -10,10 +10,6 @@
 #include "MuT/CodeOptimiser.h"
 #include "MuT/StreamsPrivate.h"
 
-#include <algorithm>
-#include <memory>
-#include <utility>
-
 
 namespace mu
 {
diff --git a/Engine/Plugins/Experimental/Mutable/Source/MutableTools/Private/MuT/ASTOpLayoutFromMesh.cpp b/Engine/Plugins/Experimental/Mutable/Source/MutableTools/Private/MuT/ASTOpLayoutFromMesh.cpp
--- a/Engine/Plugins/Experimental/Mutable/Source/MutableTools/Private/MuT/ASTOpLayoutFromMesh.cpp
+++ b/Engine/Plugins/Experimental/Mutable/Source/MutableTools/Private/MuT/ASTOpLayoutFromMesh.cpp
@@ -24,8 +24,6 @@
 #include "MuT/ASTOpLayoutMerge.h"
 #include "MuT/StreamsPrivate.h"
 
-#include <functional>
-
 
 namespace mu
 {
diff --git a/Engine/Plugins/Experimental/Mutable/Source/MutableTools/Private/MuT/ASTOpMeshApplyPose.cpp b/Engine/Plugins/Experimental/Mutable/Source/MutableTools/Private/MuT/ASTOpMeshApplyPose.cpp
--- a/Engine/Plugins/Experimental/Mutable/Source/MutableTools/Private/MuT/ASTOpMeshApplyPose.cpp
+++ b/Engine/Plugins/Experimental/Mutable/Source/MutableTools/Private/MuT/ASTOpMeshApplyPose.cpp
@@ -8,9 +8,6 @@
 #include "MuR/Types.h"
 #include "MuT/StreamsPrivate.h"
 
-#include <memory>
-#include <utility>
-
 
 namespace mu
 {
diff --git a/Engine/Plugins/Experimental/Mutable/Source/MutableTools/Private/MuT/ASTOpMeshApplyShape.cpp b/Engine/Plugins/Experimental/Mutable/Source/MutableTools/Private/MuT/ASTOpMeshApplyShape.cpp
--- a/Engine/Plugins/Experimental/Mutable/Source/MutableTools/Private/MuT/ASTOpMeshApplyShape.cpp
+++ b/Engine/Plugins/Experimental/Mutable/Source/MutableTools/Private/MuT/ASTOpMeshApplyShape.cpp
@@ -8,9 +8,6 @@
 #include "MuR/Types.h"
 #include "MuT/StreamsPrivate.h"
 
-#include <memory>
-#include <utility>
-
 
 namespace mu
 {
diff --git a/Engine/Plugins/Experimental/Mutable/Source/MutableTools/Private/MuT/ASTOpMeshBindShape.cpp b/Engine/Plugins/Experimental/Mutable/Source/MutableTools/Private/MuT/ASTOpMeshBindShape.cpp
--- a/Engine/Plugins/Experimental/Mutable/Source/MutableTools/Private/MuT/ASTOpMeshBindShape.cpp
+++ b/Engine/Plugins/Experimental/Mutable/Source/MutableTools/Private/MuT/ASTOpMeshBindShape.cpp
@@ -13,9 +13,6 @@
 #include "MuT/ASTOpMeshAddTags.h"
 #include "MuT/StreamsPrivate.h"
 
-#include <memory>
-#include <utility>
-
 
 namespace mu
 {
diff --git a/Engine/Plugins/Experimental/Mutable/Source/MutableTools/Private/MuT/ASTOpMeshClipDeform.cpp b/Engine/Plugins/Experimental/Mutable/Source/MutableTools/Private/MuT/ASTOpMeshClipDeform.cpp
--- a/Engine/Plugins/Experimental/Mutable/Source/MutableTools/Private/MuT/ASTOpMeshClipDeform.cpp
+++ b/Engine/Plugins/Experimental/Mutable/Source/MutableTools/Private/MuT/ASTOpMeshClipDeform.cpp
@@ -8,9 +8,6 @@
 #include "MuR/Types.h"
 #include "MuT/StreamsPrivate.h"
 
-#include <memory>
-#include <utility>
-
 
 namespace mu
 {
diff --git a/Engine/Plugins/Experimental/Mutable/Source/MutableTools/Private/MuT/ASTOpMeshDifference.cpp b/Engine/Plugins/Experimental/Mutable/Source/MutableTools/Private/MuT/ASTOpMeshDifference.cpp
--- a/Engine/Plugins/Experimental/Mutable/Source/MutableTools/Private/MuT/ASTOpMeshDifference.cpp
+++ b/Engine/Plugins/Experimental/Mutable/Source/MutableTools/Private/MuT/ASTOpMeshDifference.cpp
@@ -12,8 +12,6 @@
 #include "MuT/ASTOpMeshRemoveMask.h"
 #include "MuT/ASTOpMeshAddTags.h"
 
-#include <memory>
-
 
 namespace mu
 {
diff --git a/Engine/Plugins/Experimental/Mutable/Source/MutableTools/Private/MuT/ASTOpMeshFormat.cpp b/Engine/Plugins/Experimental/Mutable/Source/MutableTools/Private/MuT/ASTOpMeshFormat.cpp
--- a/Engine/Plugins/Experimental/Mutable/Source/MutableTools/Private/MuT/ASTOpMeshFormat.cpp
+++ b/Engine/Plugins/Experimental/Mutable/Source/MutableTools/Private/MuT/ASTOpMeshFormat.cpp
@@ -20,10 +20,6 @@
 #include "MuT/ASTOpSwitch.h"
 #include "MuT/StreamsPrivate.h"
 
-#include <memory>
-#include <utility>
-
-
 
 namespace mu
 {
diff --git a/Engine/Plugins/Experimental/Mutable/Source/MutableTools/Private/MuT/ASTOpMeshGeometryOperation.cpp b/Engine/Plugins/Experimental/Mutable/Source/MutableTools/Private/MuT/ASTOpMeshGeometryOperation.cpp
--- a/Engine/Plugins/Experimental/Mutable/Source/MutableTools/Private/MuT/ASTOpMeshGeometryOperation.cpp
+++ b/Engine/Plugins/Experimental/Mutable/Source/MutableTools/Private/MuT/ASTOpMeshGeometryOperation.cpp
@@ -7,9 +7,6 @@
 #include "MuR/RefCounted.h"
 #include "MuT/StreamsPrivate.h"
 
-#include <memory>
-#include <utility>
-
 namespace mu
 {
 
diff --git a/Engine/Plugins/Experimental/Mutable/Source/MutableTools/Private/MuT/ASTOpMeshMorphReshape.cpp b/Engine/Plugins/Experimental/Mutable/Source/MutableTools/Private/MuT/ASTOpMeshMorphReshape.cpp
--- a/Engine/Plugins/Experimental/Mutable/Source/MutableTools/Private/MuT/ASTOpMeshMorphReshape.cpp
+++ b/Engine/Plugins/Experimental/Mutable/Source/MutableTools/Private/MuT/ASTOpMeshMorphReshape.cpp
@@ -8,10 +8,6 @@
 #include "MuR/Types.h"
 #include "MuT/StreamsPrivate.h"
 
-#include <memory>
-#include <utility>
-
-
 namespace mu
 {
 
diff --git a/Engine/Plugins/Experimental/Mutable/Source/MutableTools/Private/MuT/ASTOpMeshOptimizeSkinning.cpp b/Engine/Plugins/Experimental/Mutable/Source/MutableTools/Private/MuT/ASTOpMeshOptimizeSkinning.cpp
--- a/Engine/Plugins/Experimental/Mutable/Source/MutableTools/Private/MuT/ASTOpMeshOptimizeSkinning.cpp
+++ b/Engine/Plugins/Experimental/Mutable/Source/MutableTools/Private/MuT/ASTOpMeshOptimizeSkinning.cpp
@@ -8,8 +8,6 @@
 #include "MuR/Types.h"
 #include "MuT/StreamsPrivate.h"
 
-#include <utility>
-
 
 namespace mu
 {
diff --git a/Engine/Plugins/Experimental/Mutable/Source/MutableTools/Private/MuT/ASTOpMeshTransform.cpp b/Engine/Plugins/Experimental/Mutable/Source/MutableTools/Private/MuT/ASTOpMeshTransform.cpp
--- a/Engine/Plugins/Experimental/Mutable/Source/MutableTools/Private/MuT/ASTOpMeshTransform.cpp
+++ b/Engine/Plugins/Experimental/Mutable/Source/MutableTools/Private/MuT/ASTOpMeshTransform.cpp
@@ -8,9 +8,6 @@
 #include "MuR/Types.h"
 #include "MuT/StreamsPrivate.h"
 
-#include <memory>
-#include <utility>
-
 
 namespace mu
 {
diff --git a/Engine/Plugins/Experimental/Mutable/Source/MutableTools/Private/MuT/ASTOpParameter.cpp b/Engine/Plugins/Experimental/Mutable/Source/MutableTools/Private/MuT/ASTOpParameter.cpp
--- a/Engine/Plugins/Experimental/Mutable/Source/MutableTools/Private/MuT/ASTOpParameter.cpp
+++ b/Engine/Plugins/Experimental/Mutable/Source/MutableTools/Private/MuT/ASTOpParameter.cpp
@@ -10,9 +10,6 @@
 #include "MuR/Types.h"
 #include "MuT/StreamsPrivate.h"
 
-#include <memory>
-#include <utility>
-
 
 namespace mu
 {
diff --git a/Engine/Plugins/Experimental/Mutable/Source/MutableTools/Private/MuT/ASTOpScalarCurve.cpp b/Engine/Plugins/Experimental/Mutable/Source/MutableTools/Private/MuT/ASTOpScalarCurve.cpp
--- a/Engine/Plugins/Experimental/Mutable/Source/MutableTools/Private/MuT/ASTOpScalarCurve.cpp
+++ b/Engine/Plugins/Experimental/Mutable/Source/MutableTools/Private/MuT/ASTOpScalarCurve.cpp
@@ -8,9 +8,6 @@
 #include "MuR/Types.h"
 #include "MuT/StreamsPrivate.h"
 
-#include <memory>
-#include <utility>
-
 
 namespace mu
 {
diff --git a/Engine/Plugins/Experimental/Mutable/Source/MutableTools/Private/MuT/ASTOpSwitch.cpp b/Engine/Plugins/Experimental/Mutable/Source/MutableTools/Private/MuT/ASTOpSwitch.cpp
--- a/Engine/Plugins/Experimental/Mutable/Source/MutableTools/Private/MuT/ASTOpSwitch.cpp
+++ b/Engine/Plugins/Experimental/Mutable/Source/MutableTools/Private/MuT/ASTOpSwitch.cpp
@@ -12,9 +12,6 @@
 #include "MuR/Types.h"
 #include "MuT/StreamsPrivate.h"
 
-#include <memory>
-#include <utility>
-
 
 namespace mu
 {
diff --git a/Engine/Plugins/Experimental/Mutable/Source/MutableTools/Private/MuT/CodeOptimiser_Logic.cpp b/Engine/Plugins/Experimental/Mutable/Source/MutableTools/Private/MuT/CodeOptimiser_Logic.cpp
--- a/Engine/Plugins/Experimental/Mutable/Source/MutableTools/Private/MuT/CodeOptimiser_Logic.cpp
+++ b/Engine/Plugins/Experimental/Mutable/Source/MutableTools/Private/MuT/CodeOptimiser_Logic.cpp
@@ -11,10 +11,6 @@
 #include "MuT/ASTOpSwitch.h"
 #include "MuT/CodeOptimiser.h"
 
-#include <memory>
-#include <utility>
-
-
 
 namespace mu
 {
diff --git a/Engine/Plugins/Experimental/Mutable/Source/MutableTools/Private/MuT/CodeOptimiser_Parameters.cpp b/Engine/Plugins/Experimental/Mutable/Source/MutableTools/Private/MuT/CodeOptimiser_Parameters.cpp
--- a/Engine/Plugins/Experimental/Mutable/Source/MutableTools/Private/MuT/CodeOptimiser_Parameters.cpp
+++ b/Engine/Plugins/Experimental/Mutable/Source/MutableTools/Private/MuT/CodeOptimiser_Parameters.cpp
@@ -1,18 +1,5 @@
 // Copyright Epic Games, Inc. All Rights Reserved.
 
-#include "HAL/PlatformCrt.h"
-#include "HAL/PlatformMath.h"
-#include "Logging/LogCategory.h"
-#include "Logging/LogMacros.h"
-#include "Misc/AssertionMacros.h"
-#include "MuR/Image.h"
-#include "MuR/ImagePrivate.h"
-#include "MuR/MutableTrace.h"
-#include "MuR/Operations.h"
-#include "MuR/ParametersPrivate.h"
-#include "MuR/Platform.h"
-#include "MuR/Ptr.h"
-#include "MuR/RefCounted.h"
 #include "MuT/AST.h"
 #include "MuT/ASTOpAddLOD.h"
 #include "MuT/ASTOpConditional.h"
@@ -30,13 +17,24 @@
 #include "MuT/CompilerPrivate.h"
 #include "MuT/DataPacker.h"
 #include "MuT/Platform.h"
+#include "MuR/Image.h"
+#include "MuR/ImagePrivate.h"
+#include "MuR/MutableTrace.h"
+#include "MuR/Operations.h"
+#include "MuR/ParametersPrivate.h"
+#include "MuR/Platform.h"
+#include "MuR/Ptr.h"
+#include "MuR/RefCounted.h"
+#include "MuR/MutableRuntimeModule.h"
+
+#include "HAL/PlatformCrt.h"
+#include "HAL/PlatformMath.h"
+#include "Logging/LogCategory.h"
+#include "Logging/LogMacros.h"
+#include "Misc/AssertionMacros.h"
 #include "Trace/Detail/Channel.h"
 
 #include <array>
-#include <memory>
-#include <utility>
-
-#include "MuR/MutableRuntimeModule.h"
 
 
 namespace mu
diff --git a/Engine/Plugins/Experimental/Mutable/Source/MutableTools/Private/MuT/Compiler.cpp b/Engine/Plugins/Experimental/Mutable/Source/MutableTools/Private/MuT/Compiler.cpp
--- a/Engine/Plugins/Experimental/Mutable/Source/MutableTools/Private/MuT/Compiler.cpp
+++ b/Engine/Plugins/Experimental/Mutable/Source/MutableTools/Private/MuT/Compiler.cpp
@@ -34,10 +34,6 @@
 #include "MuT/Table.h"
 #include "Trace/Detail/Channel.h"
 
-#include <algorithm>
-#include <memory>
-#include <utility>
-
 
 namespace mu
 {
diff --git a/Engine/Plugins/Experimental/Mutable/Source/MutableTools/Private/MuT/DataPacker.h b/Engine/Plugins/Experimental/Mutable/Source/MutableTools/Private/MuT/DataPacker.h
--- a/Engine/Plugins/Experimental/Mutable/Source/MutableTools/Private/MuT/DataPacker.h
+++ b/Engine/Plugins/Experimental/Mutable/Source/MutableTools/Private/MuT/DataPacker.h
@@ -5,12 +5,10 @@
 #include "MuR/Operations.h"
 #include "MuT/AST.h"
 
-#include <stdint.h>
-
 namespace mu { struct FProgram; }
 
 
-#define MUTABLE_HASH_SEED					((uint32_t)0xcadababa)
+#define MUTABLE_HASH_SEED					((uint32)0xcadababa)
 
 
 namespace mu
diff --git a/Engine/Plugins/Experimental/Mutable/Source/MutableTools/Private/MuT/ErrorLog.cpp b/Engine/Plugins/Experimental/Mutable/Source/MutableTools/Private/MuT/ErrorLog.cpp
--- a/Engine/Plugins/Experimental/Mutable/Source/MutableTools/Private/MuT/ErrorLog.cpp
+++ b/Engine/Plugins/Experimental/Mutable/Source/MutableTools/Private/MuT/ErrorLog.cpp
@@ -25,9 +25,6 @@
 #include "Templates/Tuple.h"
 #include "Trace/Detail/Channel.h"
 
-#include <memory>
-#include <utility>
-
 namespace mu
 {
 
diff --git a/Engine/Plugins/Experimental/Mutable/Source/MutableTools/Private/MuT/NodeBool.cpp b/Engine/Plugins/Experimental/Mutable/Source/MutableTools/Private/MuT/NodeBool.cpp
--- a/Engine/Plugins/Experimental/Mutable/Source/MutableTools/Private/MuT/NodeBool.cpp
+++ b/Engine/Plugins/Experimental/Mutable/Source/MutableTools/Private/MuT/NodeBool.cpp
@@ -7,8 +7,6 @@
 #include "MuR/Serialisation.h"
 #include "MuR/SerialisationPrivate.h"
 
-#include <stdint.h>
-
 
 namespace mu
 {
diff --git a/Engine/Plugins/Experimental/Mutable/Source/MutableTools/Private/MuT/NodeBoolParameter.cpp b/Engine/Plugins/Experimental/Mutable/Source/MutableTools/Private/MuT/NodeBoolParameter.cpp
--- a/Engine/Plugins/Experimental/Mutable/Source/MutableTools/Private/MuT/NodeBoolParameter.cpp
+++ b/Engine/Plugins/Experimental/Mutable/Source/MutableTools/Private/MuT/NodeBoolParameter.cpp
@@ -10,9 +10,6 @@
 #include "MuT/NodePrivate.h"
 #include "MuT/NodeRange.h"
 
-#include <memory>
-#include <utility>
-
 
 namespace mu
 {
diff --git a/Engine/Plugins/Experimental/Mutable/Source/MutableTools/Private/MuT/NodeColour.cpp b/Engine/Plugins/Experimental/Mutable/Source/MutableTools/Private/MuT/NodeColour.cpp
--- a/Engine/Plugins/Experimental/Mutable/Source/MutableTools/Private/MuT/NodeColour.cpp
+++ b/Engine/Plugins/Experimental/Mutable/Source/MutableTools/Private/MuT/NodeColour.cpp
@@ -15,8 +15,6 @@
 #include "MuT/NodeColourTable.h"
 #include "MuT/NodeColourVariation.h"
 
-#include <stdint.h>
-
 
 namespace mu
 {
diff --git a/Engine/Plugins/Experimental/Mutable/Source/MutableTools/Private/MuT/NodeComponent.cpp b/Engine/Plugins/Experimental/Mutable/Source/MutableTools/Private/MuT/NodeComponent.cpp
--- a/Engine/Plugins/Experimental/Mutable/Source/MutableTools/Private/MuT/NodeComponent.cpp
+++ b/Engine/Plugins/Experimental/Mutable/Source/MutableTools/Private/MuT/NodeComponent.cpp
@@ -9,8 +9,6 @@
 #include "MuT/NodeComponentEdit.h"
 #include "MuT/NodeComponentNew.h"
 
-#include <stdint.h>
-
 
 namespace mu
 {
diff --git a/Engine/Plugins/Experimental/Mutable/Source/MutableTools/Private/MuT/NodeRange.cpp b/Engine/Plugins/Experimental/Mutable/Source/MutableTools/Private/MuT/NodeRange.cpp
--- a/Engine/Plugins/Experimental/Mutable/Source/MutableTools/Private/MuT/NodeRange.cpp
+++ b/Engine/Plugins/Experimental/Mutable/Source/MutableTools/Private/MuT/NodeRange.cpp
@@ -9,8 +9,6 @@
 #include "MuT/Node.h"
 #include "MuT/NodeRangeFromScalar.h"
 
-#include <stdint.h>
-
 
 namespace mu
 {
diff --git a/Engine/Plugins/Experimental/Mutable/Source/MutableTools/Private/MuT/NodeScalar.cpp b/Engine/Plugins/Experimental/Mutable/Source/MutableTools/Private/MuT/NodeScalar.cpp
--- a/Engine/Plugins/Experimental/Mutable/Source/MutableTools/Private/MuT/NodeScalar.cpp
+++ b/Engine/Plugins/Experimental/Mutable/Source/MutableTools/Private/MuT/NodeScalar.cpp
@@ -16,8 +16,6 @@
 #include "MuT/NodeScalarTable.h"
 #include "MuT/NodeScalarVariation.h"
 
-#include <stdint.h>
-
 
 namespace mu
 {
diff --git a/Engine/Plugins/Experimental/Mutable/Source/MutableTools/Private/MuT/NodeScalarSwitch.cpp b/Engine/Plugins/Experimental/Mutable/Source/MutableTools/Private/MuT/NodeScalarSwitch.cpp
--- a/Engine/Plugins/Experimental/Mutable/Source/MutableTools/Private/MuT/NodeScalarSwitch.cpp
+++ b/Engine/Plugins/Experimental/Mutable/Source/MutableTools/Private/MuT/NodeScalarSwitch.cpp
@@ -8,9 +8,6 @@
 #include "MuT/NodeScalar.h"
 #include "MuT/NodeScalarSwitchPrivate.h"
 
-#include <memory>
-#include <utility>
-
 
 namespace mu
 {
diff --git a/Engine/Plugins/Experimental/Mutable/Source/MutableTools/Private/MuT/NodeString.cpp b/Engine/Plugins/Experimental/Mutable/Source/MutableTools/Private/MuT/NodeString.cpp
--- a/Engine/Plugins/Experimental/Mutable/Source/MutableTools/Private/MuT/NodeString.cpp
+++ b/Engine/Plugins/Experimental/Mutable/Source/MutableTools/Private/MuT/NodeString.cpp
@@ -10,8 +10,6 @@
 #include "MuT/NodeStringConstant.h"
 #include "MuT/NodeStringParameter.h"
 
-#include <stdint.h>
-
 
 namespace mu
 {
diff --git a/Engine/Plugins/Experimental/Mutable/Source/MutableTools/Private/MuT/Platform.h b/Engine/Plugins/Experimental/Mutable/Source/MutableTools/Private/MuT/Platform.h
--- a/Engine/Plugins/Experimental/Mutable/Source/MutableTools/Private/MuT/Platform.h
+++ b/Engine/Plugins/Experimental/Mutable/Source/MutableTools/Private/MuT/Platform.h
@@ -7,11 +7,6 @@
 #include <set>
 #include <string>
 #include <unordered_map>
-#include <sys/stat.h>
-
-
-// This file contains some hacks to solve differences between platforms necessary for the tools
-// library and not present in the run-time library.
 
 
 namespace mu
diff --git a/Engine/Plugins/Experimental/Mutable/Source/MutableTools/Private/MuT/Table.cpp b/Engine/Plugins/Experimental/Mutable/Source/MutableTools/Private/MuT/Table.cpp
--- a/Engine/Plugins/Experimental/Mutable/Source/MutableTools/Private/MuT/Table.cpp
+++ b/Engine/Plugins/Experimental/Mutable/Source/MutableTools/Private/MuT/Table.cpp
@@ -12,10 +12,6 @@
 #include "MuR/SerialisationPrivate.h"
 #include "MuT/TablePrivate.h"
 
-#include <memory>
-#include <string>
-#include <utility>
-
 
 namespace mu
 {
