[Mutable] Moved private classes to internal.

In preparation of Pimpl pattern.

- Moved UCustomizableObjectSystemPrivate and UCustomizableObjectInstancePrivate
- Renamed file CustomizableInstancePrivateData to CustomizableObjectInstancePrivate.

#rnx

diff --git a/Engine/Plugins/Experimental/Mutable/Source/CustomizableObject/Internal/MuCO/CustomizableInstancePrivateData.h b/Engine/Plugins/Experimental/Mutable/Source/CustomizableObject/Internal/MuCO/CustomizableInstancePrivateData.h
deleted file mode 100644
diff --git a/Engine/Plugins/Experimental/Mutable/Source/CustomizableObject/Internal/MuCO/CustomizableObjectInstancePrivate.h b/Engine/Plugins/Experimental/Mutable/Source/CustomizableObject/Internal/MuCO/CustomizableObjectInstancePrivate.h
new file mode 100644
--- /dev/null
+++ b/Engine/Plugins/Experimental/Mutable/Source/CustomizableObject/Internal/MuCO/CustomizableObjectInstancePrivate.h
@@ -0,0 +1,347 @@
+// Copyright Epic Games, Inc. All Rights Reserved.
+
+#pragma once
+
+
+#include "Materials/MaterialInterface.h"
+#include "Rendering/SkeletalMeshLODModel.h"
+#include "Rendering/SkeletalMeshModel.h"
+#include "MuCO/CustomizableObjectInstance.h"
+#include "MuCO/CustomizableObjectPrivate.h"
+#include "MuR/Instance.h"
+#include "GameplayTagContainer.h"
+#include "MuCO/DescriptorHash.h"
+#include "UObject/Package.h"
+
+#include "CustomizableObjectInstancePrivate.generated.h"
+
+namespace mu 
+{
+	class PhysicsBody;
+	class Mesh;
+	typedef uint64 FResourceID;
+}
+
+struct FModelResources;
+struct FMutableModelImageProperties;
+struct FMutableRefSkeletalMeshData;
+struct FMutableImageCacheKey;
+struct FStreamableHandle;
+struct FGeneratedMaterial;
+struct FGeneratedTexture;
+class UPhysicsAsset;
+class USkeleton;
+
+
+/** CustomizableObject Instance flags for internal use  */
+enum ECOInstanceFlags
+{
+	ECONone							= 0,  // Should not use the name None here.. it collides with other enum in global namespace
+
+	// Update process
+	ReuseTextures					= 1 << 3, 	// 
+	ReplacePhysicsAssets			= 1 << 4,	// Merge active PhysicsAssets and replace the base physics asset
+
+	// Update priorities
+	UsedByComponent					= 1 << 5,	// If any components are using this instance, they will set flag every frame
+	UsedByComponentInPlay			= 1 << 6,	// If any components are using this instance in play, they will set flag every frame
+	UsedByPlayerOrNearIt			= 1 << 7,	// The instance is used by the player or is near the player, used to give more priority to its updates
+	DiscardedByNumInstancesLimit	= 1 << 8,	// The instance is descarded because we exceeded the limit of instances generated 
+
+	// Types of updates
+	PendingLODsUpdate				= 1 << 9,	// Used to queue an update due to a change in LODs required by the instance
+	PendingLODsDowngrade			= 1 << 10,	// Used to queue a downgrade update to reduce the number of LODs. LOD update goes from a high res level to a low res one, ex: 0 to 1 or 1 to 2
+	
+	// Generation
+	ForceGenerateMipTail			= 1 << 13,	// If set, SkipGenerateResidentMips will be ignored and the mip tail will be generated
+};
+
+ENUM_CLASS_FLAGS(ECOInstanceFlags);
+
+
+USTRUCT()
+struct FReferencedPhysicsAssets
+{
+	GENERATED_USTRUCT_BODY();
+	
+	TArray<int32> PhysicsAssetToLoad;
+	
+	UPROPERTY(Transient)
+	TArray< TObjectPtr<UPhysicsAsset> > PhysicsAssetsToMerge;
+
+	TArray<int32> AdditionalPhysicsAssetsToLoad;
+
+	UPROPERTY(Transient)
+	TArray<TObjectPtr<UPhysicsAsset>> AdditionalPhysicsAssets;
+};
+
+
+USTRUCT()
+struct FReferencedSkeletons
+{
+	GENERATED_USTRUCT_BODY();
+
+	// Merged skeleton if found in the cache
+	UPROPERTY()
+	TObjectPtr<USkeleton> Skeleton;
+
+	UPROPERTY()
+	TArray<uint16> SkeletonIds;
+
+	UPROPERTY()
+	TArray< TObjectPtr<USkeleton> > SkeletonsToMerge;
+};
+
+
+USTRUCT()
+struct FCustomizableInstanceComponentData
+{
+	GENERATED_USTRUCT_BODY();
+
+	// AnimBP data gathered for a component from its constituent meshes
+	UPROPERTY(Transient, Category = CustomizableObjectInstance, editfixedsize, VisibleAnywhere)
+	TMap<FName, TSoftClassPtr<UAnimInstance>> AnimSlotToBP;
+
+	// AssetUserData gathered for a component from its constituent meshes
+	UPROPERTY(Transient, Category = CustomizableObjectInstance, editfixedsize, VisibleAnywhere)
+	TSet<TObjectPtr<UAssetUserData>> AssetUserDataArray;
+
+	// Index of the resource in the StreamedResourceData array of the CustomizableObject.
+	TArray<int32> StreamedResourceIndex;
+
+#if WITH_EDITORONLY_DATA
+	// Just used for mutable.EnableMutableAnimInfoDebugging command
+	TArray<FString> MeshPartPaths;
+#endif
+
+	/** Skeletons required by the current generated instance. Skeletons to be loaded and merged.*/
+	UPROPERTY(Transient)
+	FReferencedSkeletons Skeletons;
+	
+	/** PhysicsAssets required by the current generated instance. PhysicsAssets to be loaded and merged.*/
+	UPROPERTY(Transient)
+	FReferencedPhysicsAssets PhysicsAssets;
+
+	/** Clothing PhysicsAssets required by the current generated instance. PhysicsAssets to be loaded and merged.*/
+	TArray<TPair<int32, int32>> ClothingPhysicsAssetsToStream;
+
+	/** Array of generated MeshIds per each LOD, used to decide if the mesh should be updated or not.
+	 *  Size == NumLODsAvailable
+	 *  LODs without mesh will be set to the maximum value of FResourceID (Max_uint64). */
+	TArray<mu::FResourceID> LastMeshIdPerLOD;
+
+	UPROPERTY(Transient)
+	TArray<TObjectPtr<UMaterialInterface>> OverrideMaterials;
+};
+
+USTRUCT()
+struct FAnimInstanceOverridePhysicsAsset
+{
+	GENERATED_BODY()
+
+	UPROPERTY(Transient)
+	int32 PropertyIndex = 0;
+
+	UPROPERTY(Transient)
+	TObjectPtr<UPhysicsAsset> PhysicsAsset;
+};
+
+USTRUCT()
+struct FAnimBpGeneratedPhysicsAssets
+{
+	GENERATED_BODY()
+
+	UPROPERTY(Transient)
+	TArray<FAnimInstanceOverridePhysicsAsset> AnimInstancePropertyIndexAndPhysicsAssets;
+};
+
+
+UCLASS()
+class UCustomizableInstancePrivate : public UObject
+{
+public:
+	GENERATED_BODY()
+
+	UCustomizableInstancePrivate();
+
+	/** The generated skeletal meshes for this Instance, one for each component */
+	UPROPERTY(Transient, VisibleAnywhere, Category = NoCategory)
+	TArray<TObjectPtr<USkeletalMesh>> SkeletalMeshes;
+
+	UPROPERTY(Transient, VisibleAnywhere, Category = NoCategory)
+	TArray<FGeneratedMaterial> GeneratedMaterials;
+
+	UPROPERTY( Transient )
+	TArray<FGeneratedTexture> GeneratedTextures;
+
+	// Indices of the parameters that are relevant for the given parameter values.
+	// This only gets updated if parameter decorations are generated.
+	TArray<int> RelevantParameters;
+
+	// If Texture reuse is enabled, stores which texture is being used in a particular <LODIndex, ComponentIndex, MeshSurfaceIndex, image>
+	// \TODO: Create a key based on a struct instead of generating strings dynamically.
+	UPROPERTY(Transient)
+	TMap<FString, TWeakObjectPtr<UTexture2D>> TextureReuseCache;
+
+	// Handle used to store a streaming request operation if one is ongoing.
+	TSharedPtr<FStreamableHandle> StreamingHandle;
+
+	// Only used in LiveUpdateMode to reuse core instances between updates and their temp data to speed up updates, but spend way more memory
+	mu::Instance::ID LiveUpdateModeInstanceID = 0;
+
+	/** Cached Texture Parameters from the last update.
+	 * This cache is required since the Instance can have a LOD Update at any time.
+	 * So we need to make sure that the initially provided Texture Parameters by the user will be available until the user decides to change them. */
+	TArray<FName> UpdateTextureParameters;
+
+#if WITH_EDITOR
+
+	virtual void PostDuplicate(bool bDuplicateForPIE) override;
+
+	void BindObjectDelegates(UCustomizableObject* CurrentCustomizableObject, UCustomizableObject* NewCustomizableObject);
+
+	void OnPostCompile();
+	void OnObjectStatusChanged(FCustomizableObjectStatus::EState Previous, FCustomizableObjectStatus::EState Next);
+#endif
+
+	/** Invalidates the previously generated data and retrieves information from the CObject after specific actions.
+	 *  It'll be called in the PostLoad, after Compiling the CO, and after changing the CO of the Instance. */
+	void InitCustomizableObjectData(const UCustomizableObject* InCustomizableObject);
+	
+	FCustomizableInstanceComponentData* GetComponentData(int32 ComponentIndex);
+	const FCustomizableInstanceComponentData* GetComponentData(int32 ComponentIndex) const;
+
+	ECOInstanceFlags GetCOInstanceFlags() const { return InstanceFlagsPrivate; }
+	void SetCOInstanceFlags(ECOInstanceFlags FlagsToSet) { InstanceFlagsPrivate = (ECOInstanceFlags)(InstanceFlagsPrivate | FlagsToSet); }
+	void ClearCOInstanceFlags(ECOInstanceFlags FlagsToClear) { InstanceFlagsPrivate = (ECOInstanceFlags)(InstanceFlagsPrivate & ~FlagsToClear); }
+	bool HasCOInstanceFlags(ECOInstanceFlags FlagsToCheck) const { return (InstanceFlagsPrivate & FlagsToCheck) != 0; }
+
+	void BuildMaterials(const TSharedRef<FUpdateContextPrivate>& OperationData, UCustomizableObjectInstance* Public);
+
+	void ReuseTexture(UTexture2D* Texture, TSharedRef<FTexturePlatformData, ESPMode::ThreadSafe>& PlatformData);
+
+	// Return an event that will be fired when the assets  have been loaded. It returns null if no asset needs loading.
+	FGraphEventRef LoadAdditionalAssetsAsync(const TSharedRef<FUpdateContextPrivate>& OperationData, UCustomizableObjectInstance* Public, struct FStreamableManager &StreamableManager);
+	void AdditionalAssetsAsyncLoaded(UCustomizableObjectInstance* Public);
+
+	void TickUpdateCloseCustomizableObjects(UCustomizableObjectInstance& Publics, FMutableInstanceUpdateMap& InOutRequestedUpdates);
+	void UpdateInstanceIfNotGenerated(UCustomizableObjectInstance& Public, FMutableInstanceUpdateMap& InOutRequestedUpdates);
+
+	// Returns true if success (?)
+	bool UpdateSkeletalMesh_PostBeginUpdate0(UCustomizableObjectInstance* Public, const TSharedRef<FUpdateContextPrivate>& OperationData);
+
+	static void ReleaseMutableTexture(const FMutableImageCacheKey& MutableTextureKey, UTexture2D* Texture, struct FMutableResourceCache& Cache);
+
+	// Copy data generated in the mutable thread over to the instance and initializes additional data required during the update
+	void PrepareForUpdate(const TSharedRef<FUpdateContextPrivate>& OperationData);
+
+	int32 GetNumLODsAvailable() const { return NumLODsAvailable; }
+	
+	// The following method is basically copied from PostEditChangeProperty and/or SkeletalMesh.cpp to be able to replicate PostEditChangeProperty without the editor
+	void PostEditChangePropertyWithoutEditor();
+	
+	/** Calls ReleaseResources on all SkeletalMeshes generated by this instance and invalidates the generated data.
+	  * It should not be called if the meshes are still in use or shared with other instances. */
+	void DiscardResources();
+
+	// Releases all the mutable resources this instance holds, should only be called when it is not going to be used any more.
+	void ReleaseMutableResources(bool bCalledFromBeginDestroy, const UCustomizableObjectInstance& Instance);
+
+	/** Set the reference SkeletalMesh, or an empty mesh, to all actors using this instance. */
+	void SetDefaultSkeletalMesh(bool bSetEmptyMesh = false) const;
+
+	const TArray<FAnimInstanceOverridePhysicsAsset>* GetGeneratedPhysicsAssetsForAnimInstance(TSubclassOf<UAnimInstance> AnimInstance) const;
+
+	/** */
+#if WITH_EDITORONLY_DATA
+	void RegenerateImportedModels();
+#endif
+
+private:
+
+	void InitSkeletalMeshData(const TSharedRef<FUpdateContextPrivate>& OperationData, USkeletalMesh* SkeletalMesh, const FMutableRefSkeletalMeshData& RefSkeletalMeshData, const UCustomizableObject& CustomizableObject, int32 ComponentIndex);
+
+	bool BuildSkeletonData(const TSharedRef<FUpdateContextPrivate>& OperationData, USkeletalMesh& SkeletalMesh, const FMutableRefSkeletalMeshData& RefSkeletalMeshData, UCustomizableObject& CustomizableObject, int32 ComponentIndex);
+	void BuildMeshSockets(const TSharedRef<FUpdateContextPrivate>& OperationData, USkeletalMesh* SkeletalMesh, const FModelResources& ModelResources, const FMutableRefSkeletalMeshData& RefSkeletalMeshData, mu::Ptr<const mu::Mesh> MutableMesh);
+	void BuildOrCopyElementData(const TSharedRef<FUpdateContextPrivate>& OperationData, USkeletalMesh* SkeletalMesh, UCustomizableObjectInstance* CustomizableObjectInstance, int32 ComponentIndex);
+	void BuildOrCopyMorphTargetsData(const TSharedRef<FUpdateContextPrivate>& OperationData, USkeletalMesh* SkeletalMesh, const USkeletalMesh* SrcSkeletalMesh, UCustomizableObjectInstance* CustomizableObjectInstance, int32 ComponentIndex);
+	bool BuildOrCopyRenderData(const TSharedRef<FUpdateContextPrivate>& OperationData, USkeletalMesh* SkeletalMesh, const USkeletalMesh* SrcSkeletalMesh, UCustomizableObjectInstance* CustomizableObjectInstance, int32 ComponentIndex);
+	void BuildOrCopyClothingData(const TSharedRef<FUpdateContextPrivate>& OperationData, USkeletalMesh* SkeletalMesh, const USkeletalMesh* SrcSkeletalMesh, UCustomizableObjectInstance* CustomizableObjectInstance, int32 ComponentIndex);
+	
+	//
+	USkeleton* MergeSkeletons(UCustomizableObject& CustomizableObject, const FMutableRefSkeletalMeshData& RefSkeletalMeshData, int32 ComponentIndex);
+
+	//
+	UPhysicsAsset* GetOrBuildMainPhysicsAsset(TObjectPtr<class UPhysicsAsset> TamplateAsset, const mu::PhysicsBody* PhysicsBody, const UCustomizableObject& CustomizableObject, int32 ComponentIndex, bool bDisableCollisionBetweenAssets);
+	
+	// Create a transient texture and add it to the TextureTrackerArray
+	UTexture2D* CreateTexture();
+
+	void InvalidateGeneratedData();
+
+	bool DoComponentsNeedUpdate(UCustomizableObjectInstance* CustomizableObjectInstance, const TSharedRef<FUpdateContextPrivate>& OperationData, TArray<bool>& OutComponentNeedsUpdate, bool& bOutEmptyMesh);
+
+	mu::FResourceID GetLastMeshId(int32 ComponentIndex, int32 LODIndex) const;
+	void SetLastMeshId(int32 ComponentIndex, int32 LODIndex, mu::FResourceID MeshId);
+
+public:
+	// If any components are using this instance, they will store the min of their distances to the player here every frame for LOD purposes
+	float MinSquareDistFromComponentToPlayer;
+	float LastMinSquareDistFromComponentToPlayer; // The same as the previous dist for last frame
+												
+	// This is the LODs that the Customizable Object has
+	int32 NumLODsAvailable;
+
+	// First SkeletalMesh LOD we can generate on the running platform
+	uint8 FirstLODAvailable;
+
+	// Maximum number of SkeletalMesh LODs to stream
+	uint8 NumMaxLODsToStream;
+	
+	UPROPERTY(Transient)
+	TArray<FCustomizableInstanceComponentData> ComponentsData;
+
+	UPROPERTY(Transient)
+	TArray< TObjectPtr<UMaterialInterface> > ReferencedMaterials;
+
+	// Converts a ReferencedMaterials index from the CustomizableObject to an index in the ReferencedMaterials in the Instance
+	TMap<uint32, uint32> ObjectToInstanceIndexMap;
+
+	TArray<FGeneratedTexture> TexturesToRelease;
+
+	UPROPERTY(Transient)
+	TArray< TObjectPtr<UPhysicsAsset> > ClothingPhysicsAssets;
+
+	// To keep loaded AnimBPs referenced and prevent GC
+	UPROPERTY(Transient, Category = Animation, editfixedsize, VisibleAnywhere)
+	TArray<TSubclassOf<UAnimInstance>> GatheredAnimBPs;
+
+	UPROPERTY(Transient, Category = Animation, editfixedsize, VisibleAnywhere)
+	FGameplayTagContainer AnimBPGameplayTags;
+
+	UPROPERTY(Transient, Category = Animation, editfixedsize, VisibleAnywhere)
+	TMap<TSubclassOf<UAnimInstance>, FAnimBpGeneratedPhysicsAssets> AnimBpPhysicsAssets;
+
+	// The pass-through textures that will be loaded during an update
+	TArray<TSoftObjectPtr<UTexture>> PassThroughTexturesToLoad;
+
+	// Used during an update to prevent the pass-through textures loaded by LoadAdditionalAssetsAsync() from being unloaded by GC
+	// between AdditionalAssetsAsyncLoaded() and their setting into the generated materials in BuildMaterials()
+	UPROPERTY(Transient)
+	TArray<TObjectPtr<UTexture>> LoadedPassThroughTexturesPendingSetMaterial;
+
+private:
+	
+	ECOInstanceFlags InstanceFlagsPrivate = ECOInstanceFlags::ECONone;
+
+public:
+	/** Hash of the UCustomizableObjectInstance::Descriptor on the last update request. */
+	FDescriptorHash UpdateDescriptorHash;
+	
+	/** Hash of the UCustomizableObjectInstance::Descriptor on the last successful update. */
+	FDescriptorHash DescriptorHash;
+
+	/** Status of the generated Skeletal Mesh. Not to be confused with the Update Result. */
+	ESkeletalMeshStatus SkeletalMeshStatus = ESkeletalMeshStatus::NotGenerated;
+};
+
diff --git a/Engine/Plugins/Experimental/Mutable/Source/CustomizableObject/Internal/MuCO/CustomizableObjectSystemPrivate.h b/Engine/Plugins/Experimental/Mutable/Source/CustomizableObject/Internal/MuCO/CustomizableObjectSystemPrivate.h
new file mode 100644
--- /dev/null
+++ b/Engine/Plugins/Experimental/Mutable/Source/CustomizableObject/Internal/MuCO/CustomizableObjectSystemPrivate.h
@@ -0,0 +1,672 @@
+// Copyright Epic Games, Inc. All Rights Reserved.
+
+#pragma once
+
+#include "LogBenchmarkUtil.h"
+#include "Containers/Queue.h"
+#include "MuCO/DescriptorHash.h"
+#include "MuCO/CustomizableObject.h"
+#include "MuCO/CustomizableObjectInstance.h"
+#include "MuCO/CustomizableObjectExtension.h"
+#include "Containers/Ticker.h"
+
+#include "MuCO/CustomizableObjectInstanceDescriptor.h"
+#include "MuR/Mesh.h"
+#include "MuR/Parameters.h"
+#include "MuR/System.h"
+#include "MuR/Image.h"
+#include "UObject/GCObject.h"
+#include "WorldCollision.h"
+#include "MuCO/FMutableTaskGraph.h"
+
+// This define could come from MuR/System.h
+#ifdef MUTABLE_USE_NEW_TASKGRAPH
+	#include "Tasks/Task.h"
+#else
+	#include "Async/TaskGraphInterfaces.h"
+#endif
+
+#include "CustomizableObjectSystemPrivate.generated.h"
+
+#if WITH_EDITORONLY_DATA
+class UEditorImageProvider;
+#endif
+
+class UCustomizableObjectSystem;
+namespace LowLevelTasks { enum class ETaskPriority : int8; }
+struct FTexturePlatformData;
+
+
+struct FMutablePendingInstanceUpdate
+{
+	TSharedRef<FUpdateContextPrivate> Context;
+
+	FMutablePendingInstanceUpdate(const TSharedRef<FUpdateContextPrivate>& InContext);
+
+	bool operator==(const FMutablePendingInstanceUpdate& Other) const;
+
+	bool operator<(const FMutablePendingInstanceUpdate& Other) const;
+};
+
+
+inline uint32 GetTypeHash(const FMutablePendingInstanceUpdate& Update);
+
+
+struct FPendingInstanceUpdateKeyFuncs : BaseKeyFuncs<FMutablePendingInstanceUpdate, TWeakObjectPtr<const UCustomizableObjectInstance>>
+{
+	FORCEINLINE static TWeakObjectPtr<const UCustomizableObjectInstance> GetSetKey(const FMutablePendingInstanceUpdate& PendingUpdate);
+
+	FORCEINLINE static bool Matches(const TWeakObjectPtr<const UCustomizableObjectInstance>& A, const TWeakObjectPtr<const UCustomizableObjectInstance>& B);
+
+	FORCEINLINE static uint32 GetKeyHash(const TWeakObjectPtr<const UCustomizableObjectInstance>& Identifier);
+};
+
+
+struct FMutablePendingInstanceDiscard
+{
+	TWeakObjectPtr<UCustomizableObjectInstance> CustomizableObjectInstance;
+
+	FMutablePendingInstanceDiscard(UCustomizableObjectInstance* InCustomizableObjectInstance)
+	{
+		CustomizableObjectInstance = InCustomizableObjectInstance;
+	}
+
+	friend bool operator ==(const FMutablePendingInstanceDiscard& A, const FMutablePendingInstanceDiscard& B)
+	{
+		return A.CustomizableObjectInstance.HasSameIndexAndSerialNumber(B.CustomizableObjectInstance);
+	}
+};
+
+
+inline uint32 GetTypeHash(const FMutablePendingInstanceDiscard& Discard)
+{
+	return GetTypeHash(Discard.CustomizableObjectInstance.GetWeakPtrTypeHash());
+}
+
+
+struct FPendingInstanceDiscardKeyFuncs : BaseKeyFuncs<FMutablePendingInstanceUpdate, TWeakObjectPtr<UCustomizableObjectInstance>>
+{
+	FORCEINLINE static const TWeakObjectPtr<UCustomizableObjectInstance>& GetSetKey(const FMutablePendingInstanceDiscard& PendingDiscard)
+	{
+		return PendingDiscard.CustomizableObjectInstance;
+	}
+
+	FORCEINLINE static bool Matches(const TWeakObjectPtr<UCustomizableObjectInstance>& A, const TWeakObjectPtr<UCustomizableObjectInstance>& B)
+	{
+		return A.HasSameIndexAndSerialNumber(B);
+	}
+
+	FORCEINLINE static uint32 GetKeyHash(const TWeakObjectPtr<UCustomizableObjectInstance>& Identifier)
+	{
+		return GetTypeHash(Identifier.GetWeakPtrTypeHash());
+	}
+};
+
+
+/** Instance updates queue.
+ *
+ * The queues will only contain a single operation per UCustomizableObjectInstance.
+ * If there is already an operation it will be replaced. */
+class FMutablePendingInstanceWork
+{
+	TSet<FMutablePendingInstanceUpdate, FPendingInstanceUpdateKeyFuncs> PendingInstanceUpdates;
+
+	TSet<FMutablePendingInstanceDiscard, FPendingInstanceDiscardKeyFuncs> PendingInstanceDiscards;
+
+	TSet<mu::Instance::ID> PendingIDsToRelease;
+
+	int32 NumLODUpdatesLastTick = 0;
+
+public:
+	// Returns the number of pending instance updates, LOD Updates, discards and releases last tick.
+	int32 Num() const;
+
+	void SetLODUpdatesLastTick(int32 NumLODUpdates);
+
+	// Adds a new instance update
+	void AddUpdate(const FMutablePendingInstanceUpdate& UpdateToAdd);
+
+	// Removes an instance update
+	void RemoveUpdate(const TWeakObjectPtr<UCustomizableObjectInstance>& Instance);
+
+	const FMutablePendingInstanceUpdate* GetUpdate(const TWeakObjectPtr<const UCustomizableObjectInstance>& Instance) const;
+
+	TSet<FMutablePendingInstanceUpdate, FPendingInstanceUpdateKeyFuncs>::TIterator GetUpdateIterator()
+	{
+		return PendingInstanceUpdates.CreateIterator();
+	}
+
+	TSet<FMutablePendingInstanceDiscard, FPendingInstanceDiscardKeyFuncs>::TIterator GetDiscardIterator()
+	{
+		return PendingInstanceDiscards.CreateIterator();
+	}
+
+	TSet<mu::Instance::ID>::TIterator GetIDsToReleaseIterator()
+	{
+		return PendingIDsToRelease.CreateIterator();
+	}
+
+	void AddDiscard(const FMutablePendingInstanceDiscard& TaskToEnqueue);
+	void AddIDRelease(mu::Instance::ID IDToRelease);
+
+	void RemoveAllUpdatesAndDiscardsAndReleases();
+};
+
+
+struct FMutableImageCacheKey
+{
+	mu::FResourceID Resource = 0;
+	int32 SkippedMips = 0;
+
+	FMutableImageCacheKey() {};
+
+	FMutableImageCacheKey(mu::FResourceID InResource, int32 InSkippedMips)
+		: Resource(InResource), SkippedMips(InSkippedMips) {}
+
+	inline bool operator==(const FMutableImageCacheKey& Other) const
+	{
+		return Resource == Other.Resource && SkippedMips == Other.SkippedMips;
+	}
+};
+
+
+inline uint32 GetTypeHash(const FMutableImageCacheKey& Key)
+{
+	return HashCombine(GetTypeHash(Key.Resource), GetTypeHash(Key.SkippedMips));
+}
+
+
+// Cache of weak references to generated resources for one single model
+struct FMutableResourceCache
+{
+	TWeakObjectPtr<const UCustomizableObject> Object;
+	TMap<mu::FResourceID, TWeakObjectPtr<USkeletalMesh> > Meshes;
+	TMap<FMutableImageCacheKey, TWeakObjectPtr<UTexture2D> > Images;
+
+	void Clear()
+	{
+		Meshes.Reset();
+		Images.Reset();
+	}
+};
+
+
+USTRUCT()
+struct FGeneratedTexture
+{
+	GENERATED_USTRUCT_BODY();
+
+	FMutableImageCacheKey Key;
+
+	UPROPERTY(Category = NoCategory, VisibleAnywhere)
+	FString Name;
+
+	UPROPERTY(Category = NoCategory, VisibleAnywhere)
+	TObjectPtr<UTexture> Texture = nullptr;
+
+	bool operator==(const FGeneratedTexture& Other) const = default;
+};
+
+
+USTRUCT()
+struct FGeneratedMaterial
+{
+	GENERATED_USTRUCT_BODY();
+
+	UPROPERTY(Category = NoCategory, VisibleAnywhere)
+	TObjectPtr<UMaterialInterface> MaterialInterface;
+
+	UPROPERTY(Category = NoCategory, VisibleAnywhere)
+	TArray< FGeneratedTexture > Textures;
+
+	// Surface or SharedSurface Id
+	uint32 SurfaceId = 0;
+
+	// Index of the material to instantiate (UCustomizableObject::ReferencedMaterials)
+	uint32 MaterialIndex = 0;
+
+	bool operator==(const FGeneratedMaterial& Other) const { return SurfaceId == Other.SurfaceId && MaterialIndex == Other.MaterialIndex; };
+};
+
+
+DECLARE_DELEGATE(FMutableTaskDelegate);
+struct FMutableTask
+{
+	/** Actual function to perform in this task. */
+	FMutableTaskDelegate Function;
+
+	/** We can have 2 types of depencies:
+		From the new task system: */
+#ifdef MUTABLE_USE_NEW_TASKGRAPH
+	UE::Tasks::FTask Dependency;
+#endif
+
+	/** From the traditional event graph system, still used to wait for async loads. */
+	FGraphEventRef GraphDependency0;
+	
+#ifdef MUTABLE_USE_NEW_TASKGRAPH
+	UE::Tasks::FTask GraphDependency1;
+#else
+	FGraphEventRef GraphDependency1;
+#endif
+
+	/** Check if all the dependencies of this task have been completed. */
+	inline bool AreDependenciesComplete() const
+	{
+#ifdef MUTABLE_USE_NEW_TASKGRAPH
+		return (!Dependency.IsValid() || Dependency.IsCompleted())
+			&&
+			(!GraphDependency0 || GraphDependency0->IsComplete())
+			&&
+			(GraphDependency1.IsCompleted());
+#else
+		return 
+			(!GraphDependency0 || GraphDependency0->IsComplete())
+			&&
+			(!GraphDependency1 || GraphDependency1->IsComplete());
+#endif
+	}
+
+	/** Free the handles for any dependency of this task. */
+	inline void ClearDependencies()
+	{
+#ifdef MUTABLE_USE_NEW_TASKGRAPH
+		Dependency = {};
+#endif
+		GraphDependency0 = nullptr;
+
+#ifdef MUTABLE_USE_NEW_TASKGRAPH
+		GraphDependency1 = {};
+#else
+		GraphDependency1 = nullptr;
+#endif
+	}
+};
+
+
+// Mutable data generated during the update steps.
+// We keep it from begin to end update, and it is used in several steps.
+struct FInstanceUpdateData
+{
+	struct FImage
+	{
+		FName Name;
+		mu::FResourceID ImageID;
+		
+		// LOD of the ImageId. If the texture is shared between LOD, first LOD where this image can be found. 
+		int32 BaseLOD;
+		int32 BaseMip;
+		
+		uint16 FullImageSizeX, FullImageSizeY;
+		mu::Ptr<const mu::Image> Image;
+		TWeakObjectPtr<UTexture2D> Cached;
+
+		bool bIsPassThrough = false;
+	};
+
+	struct FVector
+	{
+		FName Name;
+		FLinearColor Vector;
+	};
+
+	struct FScalar
+	{
+		FName Name;
+		float Scalar;
+	};
+
+	struct FSurface
+	{
+		/** Range in the Images array */
+		uint16 FirstImage = 0;
+		uint16 ImageCount = 0;
+
+		/** Range in the Vectors array */
+		uint16 FirstVector = 0;
+		uint16 VectorCount = 0;
+
+		/** Range in the Scalar array */
+		uint16 FirstScalar = 0;
+		uint16 ScalarCount = 0;
+
+		uint32 MaterialIndex = 0;
+
+		/** Id of the surface in the mutable core instance. */
+		uint32 SurfaceId = 0;
+	};
+
+	struct FComponent
+	{
+		uint16 Id = 0;
+		
+		// True if the Mesh is valid
+		bool bGenerated = false;
+
+		mu::FResourceID MeshID;
+		mu::MeshPtrConst Mesh;
+
+		/** Range in the Surfaces array */
+		uint16 FirstSurface = 0;
+		uint16 SurfaceCount = 0;
+
+		// \TODO: Flatten
+		TArray<uint16> ActiveBones;
+		/** Range in the external Bones array */
+		//uint32 FirstActiveBone;
+		//uint32 ActiveBoneCount;
+
+		/** Range in the external Bones array */
+		uint32 FirstBoneMap = 0;
+		uint32 BoneMapCount = 0;
+	};
+
+	struct FLOD
+	{
+		/** Range in the Components array */
+		uint16 FirstComponent = 0;
+		uint16 ComponentCount = 0;
+	};
+
+	TArray<FLOD> LODs;
+	TArray<FComponent> Components;
+	TArray<FSurface> Surfaces;
+	TArray<FImage> Images;
+	TArray<FVector> Vectors;
+	TArray<FScalar> Scalars;
+
+	TArray<uint16> BoneMaps;
+
+	struct FSkeletonData
+	{
+		int16 ComponentIndex = INDEX_NONE;
+
+		TArray<uint16> SkeletonIds;
+
+		TArray<uint16> BoneIds;
+		TArray<FMatrix44f> BoneMatricesWithScale;
+	};
+
+	TArray<FSkeletonData> Skeletons;
+
+	struct FNamedExtensionData
+	{
+		mu::ExtensionDataPtrConst Data;
+		FName Name;
+	};
+	TArray<FNamedExtensionData> ExtendedInputPins;
+
+	/** */
+	void Clear()
+	{
+		LODs.Empty();
+		Components.Empty();
+		Surfaces.Empty();
+		Images.Empty();
+		Scalars.Empty();
+		Vectors.Empty();
+		Skeletons.Empty();
+		ExtendedInputPins.Empty();
+	}
+};
+
+
+/** Update Context.
+ *
+ * Alive from the start to the end of the update (both API and LOD update). */
+class FUpdateContextPrivate : public FGCObject
+{
+public:
+	FUpdateContextPrivate(UCustomizableObjectInstance& InInstance);
+	virtual ~FUpdateContextPrivate() override;
+
+	virtual FString GetReferencerName() const override;
+	virtual void AddReferencedObjects(FReferenceCollector& Collector) override;
+	
+	EQueuePriorityType PriorityType = EQueuePriorityType::Low;
+	
+	FInstanceUpdateDelegate UpdateCallback;
+	FInstanceUpdateNativeDelegate UpdateNativeCallback;
+
+	/** Weak reference to the instance we are operating on.
+	 *It is weak because we don't want to lock it in case it becomes irrelevant in the game while operations are pending and it needs to be destroyed. */
+	TWeakObjectPtr<UCustomizableObjectInstance> Instance;
+
+	/** Hash of the UCustomizableObjectInstance::Descriptor at the time of the update request. */
+	FDescriptorHash InstanceDescriptorHash;
+			
+	/** Instance parameters at the time of the operation request. */
+	mu::ParametersPtr Parameters; 
+	
+	TArray<FName> TextureParameters;
+
+	bool bBuildParameterRelevancy = false;
+
+	/** Instance state. */
+	int32 State = 0;
+
+	bool bOnlyUpdateIfNotGenerated = false;
+	bool bIgnoreCloseDist = false;
+	bool bForceHighPriority = false;
+	
+	FInstanceUpdateData InstanceUpdateData;
+	TArray<int32> RelevantParametersInProgress;
+
+	TArray<FString> LowPriorityTextures;
+
+	/** This option comes from the operation request */
+	bool bNeverStream = false;
+	
+	/** When this option is enabled it will reuse the Mutable core instance and its temp data between updates.  */
+	bool bLiveUpdateMode = false;
+	bool bReuseInstanceTextures = false;
+	bool bUseMeshCache = false;
+	
+	/** This option comes from the operation request. It is used to reduce the number of mipmaps that mutable must generate for images.  */
+	int32 MipsToSkip = 0;
+
+	mu::Instance::ID InstanceID = 0; // Redundant
+	const mu::Instance* MutableInstance = nullptr;
+
+	int32 NumComponents = 0;
+	int32 NumLODsAvailable = 0;
+
+	int32 CurrentMinLOD = 0;
+	int32 CurrentMaxLOD = 0;
+
+	TArray<uint16> RequestedLODs;
+
+	TMap<uint32, FTexturePlatformData*> ImageToPlatformDataMap;
+
+	EUpdateResult UpdateResult = EUpdateResult::Success;
+
+	mu::FImageOperator::FImagePixelFormatFunc PixelFormatOverride;
+
+	/** Mutable Meshes required for each component. Outermost index is the component, inner index is the LOD. */
+	TArray<TArray<mu::FResourceID>> MeshDescriptors;
+	
+	bool UpdateStarted = false;
+
+	// Update stats
+	double StartQueueTime = 0.0;
+	double QueueTime = 0.0;
+	
+	double StartUpdateTime = 0.0;
+	double UpdateTime = 0.0;
+
+	double TaskGetMeshTime = 0.0;
+	double TaskLockCacheTime = 0.0;
+	double TaskGetImagesTime = 0.0;
+	double TaskConvertResourcesTime = 0.0f;
+	double TaskCallbacksTime = 0.0;
+
+	// Update Memory stats
+	int64 UpdateStartBytes = 0;
+	int64 UpdateEndPeakBytes = 0;
+	int64 UpdateEndRealPeakBytes = 0;
+	
+#if WITH_EDITOR
+	/** Used for profiling in the editor. */
+	uint32 MutableRuntimeCycles = 0;
+#endif
+
+	/** Hard references to objects. Avoids GC to collect them. */
+	TArray<TObjectPtr<const UObject>> Objects;
+};
+
+
+/** Runtime data used during a mutable instance update */
+struct FMutableReleasePlatformOperationData
+{
+	TMap<uint32, FTexturePlatformData*> ImageToPlatformDataMap;
+};
+
+
+UCLASS()
+class UCustomizableObjectSystemPrivate : public UObject
+{
+	GENERATED_BODY()
+	
+public:
+	// Singleton for the unreal mutable system.
+	static UCustomizableObjectSystem* SSystem;
+
+	// Pointer to the lower level mutable system that actually does the work.
+	mu::Ptr<mu::System> MutableSystem;
+
+	/** Store the last streaming memory size in bytes, to change it when it is safe. */
+	uint64 LastWorkingMemoryBytes = 0;
+	uint32 LastGeneratedResourceCacheSize = 0;
+
+	// This object is responsible for streaming data to the MutableSystem.
+	TSharedPtr<class FUnrealMutableModelBulkReader> Streamer;
+
+	TSharedPtr<class FUnrealExtensionDataStreamer> ExtensionDataStreamer;
+
+	// This object is responsible for providing custom images to mutable models (for image parameters)
+	// This object is called from the mutable thread, and it should only access data already safely submitted from
+	// the game thread and stored in FUnrealMutableImageProvider::GlobalExternalImages.
+	TSharedPtr<class FUnrealMutableImageProvider> ImageProvider;
+
+	// Cache of weak references to generated resources to see if they can be reused.
+	TArray<FMutableResourceCache> ModelResourcesCache;
+
+	// List of textures currently cached and valid for the current object that we are operating on.
+	// This array gets generated when the object cached resources are protected in SetResourceCacheProtected
+	// from the game thread, and it is read from the Mutable thread only while updating the instance.
+	TArray<mu::FResourceID> ProtectedObjectCachedImages;
+
+	// The pending instance updates, discards or releases
+	FMutablePendingInstanceWork MutablePendingInstanceWork;
+
+	// Queue of game-thread tasks that need to be executed for the current operation. TQueue is thread safe
+	TQueue<FMutableTask> PendingTasks;
+
+	static int32 EnableMutableProgressiveMipStreaming;
+	static int32 EnableMutableLiveUpdate;
+	static int32 EnableReuseInstanceTextures;
+	static int32 EnableMutableAnimInfoDebugging;
+	static int32 EnableSkipGenerateResidentMips;
+	static int32 EnableOnlyGenerateRequestedLODs;
+	static int32 MaxTextureSizeToGenerate;
+	static int32 SkeletalMeshMinLodQualityLevel;
+
+#if WITH_EDITOR
+	mu::FImageOperator::FImagePixelFormatFunc ImageFormatOverrideFunc;
+#endif
+
+	void AddGameThreadTask(const FMutableTask& Task);
+
+	// Remove references to cached objects that have been deleted in the unreal
+	// side, and cannot be cached anyway.
+	// This should only happen in the game thread
+	void CleanupCache();
+
+	// This should only happen in the game thread
+	FMutableResourceCache& GetObjectCache(const UCustomizableObject* Object);
+
+	void AddTextureReference(const FMutableImageCacheKey& TextureId);
+
+	// Returns true if the texture's references become zero
+	bool RemoveTextureReference(const FMutableImageCacheKey& TextureId);
+
+	bool TextureHasReferences(const FMutableImageCacheKey& TextureId) const;
+
+	EUpdateRequired IsUpdateRequired(const UCustomizableObjectInstance& Instance, bool bOnlyUpdateIfNotGenerated, bool bOnlyUpdateIfLOD, bool bIgnoreCloseDist) const;
+
+	EQueuePriorityType GetUpdatePriority(const UCustomizableObjectInstance& Instance, bool bForceHighPriority) const;
+
+	void EnqueueUpdateSkeletalMesh(const TSharedRef<FUpdateContextPrivate>& Context);
+		
+	// Init an async and safe release of the UE and Mutable resources used by the instance without actually destroying the instance, for example if it's very far away
+	void InitDiscardResourcesSkeletalMesh(UCustomizableObjectInstance* InCustomizableObjectInstance);
+
+	// Init the async release of a Mutable Core Instance ID and all the temp resources associated with it
+	void InitInstanceIDRelease(mu::Instance::ID IDToRelease);
+
+	void GetMipStreamingConfig(const UCustomizableObjectInstance& Instance, bool& bOutNeverStream, int32& OutMipsToSkip) const;
+	
+	bool IsReplaceDiscardedWithReferenceMeshEnabled() const;
+	void SetReplaceDiscardedWithReferenceMeshEnabled(bool bIsEnabled);
+
+	/** Updated at the beginning of each tick. */
+	int32 GetNumSkeletalMeshes() const;
+
+	bool bReplaceDiscardedWithReferenceMesh = false;
+	bool bReleaseTexturesImmediately = false;
+
+	bool bSupport16BitBoneIndex = false;
+
+	static FCustomizableObjectCompilerBase* (*NewCompilerFunc)();
+
+	TMap<FMutableImageCacheKey, uint32> TextureReferenceCount; // Keeps a count of texture usage to decide if they have to be blocked from GC during an update
+
+	UPROPERTY(Transient)
+	TObjectPtr<UCustomizableObjectInstance> CurrentInstanceBeingUpdated = nullptr;
+
+	TSharedPtr<FUpdateContextPrivate> CurrentMutableOperation = nullptr;
+
+	// Handle to the registered TickDelegate.
+	FTSTicker::FDelegateHandle TickDelegateHandle;
+	FTickerDelegate TickDelegate;
+
+	/** Change the current status of Mutable. Enabling/Disabling core features.	
+	 * Disabling Mutable will turn off compilation, generation, and streaming and will remove the system ticker. */
+	static void OnMutableEnabledChanged(IConsoleVariable* CVar = nullptr);
+
+	/** Update the last set amount of internal memory Mutable can use to build objects. */
+	void UpdateMemoryLimit();
+
+	bool IsMutableAnimInfoDebuggingEnabled() const;
+
+	FUnrealMutableImageProvider* GetImageProviderChecked() const;
+
+	/** Start the actual work of Update Skeletal Mesh process (Update Skeletal Mesh without the queue). */
+	void StartUpdateSkeletalMesh(const TSharedRef<FUpdateContextPrivate>& Context);
+
+	/** See UCustomizableObjectInstance::IsUpdating. */
+	bool IsUpdating(const UCustomizableObjectInstance& Instance) const;
+
+	/** Update stats at each tick.
+	 * Used for stats that are costly to update. */
+	void UpdateStats();
+	
+	/** Mutable TaskGraph system (Mutable Thread). */
+	FMutableTaskGraph MutableTaskGraph;
+	
+#if WITH_EDITORONLY_DATA
+	/** Mutable default image provider. Used by the COIEditor and Instance/Descriptor APIs. */
+	UPROPERTY(Transient)
+	TObjectPtr<UEditorImageProvider> EditorImageProvider = nullptr;
+#endif
+
+	FLogBenchmarkUtil LogBenchmarkUtil;
+
+	int32 NumSkeletalMeshes = 0;
+
+	bool bAutoCompileCommandletEnabled = false;
+};
+
+
+/** Set OnlyLOD to -1 to generate all mips */
+CUSTOMIZABLEOBJECT_API FTexturePlatformData* MutableCreateImagePlatformData(mu::Ptr<const mu::Image> MutableImage, int32 OnlyLOD, uint16 FullSizeX, uint16 FullSizeY);
diff --git a/Engine/Plugins/Experimental/Mutable/Source/CustomizableObject/Internal/MuCO/LogBenchmarkUtil.h b/Engine/Plugins/Experimental/Mutable/Source/CustomizableObject/Internal/MuCO/LogBenchmarkUtil.h
new file mode 100644
--- /dev/null
+++ b/Engine/Plugins/Experimental/Mutable/Source/CustomizableObject/Internal/MuCO/LogBenchmarkUtil.h
@@ -0,0 +1,113 @@
+// Copyright Epic Games, Inc. All Rights Reserved.
+
+#pragma once
+
+#include "Containers/Map.h"
+#include "UObject/NameTypes.h"
+#include "Stats/Stats.h"
+#include "HAL/IConsoleManager.h"
+
+class UCustomizableObjectInstance;
+class UTexture2D;
+
+
+/** Stat which automatically gets updated to Insights when modified. */
+#define DECLARE_BENCHMARK_STAT(Name, Type) \
+	struct F##Name \
+	{ \
+	private: \
+		Type Value = 0; \
+		\
+		mutable FCriticalSection Lock##Name; \
+		\
+	public: \
+		Type GetValue() const \
+		{ \
+			FScopeLock ScopedLock(&Lock##Name); \
+			return Value; \
+		} \
+		\
+		F##Name& operator+=(const Type& Rhs) \
+		{ \
+			FScopeLock ScopedLock(&Lock##Name); \
+			Value += Rhs; \
+			return *this; \
+		} \
+		\
+		F##Name& operator=(const Type& Rhs) \
+		{ \
+			FScopeLock ScopedLock(&Lock##Name); \
+			Value = Rhs; \
+		\
+		SET_DWORD_STAT(STAT_Mutable##Name, Value); \
+		\
+		return *this; \
+		} \
+	}; \
+	F##Name Name;
+
+#define DECLARE_BENCHMARK_INSIGHTS(Name, Description) \
+	DECLARE_DWORD_ACCUMULATOR_STAT(Description, STAT_Mutable##Name, STATGROUP_Mutable);
+
+
+class FUpdateContextPrivate;
+// Mutable stats
+DECLARE_STATS_GROUP(TEXT("Mutable"), STATGROUP_Mutable, STATCAT_Advanced);
+
+DECLARE_BENCHMARK_INSIGHTS(NumAllocatedSkeletalMeshes, TEXT("Num Allocated Mutable Skeletal Meshes"));
+DECLARE_BENCHMARK_INSIGHTS(NumAllocatedTextures, TEXT("Num Allocated Mutable Textures"));
+DECLARE_BENCHMARK_INSIGHTS(TextureGPUSize, TEXT("Size of the Mutable Texture mips that are resident on the GPU"));
+DECLARE_BENCHMARK_INSIGHTS(NumInstances, TEXT("Num Instances"));
+DECLARE_BENCHMARK_INSIGHTS(NumInstancesLOD0, TEXT("Num Instances at LOD 0"));
+DECLARE_BENCHMARK_INSIGHTS(NumInstancesLOD1, TEXT("Num Instances at LOD 1"));
+DECLARE_BENCHMARK_INSIGHTS(NumInstancesLOD2, TEXT("Num Instances at LOD 2 or more"));
+DECLARE_BENCHMARK_INSIGHTS(NumPendingInstanceUpdates, TEXT("Num Pending Instance Updates"));
+DECLARE_BENCHMARK_INSIGHTS(NumBuiltInstances, TEXT("Num Built Instances"))
+DECLARE_BENCHMARK_INSIGHTS(InstanceBuildTimeAvrg, TEXT("Avrg Instance Build Time"));
+
+
+extern TAutoConsoleVariable<bool> CVarEnableBenchmark;
+
+
+/** Benchmarking system. Gathers stats and send it to Insights an Benchmarking Files. */
+class FLogBenchmarkUtil
+{
+public:
+	~FLogBenchmarkUtil();
+
+	/** Get stats. */
+	void GetInstancesStats(int32& OutNumInstances, int32& OutNumBuiltInstances, int32& OutNumInstancesLOD0, int32& OutNumInstancesLOD1, int32& OutNumInstancesLOD2, int32& OutNumAllocatedSkeletalMeshes) const;
+
+	/** Add Mutable created Texture to track. */
+	void AddTexture(UTexture2D& Texture);
+
+	/** Update stats which can only be updated on the tick. */
+	void UpdateStats();
+
+	/** Gathers update stats when it has finished. */
+	void FinishUpdateMesh(const TSharedRef<FUpdateContextPrivate>& Context);
+
+	void FinishUpdateImage(const FString& CustomizableObjectPathName, const FString& InstancePathName, double TaskUpdateImageTime, const int64 TaskUpdateImageMemoryPeak, const int64 TaskUpdateImageRealMemoryPeak) const;
+	
+	DECLARE_BENCHMARK_STAT(NumAllocatedTextures, uint32);
+	DECLARE_BENCHMARK_STAT(TextureGPUSize, uint32);
+
+	DECLARE_BENCHMARK_STAT(NumAllocatedSkeletalMeshes, int32);
+
+	DECLARE_BENCHMARK_STAT(NumInstances, int32);
+	DECLARE_BENCHMARK_STAT(NumInstancesLOD0, int32);
+	DECLARE_BENCHMARK_STAT(NumInstancesLOD1, int32);
+	DECLARE_BENCHMARK_STAT(NumInstancesLOD2, int32);
+
+	DECLARE_BENCHMARK_STAT(NumPendingInstanceUpdates, int32)
+	DECLARE_BENCHMARK_STAT(NumBuiltInstances, int32)
+	DECLARE_BENCHMARK_STAT(InstanceBuildTimeAvrg, double);
+
+	TArray<TWeakObjectPtr<UTexture2D>> TextureTrackerArray;
+
+private:
+	double TotalUpdateTime = 0;
+	uint32 NumUpdates = 0;
+
+	TSharedPtr<FArchive> Archive;
+};
diff --git a/Engine/Plugins/Experimental/Mutable/Source/CustomizableObject/Private/MuCO/CustomizableObjectInstance.cpp b/Engine/Plugins/Experimental/Mutable/Source/CustomizableObject/Private/MuCO/CustomizableObjectInstance.cpp
--- a/Engine/Plugins/Experimental/Mutable/Source/CustomizableObject/Private/MuCO/CustomizableObjectInstance.cpp
+++ b/Engine/Plugins/Experimental/Mutable/Source/CustomizableObject/Private/MuCO/CustomizableObjectInstance.cpp
@@ -21,7 +21,7 @@
 
 #include "MuCO/CustomizableObjectSystemPrivate.h"
 #include "MuCO/CustomizableInstanceLODManagement.h"
-#include "MuCO/CustomizableInstancePrivateData.h"
+#include "MuCO/CustomizableObjectInstancePrivate.h"
 #include "MuCO/CustomizableObjectExtension.h"
 #include "MuCO/CustomizableObjectMipDataProvider.h"
 #include "MuCO/CustomizableObjectPrivate.h"
diff --git a/Engine/Plugins/Experimental/Mutable/Source/CustomizableObject/Private/MuCO/CustomizableObjectInstanceDescriptor.cpp b/Engine/Plugins/Experimental/Mutable/Source/CustomizableObject/Private/MuCO/CustomizableObjectInstanceDescriptor.cpp
--- a/Engine/Plugins/Experimental/Mutable/Source/CustomizableObject/Private/MuCO/CustomizableObjectInstanceDescriptor.cpp
+++ b/Engine/Plugins/Experimental/Mutable/Source/CustomizableObject/Private/MuCO/CustomizableObjectInstanceDescriptor.cpp
@@ -2,7 +2,8 @@
 
 #include "MuCO/CustomizableObjectInstanceDescriptor.h"
 
-#include "CustomizableObjectSystemPrivate.h"
+#include "MuCO/CustomizableObjectInstancePrivate.h"
+#include "MuCO/CustomizableObjectSystemPrivate.h"
 #include "UnrealMutableImageProvider.h"
 #include "MuCO/CustomizableObject.h"
 #include "MuCO/CustomizableObjectPrivate.h"
diff --git a/Engine/Plugins/Experimental/Mutable/Source/CustomizableObject/Private/MuCO/CustomizableObjectInstanceUsage.cpp b/Engine/Plugins/Experimental/Mutable/Source/CustomizableObject/Private/MuCO/CustomizableObjectInstanceUsage.cpp
--- a/Engine/Plugins/Experimental/Mutable/Source/CustomizableObject/Private/MuCO/CustomizableObjectInstanceUsage.cpp
+++ b/Engine/Plugins/Experimental/Mutable/Source/CustomizableObject/Private/MuCO/CustomizableObjectInstanceUsage.cpp
@@ -3,7 +3,7 @@
 #include "MuCO/CustomizableObjectInstanceUsage.h"
 #include "MuCO/CustomizableSkeletalComponent.h"
 #include "MuCO/CustomizableObjectSystem.h"
-#include "MuCO/CustomizableInstancePrivateData.h"
+#include "MuCO/CustomizableObjectInstancePrivate.h"
 #include "MuCO/UnrealPortabilityHelpers.h"
 
 #include "AnimationRuntime.h"
diff --git a/Engine/Plugins/Experimental/Mutable/Source/CustomizableObject/Private/MuCO/CustomizableObjectSystem.cpp b/Engine/Plugins/Experimental/Mutable/Source/CustomizableObject/Private/MuCO/CustomizableObjectSystem.cpp
--- a/Engine/Plugins/Experimental/Mutable/Source/CustomizableObject/Private/MuCO/CustomizableObjectSystem.cpp
+++ b/Engine/Plugins/Experimental/Mutable/Source/CustomizableObject/Private/MuCO/CustomizableObjectSystem.cpp
@@ -12,7 +12,7 @@
 #include "GameFramework/PlayerController.h"
 #include "Interfaces/ITargetPlatform.h"
 #include "MuCO/CustomizableInstanceLODManagement.h"
-#include "MuCO/CustomizableInstancePrivateData.h"
+#include "MuCO/CustomizableObjectInstancePrivate.h"
 #include "MuCO/CustomizableObjectPrivate.h"
 #include "MuCO/CustomizableObjectUIData.h"
 #include "MuCO/DefaultImageProvider.h"
@@ -31,7 +31,7 @@
 #include "ContentStreaming.h"
 #include "MuCO/EditorImageProvider.h"
 #include "Components/SkeletalMeshComponent.h"
-#include "CustomizableObjectSystemPrivate.h"
+#include "MuCO/CustomizableObjectSystemPrivate.h"
 
 #if WITH_EDITOR
 #include "Editor.h"
diff --git a/Engine/Plugins/Experimental/Mutable/Source/CustomizableObject/Private/MuCO/CustomizableObjectSystemPrivate.h b/Engine/Plugins/Experimental/Mutable/Source/CustomizableObject/Private/MuCO/CustomizableObjectSystemPrivate.h
deleted file mode 100644
diff --git a/Engine/Plugins/Experimental/Mutable/Source/CustomizableObject/Private/MuCO/LogBenchmarkUtil.cpp b/Engine/Plugins/Experimental/Mutable/Source/CustomizableObject/Private/MuCO/LogBenchmarkUtil.cpp
--- a/Engine/Plugins/Experimental/Mutable/Source/CustomizableObject/Private/MuCO/LogBenchmarkUtil.cpp
+++ b/Engine/Plugins/Experimental/Mutable/Source/CustomizableObject/Private/MuCO/LogBenchmarkUtil.cpp
@@ -17,7 +17,7 @@
 #include "MuCO/UnrealPortabilityHelpers.h"
 #include "MuCO/CustomizableObjectSystemPrivate.h"
 #include "MuCO/CustomizableObjectSystem.h"
-#include "MuCO/CustomizableInstancePrivateData.h"
+#include "MuCO/CustomizableObjectInstancePrivate.h"
 #include "MuCO/CustomizableObjectInstanceUsage.h"
 
 
diff --git a/Engine/Plugins/Experimental/Mutable/Source/CustomizableObject/Private/MuCO/LogBenchmarkUtil.h b/Engine/Plugins/Experimental/Mutable/Source/CustomizableObject/Private/MuCO/LogBenchmarkUtil.h
deleted file mode 100644
diff --git a/Engine/Plugins/Experimental/Mutable/Source/CustomizableObject/Private/MuCO/LogInformationUtil.cpp b/Engine/Plugins/Experimental/Mutable/Source/CustomizableObject/Private/MuCO/LogInformationUtil.cpp
--- a/Engine/Plugins/Experimental/Mutable/Source/CustomizableObject/Private/MuCO/LogInformationUtil.cpp
+++ b/Engine/Plugins/Experimental/Mutable/Source/CustomizableObject/Private/MuCO/LogInformationUtil.cpp
@@ -2,11 +2,11 @@
 
 #include "MuCO/LogInformationUtil.h"
 
-#include "CustomizableObjectSystemPrivate.h"
+#include "MuCO/CustomizableObjectSystemPrivate.h"
 #include "Engine/SkeletalMesh.h"
 #include "Engine/World.h"
 #include "GameFramework/PlayerController.h"
-#include "MuCO/CustomizableInstancePrivateData.h"
+#include "MuCO/CustomizableObjectInstancePrivate.h"
 #include "TextureResource.h"
 #include "MuCO/CustomizableObject.h"
 
diff --git a/Engine/Plugins/Experimental/Mutable/Source/CustomizableObjectEditor/Private/MuCOE/CustomizableInstanceDetails.cpp b/Engine/Plugins/Experimental/Mutable/Source/CustomizableObjectEditor/Private/MuCOE/CustomizableInstanceDetails.cpp
--- a/Engine/Plugins/Experimental/Mutable/Source/CustomizableObjectEditor/Private/MuCOE/CustomizableInstanceDetails.cpp
+++ b/Engine/Plugins/Experimental/Mutable/Source/CustomizableObjectEditor/Private/MuCOE/CustomizableInstanceDetails.cpp
@@ -14,7 +14,7 @@
 #include "Serialization/BufferArchive.h"
 #include "Toolkits/ToolkitManager.h"
 
-#include "MuCO/CustomizableInstancePrivateData.h"
+#include "MuCO/CustomizableObjectInstancePrivate.h"
 #include "MuCO/CustomizableObject.h"
 #include "MuCO/CustomizableObjectInstance.h"
 #include "MuCO/CustomizableObjectSystem.h"
diff --git a/Engine/Plugins/Experimental/Mutable/Source/CustomizableObjectEditor/Private/MuCOE/CustomizableObjectEditor.cpp b/Engine/Plugins/Experimental/Mutable/Source/CustomizableObjectEditor/Private/MuCOE/CustomizableObjectEditor.cpp
--- a/Engine/Plugins/Experimental/Mutable/Source/CustomizableObjectEditor/Private/MuCOE/CustomizableObjectEditor.cpp
+++ b/Engine/Plugins/Experimental/Mutable/Source/CustomizableObjectEditor/Private/MuCOE/CustomizableObjectEditor.cpp
@@ -52,7 +52,7 @@
 #include "PropertyEditorModule.h"
 #include "ScopedTransaction.h"
 #include "MuCOE/SCustomizableInstanceProperties.h"
-#include "MuCO/CustomizableInstancePrivateData.h"
+#include "MuCO/CustomizableObjectInstancePrivate.h"
 #include "MuCO/CustomizableObjectPrivate.h"
 #include "UObject/EnumProperty.h"
 #include "Widgets/Docking/SDockTab.h"
diff --git a/Engine/Plugins/Experimental/Mutable/Source/CustomizableObjectEditor/Private/MuCOE/CustomizableObjectInstanceEditor.cpp b/Engine/Plugins/Experimental/Mutable/Source/CustomizableObjectEditor/Private/MuCOE/CustomizableObjectInstanceEditor.cpp
--- a/Engine/Plugins/Experimental/Mutable/Source/CustomizableObjectEditor/Private/MuCOE/CustomizableObjectInstanceEditor.cpp
+++ b/Engine/Plugins/Experimental/Mutable/Source/CustomizableObjectEditor/Private/MuCOE/CustomizableObjectInstanceEditor.cpp
@@ -34,7 +34,7 @@
 #include "Widgets/Notifications/SNotificationList.h"
 #include "WorkspaceMenuStructure.h"
 #include "WorkspaceMenuStructureModule.h"
-#include "MuCO/CustomizableInstancePrivateData.h"
+#include "MuCO/CustomizableObjectInstancePrivate.h"
 
 class FAdvancedPreviewScene;
 class FProperty;
diff --git a/Engine/Plugins/Experimental/Mutable/Source/CustomizableObjectEditor/Private/MuCOE/SCustomizableObjectEditorViewportToolBar.cpp b/Engine/Plugins/Experimental/Mutable/Source/CustomizableObjectEditor/Private/MuCOE/SCustomizableObjectEditorViewportToolBar.cpp
--- a/Engine/Plugins/Experimental/Mutable/Source/CustomizableObjectEditor/Private/MuCOE/SCustomizableObjectEditorViewportToolBar.cpp
+++ b/Engine/Plugins/Experimental/Mutable/Source/CustomizableObjectEditor/Private/MuCOE/SCustomizableObjectEditorViewportToolBar.cpp
@@ -22,7 +22,7 @@
 #include "Settings/LevelEditorViewportSettings.h"
 #include "UnrealEdGlobals.h"
 #include "AssetRegistry/AssetRegistryModule.h"
-#include "MuCO/CustomizableInstancePrivateData.h"
+#include "MuCO/CustomizableObjectInstancePrivate.h"
 #include "MuCO/CustomizableObject.h"
 #include "MuCO/CustomizableObjectPrivate.h"
 #include "MuCO/CustomizableObjectSystem.h"
