[mutable] Fixed instance baking not respecting sRGB settings for textures in an uncompressed RGB* format.
#rnx

#tests Editor

diff --git a/Engine/Plugins/Experimental/Mutable/Source/CustomizableObjectEditor/Private/MuCOE/UnrealBakeHelpers.cpp b/Engine/Plugins/Experimental/Mutable/Source/CustomizableObjectEditor/Private/MuCOE/UnrealBakeHelpers.cpp
--- a/Engine/Plugins/Experimental/Mutable/Source/CustomizableObjectEditor/Private/MuCOE/UnrealBakeHelpers.cpp
+++ b/Engine/Plugins/Experimental/Mutable/Source/CustomizableObjectEditor/Private/MuCOE/UnrealBakeHelpers.cpp
@@ -257,6 +257,8 @@ UTexture2D* FUnrealBakeHelpers::BakeHelper_CreateAssetTexture(UTexture2D* SrcTex
 	{
 		FTextureFormatSettings Settings;
 
+		Settings.SRGB = SrcTex->SRGB;
+
 		if (bNeeds_TC_Grayscale)
 		{
 			// If compression settings are not set to TC_Grayscale the texture will get a DXT format
