[Mutable] Instance Descriptor copy.

Now, changes to the Instance Descriptor only take effect before calling UpdateSkeletalMeshAsync. Changes after the call, will only take effect on the next update.

#rnx

diff --git a/Engine/Plugins/Experimental/Mutable/Source/CustomizableObject/Classes/MuCO/CustomizableObjectInstance.h b/Engine/Plugins/Experimental/Mutable/Source/CustomizableObject/Classes/MuCO/CustomizableObjectInstance.h
--- a/Engine/Plugins/Experimental/Mutable/Source/CustomizableObject/Classes/MuCO/CustomizableObjectInstance.h
+++ b/Engine/Plugins/Experimental/Mutable/Source/CustomizableObject/Classes/MuCO/CustomizableObjectInstance.h
@@ -635,7 +635,6 @@ public:
 	UFUNCTION(BlueprintCallable, Category = CustomizableObjectInstance)
 	int32 GetNumComponents() const;
 
-
 	int32 GetMinLODToLoad() const;
 	int32 GetMaxLODToLoad() const;
 	int32 GetNumLODsAvailable() const;
@@ -649,9 +648,6 @@ public:
 	/** Return the Max LOD this Instance is using (from the beginning of an update. If an update fails this value will be incorrect). */
 	int32 GetCurrentMaxLOD() const;
 
-	/** Save the Min and Max LOD that will be used for the update. */
-	void CommitMinMaxLOD();
-
 	/** Sets an array of LODs to generate per component. Mutable will generate those plus the currently generated LODs (if any).
 	 * Requires mutable.EnableOnlyGenerateRequestedLODs and CurrentInstanceLODManagement->IsOnlyGenerateRequestedLODLevelsEnabled() to be true.
 	 * @param InMinLOD - MinLOD to generate.
@@ -690,10 +686,6 @@ private:
 	UPROPERTY()
 	TObjectPtr<UCustomizableInstancePrivate> PrivateData;
 
-	/** LODs applied on the beginning of the last update. Represent the actual LODs the Instance is using (not strictly true since an update can fail). */
-	int32 CurrentMinLOD = -1;
-	int32 CurrentMaxLOD = -1;
-
 #if WITH_EDITORONLY_DATA
 	/** Textures which can used as values in Texture Parameters. */
 	UPROPERTY(EditAnywhere, Category = TextureParameter)
diff --git a/Engine/Plugins/Experimental/Mutable/Source/CustomizableObject/Internal/MuCO/CustomizableObjectInstancePrivate.h b/Engine/Plugins/Experimental/Mutable/Source/CustomizableObject/Internal/MuCO/CustomizableObjectInstancePrivate.h
--- a/Engine/Plugins/Experimental/Mutable/Source/CustomizableObject/Internal/MuCO/CustomizableObjectInstancePrivate.h
+++ b/Engine/Plugins/Experimental/Mutable/Source/CustomizableObject/Internal/MuCO/CustomizableObjectInstancePrivate.h
@@ -189,11 +189,6 @@ public:
 	// Only used in LiveUpdateMode to reuse core instances between updates and their temp data to speed up updates, but spend way more memory
 	mu::Instance::ID LiveUpdateModeInstanceID = 0;
 
-	/** Cached Texture Parameters from the last update.
-	 * This cache is required since the Instance can have a LOD Update at any time.
-	 * So we need to make sure that the initially provided Texture Parameters by the user will be available until the user decides to change them. */
-	TArray<FName> UpdateTextureParameters;
-
 #if WITH_EDITOR
 
 	virtual void PostDuplicate(bool bDuplicateForPIE) override;
@@ -224,7 +219,7 @@ public:
 	FGraphEventRef LoadAdditionalAssetsAsync(const TSharedRef<FUpdateContextPrivate>& OperationData, UCustomizableObjectInstance* Public, struct FStreamableManager &StreamableManager);
 	void AdditionalAssetsAsyncLoaded(UCustomizableObjectInstance* Public);
 
-	void TickUpdateCloseCustomizableObjects(UCustomizableObjectInstance& Publics, FMutableInstanceUpdateMap& InOutRequestedUpdates);
+	void TickUpdateCloseCustomizableObjects(UCustomizableObjectInstance& Public, FMutableInstanceUpdateMap& InOutRequestedUpdates);
 	void UpdateInstanceIfNotGenerated(UCustomizableObjectInstance& Public, FMutableInstanceUpdateMap& InOutRequestedUpdates);
 
 	// Returns true if success (?)
@@ -331,16 +326,16 @@ public:
 	TArray<TObjectPtr<UTexture>> LoadedPassThroughTexturesPendingSetMaterial;
 
 private:
-	
 	ECOInstanceFlags InstanceFlagsPrivate = ECOInstanceFlags::ECONone;
 
 public:
-	/** Hash of the UCustomizableObjectInstance::Descriptor on the last update request. */
-	FDescriptorHash UpdateDescriptorHash;
-	
-	/** Hash of the UCustomizableObjectInstance::Descriptor on the last successful update. */
-	FDescriptorHash DescriptorHash;
+	/** Copy of the descriptor of the latest successful update. */
+	UPROPERTY(Transient)
+	FCustomizableObjectInstanceDescriptor CommittedDescriptor;
 
+	/** Hash of the descriptor copy of the latest successful update. */
+	FDescriptorHash CommittedDescriptorHash;
+	
 	/** Status of the generated Skeletal Mesh. Not to be confused with the Update Result. */
 	ESkeletalMeshStatus SkeletalMeshStatus = ESkeletalMeshStatus::NotGenerated;
 };
diff --git a/Engine/Plugins/Experimental/Mutable/Source/CustomizableObject/Internal/MuCO/CustomizableObjectSystemPrivate.h b/Engine/Plugins/Experimental/Mutable/Source/CustomizableObject/Internal/MuCO/CustomizableObjectSystemPrivate.h
--- a/Engine/Plugins/Experimental/Mutable/Source/CustomizableObject/Internal/MuCO/CustomizableObjectSystemPrivate.h
+++ b/Engine/Plugins/Experimental/Mutable/Source/CustomizableObject/Internal/MuCO/CustomizableObjectSystemPrivate.h
@@ -417,11 +417,31 @@ struct FInstanceUpdateData
 class FUpdateContextPrivate : public FGCObject
 {
 public:
+	FUpdateContextPrivate(UCustomizableObjectInstance& InInstance, const FCustomizableObjectInstanceDescriptor& Descriptor);
+
 	FUpdateContextPrivate(UCustomizableObjectInstance& InInstance);
 	virtual ~FUpdateContextPrivate() override;
 
 	virtual FString GetReferencerName() const override;
 	virtual void AddReferencedObjects(FReferenceCollector& Collector) override;
+
+	int32 GetMinLOD() const;
+
+	void SetMinLOD(int32 MinLOD);
+
+	int32 GetMaxLOD() const;
+
+	void SetMaxLOD(int32 MaxLOD);
+
+	const TArray<uint16>& GetRequestedLODs() const;
+
+	void SetRequestedLODs(TArray<uint16>& RequestedLODs);
+
+	const FCustomizableObjectInstanceDescriptor& GetCapturedDescriptor() const;
+
+	const FDescriptorHash& GetCapturedDescriptorHash() const;
+
+	const FCustomizableObjectInstanceDescriptor&& MoveCommittedDescriptor();
 	
 	EQueuePriorityType PriorityType = EQueuePriorityType::Low;
 	
@@ -432,19 +452,17 @@ public:
 	 *It is weak because we don't want to lock it in case it becomes irrelevant in the game while operations are pending and it needs to be destroyed. */
 	TWeakObjectPtr<UCustomizableObjectInstance> Instance;
 
-	/** Hash of the UCustomizableObjectInstance::Descriptor at the time of the update request. */
-	FDescriptorHash InstanceDescriptorHash;
-			
+private:
+	/** Descriptor which the update will be performed on. */
+	FCustomizableObjectInstanceDescriptor CapturedDescriptor;
+
+	/** Hash of the descriptor. */
+	FDescriptorHash CapturedDescriptorHash;
+
+public:
 	/** Instance parameters at the time of the operation request. */
 	mu::ParametersPtr Parameters; 
 	
-	TArray<FName> TextureParameters;
-
-	bool bBuildParameterRelevancy = false;
-
-	/** Instance state. */
-	int32 State = 0;
-
 	bool bOnlyUpdateIfNotGenerated = false;
 	bool bIgnoreCloseDist = false;
 	bool bForceHighPriority = false;
@@ -471,11 +489,6 @@ public:
 	int32 NumComponents = 0;
 	int32 NumLODsAvailable = 0;
 
-	int32 CurrentMinLOD = 0;
-	int32 CurrentMaxLOD = 0;
-
-	TArray<uint16> RequestedLODs;
-
 	TMap<uint32, FTexturePlatformData*> ImageToPlatformDataMap;
 
 	EUpdateResult UpdateResult = EUpdateResult::Success;
@@ -650,6 +663,10 @@ public:
 	/** Update stats at each tick.
 	 * Used for stats that are costly to update. */
 	void UpdateStats();
+
+	void CacheTextureParameters(const TArray<FCustomizableObjectTextureParameterValue>& TextureParameters) const;
+
+	void UnCacheTextureParameters(const TArray<FCustomizableObjectTextureParameterValue>& TextureParameters) const;
 	
 	/** Mutable TaskGraph system (Mutable Thread). */
 	FMutableTaskGraph MutableTaskGraph;
diff --git a/Engine/Plugins/Experimental/Mutable/Source/CustomizableObject/Internal/MuCO/DescriptorHash.h b/Engine/Plugins/Experimental/Mutable/Source/CustomizableObject/Internal/MuCO/DescriptorHash.h
--- a/Engine/Plugins/Experimental/Mutable/Source/CustomizableObject/Internal/MuCO/DescriptorHash.h
+++ b/Engine/Plugins/Experimental/Mutable/Source/CustomizableObject/Internal/MuCO/DescriptorHash.h
@@ -21,21 +21,12 @@ public:
 	/** Return true if this Hash is a subset of the other Hash (i.e., this Descriptor is a subset of the other Descriptor). */
 	bool IsSubset(const FDescriptorHash& Other) const;
 
-	void UpdateMinMaxLOD(int32 InMinLOD, int32 InMaxLOD);
-
-	int32 GetMinLOD() const;
-
-	int32 GetMaxLOD() const;
-
-	void UpdateRequestedLODs(const TArray<uint16>& InRequestedLODs);
-
-	const TArray<uint16>& GetRequestedLODs() const;
-
 	FString ToString() const;
 
 private:
 	uint32 Hash = 0;
 
+public:
 	int32 MinLOD = 0;
 	int32 MaxLOD = INT32_MAX;
 
diff --git a/Engine/Plugins/Experimental/Mutable/Source/CustomizableObject/Private/MuCO/CustomizableObjectInstance.cpp b/Engine/Plugins/Experimental/Mutable/Source/CustomizableObject/Private/MuCO/CustomizableObjectInstance.cpp
--- a/Engine/Plugins/Experimental/Mutable/Source/CustomizableObject/Private/MuCO/CustomizableObjectInstance.cpp
+++ b/Engine/Plugins/Experimental/Mutable/Source/CustomizableObject/Private/MuCO/CustomizableObjectInstance.cpp
@@ -194,8 +194,8 @@ void UCustomizableInstancePrivate::InvalidateGeneratedData()
 	SkeletalMeshStatus = ESkeletalMeshStatus::NotGenerated;
 	SkeletalMeshes.Reset();
 
-	DescriptorHash = FDescriptorHash();
-	UpdateDescriptorHash = FDescriptorHash();
+	CommittedDescriptor = {};
+	CommittedDescriptorHash = {};
 
 	// Init Component Data
 	FCustomizableInstanceComponentData TemplateComponentData;
@@ -478,12 +478,7 @@ void UCustomizableInstancePrivate::ReleaseMutableResources(bool bCalledFromBegin
 		// Remove all references to cached Texture Parameters. Only if we're destroying the COI.
 		if (bCalledFromBeginDestroy)
 		{
-			FUnrealMutableImageProvider* ImageProvider = UCustomizableObjectSystem::GetInstance()->GetPrivateChecked()->GetImageProviderChecked();
-
-			for (const FName& TextureParameter : UpdateTextureParameters)
-			{
-				ImageProvider->UnCacheImage(TextureParameter, false);
-			}
+			CustomizableObjectSystem->UnCacheTextureParameters(CommittedDescriptor.GetTextureParameters());
 		}
 	}
 
@@ -1734,7 +1729,7 @@ bool UCustomizableInstancePrivate::DoComponentsNeedUpdate(UCustomizableObjectIns
 	MeshIDs.Init(MAX_uint64, NumComponents * MAX_MESH_LOD_COUNT);
 
 	// Gather the Mesh Ids of all components, and validate the integrity of the meshes to generate. 
-	for (int32 LODIndex = OperationData->CurrentMinLOD; LODIndex <= OperationData->CurrentMaxLOD && LODIndex < NumLODs; ++LODIndex)
+	for (int32 LODIndex = OperationData->GetMinLOD(); LODIndex <= OperationData->GetMaxLOD() && LODIndex < NumLODs; ++LODIndex)
 	{
 		const FInstanceUpdateData::FLOD& LOD = OperationData->InstanceUpdateData.LODs[LODIndex];
 
@@ -1759,7 +1754,7 @@ bool UCustomizableInstancePrivate::DoComponentsNeedUpdate(UCustomizableObjectIns
 			{
 				UE_LOG(LogMutable, Error, TEXT("Failed to generate SkeletalMesh for CO Instance [%s]. CO [%s] has invalid geometry for LOD [%d] Component [%d]."),
 					*Public->GetName(), *CustomizableObject->GetName(),
-					OperationData->CurrentMaxLOD, ComponentIndex);
+					OperationData->GetMaxLOD(), ComponentIndex);
 				bHasInvalidMesh = true;
 				continue;
 			}
@@ -1785,11 +1780,11 @@ bool UCustomizableInstancePrivate::DoComponentsNeedUpdate(UCustomizableObjectIns
 		}
 
 		// Components with mesh must have valid geometry at CurrentMaxLOD
-		if (ComponentWithMesh[ComponentIndex] && MeshIDs[ComponentIndex * MAX_MESH_LOD_COUNT + OperationData->CurrentMaxLOD] == MAX_uint64)
+		if (ComponentWithMesh[ComponentIndex] && MeshIDs[ComponentIndex * MAX_MESH_LOD_COUNT + OperationData->GetMaxLOD()] == MAX_uint64)
 		{
 			UE_LOG(LogMutable, Error, TEXT("Failed to generate SkeletalMesh for CO Instance [%s]. CO [%s] is missing geometry for LOD [%d] Component [%d]."),
 				*Public->GetName(), *CustomizableObject->GetName(),
-				OperationData->CurrentMaxLOD, ComponentIndex);
+				OperationData->GetMaxLOD(), ComponentIndex);
 			bHasInvalidMesh = true;
 			continue;
 		}
@@ -2932,29 +2927,17 @@ void FMutableUpdateCandidate::ApplyLODUpdateParamsToInstance(FUpdateContextPriva
 {
 	CustomizableObjectInstance->Descriptor.MinLOD = MinLOD;
 	CustomizableObjectInstance->Descriptor.MaxLOD = MaxLOD;
-
 	CustomizableObjectInstance->Descriptor.RequestedLODLevels = RequestedLODLevels;
 
-	CustomizableObjectInstance->GetPrivate()->UpdateDescriptorHash.UpdateMinMaxLOD(MinLOD, MaxLOD);
-	CustomizableObjectInstance->GetPrivate()->UpdateDescriptorHash.UpdateRequestedLODs(CustomizableObjectInstance->Descriptor.RequestedLODLevels);
-
 	if (Context)
 	{
-		Context->InstanceDescriptorHash = CustomizableObjectInstance->GetPrivate()->UpdateDescriptorHash;
-		Context->CurrentMinLOD = MinLOD;
-		Context->CurrentMaxLOD = MaxLOD;
-		Context->RequestedLODs = RequestedLODLevels;
+		Context->SetMinLOD(MinLOD);
+		Context->SetMaxLOD(MaxLOD);
+		Context->SetRequestedLODs(RequestedLODLevels);
 	}
 }
 
 
-void UCustomizableObjectInstance::CommitMinMaxLOD()
-{
-	CurrentMinLOD = Descriptor.MinLOD;
-	CurrentMaxLOD = Descriptor.MaxLOD;	
-}
-
-
 // The memory allocated in the function and pointed by the returned pointer is owned by the caller and must be freed. 
 // If assigned to a UTexture2D, it will be freed by that UTexture2D
 FTexturePlatformData* MutableCreateImagePlatformData(mu::Ptr<const mu::Image> MutableImage, int32 OnlyLOD, uint16 FullSizeX, uint16 FullSizeY)
@@ -3366,7 +3349,7 @@ void UCustomizableInstancePrivate::InitSkeletalMeshData(const TSharedRef<FUpdate
 		SkeletalMesh->AllocateResourceForRendering();
 
 		FSkeletalMeshRenderData* RenderData = SkeletalMesh->GetResourceForRendering();
-		for (int32 LODIndex = 0; LODIndex <= OperationData->CurrentMaxLOD; ++LODIndex)
+		for (int32 LODIndex = 0; LODIndex <= OperationData->GetMaxLOD(); ++LODIndex)
 		{
 			RenderData->LODRenderData.Add(new FSkeletalMeshLODRenderData());
 
@@ -3634,7 +3617,7 @@ void UCustomizableInstancePrivate::BuildOrCopyElementData(const TSharedRef<FUpda
 {
 	MUTABLE_CPUPROFILER_SCOPE(UCustomizableInstancePrivate::BuildOrCopyElementData);
 
-	for (int32 LODIndex = OperationData->CurrentMaxLOD; LODIndex >= FirstLODAvailable; --LODIndex)
+	for (int32 LODIndex = OperationData->GetMaxLOD(); LODIndex >= FirstLODAvailable; --LODIndex)
 	{
 		const FInstanceUpdateData::FLOD& LOD = OperationData->InstanceUpdateData.LODs[LODIndex];
 
@@ -3691,12 +3674,12 @@ void UCustomizableInstancePrivate::BuildOrCopyMorphTargetsData(const TSharedRef<
 
 	const int32 SkeletalMeshMorphTargetsNum = SkeletalMesh->GetMorphTargets().Num();
 
-	int32 LastValidLODIndex = OperationData->CurrentMaxLOD;
-	for (int32 LODIndex = OperationData->CurrentMaxLOD; LODIndex >= FirstLODAvailable; --LODIndex)
+	int32 LastValidLODIndex = OperationData->GetMaxLOD();
+	for (int32 LODIndex = OperationData->GetMaxLOD(); LODIndex >= FirstLODAvailable; --LODIndex)
 	{
 		const FInstanceUpdateData::FLOD& LOD = OperationData->InstanceUpdateData.LODs[LODIndex];
 
-		if (LODIndex >= OperationData->CurrentMinLOD && OperationData->InstanceUpdateData.Components[LOD.FirstComponent + ComponentIndex].bGenerated)
+		if (LODIndex >= OperationData->GetMinLOD() && OperationData->InstanceUpdateData.Components[LOD.FirstComponent + ComponentIndex].bGenerated)
 		{
 			LastValidLODIndex = LODIndex;
 		}
@@ -3852,7 +3835,7 @@ void UCustomizableInstancePrivate::BuildOrCopyClothingData(const TSharedRef<FUpd
 	{
 		MUTABLE_CPUPROFILER_SCOPE(DiscoverSectionsWithCloth);
 
-		for (int32 LODIndex = OperationData->CurrentMinLOD; LODIndex <= OperationData->CurrentMaxLOD; ++LODIndex)
+		for (int32 LODIndex = OperationData->GetMinLOD(); LODIndex <= OperationData->GetMaxLOD(); ++LODIndex)
 		{
 			const FInstanceUpdateData::FLOD& LOD = OperationData->InstanceUpdateData.LODs[LODIndex];
 			if (ComponentIndex >= LOD.ComponentCount)
@@ -3978,7 +3961,7 @@ void UCustomizableInstancePrivate::BuildOrCopyClothingData(const TSharedRef<FUpd
 			// Only initilialize once, multiple sections with cloth could point to the same cloth asset.
 			if (!DstAssetData.LodMap.Num())
 			{
-				DstAssetData.LodMap.Init( INDEX_NONE, OperationData->CurrentMaxLOD + 1 );
+				DstAssetData.LodMap.Init( INDEX_NONE, OperationData->GetMaxLOD() + 1 );
 
 				DstAssetData.LodData.SetNum( SrcAssetData.LodData.Num() );
 				DstAssetData.UsedBoneNames = SrcAssetData.UsedBoneNames;
@@ -4512,7 +4495,7 @@ void UCustomizableInstancePrivate::BuildOrCopyClothingData(const TSharedRef<FUpd
 	{
 		int32 NumSectionsWithClothProcessed = 0;
 
-		for (int32 LODIndex = OperationData->CurrentMinLOD; LODIndex <= OperationData->CurrentMaxLOD; ++LODIndex)
+		for (int32 LODIndex = OperationData->GetMinLOD(); LODIndex <= OperationData->GetMaxLOD(); ++LODIndex)
 		{
 			const FInstanceUpdateData::FLOD& LOD = OperationData->InstanceUpdateData.LODs[LODIndex];
 			if (ComponentIndex >= LOD.ComponentCount)
@@ -4565,7 +4548,7 @@ void UCustomizableInstancePrivate::BuildOrCopyClothingData(const TSharedRef<FUpd
 	{
 		MUTABLE_CPUPROFILER_SCOPE(InitClothRenderData)
 		// Based on FSkeletalMeshLODModel::GetClothMappingData().
-		for (int32 LODIndex = OperationData->CurrentMinLOD; LODIndex <= OperationData->CurrentMaxLOD; ++LODIndex)
+		for (int32 LODIndex = OperationData->GetMinLOD(); LODIndex <= OperationData->GetMaxLOD(); ++LODIndex)
 		{
 			FSkeletalMeshLODRenderData& LODModel = RenderResource->LODRenderData[LODIndex];
 	
@@ -4683,8 +4666,8 @@ bool UCustomizableInstancePrivate::BuildOrCopyRenderData(const TSharedRef<FUpdat
 	check(CustomizableObject);
 	const FModelResources& ModelResources = CustomizableObject->GetPrivate()->GetModelResources();
 
-	int32 LastValidLODIndex = OperationData->CurrentMaxLOD;
-	for (int32 LODIndex = OperationData->CurrentMaxLOD; LODIndex >= FirstLODAvailable; --LODIndex)
+	int32 LastValidLODIndex = OperationData->GetMaxLOD();
+	for (int32 LODIndex = OperationData->GetMaxLOD(); LODIndex >= FirstLODAvailable; --LODIndex)
 	{
 		MUTABLE_CPUPROFILER_SCOPE(BuildOrCopyRenderData_LODLoop);
 
@@ -4694,7 +4677,7 @@ bool UCustomizableInstancePrivate::BuildOrCopyRenderData(const TSharedRef<FUpdat
 		if (ComponentIndex >=  LOD.ComponentCount)
 		{
 			// Interrupt the generation if the LOD is empty and it should have been generated.
-			if (LODIndex >= OperationData->CurrentMinLOD && LODIndex <= OperationData->CurrentMaxLOD)
+			if (LODIndex >= OperationData->GetMinLOD() && LODIndex <= OperationData->GetMaxLOD())
 			{
 				UE_LOG(LogMutable, Warning, TEXT("Building instance: generated mesh [%s] has LOD [%d] with no component.")
 					, *SkeletalMesh->GetName()
@@ -4938,7 +4921,7 @@ FGraphEventRef UCustomizableInstancePrivate::LoadAdditionalAssetsAsync(const TSh
 	// Load assets coming from SubMeshes of the newly generated Mesh
 	if (OperationData->InstanceUpdateData.LODs.Num())
 	{
-		const FInstanceUpdateData::FLOD& LOD = OperationData->InstanceUpdateData.LODs[OperationData->CurrentMinLOD];
+		const FInstanceUpdateData::FLOD& LOD = OperationData->InstanceUpdateData.LODs[OperationData->GetMinLOD()];
 		for (int32 ComponentIndex = 0; ComponentIndex < LOD.ComponentCount; ++ComponentIndex)
 		{
 			const FInstanceUpdateData::FComponent& Component = Components[LOD.FirstComponent + ComponentIndex];
@@ -5420,7 +5403,7 @@ void UCustomizableInstancePrivate::BuildMaterials(const TSharedRef<FUpdateContex
 		UCustomizableObjectSystem::GetInstance()->GetPrivate()->MutableSystem,
 		CustomizableObject->GetModel(),
 		OperationData->Parameters,
-	    OperationData->State);
+	    OperationData->GetCapturedDescriptor().GetState());
 	
 	const bool bReuseTextures = OperationData->bReuseInstanceTextures;
 
@@ -5457,7 +5440,7 @@ void UCustomizableInstancePrivate::BuildMaterials(const TSharedRef<FUpdateContex
 
 		MUTABLE_CPUPROFILER_SCOPE(BuildMaterials_LODLoop);
 
-		for (int32 LODIndex = FirstLODAvailable; LODIndex <= OperationData->CurrentMaxLOD; LODIndex++)
+		for (int32 LODIndex = FirstLODAvailable; LODIndex <= OperationData->GetMaxLOD(); LODIndex++)
 		{
 			const FInstanceUpdateData::FLOD& LOD = OperationData->InstanceUpdateData.LODs[LODIndex];
 			const FInstanceUpdateData::FComponent& Component = OperationData->InstanceUpdateData.Components[LOD.FirstComponent + ComponentIndex];
@@ -5970,8 +5953,8 @@ void UCustomizableInstancePrivate::BuildMaterials(const TSharedRef<FUpdateContex
 
 		{
 			// Copy data from valid LODs into the skipped ones.
-			int32 LastValidLODIndex = OperationData->CurrentMaxLOD;
-			for (int32 LODIndex = OperationData->CurrentMaxLOD; LODIndex >= FirstLODAvailable; --LODIndex)
+			int32 LastValidLODIndex = OperationData->GetMaxLOD();
+			for (int32 LODIndex = OperationData->GetMaxLOD(); LODIndex >= FirstLODAvailable; --LODIndex)
 			{
 				// Copy information from the LastValidLODIndex
 				if (LODsSkipped[LODIndex])
@@ -6178,13 +6161,13 @@ int32 UCustomizableObjectInstance::GetNumLODsAvailable() const
 
 int32 UCustomizableObjectInstance::GetCurrentMinLOD() const
 {
-	return CurrentMinLOD;
+	return GetPrivate()->CommittedDescriptor.GetMinLod();
 }
 
 
 int32 UCustomizableObjectInstance::GetCurrentMaxLOD() const
 {
-	return CurrentMaxLOD;
+	return GetPrivate()->CommittedDescriptor.GetMaxLod();
 }
 
 #if !UE_BUILD_SHIPPING
@@ -6257,9 +6240,9 @@ void UCustomizableObjectInstance::SetRequestedLODs(int32 InMinLOD, int32 InMaxLO
 		const int32 ComponentCount = GetNumComponents();
 		MutableUpdateCandidate.RequestedLODLevels.SetNumZeroed(ComponentCount);
 
-		const TArray<uint16>& GeneratedLODsPerComponent = GetPrivate()->DescriptorHash.GetRequestedLODs();
+		const TArray<uint16>& GeneratedLODsPerComponent = GetPrivate()->CommittedDescriptorHash.RequestedLODsPerComponent;
 
-		const bool bIgnoreGeneratedLODs = GetPrivate()->DescriptorHash != GetPrivate()->UpdateDescriptorHash || GeneratedLODsPerComponent.Num() != ComponentCount;
+		const bool bIgnoreGeneratedLODs = GeneratedLODsPerComponent.Num() != ComponentCount;
 
 		if (bMinMaxLODChanged || bIgnoreGeneratedLODs || Descriptor.GetRequestedLODLevels() != InRequestedLODsPerComponent)
 		{
diff --git a/Engine/Plugins/Experimental/Mutable/Source/CustomizableObject/Private/MuCO/CustomizableObjectSystem.cpp b/Engine/Plugins/Experimental/Mutable/Source/CustomizableObject/Private/MuCO/CustomizableObjectSystem.cpp
--- a/Engine/Plugins/Experimental/Mutable/Source/CustomizableObject/Private/MuCO/CustomizableObjectSystem.cpp
+++ b/Engine/Plugins/Experimental/Mutable/Source/CustomizableObject/Private/MuCO/CustomizableObjectSystem.cpp
@@ -49,9 +49,6 @@
 class AActor;
 class UAnimInstance;
 
-void CacheTexturesParameters(const TArray<FName>& TextureParameters);
-void UnCacheTexturesParameters(const TArray<FName>& TextureParameters);
-
 
 DECLARE_CYCLE_STAT(TEXT("MutablePendingRelease Time"), STAT_MutablePendingRelease, STATGROUP_Game);
 DECLARE_CYCLE_STAT(TEXT("MutableTask"), STAT_MutableTask, STATGROUP_Game);
@@ -128,32 +125,37 @@ static void CVarMutableSinkFunction()
 static FAutoConsoleVariableSink CVarMutableSink(FConsoleCommandDelegate::CreateStatic(&CVarMutableSinkFunction));
 
 
-FUpdateContextPrivate::FUpdateContextPrivate(UCustomizableObjectInstance& InInstance)
+FUpdateContextPrivate::FUpdateContextPrivate(UCustomizableObjectInstance& InInstance, const FCustomizableObjectInstanceDescriptor& Descriptor)
 {
 	check(InInstance.GetPrivate());
 	check(InInstance.GetCustomizableObject());
 
 	Instance = &InInstance;
-	InstanceDescriptorHash = FDescriptorHash(Instance->GetDescriptor());
-	InInstance.GetPrivate()->UpdateDescriptorHash = InstanceDescriptorHash; // TODO GMTFuture Remove on MTBL-1409
-	State = InInstance.GetState();
-	bBuildParameterRelevancy = InInstance.GetBuildParameterRelevancy();
-	Parameters = InInstance.GetDescriptor().GetParameters();
-	TextureParameters = InInstance.GetPrivate()->UpdateTextureParameters;
+	CapturedDescriptor = Descriptor;
+	CapturedDescriptorHash = FDescriptorHash(Descriptor);
+	Parameters = Descriptor.GetParameters();
 	NumComponents = InInstance.GetCustomizableObject()->GetComponentCount();
-	CurrentMinLOD = InInstance.GetCurrentMinLOD();
-	CurrentMaxLOD = InInstance.GetCurrentMaxLOD();
-	RequestedLODs = InInstance.GetRequestedLODsPerComponent();
 	
-	InInstance.GetCustomizableObject()->ApplyStateForcedValuesToParameters(State, Parameters.get());
+	InInstance.GetCustomizableObject()->ApplyStateForcedValuesToParameters(CapturedDescriptor.GetState(), Parameters.get());
+
+	UCustomizableObjectSystem* System = UCustomizableObjectSystem::GetInstance();
+	System->GetPrivate()->CacheTextureParameters(CapturedDescriptor.GetTextureParameters());
+}
+
 
-	CacheTexturesParameters(TextureParameters);
+FUpdateContextPrivate::FUpdateContextPrivate(UCustomizableObjectInstance& InInstance) :
+	FUpdateContextPrivate(InInstance, InInstance.GetDescriptor())
+{
 }
 
 
 FUpdateContextPrivate::~FUpdateContextPrivate()
 {
-	UnCacheTexturesParameters(TextureParameters);
+	if (UCustomizableObjectSystem::IsCreated())
+	{
+		UCustomizableObjectSystem* System = UCustomizableObjectSystem::GetInstance();
+		System->GetPrivate()->UnCacheTextureParameters(CapturedDescriptor.GetTextureParameters());
+	}
 }
 
 
@@ -169,6 +171,63 @@ void FUpdateContextPrivate::AddReferencedObjects(FReferenceCollector& Collector)
 }
 
 
+int32 FUpdateContextPrivate::GetMinLOD() const
+{
+	return CapturedDescriptor.GetMinLod();
+}
+
+
+void FUpdateContextPrivate::SetMinLOD(int32 MinLOD)
+{
+	CapturedDescriptor.SetMinLod(MinLOD);
+	CapturedDescriptorHash.MinLOD = MinLOD;
+}
+
+
+int32 FUpdateContextPrivate::GetMaxLOD() const
+{
+	return CapturedDescriptor.GetMinLod();
+}
+
+
+void FUpdateContextPrivate::SetMaxLOD(int32 MaxLOD)
+{
+	CapturedDescriptor.SetMaxLod(MaxLOD);
+	CapturedDescriptorHash.MaxLOD = MaxLOD;
+}
+
+
+const TArray<uint16>& FUpdateContextPrivate::GetRequestedLODs() const
+{
+	return CapturedDescriptor.GetRequestedLODLevels();
+}
+
+
+void FUpdateContextPrivate::SetRequestedLODs(TArray<uint16>& RequestedLODs)
+{
+	CapturedDescriptor.SetRequestedLODLevels(RequestedLODs);
+	CapturedDescriptorHash.RequestedLODsPerComponent = RequestedLODs;
+}
+
+
+const FCustomizableObjectInstanceDescriptor& FUpdateContextPrivate::GetCapturedDescriptor() const
+{
+	return CapturedDescriptor;
+}
+
+
+const FDescriptorHash& FUpdateContextPrivate::GetCapturedDescriptorHash() const
+{
+	return CapturedDescriptorHash;
+}
+
+
+const FCustomizableObjectInstanceDescriptor&& FUpdateContextPrivate::MoveCommittedDescriptor()
+{
+	return MoveTemp(CapturedDescriptor);	
+}
+
+
 FMutablePendingInstanceUpdate::FMutablePendingInstanceUpdate(const TSharedRef<FUpdateContextPrivate>& InContext) :
 	Context(InContext)
 {
@@ -759,8 +818,20 @@ void FinishUpdateGlobal(const TSharedRef<FUpdateContextPrivate>& Context)
 		case EUpdateResult::Warning:
 			PrivateInstance->SkeletalMeshStatus = ESkeletalMeshStatus::Success;
 
-			PrivateInstance->DescriptorHash = Context->InstanceDescriptorHash;
+			if (SystemPrivate)
+			{
+				SystemPrivate->UnCacheTextureParameters(PrivateInstance->CommittedDescriptor.GetTextureParameters());				
+			}
+
+			PrivateInstance->CommittedDescriptor = Context->MoveCommittedDescriptor();
+			PrivateInstance->CommittedDescriptorHash = Context->GetCapturedDescriptorHash();
 
+			if (SystemPrivate)
+			{
+				// Cache new Texture Parameters
+				SystemPrivate->CacheTextureParameters(PrivateInstance->CommittedDescriptor.GetTextureParameters());
+			}
+			
 			// Delegates must be called only after updating the Instance flags.
 			Instance->UpdatedDelegate.Broadcast(Instance);
 			Instance->UpdatedNativeDelegate.Broadcast(Instance);
@@ -1149,7 +1220,7 @@ void UCustomizableObjectSystemPrivate::EnqueueUpdateSkeletalMesh(const TSharedRe
 
 		if (const FMutablePendingInstanceUpdate* QueueElem = MutablePendingInstanceWork.GetUpdate(Instance))
 		{
-			if (InstancePrivate->UpdateDescriptorHash.IsSubset(FDescriptorHash(QueueElem->Context->InstanceDescriptorHash)))
+			if (Context->GetCapturedDescriptorHash().IsSubset(QueueElem->Context->GetCapturedDescriptorHash()))
 			{
 				Context->UpdateResult = EUpdateResult::ErrorOptimized;
 				FinishUpdateGlobal(Context);			
@@ -1159,14 +1230,14 @@ void UCustomizableObjectSystemPrivate::EnqueueUpdateSkeletalMesh(const TSharedRe
 
 		if (CurrentMutableOperation &&
 			Instance == CurrentMutableOperation->Instance &&
-			InstancePrivate->UpdateDescriptorHash.IsSubset(CurrentMutableOperation->InstanceDescriptorHash))
+			Context->GetCapturedDescriptorHash().IsSubset(CurrentMutableOperation->GetCapturedDescriptorHash()))
 		{
 			Context->UpdateResult = EUpdateResult::ErrorOptimized;
 			FinishUpdateGlobal(Context);
 			return; // The requested update is equal to the running update.
 		}
 	
-		if (InstancePrivate->UpdateDescriptorHash.IsSubset(InstancePrivate->DescriptorHash) &&
+		if (Context->GetCapturedDescriptorHash().IsSubset(InstancePrivate->CommittedDescriptorHash) &&
 			!(CurrentMutableOperation &&
 			Instance == CurrentMutableOperation->Instance)) // This condition is necessary because even if the descriptor is a subset, it will be replaced by the CurrentMutableOperation
 		{
@@ -1176,43 +1247,11 @@ void UCustomizableObjectSystemPrivate::EnqueueUpdateSkeletalMesh(const TSharedRe
 		}
 		else
 		{
-			// Cache Texture Parameters being used during the update:
-			check(ImageProvider);
-
-			// Cache new Texture Parameters
-			for (const FCustomizableObjectTextureParameterValue& TextureParameters : Instance->GetDescriptor().GetTextureParameters())
-			{
-				ImageProvider->CacheImage(TextureParameters.ParameterValue, false);
-
-				for (const FName& TextureParameter : TextureParameters.ParameterRangeValues)
-				{
-					ImageProvider->CacheImage(TextureParameter, false);
-				}
-			}
-
-			// Uncache old Texture Parameters
-			for (const FName& TextureParameter : InstancePrivate->UpdateTextureParameters)
-			{
-				ImageProvider->UnCacheImage(TextureParameter, false);
-			}
-
-			// Update which ones are currently are being used
-			InstancePrivate->UpdateTextureParameters.Reset();
-			for (const FCustomizableObjectTextureParameterValue& TextureParameters : Instance->GetDescriptor().GetTextureParameters())
-			{
-				InstancePrivate->UpdateTextureParameters.Add(TextureParameters.ParameterValue);
-
-				for (const FName& TextureParameter : TextureParameters.ParameterRangeValues)
-				{
-					InstancePrivate->UpdateTextureParameters.Add(TextureParameter);
-				}
-			}
-
 			if (CVarDescriptorDebugPrint->GetBool())
 			{
 				FString String = TEXT("DESCRIPTOR DEBUG PRINT\n");
 				String += "================================\n";				
-				String += FString::Printf(TEXT("=== DESCRIPTOR HASH ===\n%s\n"), *InstancePrivate->UpdateDescriptorHash.ToString());
+				String += FString::Printf(TEXT("=== DESCRIPTOR HASH ===\n%s\n"), *Context->GetCapturedDescriptorHash().ToString());
 				String += FString::Printf(TEXT("=== DESCRIPTOR ===\n%s"), *Instance->GetDescriptor().ToString());
 				String += "================================";
 				
@@ -1600,7 +1639,7 @@ namespace impl
 			UE_LOG(LogMutable, Verbose, TEXT("Creating Mutable instance with id [%d] "), Operation->InstanceID);
 		}
 
-		Operation->MutableInstance = MutableSystem->BeginUpdate(Operation->InstanceID, Operation->Parameters, Operation->State, mu::System::AllLODs);
+		Operation->MutableInstance = MutableSystem->BeginUpdate(Operation->InstanceID, Operation->Parameters, Operation->GetCapturedDescriptor().GetState(), mu::System::AllLODs);
 	}
 
 	
@@ -1608,29 +1647,36 @@ namespace impl
 	{
 		Operation->NumLODsAvailable = Operation->MutableInstance->GetLODCount();
 
-		if (Operation->CurrentMinLOD >= Operation->NumLODsAvailable)
+		int32 CurrentMinLOD = Operation->GetMinLOD();
+		int32 CurrentMaxLOD = Operation->GetMaxLOD();
+
+		if (CurrentMinLOD >= Operation->NumLODsAvailable)
 		{
-			Operation->CurrentMinLOD = Operation->NumLODsAvailable - 1;
-			Operation->CurrentMaxLOD = Operation->CurrentMinLOD;
+			CurrentMinLOD = Operation->NumLODsAvailable - 1;
+			CurrentMaxLOD = CurrentMinLOD;
 		}
-		else if (Operation->CurrentMaxLOD >= Operation->NumLODsAvailable)
+		else if (CurrentMaxLOD >= Operation->NumLODsAvailable)
 		{
-			Operation->CurrentMaxLOD = Operation->NumLODsAvailable - 1;
+			CurrentMaxLOD = Operation->NumLODsAvailable - 1;
 		}
 
+		Operation->SetMinLOD(CurrentMinLOD);
+		Operation->SetMaxLOD(CurrentMaxLOD);
+
 		// Initialize RequestedLODs to zero if not set
-		Operation->RequestedLODs.SetNumZeroed(Operation->NumComponents);
+		TArray<uint16> RequestedLODs;
+		RequestedLODs.SetNumZeroed(Operation->NumComponents);
 
 		for (int32 ComponentIndex = 0; ComponentIndex < Operation->NumComponents; ++ComponentIndex)
 		{
 			// Ensure we're generating at least one LOD
-			for (int32 LODIndex = Operation->CurrentMaxLOD; LODIndex < MAX_MESH_LOD_COUNT; ++LODIndex)
+			for (int32 LODIndex = CurrentMaxLOD; LODIndex < MAX_MESH_LOD_COUNT; ++LODIndex)
 			{
-				Operation->RequestedLODs[ComponentIndex] |= (1 << LODIndex);
+				RequestedLODs[ComponentIndex] |= (1 << LODIndex);
 			}
 		}
 
-		Operation->InstanceDescriptorHash.UpdateRequestedLODs(Operation->RequestedLODs);
+		Operation->SetRequestedLODs(RequestedLODs);
 	}
 	
 	
@@ -1682,7 +1728,7 @@ namespace impl
 		for (int32 MutableLODIndex = 0; MutableLODIndex < Instance->GetLODCount(); ++MutableLODIndex)
 		{
 			// Skip LODs outside the range we want to generate
-			if (MutableLODIndex < OperationData->CurrentMinLOD || MutableLODIndex > OperationData->CurrentMaxLOD)
+			if (MutableLODIndex < OperationData->GetMinLOD() || MutableLODIndex > OperationData->GetMaxLOD())
 			{
 				continue;
 			}
@@ -1699,7 +1745,8 @@ namespace impl
 				Component.FirstSurface = OperationData->InstanceUpdateData.Surfaces.Num();
 				Component.SurfaceCount = 0;
 
-				const bool bGenerateLOD = OperationData->RequestedLODs.IsValidIndex(Component.Id) ? (OperationData->RequestedLODs[Component.Id] & (1 << MutableLODIndex)) != 0 : true;
+				const TArray<uint16>& RequestedLODs = OperationData->GetRequestedLODs();
+				const bool bGenerateLOD = RequestedLODs.IsValidIndex(Component.Id) ? (RequestedLODs[Component.Id] & (1 << MutableLODIndex)) != 0 : true;
 
 				// Mesh
 				if (Instance->GetMeshCount(MutableLODIndex, ComponentIndex) > 0)
@@ -2096,7 +2143,7 @@ namespace impl
 		// TODO: Not strictly mutable: move to another worker thread task to free mutable access?
 		Subtask_Mutable_PrepareSkeletonData(OperationData);
 
-		if (OperationData->bBuildParameterRelevancy)
+		if (OperationData->GetCapturedDescriptor().GetBuildParameterRelevancy())
 		{
 			Subtask_Mutable_UpdateParameterRelevancy(OperationData);
 		}
@@ -2415,7 +2462,7 @@ namespace impl
 			return;
 		}
 		
-		if (OperationData->bBuildParameterRelevancy)
+		if (OperationData->GetCapturedDescriptor().GetBuildParameterRelevancy())
 		{
 			// Relevancy
 			ObjectInstancePrivateData->RelevantParameters = OperationData->RelevantParametersInProgress;
@@ -2593,10 +2640,12 @@ namespace impl
 		{
 			TArray<mu::FResourceID>& MeshId = Operation->MeshDescriptors[ComponentIndex];
 			MeshId.Init(MAX_uint64, MAX_MESH_LOD_COUNT);
+
+			const TArray<uint16>& RequestedLODs = Operation->GetRequestedLODs();
 			
-			for (int32 LODIndex = Operation->CurrentMinLOD; LODIndex <= Operation->CurrentMaxLOD; ++LODIndex)
+			for (int32 LODIndex = Operation->GetMinLOD(); LODIndex <= Operation->GetMaxLOD(); ++LODIndex)
 			{
-				const bool bGenerateLOD = Operation->RequestedLODs.IsValidIndex(ComponentIndex) ? (Operation->RequestedLODs[ComponentIndex] & (1 << LODIndex)) != 0 : true;
+				const bool bGenerateLOD = RequestedLODs.IsValidIndex(ComponentIndex) ? (RequestedLODs[ComponentIndex] & (1 << LODIndex)) != 0 : true;
 				if (bGenerateLOD)
 				{
 					MeshId[LODIndex] = Operation->MutableInstance->GetMeshId(LODIndex, ComponentIndex, 0);
@@ -2646,7 +2695,7 @@ namespace impl
 		}
 
 		// Skip update, the requested update is equal to the running update.
-		if (Operation->InstanceDescriptorHash.IsSubset(CandidateInstance->GetPrivate()->DescriptorHash))
+		if (Operation->GetCapturedDescriptorHash().IsSubset(CandidateInstancePrivateData->CommittedDescriptorHash))
 		{
 			System->ClearCurrentMutableOperation();
 
@@ -2699,8 +2748,6 @@ namespace impl
 		check(SystemPrivateData->ExtensionDataStreamer != nullptr);
 		SystemPrivateData->ExtensionDataStreamer->SetActiveObject(CustomizableObject);
 
-		CandidateInstance->CommitMinMaxLOD();
-
 		FString StateName = CandidateInstance->GetCustomizableObject()->GetStateName(CandidateInstance->GetState());
 		const FParameterUIData* StateData = CandidateInstance->GetCustomizableObject()->GetPrivate()->GetModelResources().StateUIDataMap.Find(StateName);
 
@@ -2747,8 +2794,6 @@ namespace impl
 		
 		// Task: Mutable Update and GetMesh
 		//-------------------------------------------------------------
-		Operation->CurrentMinLOD = Operation->InstanceDescriptorHash.GetMinLOD();
-		Operation->CurrentMaxLOD = Operation->InstanceDescriptorHash.GetMaxLOD();
 		Operation->InstanceID = Operation->bLiveUpdateMode ? CandidateInstancePrivateData->LiveUpdateModeInstanceID : 0;
 		Operation->bUseMeshCache = CustomizableObject->IsMeshCacheEnabled() && !Operation->bLiveUpdateMode && CVarEnableMeshCache.GetValueOnGameThread();
 #if WITH_EDITOR
@@ -2798,8 +2843,10 @@ namespace impl
 			!System->CurrentInstanceLODManagement->IsOnlyGenerateRequestedLODLevelsEnabled() ||
 			bIsInEditorViewport)
 		{
-			Operation->RequestedLODs.Init(MAX_uint8, Operation->NumComponents);
-			Operation->InstanceDescriptorHash.UpdateRequestedLODs(Operation->RequestedLODs);
+			TArray<uint16> RequestedLODs = Operation->GetRequestedLODs();
+			RequestedLODs.Init(MAX_uint8, Operation->NumComponents);
+
+			Operation->SetRequestedLODs(RequestedLODs);
 		}
 
 #ifdef MUTABLE_USE_NEW_TASKGRAPH
@@ -3079,7 +3126,10 @@ bool UCustomizableObjectSystem::Tick(float DeltaTime)
 				// Commit the LOD changes
 				LODUpdateCandidateFound->ApplyLODUpdateParamsToInstance();
 
-				const TSharedRef<FUpdateContextPrivate> Context = MakeShared<FUpdateContextPrivate>(*LODUpdateCandidateFound->CustomizableObjectInstance);
+				UCustomizableObjectInstance* Instance = LODUpdateCandidateFound->CustomizableObjectInstance;
+				FCustomizableObjectInstanceDescriptor& Descriptor = Instance->GetPrivate()->SkeletalMeshStatus == ESkeletalMeshStatus::NotGenerated ? Instance->GetDescriptor() : Instance->GetPrivate()->CommittedDescriptor;
+
+				const TSharedRef<FUpdateContextPrivate> Context = MakeShared<FUpdateContextPrivate>(*Instance, Descriptor);
 				Private->StartUpdateSkeletalMesh(Context);
 			}
 		}
@@ -3235,31 +3285,29 @@ void UCustomizableObjectSystem::UnregisterImageProvider(UCustomizableSystemImage
 }
 
 
-void CacheTexturesParameters(const TArray<FName>& TextureParameters)
+void UCustomizableObjectSystemPrivate::CacheTextureParameters(const TArray<FCustomizableObjectTextureParameterValue>& TextureParameters) const
 {
-	if (!TextureParameters.IsEmpty() && UCustomizableObjectSystem::IsCreated())
+	for (const FCustomizableObjectTextureParameterValue& TextureParameter : TextureParameters)
 	{
-		FUnrealMutableImageProvider* ImageProvider = UCustomizableObjectSystem::GetInstance()->GetPrivateChecked()->GetImageProviderChecked();
-		check(ImageProvider);
+		ImageProvider->CacheImage(TextureParameter.ParameterValue, false);
 
-		for (const FName& TextureParameter : TextureParameters)
+		for (const FName& RangeValue : TextureParameter.ParameterRangeValues)
 		{
-			ImageProvider->CacheImage(TextureParameter, false);
+			ImageProvider->CacheImage(RangeValue, false);
 		}
 	}
 }
 
 
-void UnCacheTexturesParameters(const TArray<FName>& TextureParameters)
+void UCustomizableObjectSystemPrivate::UnCacheTextureParameters(const TArray<FCustomizableObjectTextureParameterValue>& TextureParameters) const
 {
-	if (!TextureParameters.IsEmpty() && UCustomizableObjectSystem::IsCreated())
+	for (const FCustomizableObjectTextureParameterValue& TextureParameter : TextureParameters)
 	{
-		FUnrealMutableImageProvider* ImageProvider = UCustomizableObjectSystem::GetInstance()->GetPrivateChecked()->GetImageProviderChecked();
-		check(ImageProvider);
+		ImageProvider->UnCacheImage(TextureParameter.ParameterValue, false);
 
-		for (const FName& TextureParameter : TextureParameters)
+		for (const FName& RangeValue : TextureParameter.ParameterRangeValues)
 		{
-			ImageProvider->UnCacheImage(TextureParameter, false);
+			ImageProvider->UnCacheImage(RangeValue, false);
 		}
 	}
 }
diff --git a/Engine/Plugins/Experimental/Mutable/Source/CustomizableObject/Private/MuCO/DescriptorHash.cpp b/Engine/Plugins/Experimental/Mutable/Source/CustomizableObject/Private/MuCO/DescriptorHash.cpp
--- a/Engine/Plugins/Experimental/Mutable/Source/CustomizableObject/Private/MuCO/DescriptorHash.cpp
+++ b/Engine/Plugins/Experimental/Mutable/Source/CustomizableObject/Private/MuCO/DescriptorHash.cpp
@@ -110,37 +110,6 @@ bool FDescriptorHash::IsSubset(const FDescriptorHash& Other) const
 }
 
 
-void FDescriptorHash::UpdateMinMaxLOD(const int32 InMinLOD, const int32 InMaxLOD)
-{
-	MinLOD = InMinLOD;
-	MaxLOD = InMaxLOD;	
-}
-
-
-int32 FDescriptorHash::GetMinLOD() const
-{
-	return MinLOD;
-}
-
-
-int32 FDescriptorHash::GetMaxLOD() const
-{
-	return MaxLOD;
-}
-
-
-void FDescriptorHash::UpdateRequestedLODs(const TArray<uint16>& InRequestedLODs)
-{
-	RequestedLODsPerComponent = InRequestedLODs;
-}
-
-
-const TArray<uint16>& FDescriptorHash::GetRequestedLODs() const
-{
-	return RequestedLODsPerComponent;
-}
-
-
 FString FDescriptorHash::ToString() const
 {
 	TStringBuilder<150> Builder;
