[Mutable] UCustomizableObject Moved non-serialized members to private.

- Moved all non-serialized members to private.
- Moved non-public classes and structs to private.
- Moved the following functions to private:
    - IsLocked
    - GetModel

#rnx

diff --git a/Engine/Plugins/Experimental/Mutable/Source/CustomizableObject/Classes/MuCO/CustomizableObject.h b/Engine/Plugins/Experimental/Mutable/Source/CustomizableObject/Classes/MuCO/CustomizableObject.h
--- a/Engine/Plugins/Experimental/Mutable/Source/CustomizableObject/Classes/MuCO/CustomizableObject.h
+++ b/Engine/Plugins/Experimental/Mutable/Source/CustomizableObject/Classes/MuCO/CustomizableObject.h
@@ -7,7 +7,6 @@
 #include "MuCO/CustomizableObjectResourceData.h"
 #include "MuCO/CustomizableObjectResourceDataTypes.h"
 #include "MuCO/CustomizableObjectStreamedResourceData.h"
-#include "MuCO/CustomizableObjectIdentifier.h"
 #include "MuCO/CustomizableObjectParameterTypeDefinitions.h"
 #include "MuCO/CustomizableObjectUIData.h"
 
@@ -24,9 +23,11 @@ class ITargetPlatform;
 class UCustomizableObject;
 class UEdGraph;
 class USkeletalMesh;
+class UCustomizableObjectPrivate;
 struct FFrame;
 struct FStreamableHandle;
 template <typename FuncType> class TFunctionRef;
+enum class ECustomizableObjectCompilationState : uint8;
 
 DECLARE_MULTICAST_DELEGATE(FPostCompileDelegate)
  
@@ -55,8 +56,6 @@ struct FFParameterOptionsTags
 USTRUCT()
 struct FParameterTags
 {
-public:
-
 	GENERATED_USTRUCT_BODY()
 
 	/** List of tags of a parameter */
@@ -79,9 +78,6 @@ enum class ECustomizableObjectRelevancy : uint8
 };
 
 
-// This is used to hide Mutable SDK members in the public headers.
-class UCustomizableObjectPrivate;
-
 USTRUCT()
 struct FProfileParameterDat
 {
@@ -202,97 +198,6 @@ struct FCompilationOptions
 };
 
 
-UENUM()
-enum class ECustomizableObjectCompilationState : uint8
-{
-	//
-	None,
-	// 
-	InProgress,
-	//
-	Completed,
-	//
-	Failed
-};
-
-
-class FCustomizableObjectCompilerBase
-{
-public:
-
-	FCustomizableObjectCompilerBase() {};
-
-	// Ensure virtual destruction
-	virtual ~FCustomizableObjectCompilerBase() {};
-	
-	virtual void Compile(UCustomizableObject& Object, const FCompilationOptions& Options, bool bAsync) {};
-
-	virtual bool Tick() { return false; }
-	virtual void ForceFinishCompilation() {};
-
-	// Return true if this object doesn't reference a parent object.
-	virtual bool IsRootObject(const class UCustomizableObject* Object) const { return true; }
-
-	/** Returns the Customizable Object that does start the CO tree */
-	virtual UCustomizableObject* GetRootObject( class UCustomizableObject* Object) = 0;
-
-	/** Provides the caller with the warning and error messages produced during compilation */
-	virtual void GetCompilationMessages(TArray<FText>& OutWarningMessages, TArray<FText>& OutErrorMessages) const = 0;
-	
-	virtual ECustomizableObjectCompilationState GetCompilationState() const { return ECustomizableObjectCompilationState::None; }
-};
-
-
-// Warning! MutableCompiledDataHeader must be the first data serialized in a stream
-struct MutableCompiledDataStreamHeader
-{
-	int32 InternalVersion;
-	FGuid VersionId;
-
-	MutableCompiledDataStreamHeader() { }
-	MutableCompiledDataStreamHeader(int32 InInternalVersion, FGuid InVersionId) : InternalVersion(InInternalVersion), VersionId(InVersionId) { }
-
-	friend FArchive& operator<<(FArchive& Ar, MutableCompiledDataStreamHeader& Header)
-	{
-		Ar << Header.InternalVersion;
-		Ar << Header.VersionId;
-
-		return Ar;
-	}
-};
-
-
-USTRUCT()
-struct FMutableModelParameterValue
-{
-	GENERATED_USTRUCT_BODY()
-
-	FMutableModelParameterValue() = default;
-
-	UPROPERTY()
-	FString Name;
-
-	UPROPERTY()
-	int Value = 0;
-};
-
-
-USTRUCT()
-struct FMutableModelParameterProperties
-{
-	GENERATED_USTRUCT_BODY()
-
-	FMutableModelParameterProperties() = default;
-	FString Name;
-
-	UPROPERTY()
-	EMutableParameterType Type = EMutableParameterType::None;
-
-	UPROPERTY()
-	TArray<FMutableModelParameterValue> PossibleValues;
-};
-
-
 USTRUCT()
 struct FMorphTargetInfo
 {
@@ -479,17 +384,6 @@ struct CUSTOMIZABLEOBJECT_API FMutableStreamableBlock
 template<> struct TCanBulkSerialize<FMutableStreamableBlock> { enum { Value = true }; };
 
 
-struct CUSTOMIZABLEOBJECT_API FMutableCachedPlatformData
-{
-	/** */
-	TArray64<uint8> ModelData;
-
-	/** */
-	TArray64<uint8> StreamableData;
-};
-
-
-
 USTRUCT()
 struct FMutableLODSettings
 {
@@ -627,10 +521,7 @@ public:
 
 	UPROPERTY(EditAnywhere, Category = CustomizableObject, meta = (DisplayName = "LOD Settings"))
 	FMutableLODSettings LODSettings;
-
-	/** This is a non-user-controlled flag to disable streaming (set at object compilation time, depending on optimization). */
-	bool bDisableTextureStreaming = false;
-
+	
 	UPROPERTY(Transient)
 	TArray<FMorphTargetInfo> ContributingMorphTargetsInfo;
 	
@@ -727,14 +618,6 @@ public:
 	UPROPERTY()
 	FGuid VersionId;
 
-	/* Map to identify what CustomizableObject owns a parameter. Used to display a tooltip when hovering a parameter
-	   in the Prev. instance panel */
-	UPROPERTY(Transient)
-	TMap<FString, FString> CustomizableObjectPathMap;
-
-	UPROPERTY(Transient)
-	TMap<FString, FCustomizableObjectIdPair> GroupNodeMap; 
-
 	// 
 	UPROPERTY()
 	TArray<FProfileParameterDat>  InstancePropertiesProfiles;
@@ -875,23 +758,9 @@ public:
 	bool IsParameterMultidimensional ( UPARAM(DisplayName = "Parameter Name") const FString& InParameterName) const;
 	
 private: 
-
-	// This is information about the parameters in the model that is generated at model compile time.
-	UPROPERTY(Transient)
-	TArray<FMutableModelParameterProperties> ParameterProperties;
-
-	// Map of name to index of ParameterProperties.
-	// use this to lookup fast by Name
-	TMap<FString, int32> ParameterPropertiesLookupTable;
-
 	// Rebuild ParameterProperties from the current compiled model.
 	void UpdateParameterPropertiesFromModel(const TSharedPtr<mu::Model>& Model);
 
-	// This is a manual version number for the binary blobs in this asset.
-	// Increasing it invalidates all the previously compiled models.
-	// Warning: If while merging code both versions have changed, take the highest+1.
-	static const int32 CurrentSupportedVersion = 421;
-
 public:
 
 	void AddUncompiledCOWarning(const FString& AdditionalLoggingInfo);
@@ -916,8 +785,6 @@ public:
 	void UpdateVersionId();
 	FGuid GetVersionId() const { return VersionId; }
 
-	int32 GetCurrentSupportedVersion() const { return CurrentSupportedVersion; };
-
 	// Compose folder name where the data is stored
 	FString GetCompiledDataFolderPath() const;
 
@@ -959,13 +826,7 @@ public:
 	 * Returns whether the data was successfully loaded.
 	 */
 	bool TryLoadCompiledCookDataForPlatform(const ITargetPlatform* TargetPlatform);
-#if WITH_EDITORONLY_DATA
 
-	/** Map of PlatformName to CachedPlatformData. Only valid while cooking. */
-	TMap<FString, FMutableCachedPlatformData> CachedPlatformsData;
-	
-#endif
-	
 	void SaveEmbeddedData(FArchive& Ar);
 
 #endif
@@ -1053,17 +914,12 @@ public:
 	UPROPERTY()
 	bool bIsChildObject = false;
 
-	/** If the object is compiled, this flag is true if it was compiled with maximum optimizations. If the object is not compiled, its value is meaningless. */
-	bool bIsCompiledWithOptimization = true;
-
-	ECustomizableObjectCompilationState CompilationState = ECustomizableObjectCompilationState::None;
-
-	FPostCompileDelegate PostCompileDelegate;
-
 	void PostCompile();
 
 #endif
 
+	FPostCompileDelegate& GetPostCompileDelegate() const;
+	
 	/** Create a new instance of this object. The instance parameters will be initialized with the object default values. */
 	UFUNCTION(BlueprintCallable, Category = CustomizableObject)
 	class UCustomizableObjectInstance* CreateInstance();
@@ -1077,11 +933,6 @@ public:
 	UFUNCTION(BlueprintCallable, Category = CustomizableObject)
 	bool IsCompiled() const;
 
-	// See UCustomizableObjectSystem::LockObject()
-	bool IsLocked() const;
-
-	TSharedPtr<mu::Model, ESPMode::ThreadSafe> GetModel() const;
-
 #if WITH_EDITOR
 	void SetModel(TSharedPtr<mu::Model, ESPMode::ThreadSafe> Model);
 #endif
diff --git a/Engine/Plugins/Experimental/Mutable/Source/CustomizableObject/Internal/MuCO/CustomizableObjectPrivate.h b/Engine/Plugins/Experimental/Mutable/Source/CustomizableObject/Internal/MuCO/CustomizableObjectPrivate.h
--- a/Engine/Plugins/Experimental/Mutable/Source/CustomizableObject/Internal/MuCO/CustomizableObjectPrivate.h
+++ b/Engine/Plugins/Experimental/Mutable/Source/CustomizableObject/Internal/MuCO/CustomizableObjectPrivate.h
@@ -2,8 +2,10 @@
 
 #pragma once
 
+#include "MuCO/CustomizableObject.h"
 #include "MuCO/StateMachine.h"
 #include "MuCO/CustomizableObjectUIData.h"
+#include "MuCO/CustomizableObjectIdentifier.h"
 #include "Templates/SharedPointer.h"
 #include "UObject/WeakObjectPtr.h"
 #include "UObject/SoftObjectPtr.h"
@@ -16,7 +18,6 @@
 #include "CustomizableObjectPrivate.generated.h"
 
 namespace mu { class Model; }
-class UCustomizableObject;
 class USkeletalMesh;
 class USkeleton;
 class UPhysicsAsset;
@@ -24,6 +25,85 @@ class UMaterialInterface;
 class UTexture;
 class UAnimInstance;
 class UAssetUserData;
+class UCustomizableObject;
+
+
+class CUSTOMIZABLEOBJECT_API FCustomizableObjectCompilerBase
+{
+public:
+
+	FCustomizableObjectCompilerBase() {};
+
+	// Ensure virtual destruction
+	virtual ~FCustomizableObjectCompilerBase() {};
+	
+	virtual void Compile(UCustomizableObject& Object, const FCompilationOptions& Options, bool bAsync) {};
+
+	virtual bool Tick() { return false; }
+	virtual void ForceFinishCompilation() {};
+
+	// Return true if this object doesn't reference a parent object.
+	virtual bool IsRootObject(const class UCustomizableObject* Object) const { return true; }
+
+	/** Returns the Customizable Object that does start the CO tree */
+	virtual UCustomizableObject* GetRootObject( class UCustomizableObject* Object) = 0;
+
+	/** Provides the caller with the warning and error messages produced during compilation */
+	virtual void GetCompilationMessages(TArray<FText>& OutWarningMessages, TArray<FText>& OutErrorMessages) const = 0;
+	
+	virtual ECustomizableObjectCompilationState GetCompilationState() const;
+};
+
+
+// Warning! MutableCompiledDataHeader must be the first data serialized in a stream
+struct MutableCompiledDataStreamHeader
+{
+	int32 InternalVersion;
+	FGuid VersionId;
+
+	MutableCompiledDataStreamHeader() { }
+	MutableCompiledDataStreamHeader(int32 InInternalVersion, FGuid InVersionId) : InternalVersion(InInternalVersion), VersionId(InVersionId) { }
+
+	friend FArchive& operator<<(FArchive& Ar, MutableCompiledDataStreamHeader& Header)
+	{
+		Ar << Header.InternalVersion;
+		Ar << Header.VersionId;
+
+		return Ar;
+	}
+};
+
+
+USTRUCT()
+struct FMutableModelParameterValue
+{
+	GENERATED_USTRUCT_BODY()
+
+	FMutableModelParameterValue() = default;
+
+	UPROPERTY()
+	FString Name;
+
+	UPROPERTY()
+	int Value = 0;
+};
+
+
+USTRUCT()
+struct FMutableModelParameterProperties
+{
+	GENERATED_USTRUCT_BODY()
+
+	FMutableModelParameterProperties() = default;
+	FString Name;
+
+	UPROPERTY()
+	EMutableParameterType Type = EMutableParameterType::None;
+
+	UPROPERTY()
+	TArray<FMutableModelParameterValue> PossibleValues;
+};
+
 
 class FMeshCache
 {
@@ -76,6 +156,20 @@ struct FCustomizableObjectStatusTypes
 using FCustomizableObjectStatus = FStateMachine<FCustomizableObjectStatusTypes>;
 
 
+UENUM()
+enum class ECustomizableObjectCompilationState : uint8
+{
+	//
+	None,
+	// 
+	InProgress,
+	//
+	Completed,
+	//
+	Failed
+};
+
+
 USTRUCT()
 struct CUSTOMIZABLEOBJECT_API FMutableModelImageProperties
 {
@@ -441,13 +535,21 @@ struct FModelResources
 };
 
 
+struct CUSTOMIZABLEOBJECT_API FMutableCachedPlatformData
+{
+	/** */
+	TArray64<uint8> ModelData;
+
+	/** */
+	TArray64<uint8> StreamableData;
+};
+
+
 UCLASS()
-class UCustomizableObjectPrivate : public UObject
+class CUSTOMIZABLEOBJECT_API UCustomizableObjectPrivate : public UObject
 {
 	GENERATED_BODY()
 
-private:
-
 	TSharedPtr<mu::Model, ESPMode::ThreadSafe> MutableModel;
 
 	/** Stores resources to be used by MutableModel In-Game. Cooked resources. */
@@ -469,9 +571,12 @@ public:
 	const FModelResources& GetModelResources() const;
 
 #if WITH_EDITORONLY_DATA
-	CUSTOMIZABLEOBJECT_API FModelResources& GetModelResources(bool bIsCooking);
+	FModelResources& GetModelResources(bool bIsCooking);
 #endif
 
+	// See UCustomizableObjectSystem::LockObject()
+	bool IsLocked() const;
+
 	/** Cache of generated SkeletalMeshes */
 	FMeshCache MeshCache;
 
@@ -494,8 +599,44 @@ public:
 	 *
 	 * Updated each time the CO is compiled and saved in the Derived Data. */
 	TMap<FName, FGuid> ParticipatingObjects;
+
+	/** If the object is compiled, this flag is true if it was compiled with maximum optimizations. If the object is not compiled, its value is meaningless. */
+	bool bIsCompiledWithOptimization = true;
+
+	ECustomizableObjectCompilationState CompilationState = ECustomizableObjectCompilationState::None;
+	
+#if WITH_EDITOR
+	/** Map of PlatformName to CachedPlatformData. Only valid while cooking. */
+	TMap<FString, FMutableCachedPlatformData> CachedPlatformsData;
+#endif
 #endif
 
 	FCustomizableObjectStatus Status;
+
+	/** Map to identify what CustomizableObject owns a parameter. Used to display a tooltip when hovering a parameter
+	 * in the Prev. instance panel */
+	UPROPERTY(Transient)
+	TMap<FString, FString> CustomizableObjectPathMap;
+
+	UPROPERTY(Transient)
+	TMap<FString, FCustomizableObjectIdPair> GroupNodeMap;
+
+	FPostCompileDelegate PostCompileDelegate;
+
+	// This is information about the parameters in the model that is generated at model compile time.
+	UPROPERTY(Transient)
+	TArray<FMutableModelParameterProperties> ParameterProperties;
+
+	// Map of name to index of ParameterProperties.
+	// use this to lookup fast by Name
+	TMap<FString, int32> ParameterPropertiesLookupTable;
+
+	// This is a manual version number for the binary blobs in this asset.
+	// Increasing it invalidates all the previously compiled models.
+	// Warning: If while merging code both versions have changed, take the highest+1.
+	static constexpr int32 CurrentSupportedVersion = 421;
+
+	/** This is a non-user-controlled flag to disable streaming (set at object compilation time, depending on optimization). */
+	bool bDisableTextureStreaming = false;
 };
 
diff --git a/Engine/Plugins/Experimental/Mutable/Source/CustomizableObject/Internal/MuCO/CustomizableObjectSystemPrivate.h b/Engine/Plugins/Experimental/Mutable/Source/CustomizableObject/Internal/MuCO/CustomizableObjectSystemPrivate.h
--- a/Engine/Plugins/Experimental/Mutable/Source/CustomizableObject/Internal/MuCO/CustomizableObjectSystemPrivate.h
+++ b/Engine/Plugins/Experimental/Mutable/Source/CustomizableObject/Internal/MuCO/CustomizableObjectSystemPrivate.h
@@ -6,6 +6,7 @@
 #include "Containers/Queue.h"
 #include "MuCO/DescriptorHash.h"
 #include "MuCO/CustomizableObject.h"
+#include "MuCO/CustomizableObjectPrivate.h"
 #include "MuCO/CustomizableObjectInstance.h"
 #include "MuCO/CustomizableObjectExtension.h"
 #include "Containers/Ticker.h"
diff --git a/Engine/Plugins/Experimental/Mutable/Source/CustomizableObject/Private/MuCO/CustomizableInstanceLODManagement.cpp b/Engine/Plugins/Experimental/Mutable/Source/CustomizableObject/Private/MuCO/CustomizableInstanceLODManagement.cpp
--- a/Engine/Plugins/Experimental/Mutable/Source/CustomizableObject/Private/MuCO/CustomizableInstanceLODManagement.cpp
+++ b/Engine/Plugins/Experimental/Mutable/Source/CustomizableObject/Private/MuCO/CustomizableInstanceLODManagement.cpp
@@ -15,6 +15,7 @@
 #include "MuCO/CustomizableObjectInstanceUsage.h"
 #include "UObject/UObjectIterator.h"
 #include "Components/SkeletalMeshComponent.h"
+#include "MuCO/CustomizableObjectPrivate.h"
 
 #include UE_INLINE_GENERATED_CPP_BY_NAME(CustomizableInstanceLODManagement)
 
@@ -206,7 +207,7 @@ void UCustomizableInstanceLODManagement::UpdateInstanceDistsAndLODs(FMutableInst
 			{
 				if (const UCustomizableObject* CustomizableObject = CustomizableObjectInstance->GetCustomizableObject();
 					CustomizableObject &&
-					!CustomizableObject->IsLocked())
+					!CustomizableObject->GetPrivate()->IsLocked())
 				{
 					UCustomizableObjectInstance* Ptr = *CustomizableObjectInstance;
 					Ptr->SetIsDiscardedBecauseOfTooManyInstances(false);
@@ -403,7 +404,7 @@ void UCustomizableInstanceLODManagement::UpdateInstanceDistsAndLODs(FMutableInst
 			if (IsValidChecked(It.Key) && It.Key->GetPrivate())
 			{
 				const UCustomizableObject* CustomizableObject = It.Key->GetCustomizableObject();
-				if (!CustomizableObject || CustomizableObject->IsLocked())
+				if (!CustomizableObject || CustomizableObject->GetPrivate()->IsLocked())
 				{
 					continue;
 				}
diff --git a/Engine/Plugins/Experimental/Mutable/Source/CustomizableObject/Private/MuCO/CustomizableObject.cpp b/Engine/Plugins/Experimental/Mutable/Source/CustomizableObject/Private/MuCO/CustomizableObject.cpp
--- a/Engine/Plugins/Experimental/Mutable/Source/CustomizableObject/Private/MuCO/CustomizableObject.cpp
+++ b/Engine/Plugins/Experimental/Mutable/Source/CustomizableObject/Private/MuCO/CustomizableObject.cpp
@@ -241,14 +241,9 @@ void UCustomizableObject::PostLoad()
 }
 
 
-bool UCustomizableObject::IsLocked() const
+bool UCustomizableObjectPrivate::IsLocked() const
 {
-	if (Private)
-	{
-		return Private->bLocked;
-	}
-
-	return false;
+	return bLocked;
 }
 
 
@@ -270,7 +265,7 @@ void UCustomizableObject::Serialize(FArchive& Ar_Asset)
 	else
 	{
 		// Can't remove this or saved customizable objects will fail to load
-		int64 InternalVersion = CurrentSupportedVersion;
+		int64 InternalVersion = UCustomizableObjectPrivate::CurrentSupportedVersion;
 		Ar_Asset << InternalVersion;
 	}
 #else
@@ -337,8 +332,8 @@ void UCustomizableObject::ClearCompiledData(bool bIsCooking)
 	ClothSharedConfigsData.Empty();
 
 #if WITH_EDITORONLY_DATA
-	CustomizableObjectPathMap.Empty();
-	GroupNodeMap.Empty();
+	GetPrivate()->CustomizableObjectPathMap.Empty();
+	GetPrivate()->GroupNodeMap.Empty();
 	Private->ParticipatingObjects.Empty();
 #endif
 
@@ -460,7 +455,7 @@ void SerializeStreamedResources(FArchive& Ar, UObject* Object, TArray<FCustomiza
 
 void UCustomizableObject::SaveCompiledData(FArchive& MemoryWriter, bool bIsCooking)
 {
-	int32 InternalVersion = CurrentSupportedVersion;
+	int32 InternalVersion = UCustomizableObjectPrivate::CurrentSupportedVersion;
 	MutableCompiledDataStreamHeader Header(InternalVersion, VersionId);
 	MemoryWriter << Header;
 
@@ -545,9 +540,9 @@ void UCustomizableObject::SaveCompiledData(FArchive& MemoryWriter, bool bIsCooki
 	MemoryWriter << LODSettings.bLODStreamingEnabled;
 
 	// Editor Only data
-	MemoryWriter << bIsCompiledWithOptimization;
-	MemoryWriter << CustomizableObjectPathMap;
-	MemoryWriter << GroupNodeMap;
+	MemoryWriter << GetPrivate()->bIsCompiledWithOptimization;
+	MemoryWriter << GetPrivate()->CustomizableObjectPathMap;
+	MemoryWriter << GetPrivate()->GroupNodeMap;
 	MemoryWriter << GetPrivate()->ParticipatingObjects;
 }
 
@@ -559,7 +554,7 @@ void UCustomizableObject::LoadCompiledData(FArchive& MemoryReader, const ITarget
 	MutableCompiledDataStreamHeader Header;
 	MemoryReader << Header;
 
-	if (CurrentSupportedVersion == Header.InternalVersion)
+	if (UCustomizableObjectPrivate::CurrentSupportedVersion == Header.InternalVersion)
 	{
 		// Make sure mutable has been initialised.
 		UCustomizableObjectSystem::GetInstance();
@@ -667,9 +662,9 @@ void UCustomizableObject::LoadCompiledData(FArchive& MemoryReader, const ITarget
 
 		// Editor Only data
 		{
-			MemoryReader << bIsCompiledWithOptimization;
-			MemoryReader << CustomizableObjectPathMap;
-			MemoryReader << GroupNodeMap;
+			MemoryReader << GetPrivate()->bIsCompiledWithOptimization;
+			MemoryReader << GetPrivate()->CustomizableObjectPathMap;
+			MemoryReader << GetPrivate()->GroupNodeMap;
 
 			TMap<FName, FGuid>& ParticipatingObjects = GetPrivate()->ParticipatingObjects;
 			MemoryReader << ParticipatingObjects;
@@ -758,7 +753,7 @@ void UCustomizableObject::LoadCompiledDataFromDisk()
 			AuxMemoryReader << StreamableDataHeader;
 		}
 
-		if (CompiledDataHeader.InternalVersion == CurrentSupportedVersion
+		if (CompiledDataHeader.InternalVersion == UCustomizableObjectPrivate::CurrentSupportedVersion
 			&&
 			CompiledDataHeader.InternalVersion == StreamableDataHeader.InternalVersion 
 			&&
@@ -795,9 +790,9 @@ void UCustomizableObject::CachePlatformData(const ITargetPlatform* InTargetPlatf
 {
 	FString PlatformName = InTargetPlatform ? InTargetPlatform->PlatformName() : FPlatformProperties::PlatformName();
 
-	check(!CachedPlatformsData.Find(PlatformName));
+	check(!GetPrivate()->CachedPlatformsData.Find(PlatformName));
 
-	FMutableCachedPlatformData& Data = CachedPlatformsData.Add(PlatformName);
+	FMutableCachedPlatformData& Data = GetPrivate()->CachedPlatformsData.Add(PlatformName);
 
 	// Cache CO data and mu::Model
 	FMemoryWriter64 MemoryWriter(Data.ModelData);
@@ -847,7 +842,7 @@ bool UCustomizableObject::ConditionalAutoCompile()
 	check(IsInGameThread());
 
 	// Don't compile objects being compiled
-	if (IsLocked())
+	if (GetPrivate()->IsLocked())
 	{
 		return false;
 	}
@@ -856,7 +851,7 @@ bool UCustomizableObject::ConditionalAutoCompile()
 	if (IsCompiled())
 	{
 		// Show a warning if the compilation was not done with optimizations.
-		if (bIsCompiledWithOptimization)
+		if (GetPrivate()->bIsCompiledWithOptimization)
 		{
 			FString Msg = FString::Printf(TEXT("Warning: Customizable Object [%s] was compiled without optimization."), *GetName());
 			GEngine->AddOnScreenDebugMessage((uint64)((PTRINT)this), 10.0f, FColor::Red, Msg);
@@ -878,7 +873,7 @@ bool UCustomizableObject::ConditionalAutoCompile()
 	}
 
 	// Don't re-compile objects if they failed to compile. 
-	if (CompilationState == ECustomizableObjectCompilationState::Failed)
+	if (GetPrivate()->CompilationState == ECustomizableObjectCompilationState::Failed)
 	{
 		return false;
 	}
@@ -978,13 +973,13 @@ void UCustomizableObject::SaveEmbeddedData(FArchive& Ar)
 {
 	UE_LOG(LogMutable, Verbose, TEXT("Saving embedded data for Customizable Object [%s] now at position %d."), *GetName(), int(Ar.Tell()));
 
-	int32 InternalVersion = Private->GetModel() ? CurrentSupportedVersion : -1;
+	int32 InternalVersion = Private->GetModel() ? UCustomizableObjectPrivate::CurrentSupportedVersion : -1;
 	Ar << InternalVersion;
 
 	if (Private->GetModel())
 	{
 		// General derived flags
-		Ar << bDisableTextureStreaming;
+		Ar << GetPrivate()->bDisableTextureStreaming;
 
 		// Serialize morph data
 		{
@@ -1022,12 +1017,12 @@ void UCustomizableObject::LoadEmbeddedData(FArchive& Ar)
 
 	// If this fails, something went wrong with the packaging: we have data that belongs
 	// to a different version than the code.
-	check(CurrentSupportedVersion==InternalVersion);
+	check(UCustomizableObjectPrivate::CurrentSupportedVersion==InternalVersion);
 
-	if(CurrentSupportedVersion == InternalVersion)
+	if(UCustomizableObjectPrivate::CurrentSupportedVersion == InternalVersion)
 	{
 		// General derived flags
-		Ar << bDisableTextureStreaming;
+		Ar << GetPrivate()->bDisableTextureStreaming;
 
 		// Load morph data
 		{
@@ -1188,11 +1183,17 @@ FString UCustomizableObject::GetStateParameterName(int32 StateIndex, int32 Param
 #if WITH_EDITORONLY_DATA
 void UCustomizableObject::PostCompile()
 {
-	PostCompileDelegate.Broadcast();
+	GetPrivate()->PostCompileDelegate.Broadcast();
 }
 #endif
 
 
+FPostCompileDelegate& UCustomizableObject::GetPostCompileDelegate() const
+{
+	return GetPrivate()->PostCompileDelegate;
+}
+
+
 UCustomizableObjectInstance* UCustomizableObject::CreateInstance()
 {
 	MUTABLE_CPUPROFILER_SCOPE(UCustomizableObject::CreateInstance)
@@ -1207,11 +1208,6 @@ UCustomizableObjectInstance* UCustomizableObject::CreateInstance()
 }
 
 
-TSharedPtr<mu::Model, ESPMode::ThreadSafe> UCustomizableObject::GetModel() const
-{
-	return Private->GetModel();
-}
-
 #if WITH_EDITOR
 void UCustomizableObject::SetModel(TSharedPtr<mu::Model, ESPMode::ThreadSafe> Model)
 {
@@ -1253,15 +1249,15 @@ int32 UCustomizableObject::GetComponentCount() const
 
 int32 UCustomizableObject::GetParameterCount() const
 {
-	return ParameterProperties.Num();
+	return GetPrivate()->ParameterProperties.Num();
 }
 
 
 EMutableParameterType UCustomizableObject::GetParameterType(int32 ParamIndex) const
 {
-	if (ParamIndex >= 0 && ParamIndex < ParameterProperties.Num())
+	if (ParamIndex >= 0 && ParamIndex < GetPrivate()->ParameterProperties.Num())
 	{
-		return ParameterProperties[ParamIndex].Type;
+		return GetPrivate()->ParameterProperties[ParamIndex].Type;
 	}
 	else
 	{
@@ -1274,21 +1270,21 @@ EMutableParameterType UCustomizableObject::GetParameterType(int32 ParamIndex) co
 
 EMutableParameterType UCustomizableObject::GetParameterTypeByName(const FString& Name) const
 {
-	const int* IndexPtr = ParameterPropertiesLookupTable.Find(Name);
+	const int* IndexPtr = GetPrivate()->ParameterPropertiesLookupTable.Find(Name);
 	int Index = IndexPtr ? *IndexPtr : INDEX_NONE;
 
-	if (ParameterProperties.IsValidIndex(Index))
+	if (GetPrivate()->ParameterProperties.IsValidIndex(Index))
 	{
-		return ParameterProperties[Index].Type;
+		return GetPrivate()->ParameterProperties[Index].Type;
 	}
 
 	UE_LOG(LogMutable, Warning, TEXT("Name '%s' does not exist in ParameterProperties lookup table at GetParameterTypeByName at CO %s."), *Name, *GetName());
 
-	for (int32 ParamIndex = 0; ParamIndex < ParameterProperties.Num(); ++ParamIndex)
+	for (int32 ParamIndex = 0; ParamIndex < GetPrivate()->ParameterProperties.Num(); ++ParamIndex)
 	{
-		if (ParameterProperties[ParamIndex].Name == Name)
+		if (GetPrivate()->ParameterProperties[ParamIndex].Name == Name)
 		{
-			return ParameterProperties[ParamIndex].Type;
+			return GetPrivate()->ParameterProperties[ParamIndex].Type;
 		}
 	}
 
@@ -1302,9 +1298,9 @@ static const FString s_EmptyString;
 
 const FString & UCustomizableObject::GetParameterName(int32 ParamIndex) const
 {
-	if (ParamIndex >= 0 && ParamIndex < ParameterProperties.Num())
+	if (ParamIndex >= 0 && ParamIndex < GetPrivate()->ParameterProperties.Num())
 	{
-		return ParameterProperties[ParamIndex].Name;
+		return GetPrivate()->ParameterProperties[ParamIndex].Name;
 	}
 	else
 	{
@@ -1322,8 +1318,8 @@ void UCustomizableObject::UpdateParameterPropertiesFromModel(const TSharedPtr<mu
 		mu::ParametersPtr MutableParameters = mu::Model::NewParameters(Model);
 		int paramCount = MutableParameters->GetCount();
 
-		ParameterProperties.Reset(paramCount);
-		ParameterPropertiesLookupTable.Empty(paramCount);
+		GetPrivate()->ParameterProperties.Reset(paramCount);
+		GetPrivate()->ParameterPropertiesLookupTable.Empty(paramCount);
 		for (int paramIndex = 0; paramIndex<paramCount; ++paramIndex)
 		{
 			FMutableModelParameterProperties Data;
@@ -1385,14 +1381,14 @@ void UCustomizableObject::UpdateParameterPropertiesFromModel(const TSharedPtr<mu
 				break;
 			}
 
-			ParameterProperties.Add(Data);
-			ParameterPropertiesLookupTable.Add(Data.Name, paramIndex);
+			GetPrivate()->ParameterProperties.Add(Data);
+			GetPrivate()->ParameterPropertiesLookupTable.Add(Data.Name, paramIndex);
 		}
 	}
 	else
 	{
-		ParameterProperties.Empty();
-		ParameterPropertiesLookupTable.Empty();
+		GetPrivate()->ParameterProperties.Empty();
+		GetPrivate()->ParameterPropertiesLookupTable.Empty();
 	}
 }
 
@@ -1405,9 +1401,9 @@ int32 UCustomizableObject::GetParameterDescriptionCount(const FString& ParamName
 
 int32 UCustomizableObject::GetIntParameterNumOptions(int32 ParamIndex) const
 {
-	if (ParamIndex >= 0 && ParamIndex < ParameterProperties.Num())
+	if (ParamIndex >= 0 && ParamIndex < GetPrivate()->ParameterProperties.Num())
 	{
-		return ParameterProperties[ParamIndex].PossibleValues.Num();
+		return GetPrivate()->ParameterProperties[ParamIndex].PossibleValues.Num();
 	}
 	else
 	{
@@ -1420,11 +1416,11 @@ int32 UCustomizableObject::GetIntParameterNumOptions(int32 ParamIndex) const
 
 const FString& UCustomizableObject::GetIntParameterAvailableOption(int32 ParamIndex, int32 K) const
 {
-	if (ParamIndex >= 0 && ParamIndex < ParameterProperties.Num())
+	if (ParamIndex >= 0 && ParamIndex < GetPrivate()->ParameterProperties.Num())
 	{
 		if (K >= 0 && K < GetIntParameterNumOptions(ParamIndex))
 		{
-			return ParameterProperties[ParamIndex].PossibleValues[K].Name;
+			return GetPrivate()->ParameterProperties[ParamIndex].PossibleValues[K].Name;
 		}
 		else
 		{
@@ -1442,7 +1438,7 @@ const FString& UCustomizableObject::GetIntParameterAvailableOption(int32 ParamIn
 
 int32 UCustomizableObject::FindParameter(const FString& Name) const
 {
-	const int32 * Found = ParameterPropertiesLookupTable.Find(Name);
+	const int32 * Found = GetPrivate()->ParameterPropertiesLookupTable.Find(Name);
 	if (Found == nullptr)
 	{
 		return INDEX_NONE;
@@ -1454,9 +1450,9 @@ int32 UCustomizableObject::FindParameter(const FString& Name) const
 int32 UCustomizableObject::FindIntParameterValue(int32 ParamIndex, const FString& Value) const
 {
 	int32 MinValueIndex = 0;
-	if (ParamIndex >= 0 && ParamIndex < ParameterProperties.Num())
+	if (ParamIndex >= 0 && ParamIndex < GetPrivate()->ParameterProperties.Num())
 	{
-		const TArray<FMutableModelParameterValue>& PossibleValues = ParameterProperties[ParamIndex].PossibleValues;
+		const TArray<FMutableModelParameterValue>& PossibleValues = GetPrivate()->ParameterProperties[ParamIndex].PossibleValues;
 		if (PossibleValues.Num())
 		{
 			MinValueIndex = PossibleValues[0].Value;
@@ -1488,9 +1484,9 @@ int32 UCustomizableObject::FindIntParameterValue(int32 ParamIndex, const FString
 
 FString UCustomizableObject::FindIntParameterValueName(int32 ParamIndex, int32 ParamValue) const
 {
-	if (ParamIndex >= 0 && ParamIndex < ParameterProperties.Num())
+	if (ParamIndex >= 0 && ParamIndex < GetPrivate()->ParameterProperties.Num())
 	{
-		const TArray<FMutableModelParameterValue> & PossibleValues = ParameterProperties[ParamIndex].PossibleValues;
+		const TArray<FMutableModelParameterValue> & PossibleValues = GetPrivate()->ParameterProperties[ParamIndex].PossibleValues;
 
 		const int32 MinValueIndex = !PossibleValues.IsEmpty() ? PossibleValues[0].Value : 0;
 		ParamValue = ParamValue - MinValueIndex;
@@ -1546,7 +1542,7 @@ float UCustomizableObject::GetFloatParameterDefaultValue(const FString& InParame
 		return FCustomizableObjectFloatParameterValue::DEFAULT_PARAMETER_VALUE;
 	}
 
-	const TSharedPtr<mu::Model> Model = GetModel();
+	const TSharedPtr<mu::Model> Model = GetPrivate()->GetModel();
 	if (!Model)
 	{
 		checkNoEntry();
@@ -1566,7 +1562,7 @@ int32 UCustomizableObject::GetIntParameterDefaultValue(const FString& InParamete
 		return FCustomizableObjectIntParameterValue::DEFAULT_PARAMETER_VALUE;
 	}
 
-	const TSharedPtr<mu::Model> Model = GetModel();
+	const TSharedPtr<mu::Model> Model = GetPrivate()->GetModel();
 	if (!Model)
 	{
 		checkNoEntry();
@@ -1586,7 +1582,7 @@ bool UCustomizableObject::GetBoolParameterDefaultValue(const FString& InParamete
 		return FCustomizableObjectBoolParameterValue::DEFAULT_PARAMETER_VALUE;
 	}
 
-	const TSharedPtr<mu::Model> Model = GetModel();
+	const TSharedPtr<mu::Model> Model = GetPrivate()->GetModel();
 	if (!Model)
 	{
 		checkNoEntry();
@@ -1606,7 +1602,7 @@ FLinearColor UCustomizableObject::GetColorParameterDefaultValue(const FString& I
 		return FCustomizableObjectVectorParameterValue::DEFAULT_PARAMETER_VALUE;
 	}
 
-	const TSharedPtr<mu::Model> Model = GetModel();
+	const TSharedPtr<mu::Model> Model = GetPrivate()->GetModel();
 	if (!Model)
 	{
 		checkNoEntry();
@@ -1644,7 +1640,7 @@ FCustomizableObjectProjector UCustomizableObject::GetProjectorParameterDefaultVa
 		return FCustomizableObjectProjectorParameterValue::DEFAULT_PARAMETER_VALUE;
 	}
 
-	const TSharedPtr<mu::Model> Model = GetModel();
+	const TSharedPtr<mu::Model> Model = GetPrivate()->GetModel();
 	if (!Model)
 	{
 		checkNoEntry();
@@ -1669,7 +1665,7 @@ FName UCustomizableObject::GetTextureParameterDefaultValue(const FString& InPara
 		return FCustomizableObjectTextureParameterValue::DEFAULT_PARAMETER_VALUE;
 	}
 
-	const TSharedPtr<mu::Model> Model = GetModel();
+	const TSharedPtr<mu::Model> Model = GetPrivate()->GetModel();
 	if (!Model)
 	{
 		checkNoEntry();
@@ -1921,6 +1917,12 @@ FModelResources& UCustomizableObjectPrivate::GetModelResources(bool bIsCooking)
 //-------------------------------------------------------------------------------------------------
 //-------------------------------------------------------------------------------------------------
 
+ECustomizableObjectCompilationState FCustomizableObjectCompilerBase::GetCompilationState() const
+{
+	return ECustomizableObjectCompilationState::None;
+}
+
+
 void UCustomizableObjectBulk::PostLoad()
 {
 	UObject::PostLoad();
@@ -1983,7 +1985,7 @@ void UCustomizableObjectBulk::CookAdditionalFilesOverride(const TCHAR* PackageFi
 	
 	check(CustomizableObject);
 	
-	FMutableCachedPlatformData* PlatformData = CustomizableObject->CachedPlatformsData.Find(TargetPlatform->PlatformName());
+	FMutableCachedPlatformData* PlatformData = CustomizableObject->GetPrivate()->CachedPlatformsData.Find(TargetPlatform->PlatformName());
 	check(PlatformData);
 
 	// Data to serialize in separate files
@@ -2013,7 +2015,7 @@ void UCustomizableObjectBulk::PrepareBulkData(UCustomizableObject* InOuter, cons
 	BulkDataFileNames.Empty();
 	
 	// Split the Streamable data into several separate files and fix up FileIndex and Offset of each StreamableBlock
-	if (TSharedPtr<const mu::Model, ESPMode::ThreadSafe> Model = CustomizableObject->GetModel())
+	if (TSharedPtr<const mu::Model, ESPMode::ThreadSafe> Model = CustomizableObject->GetPrivate()->GetModel())
 	{
 		const uint64 MaxChunkSize = UCustomizableObjectSystem::GetInstance()->GetMaxChunkSizeForPlatform(TargetPlatform);
 
diff --git a/Engine/Plugins/Experimental/Mutable/Source/CustomizableObject/Private/MuCO/CustomizableObjectInstance.cpp b/Engine/Plugins/Experimental/Mutable/Source/CustomizableObject/Private/MuCO/CustomizableObjectInstance.cpp
--- a/Engine/Plugins/Experimental/Mutable/Source/CustomizableObject/Private/MuCO/CustomizableObjectInstance.cpp
+++ b/Engine/Plugins/Experimental/Mutable/Source/CustomizableObject/Private/MuCO/CustomizableObjectInstance.cpp
@@ -347,14 +347,14 @@ void UCustomizableInstancePrivate::BindObjectDelegates(UCustomizableObject*  Cur
 	// Unbind callback from the previous CO
 	if (CurrentCustomizableObject)
 	{
-		CurrentCustomizableObject->PostCompileDelegate.RemoveAll(this);
+		CurrentCustomizableObject->GetPrivate()->PostCompileDelegate.RemoveAll(this);
 		CurrentCustomizableObject->GetPrivate()->Status.GetOnStateChangedDelegate().RemoveAll(this);
 	}
 
 	// Bind callback to the new CO
 	if (NewCustomizableObject)
 	{
-		NewCustomizableObject->PostCompileDelegate.AddUObject(this, &UCustomizableInstancePrivate::OnPostCompile);
+		NewCustomizableObject->GetPrivate()->PostCompileDelegate.AddUObject(this, &UCustomizableInstancePrivate::OnPostCompile);
 		NewCustomizableObject->GetPrivate()->Status.GetOnStateChangedDelegate().AddUObject(this, &UCustomizableInstancePrivate::OnObjectStatusChanged);
 	}
 }
@@ -5401,7 +5401,7 @@ void UCustomizableInstancePrivate::BuildMaterials(const TSharedRef<FUpdateContex
 		CustomizableObject->GetPathName(),
 		Public->GetPathName(),
 		UCustomizableObjectSystem::GetInstance()->GetPrivate()->MutableSystem,
-		CustomizableObject->GetModel(),
+		CustomizableObject->GetPrivate()->GetModel(),
 		OperationData->Parameters,
 	    OperationData->GetCapturedDescriptor().GetState());
 	
diff --git a/Engine/Plugins/Experimental/Mutable/Source/CustomizableObject/Private/MuCO/CustomizableObjectMipDataProvider.cpp b/Engine/Plugins/Experimental/Mutable/Source/CustomizableObject/Private/MuCO/CustomizableObjectMipDataProvider.cpp
--- a/Engine/Plugins/Experimental/Mutable/Source/CustomizableObject/Private/MuCO/CustomizableObjectMipDataProvider.cpp
+++ b/Engine/Plugins/Experimental/Mutable/Source/CustomizableObject/Private/MuCO/CustomizableObjectMipDataProvider.cpp
@@ -112,7 +112,7 @@ void FMutableTextureMipDataProvider::Init(const FTextureUpdateContext& Context,
 #if WITH_EDITOR
 	check(Context.Texture->HasPendingInitOrStreaming());
 	check(CustomizableObjectInstance->GetCustomizableObject());
-	if (CustomizableObjectInstance->GetCustomizableObject()->IsLocked())
+	if (CustomizableObjectInstance->GetCustomizableObject()->GetPrivate()->IsLocked())
 	{
 		PrintWarningAndAdvanceToCleanup();
 
@@ -300,7 +300,7 @@ int32 FMutableTextureMipDataProvider::GetMips(const FTextureUpdateContext& Conte
 #if WITH_EDITOR
 	check(Context.Texture->HasPendingInitOrStreaming());
 	check(CustomizableObjectInstance->GetCustomizableObject());
-	if (CustomizableObjectInstance->GetCustomizableObject()->IsLocked())
+	if (CustomizableObjectInstance->GetCustomizableObject()->GetPrivate()->IsLocked())
 	{
 		PrintWarningAndAdvanceToCleanup();
 
@@ -398,7 +398,7 @@ bool FMutableTextureMipDataProvider::PollMips(const FTextureUpdateSyncOptions& S
 	
 #if WITH_EDITOR
 	check(CustomizableObjectInstance->GetCustomizableObject());
-	if (CustomizableObjectInstance->GetCustomizableObject()->IsLocked())
+	if (CustomizableObjectInstance->GetCustomizableObject()->GetPrivate()->IsLocked())
 	{
 		PrintWarningAndAdvanceToCleanup();
 
diff --git a/Engine/Plugins/Experimental/Mutable/Source/CustomizableObject/Private/MuCO/CustomizableObjectSystem.cpp b/Engine/Plugins/Experimental/Mutable/Source/CustomizableObject/Private/MuCO/CustomizableObjectSystem.cpp
--- a/Engine/Plugins/Experimental/Mutable/Source/CustomizableObject/Private/MuCO/CustomizableObjectSystem.cpp
+++ b/Engine/Plugins/Experimental/Mutable/Source/CustomizableObject/Private/MuCO/CustomizableObjectSystem.cpp
@@ -985,7 +985,7 @@ void UCustomizableObjectSystemPrivate::GetMipStreamingConfig(const UCustomizable
 		bOutNeverStream = State->bDisableTextureStreaming;
 
 		// Was streaming disabled at object-compilation time? 
-		if (Instance.GetCustomizableObject()->bDisableTextureStreaming)
+		if (Instance.GetCustomizableObject()->GetPrivate()->bDisableTextureStreaming)
 		{
 			bOutNeverStream = true;
 		}
@@ -3497,7 +3497,7 @@ void UCustomizableObjectSystem::OnPreBeginPIE(const bool bIsSimulatingInEditor)
 		}
 		
 		const UCustomizableObject* Object = Cast<UCustomizableObject>(Asset.GetAsset());
-		if (!Object || Object->IsCompiled() || Object->IsLocked() || Object->bIsChildObject)
+		if (!Object || Object->IsCompiled() || Object->GetPrivate()->IsLocked() || Object->bIsChildObject)
 		{
 			continue;
 		}
@@ -3547,7 +3547,7 @@ void UCustomizableObjectSystem::RecompileCustomizableObjectAsync(const FAssetDat
 		return;
 	}
 	
-	if ((InObject && InObject->IsLocked()) || ObjectsToRecompile.Find((InAssetData)) != INDEX_NONE)
+	if ((InObject && InObject->GetPrivate()->IsLocked()) || ObjectsToRecompile.Find((InAssetData)) != INDEX_NONE)
 	{
 		return;
 	}
diff --git a/Engine/Plugins/Experimental/Mutable/Source/CustomizableObject/Private/MuCO/UnrealMutableImageProvider.cpp b/Engine/Plugins/Experimental/Mutable/Source/CustomizableObject/Private/MuCO/UnrealMutableImageProvider.cpp
--- a/Engine/Plugins/Experimental/Mutable/Source/CustomizableObject/Private/MuCO/UnrealMutableImageProvider.cpp
+++ b/Engine/Plugins/Experimental/Mutable/Source/CustomizableObject/Private/MuCO/UnrealMutableImageProvider.cpp
@@ -110,7 +110,7 @@ bool FUnrealMutableImageProvider::Tick(float DeltaTime)
 	UCustomizableObject* CO = nullptr;
 	for (TObjectIterator<UCustomizableObject> It; It; ++It)
 	{
-		if (IsValid(*It) && It->GetModel().Get() == Request->ModelPtr)
+		if (IsValid(*It) && It->GetPrivate()->GetModel().Get() == Request->ModelPtr)
 		{
 			CO = *It;
 			break;
diff --git a/Engine/Plugins/Experimental/Mutable/Source/CustomizableObject/Private/MuCO/UnrealMutableModelDiskStreamer.cpp b/Engine/Plugins/Experimental/Mutable/Source/CustomizableObject/Private/MuCO/UnrealMutableModelDiskStreamer.cpp
--- a/Engine/Plugins/Experimental/Mutable/Source/CustomizableObject/Private/MuCO/UnrealMutableModelDiskStreamer.cpp
+++ b/Engine/Plugins/Experimental/Mutable/Source/CustomizableObject/Private/MuCO/UnrealMutableModelDiskStreamer.cpp
@@ -4,6 +4,7 @@
 
 #include "Async/AsyncFileHandle.h"
 #include "HAL/PlatformFile.h"
+#include "MuCO/CustomizableObjectPrivate.h"
 #include "MuCO/CustomizableObjectSystem.h"
 #include "MuR/Model.h"
 #include "MuR/MutableTrace.h"
@@ -76,14 +77,14 @@ bool FUnrealMutableModelBulkReader::PrepareStreamingForObject(UCustomizableObjec
 	// Is the object already prepared for streaming?
 	bool bAlreadyStreaming = Objects.FindByPredicate(
 		[CustomizableObject](const FObjectData& d)
-		{ return d.Model.Pin().Get() == CustomizableObject->GetModel().Get(); })
+		{ return d.Model.Pin().Get() == CustomizableObject->GetPrivate()->GetModel().Get(); })
 		!=
 		nullptr;
 
 	if (!bAlreadyStreaming)
 	{
 		FObjectData NewData;
-		NewData.Model = TWeakPtr<const mu::Model>(CustomizableObject->GetModel());
+		NewData.Model = TWeakPtr<const mu::Model>(CustomizableObject->GetPrivate()->GetModel());
 
 #if WITH_EDITOR
 		FString FolderPath = CustomizableObject->GetCompiledDataFolderPath();
@@ -146,7 +147,7 @@ void FUnrealMutableModelBulkReader::CancelStreamingForObject(const UCustomizable
 	// See if we can free previuously allocated resources
 	for (int32 ObjectIndex = 0; ObjectIndex < Objects.Num(); ++ObjectIndex)
 	{
-		if (Objects[ObjectIndex].Model.Pin() == CustomizableObject->GetModel())
+		if (Objects[ObjectIndex].Model.Pin() == CustomizableObject->GetPrivate()->GetModel())
 		{
 			check(Objects[ObjectIndex].CurrentReadRequests.IsEmpty());
 
@@ -170,7 +171,7 @@ bool FUnrealMutableModelBulkReader::AreTherePendingStreamingOperationsForObject(
 
 	for (int32 ObjectIndex = 0; ObjectIndex < Objects.Num(); ++ObjectIndex)
 	{
-		if (Objects[ObjectIndex].Model.Pin() == CustomizableObject->GetModel())
+		if (Objects[ObjectIndex].Model.Pin() == CustomizableObject->GetPrivate()->GetModel())
 		{
 			if (!Objects[ObjectIndex].CurrentReadRequests.IsEmpty())
 			{
diff --git a/Engine/Plugins/Experimental/Mutable/Source/CustomizableObjectEditor/Private/MuCOE/CustomizableInstanceDetails.cpp b/Engine/Plugins/Experimental/Mutable/Source/CustomizableObjectEditor/Private/MuCOE/CustomizableInstanceDetails.cpp
--- a/Engine/Plugins/Experimental/Mutable/Source/CustomizableObjectEditor/Private/MuCOE/CustomizableInstanceDetails.cpp
+++ b/Engine/Plugins/Experimental/Mutable/Source/CustomizableObjectEditor/Private/MuCOE/CustomizableInstanceDetails.cpp
@@ -884,10 +884,10 @@ TSharedRef<SWidget> FCustomizableInstanceDetails::GenerateIntWidget(const int32
 			{
 				ValueIndex = i;
 
-				const FString* Identifier = CustomizableObject->GroupNodeMap.FindKey(FCustomizableObjectIdPair(ParamName, PossibleValue));
+				const FString* Identifier = CustomizableObject->GetPrivate()->GroupNodeMap.FindKey(FCustomizableObjectIdPair(ParamName, PossibleValue));
 				if (Identifier)
 				{
-					if (FString* CustomizableObjectPath = CustomizableObject->CustomizableObjectPathMap.Find(*Identifier))
+					if (FString* CustomizableObjectPath = CustomizableObject->GetPrivate()->CustomizableObjectPathMap.Find(*Identifier))
 					{
 						ToolTipText = *CustomizableObjectPath;
 					}
diff --git a/Engine/Plugins/Experimental/Mutable/Source/CustomizableObjectEditor/Private/MuCOE/CustomizableObjectCompileRunnable.cpp b/Engine/Plugins/Experimental/Mutable/Source/CustomizableObjectEditor/Private/MuCOE/CustomizableObjectCompileRunnable.cpp
--- a/Engine/Plugins/Experimental/Mutable/Source/CustomizableObjectEditor/Private/MuCOE/CustomizableObjectCompileRunnable.cpp
+++ b/Engine/Plugins/Experimental/Mutable/Source/CustomizableObjectEditor/Private/MuCOE/CustomizableObjectCompileRunnable.cpp
@@ -5,6 +5,7 @@
 #include "HAL/FileManager.h"
 #include "MuCO/UnrealMutableModelDiskStreamer.h"
 #include "MuCO/UnrealToMutableTextureConversionUtils.h"
+#include "MuCO/CustomizableObjectPrivate.h"
 #include "MuR/Model.h"
 #include "MuT/Compiler.h"
 #include "MuT/ErrorLog.h"
@@ -245,10 +246,10 @@ const TArray<FCustomizableObjectCompileRunnable::FError>& FCustomizableObjectCom
 
 FCustomizableObjectSaveDDRunnable::FCustomizableObjectSaveDDRunnable(UCustomizableObject* CustomizableObject, const FCompilationOptions& InOptions)
 {
-	Model = CustomizableObject->GetModel();
+	Model = CustomizableObject->GetPrivate()->GetModel();
 	Options = InOptions;
 	
-	CustomizableObjectHeader.InternalVersion = CustomizableObject->GetCurrentSupportedVersion();
+	CustomizableObjectHeader.InternalVersion = CustomizableObject->GetPrivate()->CurrentSupportedVersion;
 	CustomizableObjectHeader.VersionId = Options.bIsCooking? FGuid::NewGuid() : CustomizableObject->GetVersionId();
 
 	if (!Options.bIsCooking)
diff --git a/Engine/Plugins/Experimental/Mutable/Source/CustomizableObjectEditor/Private/MuCOE/CustomizableObjectCompileRunnable.h b/Engine/Plugins/Experimental/Mutable/Source/CustomizableObjectEditor/Private/MuCOE/CustomizableObjectCompileRunnable.h
--- a/Engine/Plugins/Experimental/Mutable/Source/CustomizableObjectEditor/Private/MuCOE/CustomizableObjectCompileRunnable.h
+++ b/Engine/Plugins/Experimental/Mutable/Source/CustomizableObjectEditor/Private/MuCOE/CustomizableObjectCompileRunnable.h
@@ -4,6 +4,7 @@
 
 #include "HAL/Runnable.h"
 #include "MuCO/CustomizableObject.h"
+#include "MuCO/CustomizableObjectPrivate.h"
 #include "MuR/Ptr.h"
 #include "MuT/Node.h"
 #include "MuCOE/CustomizableObjectEditorLogger.h"
diff --git a/Engine/Plugins/Experimental/Mutable/Source/CustomizableObjectEditor/Private/MuCOE/CustomizableObjectCompiler.cpp b/Engine/Plugins/Experimental/Mutable/Source/CustomizableObjectEditor/Private/MuCOE/CustomizableObjectCompiler.cpp
--- a/Engine/Plugins/Experimental/Mutable/Source/CustomizableObjectEditor/Private/MuCOE/CustomizableObjectCompiler.cpp
+++ b/Engine/Plugins/Experimental/Mutable/Source/CustomizableObjectEditor/Private/MuCOE/CustomizableObjectCompiler.cpp
@@ -185,7 +185,7 @@ void FCustomizableObjectCompiler::Compile(UCustomizableObject& Object, const FCo
 	Info.FadeOutDuration = 1.0f;
 	Info.ExpireDuration = 1.0f;
 
-	if (PreloadingReferencerAssets || CompilationLaunchPending || (Object.IsLocked()))
+	if (PreloadingReferencerAssets || CompilationLaunchPending || (Object.GetPrivate()->IsLocked()))
 	{
 		UE_LOG(LogMutable, Warning, TEXT("%s"), *Message);
 		FSlateNotificationManager::Get().AddNotification(Info);
@@ -822,7 +822,7 @@ void FCustomizableObjectCompiler::SetCompilationState(ECustomizableObjectCompila
 	
 	if (CurrentObject)
 	{
-		CurrentObject->CompilationState = InState;
+		CurrentObject->GetPrivate()->CompilationState = InState;
 	}
 }
 
@@ -864,7 +864,7 @@ void FCustomizableObjectCompiler::CompileInternal(UCustomizableObject* Object, c
 		return;
 	}
 
-	if (Object->IsLocked() || !UCustomizableObjectSystem::GetInstance()->LockObject(Object))
+	if (Object->GetPrivate()->IsLocked() || !UCustomizableObjectSystem::GetInstance()->LockObject(Object))
 	{
 		UE_LOG(LogMutable, Display, TEXT("Customizable Object is already being compiled or updated %s. Please wait a few seconds and try again."), *Object->GetName());
 		return;
@@ -921,7 +921,7 @@ void FCustomizableObjectCompiler::CompileInternal(UCustomizableObject* Object, c
 			CompilerLog(FText::FromString(TEXT("Failed to generate the mutable node graph. Object not built.")), nullptr);
 		}
 
-		if (Object->IsLocked())
+		if (Object->GetPrivate()->IsLocked())
 		{
 			UCustomizableObjectSystem::GetInstance()->UnlockObject(Object);
 		}
@@ -1033,16 +1033,16 @@ void FCustomizableObjectCompiler::CompileInternal(UCustomizableObject* Object, c
 			 ClothingAssetData.ConfigsData.RemoveAllSwap(IsSharedConfigData);
 		}
 
-		Object->GroupNodeMap = GenerationContext.GroupNodeMap;
+		Object->GetPrivate()->GroupNodeMap = GenerationContext.GroupNodeMap;
 
 		if (GenerationContext.Options.OptimizationLevel == 0)
 		{
 			// If the optimization level is "none" disable texture streaming, because textures are all referenced
 			// unreal assets and progressive generation is not supported.
-			Object->bDisableTextureStreaming = true;
+			Object->GetPrivate()->bDisableTextureStreaming = true;
 		}
 		
-		Object->bIsCompiledWithOptimization = GenerationContext.Options.OptimizationLevel < UE_MUTABLE_MAX_OPTIMIZATION;
+		Object->GetPrivate()->bIsCompiledWithOptimization = GenerationContext.Options.OptimizationLevel < UE_MUTABLE_MAX_OPTIMIZATION;
 
 		Object->AlwaysLoadedExtensionData = MoveTemp(GenerationContext.AlwaysLoadedExtensionData);
 
@@ -1053,7 +1053,7 @@ void FCustomizableObjectCompiler::CompileInternal(UCustomizableObject* Object, c
 		}
 
 #if WITH_EDITORONLY_DATA
-		Object->CustomizableObjectPathMap = GenerationContext.CustomizableObjectPathMap;
+		Object->GetPrivate()->CustomizableObjectPathMap = GenerationContext.CustomizableObjectPathMap;
 #endif
 
 		Object->LODSettings.NumLODsInRoot = GenerationContext.NumLODsInRoot;
@@ -1092,7 +1092,7 @@ void FCustomizableObjectCompiler::CompileInternal(UCustomizableObject* Object, c
 
 		// Lock the object to prevent instance updates while compiling. It's ignored and returns false if it's already locked
 		// Will unlock in the FinishCompilation call.
-		if (!Object->IsLocked())
+		if (!Object->GetPrivate()->IsLocked())
 		{
 			UCustomizableObjectSystem::GetInstance()->LockObject(Object);
 		}
@@ -1245,7 +1245,7 @@ void FCustomizableObjectCompiler::FinishCompilation()
 	}
 	else
 	{
-		CurrentObject->GetModel()->GetPrivate()->UnloadRoms();
+		CurrentObject->GetPrivate()->GetModel()->GetPrivate()->UnloadRoms();
 
 		// when skipping the SaveDerivedData task unlock the object so that instances can be updated
 		UCustomizableObjectSystem::GetInstance()->UnlockObject(CurrentObject);
@@ -1267,7 +1267,7 @@ void FCustomizableObjectCompiler::FinishSavingDerivedData()
 	SaveDDThread.Reset();
 	SaveDDTask.Reset();
 
-	CurrentObject->GetModel()->GetPrivate()->UnloadRoms();
+	CurrentObject->GetPrivate()->GetModel()->GetPrivate()->UnloadRoms();
 	
 	// Unlock the object so that instances can be updated
 	UCustomizableObjectSystem::GetInstance()->UnlockObject(CurrentObject);
@@ -1323,7 +1323,7 @@ void FCustomizableObjectCompiler::ForceFinishBeforeStartCompilation(UCustomizabl
 	CleanCachedReferencers();
 	UpdateArrayGCProtect();
 
-	if (Object && Object->IsLocked())
+	if (Object && Object->GetPrivate()->IsLocked())
 	{
 		UCustomizableObjectSystem::GetInstance()->UnlockObject(Object);
 	}
diff --git a/Engine/Plugins/Experimental/Mutable/Source/CustomizableObjectEditor/Private/MuCOE/CustomizableObjectCompiler.h b/Engine/Plugins/Experimental/Mutable/Source/CustomizableObjectEditor/Private/MuCOE/CustomizableObjectCompiler.h
--- a/Engine/Plugins/Experimental/Mutable/Source/CustomizableObjectEditor/Private/MuCOE/CustomizableObjectCompiler.h
+++ b/Engine/Plugins/Experimental/Mutable/Source/CustomizableObjectEditor/Private/MuCOE/CustomizableObjectCompiler.h
@@ -3,6 +3,7 @@
 #pragma once
 
 #include "AssetRegistry/AssetData.h"
+#include "MuCO/CustomizableObjectPrivate.h"
 #include "MuCOE/CompilationMessageCache.h"
 #include "MuCOE/GenerateMutableSource/GenerateMutableSource.h"
 #include "MuCOE/CustomizableObjectEditorLogger.h"
diff --git a/Engine/Plugins/Experimental/Mutable/Source/CustomizableObjectEditor/Private/MuCOE/CustomizableObjectDetails.cpp b/Engine/Plugins/Experimental/Mutable/Source/CustomizableObjectEditor/Private/MuCOE/CustomizableObjectDetails.cpp
--- a/Engine/Plugins/Experimental/Mutable/Source/CustomizableObjectEditor/Private/MuCOE/CustomizableObjectDetails.cpp
+++ b/Engine/Plugins/Experimental/Mutable/Source/CustomizableObjectEditor/Private/MuCOE/CustomizableObjectDetails.cpp
@@ -6,6 +6,7 @@
 #include "DetailWidgetRow.h"
 #include "IDetailsView.h"
 #include "MuCO/CustomizableObject.h"
+#include "MuCO/CustomizableObjectPrivate.h"
 #include "MuR/Model.h"
 #include "MuR/Parameters.h"
 #include "MuR/Ptr.h"
@@ -96,9 +97,9 @@ void FCustomizableObjectDetails::CustomizeDetails( IDetailLayoutBuilder& DetailB
 	DetailBuilder.HideProperty("PopulationClassTags");
 
 	// Make the tree get automatically updated each time we compile the host CO
-	if (CustomizableObject && !CustomizableObject->PostCompileDelegate.IsBoundToObject(this))
+	if (CustomizableObject && !CustomizableObject->GetPrivate()->PostCompileDelegate.IsBoundToObject(this))
 	{
-		CustomizableObject->PostCompileDelegate.AddSP(this, &FCustomizableObjectDetails::UpdateTree);
+		CustomizableObject->GetPrivate()->PostCompileDelegate.AddSP(this, &FCustomizableObjectDetails::UpdateTree);
 	}
 	
 	// Cache the states defined in the CO
@@ -149,15 +150,15 @@ void FCustomizableObjectDetails::CustomizeDetails( IDetailLayoutBuilder& DetailB
 
 void FCustomizableObjectDetails::UpdateTree()
 {
-	if ( CustomizableObject &&  CustomizableObject->GetModel() )
+	if ( CustomizableObject &&  CustomizableObject->GetPrivate()->GetModel() )
 	{
 		RootTreeItems.SetNumUninitialized(0);
-		const uint32 NumElements = CustomizableObject->GetModel()->GetStateCount();
+		const uint32 NumElements = CustomizableObject->GetPrivate()->GetModel()->GetStateCount();
 		for ( uint32 i=0; i<NumElements; ++i )
 		{
 			TSharedPtr<FStateDetailsNode> SlateDetailsNode = MakeShareable( new FStateDetailsNode );
 
-			SlateDetailsNode->Model = CustomizableObject->GetModel();
+			SlateDetailsNode->Model = CustomizableObject->GetPrivate()->GetModel();
 			SlateDetailsNode->StateIndex = i;
 
 			RootTreeItems.Add( SlateDetailsNode );
diff --git a/Engine/Plugins/Experimental/Mutable/Source/CustomizableObjectEditor/Private/MuCOE/CustomizableObjectEditorFunctionLibrary.cpp b/Engine/Plugins/Experimental/Mutable/Source/CustomizableObjectEditor/Private/MuCOE/CustomizableObjectEditorFunctionLibrary.cpp
--- a/Engine/Plugins/Experimental/Mutable/Source/CustomizableObjectEditor/Private/MuCOE/CustomizableObjectEditorFunctionLibrary.cpp
+++ b/Engine/Plugins/Experimental/Mutable/Source/CustomizableObjectEditor/Private/MuCOE/CustomizableObjectEditorFunctionLibrary.cpp
@@ -17,7 +17,7 @@ ECustomizableObjectCompilationState UCustomizableObjectEditorFunctionLibrary::Co
 	}
 
 	FString ObjectPath = CustomizableObject->GetPathName();
-	if (CustomizableObject->IsLocked())
+	if (CustomizableObject->GetPrivate()->IsLocked())
 	{
 		// Take this if you need a hack:
 		// UCustomizableObjectSystem::GetInstance()->UnlockObject(CustomizableObject);
diff --git a/Engine/Plugins/Experimental/Mutable/Source/CustomizableObjectEditor/Private/MuCOE/GenerateMutableSource/GenerateMutableSource.h b/Engine/Plugins/Experimental/Mutable/Source/CustomizableObjectEditor/Private/MuCOE/GenerateMutableSource/GenerateMutableSource.h
--- a/Engine/Plugins/Experimental/Mutable/Source/CustomizableObjectEditor/Private/MuCOE/GenerateMutableSource/GenerateMutableSource.h
+++ b/Engine/Plugins/Experimental/Mutable/Source/CustomizableObjectEditor/Private/MuCOE/GenerateMutableSource/GenerateMutableSource.h
@@ -7,6 +7,7 @@
 #include "Engine/DataTable.h"
 #include "Engine/SkeletalMesh.h"
 #include "MuCO/CustomizableObject.h"
+#include "MuCO/CustomizableObjectIdentifier.h"
 #include "MuCOE/ExtensionDataCompilerInterface.h"
 #include "MuCOE/Nodes/CustomizableObjectNodeMaterialBase.h"
 #include "MuCOE/Nodes/CustomizableObjectNodeObject.h"
@@ -40,7 +41,6 @@ class UMaterialInterface;
 class UObject;
 class UPhysicsAsset;
 class UTexture2D;
-struct FCustomizableObjectIdPair;
 struct FAnimBpOverridePhysicsAssetsInfo;
 struct FMutableGraphGenerationContext;
 struct FParameterUIData;
diff --git a/Engine/Plugins/Experimental/Mutable/Source/CustomizableObjectEditor/Private/MuCOE/SCustomizableInstanceProperties.cpp b/Engine/Plugins/Experimental/Mutable/Source/CustomizableObjectEditor/Private/MuCOE/SCustomizableInstanceProperties.cpp
--- a/Engine/Plugins/Experimental/Mutable/Source/CustomizableObjectEditor/Private/MuCOE/SCustomizableInstanceProperties.cpp
+++ b/Engine/Plugins/Experimental/Mutable/Source/CustomizableObjectEditor/Private/MuCOE/SCustomizableInstanceProperties.cpp
@@ -762,10 +762,10 @@ void SCustomizableInstanceProperties::AddParameter(int32 ParamIndexInObject)
 				{
 					ValueIndex = i;
 
-					const FString* Identifier = CustomizableObject->GroupNodeMap.FindKey(FCustomizableObjectIdPair(ParamName, PossibleValue));
+					const FString* Identifier = CustomizableObject->GetPrivate()->GroupNodeMap.FindKey(FCustomizableObjectIdPair(ParamName, PossibleValue));
 					if (Identifier)
 					{
-						if (FString* CustomizableObjectPath = CustomizableObject->CustomizableObjectPathMap.Find(*Identifier))
+						if (FString* CustomizableObjectPath = CustomizableObject->GetPrivate()->CustomizableObjectPathMap.Find(*Identifier))
 						{
 							ToolTipText = *CustomizableObjectPath;
 						}
diff --git a/Engine/Plugins/Experimental/Mutable/Source/CustomizableObjectEditor/Private/MuCOE/SCustomizableObjectEditorPerformanceReport.cpp b/Engine/Plugins/Experimental/Mutable/Source/CustomizableObjectEditor/Private/MuCOE/SCustomizableObjectEditorPerformanceReport.cpp
--- a/Engine/Plugins/Experimental/Mutable/Source/CustomizableObjectEditor/Private/MuCOE/SCustomizableObjectEditorPerformanceReport.cpp
+++ b/Engine/Plugins/Experimental/Mutable/Source/CustomizableObjectEditor/Private/MuCOE/SCustomizableObjectEditorPerformanceReport.cpp
@@ -15,6 +15,7 @@
 #include "MuCOE/SCustomizableObjectEditorTextureAnalyzer.h"
 #include "MuCOE/UnrealEditorPortabilityHelpers.h"
 #include "PropertyCustomizationHelpers.h"
+#include "MuCO/CustomizableObjectPrivate.h"
 #include "Widgets/Input/SComboButton.h"
 #include "Widgets/Input/SNumericEntryBox.h"
 #include "Widgets/Layout/SScrollBox.h"
@@ -1343,7 +1344,7 @@ TSharedRef<SWidget> SPerformanceReportWorstCaseRow::GenerateWidgetForColumn(cons
 		{
 			if (const UCustomizableObject* CustomizableObject = ReportInstance->WorstCaseInstance->GetCustomizableObject();
 				CustomizableObject &&
-				!CustomizableObject->IsLocked())
+				!CustomizableObject->GetPrivate()->IsLocked())
 			{
 				const bool bIsDetailParam = ReportInstance->WorstCaseType == UWorstCasePerformanceReportInstance::EPerformanceReportWorstCaseType::SLOWEST_UPDATE_PARAM_IN_RELEVANT_STATE;
 				const FText TextState = ReportInstance->WorstCaseType == UWorstCasePerformanceReportInstance::EPerformanceReportWorstCaseType::SLOWEST_UPDATE_ENTER_STATE ?
diff --git a/Engine/Plugins/Experimental/Mutable/Source/CustomizableObjectEditor/Public/MuCOE/CustomizableObjectInstanceBakingUtils.cpp b/Engine/Plugins/Experimental/Mutable/Source/CustomizableObjectEditor/Public/MuCOE/CustomizableObjectInstanceBakingUtils.cpp
--- a/Engine/Plugins/Experimental/Mutable/Source/CustomizableObjectEditor/Public/MuCOE/CustomizableObjectInstanceBakingUtils.cpp
+++ b/Engine/Plugins/Experimental/Mutable/Source/CustomizableObjectEditor/Public/MuCOE/CustomizableObjectInstanceBakingUtils.cpp
@@ -310,7 +310,7 @@ void BakeCustomizableObjectInstance(
 	check (InstanceCO);
 
 	// Ensure the CO of the COI is accessible 
-	if (!InstanceCO || InstanceCO->IsLocked())
+	if (!InstanceCO || InstanceCO->GetPrivate()->IsLocked())
 	{
 		FCustomizableObjectEditorLogger::CreateLog(
 		LOCTEXT("CustomizableObjectCompilingTryLater_Baking", "Please wait until the Customizable Object is compiled"))
diff --git a/Engine/Plugins/Experimental/Mutable/Source/MutableValidation/Private/MuV/AssetValidator_CustomizableObjects.cpp b/Engine/Plugins/Experimental/Mutable/Source/MutableValidation/Private/MuV/AssetValidator_CustomizableObjects.cpp
--- a/Engine/Plugins/Experimental/Mutable/Source/MutableValidation/Private/MuV/AssetValidator_CustomizableObjects.cpp
+++ b/Engine/Plugins/Experimental/Mutable/Source/MutableValidation/Private/MuV/AssetValidator_CustomizableObjects.cpp
@@ -11,6 +11,7 @@
 #include "Engine/Texture.h"
 #include "Materials/Material.h"
 #include "MuCO/CustomizableObject.h"
+#include "MuCO/CustomizableObjectPrivate.h"
 #include "MuCO/CustomizableObjectSystem.h"
 #include "UObject/NameTypes.h"
 #include "UObject/Object.h"
diff --git a/Engine/Plugins/Experimental/Mutable/Source/MutableValidation/Private/MuV/COIBakingTestCommandlet.cpp b/Engine/Plugins/Experimental/Mutable/Source/MutableValidation/Private/MuV/COIBakingTestCommandlet.cpp
--- a/Engine/Plugins/Experimental/Mutable/Source/MutableValidation/Private/MuV/COIBakingTestCommandlet.cpp
+++ b/Engine/Plugins/Experimental/Mutable/Source/MutableValidation/Private/MuV/COIBakingTestCommandlet.cpp
@@ -8,6 +8,7 @@
 #include "MuCO/CustomizableObjectSystem.h"
 #include "MuCOE/CustomizableObjectInstanceBakingUtils.h"
 #include "HAL/FileManager.h"
+#include "MuCO/CustomizableObjectPrivate.h"
 #include "MuT/UnrealPixelFormatOverride.h"
 
 
diff --git a/Engine/Plugins/Experimental/Mutable/Source/MutableValidation/Private/MuV/CustomizableObjectValidationCommandlet.cpp b/Engine/Plugins/Experimental/Mutable/Source/MutableValidation/Private/MuV/CustomizableObjectValidationCommandlet.cpp
--- a/Engine/Plugins/Experimental/Mutable/Source/MutableValidation/Private/MuV/CustomizableObjectValidationCommandlet.cpp
+++ b/Engine/Plugins/Experimental/Mutable/Source/MutableValidation/Private/MuV/CustomizableObjectValidationCommandlet.cpp
@@ -6,6 +6,7 @@
 #include "AssetRegistry/AssetRegistryModule.h"
 #include "Containers/Array.h"
 #include "MuCO/CustomizableObject.h"
+#include "MuCO/CustomizableObjectPrivate.h"
 #include "MuCO/CustomizableObjectSystem.h"
 #include "MuR/Model.h"
 
@@ -102,7 +103,7 @@ int32 UCustomizableObjectValidationCommandlet::Main(const FString& Params)
 		UE_LOG(LogMutable,Display,TEXT("Customizable Object was compiled succesfully."));
 		
 		// GHet the total size of the streaming data of the model ---------------------------------------------- //
-		const TSharedPtr<const mu::Model> MutableModel = ToTestCustomizableObject->GetModel();
+		const TSharedPtr<const mu::Model> MutableModel = ToTestCustomizableObject->GetPrivate()->GetModel();
 		check (MutableModel);
 
 		// Roms ---------------------- //
