[Mutable] Fix branch merge conflict with CachedPlatformsData.

Member moved to private.

#rnx

diff --git a/Engine/Plugins/Experimental/Mutable/Source/CustomizableObject/Private/MuCO/CustomizableObject.cpp b/Engine/Plugins/Experimental/Mutable/Source/CustomizableObject/Private/MuCO/CustomizableObject.cpp
--- a/Engine/Plugins/Experimental/Mutable/Source/CustomizableObject/Private/MuCO/CustomizableObject.cpp
+++ b/Engine/Plugins/Experimental/Mutable/Source/CustomizableObject/Private/MuCO/CustomizableObject.cpp
@@ -180,7 +180,7 @@ bool UCustomizableObject::TryUpdateIsChildObject()
 #if WITH_EDITOR
 bool UCustomizableObject::TryLoadCompiledCookDataForPlatform(const ITargetPlatform* TargetPlatform)
 {
-	const FMutableCachedPlatformData* PlatformData = CachedPlatformsData.Find(TargetPlatform->PlatformName());
+	const FMutableCachedPlatformData* PlatformData = GetPrivate()->CachedPlatformsData.Find(TargetPlatform->PlatformName());
 	if (!PlatformData)
 	{
 		return false;
