[mutable] Fixed incorrect total Mutable texture size reporting due to an int32 overflow.
#rnx

#tests PIE

#virtualized

diff --git a/Engine/Plugins/Experimental/Mutable/Source/CustomizableObject/Classes/MuCO/CustomizableObjectSystem.h b/Engine/Plugins/Experimental/Mutable/Source/CustomizableObject/Classes/MuCO/CustomizableObjectSystem.h
--- a/Engine/Plugins/Experimental/Mutable/Source/CustomizableObject/Classes/MuCO/CustomizableObjectSystem.h
+++ b/Engine/Plugins/Experimental/Mutable/Source/CustomizableObject/Classes/MuCO/CustomizableObjectSystem.h
@@ -345,9 +345,9 @@ public:
 	UFUNCTION(BlueprintCallable, Category = Status)
 	int32 GetTotalInstances() const;
 
-	// Get the amount of memory in use for textures generated by mutable.
+	// Get the amount of GPU memory in use in bytes for textures generated by mutable.
 	UFUNCTION(BlueprintCallable, Category = Status)
-	int32 GetTextureMemoryUsed() const;
+	int64 GetTextureMemoryUsed() const;
 
 	// Return the average build/update time of an instance in ms.
 	UFUNCTION(BlueprintCallable, Category = Status)
diff --git a/Engine/Plugins/Experimental/Mutable/Source/CustomizableObject/Internal/MuCO/LogBenchmarkUtil.h b/Engine/Plugins/Experimental/Mutable/Source/CustomizableObject/Internal/MuCO/LogBenchmarkUtil.h
--- a/Engine/Plugins/Experimental/Mutable/Source/CustomizableObject/Internal/MuCO/LogBenchmarkUtil.h
+++ b/Engine/Plugins/Experimental/Mutable/Source/CustomizableObject/Internal/MuCO/LogBenchmarkUtil.h
@@ -90,7 +90,7 @@ public:
 	void FinishUpdateImage(const FString& CustomizableObjectPathName, const FString& InstancePathName, double TaskUpdateImageTime, const int64 TaskUpdateImageMemoryPeak, const int64 TaskUpdateImageRealMemoryPeak) const;
 	
 	DECLARE_BENCHMARK_STAT(NumAllocatedTextures, uint32);
-	DECLARE_BENCHMARK_STAT(TextureGPUSize, uint32);
+	DECLARE_BENCHMARK_STAT(TextureGPUSize, uint64);
 
 	DECLARE_BENCHMARK_STAT(NumAllocatedSkeletalMeshes, int32);
 
diff --git a/Engine/Plugins/Experimental/Mutable/Source/CustomizableObject/Private/MuCO/CustomizableObjectSystem.cpp b/Engine/Plugins/Experimental/Mutable/Source/CustomizableObject/Private/MuCO/CustomizableObjectSystem.cpp
--- a/Engine/Plugins/Experimental/Mutable/Source/CustomizableObject/Private/MuCO/CustomizableObjectSystem.cpp
+++ b/Engine/Plugins/Experimental/Mutable/Source/CustomizableObject/Private/MuCO/CustomizableObjectSystem.cpp
@@ -3348,7 +3348,7 @@ int32 UCustomizableObjectSystem::GetTotalInstances() const
 	return NumInstances;
 }
 
-int32 UCustomizableObjectSystem::GetTextureMemoryUsed() const
+int64 UCustomizableObjectSystem::GetTextureMemoryUsed() const
 {
 	return GetPrivateChecked()->LogBenchmarkUtil.TextureGPUSize.GetValue();
 }
diff --git a/Engine/Plugins/Experimental/Mutable/Source/CustomizableObject/Private/MuCO/LogBenchmarkUtil.cpp b/Engine/Plugins/Experimental/Mutable/Source/CustomizableObject/Private/MuCO/LogBenchmarkUtil.cpp
--- a/Engine/Plugins/Experimental/Mutable/Source/CustomizableObject/Private/MuCO/LogBenchmarkUtil.cpp
+++ b/Engine/Plugins/Experimental/Mutable/Source/CustomizableObject/Private/MuCO/LogBenchmarkUtil.cpp
@@ -155,7 +155,7 @@ void FLogBenchmarkUtil::UpdateStats()
 
 	// Textures
 	uint32 LocalNumAllocatedTextures = 0;
-	uint32 LocalTextureGPUSize = 0;
+	uint64 LocalTextureGPUSize = 0;
 
 	for (auto Iterator = TextureTrackerArray.CreateIterator(); Iterator; ++Iterator)
 	{
@@ -163,11 +163,16 @@ void FLogBenchmarkUtil::UpdateStats()
 		{
 			Iterator.RemoveCurrent();
 		}
-		else if (Iterator->IsValid() && (*Iterator)->GetResource())
+		else if (Iterator->IsValid())
 		{
 			++LocalNumAllocatedTextures;
 
-			LocalTextureGPUSize += (*Iterator)->CalcTextureMemorySizeEnum(ETextureMipCount::TMC_ResidentMips);
+			UTexture2D* ActualTexture = Iterator->Get();
+
+			if (ActualTexture && ActualTexture->GetResource())
+			{
+				LocalTextureGPUSize += (*Iterator)->CalcTextureMemorySizeEnum(ETextureMipCount::TMC_ResidentMips);
+			}
 		}
 	}
 
