[Mutable] Removed UCustomizableObjectSystem unused code.

#rnx

diff --git a/Engine/Plugins/Experimental/Mutable/Source/CustomizableObject/Classes/MuCO/CustomizableObjectSystem.h b/Engine/Plugins/Experimental/Mutable/Source/CustomizableObject/Classes/MuCO/CustomizableObjectSystem.h
--- a/Engine/Plugins/Experimental/Mutable/Source/CustomizableObject/Classes/MuCO/CustomizableObjectSystem.h
+++ b/Engine/Plugins/Experimental/Mutable/Source/CustomizableObject/Classes/MuCO/CustomizableObjectSystem.h
@@ -272,10 +272,6 @@ public:
 	void SetImagePixelFormatOverride(const mu::FImageOperator::FImagePixelFormatFunc&);
 #endif
 
-	void AddPendingReleaseSkeletalMesh( USkeletalMesh* SkeletalMesh );
-
-	void PurgePendingReleaseSkeletalMesh();
-
 private:
 	UPROPERTY()
 	TArray<FPendingReleaseSkeletalMeshInfo> PendingReleaseSkeletalMesh;
diff --git a/Engine/Plugins/Experimental/Mutable/Source/CustomizableObject/Private/MuCO/CustomizableObjectSystem.cpp b/Engine/Plugins/Experimental/Mutable/Source/CustomizableObject/Private/MuCO/CustomizableObjectSystem.cpp
--- a/Engine/Plugins/Experimental/Mutable/Source/CustomizableObject/Private/MuCO/CustomizableObjectSystem.cpp
+++ b/Engine/Plugins/Experimental/Mutable/Source/CustomizableObject/Private/MuCO/CustomizableObjectSystem.cpp
@@ -1481,50 +1481,6 @@ bool UCustomizableObjectSystem::IsAutoCompilationSync() const
 #endif
 
 
-void UCustomizableObjectSystem::PurgePendingReleaseSkeletalMesh()
-{
-	MUTABLE_CPUPROFILER_SCOPE(UCustomizableObjectSystem::PurgePendingReleaseSkeletalMesh);
-
-	const double CurTime = FPlatformTime::Seconds();
-	static double TimeToDelete = 1.0;
-
-	for (int32 InfoIndex = PendingReleaseSkeletalMesh.Num() - 1; InfoIndex >= 0; --InfoIndex)
-	{
-		FPendingReleaseSkeletalMeshInfo& Info = PendingReleaseSkeletalMesh[InfoIndex];
-		
-		if (Info.SkeletalMesh != nullptr)
-		{
-			if ((CurTime - Info.TimeStamp) >= TimeToDelete)
-			{
-				if (Info.SkeletalMesh->GetSkeleton())
-				{
-					Info.SkeletalMesh->GetSkeleton()->ClearCacheData();
-				}
-				Info.SkeletalMesh->GetRefSkeleton().Empty();
-				Info.SkeletalMesh->GetMaterials().Empty();
-				Info.SkeletalMesh->GetRefBasesInvMatrix().Empty();
-				Info.SkeletalMesh->ReleaseResources();
-				Info.SkeletalMesh->ReleaseResourcesFence.Wait();
-
-				PendingReleaseSkeletalMesh.RemoveAt(InfoIndex);
-			}
-		}
-	}
-}
-
-
-void UCustomizableObjectSystem::AddPendingReleaseSkeletalMesh(USkeletalMesh* SkeletalMesh)
-{
-	check(SkeletalMesh != nullptr);
-
-	FPendingReleaseSkeletalMeshInfo	Info;
-	Info.SkeletalMesh = SkeletalMesh;
-	Info.TimeStamp = FPlatformTime::Seconds();
-
-	PendingReleaseSkeletalMesh.Add(Info);
-}
-
-
 void UCustomizableObjectSystem::ClearCurrentMutableOperation()
 {
 	check(Private != nullptr);
