[mutable] Fixed Instance Descriptor's BuildParameterRelevancy parameter.

#rnx

diff --git a/Engine/Plugins/Experimental/Mutable/Source/CustomizableObject/Private/MuCO/CustomizableObjectInstance.cpp b/Engine/Plugins/Experimental/Mutable/Source/CustomizableObject/Private/MuCO/CustomizableObjectInstance.cpp
--- a/Engine/Plugins/Experimental/Mutable/Source/CustomizableObject/Private/MuCO/CustomizableObjectInstance.cpp
+++ b/Engine/Plugins/Experimental/Mutable/Source/CustomizableObject/Private/MuCO/CustomizableObjectInstance.cpp
@@ -724,11 +724,7 @@ void UCustomizableObjectInstance::SetBuildParameterDecorations(const bool Value)
 
 bool UCustomizableObjectInstance::GetBuildParameterRelevancy() const
 {
-#if WITH_EDITOR
-	return true;
-#else
-	return Descriptor.bBuildParameterRelevancy;
-#endif
+	return Descriptor.GetBuildParameterRelevancy();
 }
 
 
diff --git a/Engine/Plugins/Experimental/Mutable/Source/CustomizableObject/Private/MuCO/CustomizableObjectInstanceDescriptor.cpp b/Engine/Plugins/Experimental/Mutable/Source/CustomizableObject/Private/MuCO/CustomizableObjectInstanceDescriptor.cpp
--- a/Engine/Plugins/Experimental/Mutable/Source/CustomizableObject/Private/MuCO/CustomizableObjectInstanceDescriptor.cpp
+++ b/Engine/Plugins/Experimental/Mutable/Source/CustomizableObject/Private/MuCO/CustomizableObjectInstanceDescriptor.cpp
@@ -434,7 +434,11 @@ void FCustomizableObjectInstanceDescriptor::SetCustomizableObject(UCustomizableO
 
 bool FCustomizableObjectInstanceDescriptor::GetBuildParameterRelevancy() const
 {
+#if WITH_EDITOR
+	return true;
+#else
 	return bBuildParameterRelevancy;
+#endif
 }
 
 
@@ -661,7 +665,7 @@ FString FCustomizableObjectInstanceDescriptor::ToString() const
 	
 	Builder.Appendf(TEXT("CustomizableObject=%s\n"), *CustomizableObject.GetFullName());
 	Builder.Appendf(TEXT("State=%i\n"), State);
-	Builder.Appendf(TEXT("BuildParameterRelevancy=%i\n"), bBuildParameterRelevancy);
+	Builder.Appendf(TEXT("BuildParameterRelevancy=%i\n"), GetBuildParameterRelevancy());
 	Builder.Appendf(TEXT("MinLOD=%i\n"), MinLOD);
 	Builder.Appendf(TEXT("MaxLOD=%i\n"), MaxLOD);
 
diff --git a/Engine/Plugins/Experimental/Mutable/Source/CustomizableObject/Private/MuCO/DescriptorHash.cpp b/Engine/Plugins/Experimental/Mutable/Source/CustomizableObject/Private/MuCO/DescriptorHash.cpp
--- a/Engine/Plugins/Experimental/Mutable/Source/CustomizableObject/Private/MuCO/DescriptorHash.cpp
+++ b/Engine/Plugins/Experimental/Mutable/Source/CustomizableObject/Private/MuCO/DescriptorHash.cpp
@@ -45,7 +45,7 @@ FDescriptorHash::FDescriptorHash(const FCustomizableObjectInstanceDescriptor& De
 	}
 	
 	Hash = HashCombine(Hash, GetTypeHash(Descriptor.State));
-	Hash = HashCombine(Hash, GetTypeHash(Descriptor.bBuildParameterRelevancy));
+	Hash = HashCombine(Hash, GetTypeHash(Descriptor.GetBuildParameterRelevancy()));
 
 	for (const TTuple<FName, FMultilayerProjector>& Pair : Descriptor.MultilayerProjectors)
 	{
