[Mutable] Fix clip morph mesh operation not using mesh bone map to determine affected bones.

#rnx

diff --git a/Engine/Plugins/Experimental/Mutable/Source/MutableRuntime/Private/MuR/OpMeshClipMorphPlane.h b/Engine/Plugins/Experimental/Mutable/Source/MutableRuntime/Private/MuR/OpMeshClipMorphPlane.h
--- a/Engine/Plugins/Experimental/Mutable/Source/MutableRuntime/Private/MuR/OpMeshClipMorphPlane.h
+++ b/Engine/Plugins/Experimental/Mutable/Source/MutableRuntime/Private/MuR/OpMeshClipMorphPlane.h
@@ -96,11 +96,11 @@ namespace mu
 			return;
 		}
 
-		TArray<bool> AffectedBones;
 		TArray<vertex_bone_info> vertex_info;
 
         Ptr<const Skeleton> BaseSkeleton = pBase->GetSkeleton();
 
+		TArray<bool> AffectedBones;
 		const int32 BaseBoneIndex = BaseSkeleton ? BaseSkeleton->FindBone(BoneId) : INDEX_NONE;
         if (BaseBoneIndex != INDEX_NONE)
 		{
@@ -128,6 +128,7 @@ namespace mu
 			vertex_info.SetNum(vcount);
 			//int firstCount = pBase->GetVertexBuffers().GetElementCount();
 
+			const TArray<uint16>& BoneMap = pBase->BoneMap;
 			for (int32 vb = 0; vb < pBase->GetVertexBuffers().m_buffers.Num(); ++vb)
 			{
 				const MESH_BUFFER& result = pBase->GetVertexBuffers().m_buffers[vb];
@@ -161,7 +162,7 @@ namespace mu
 								
 								for (int j = 0; j < components; ++j)
 								{
-									vertex_info[i].bone_indices.Add(pD[j]);
+									vertex_info[i].bone_indices.Add(BoneMap[pD[j]]);
 								}
 
 								secondOffset += elemSize;
@@ -175,7 +176,7 @@ namespace mu
 
 								for (int j = 0; j < components; ++j)
 								{
-									vertex_info[i].bone_indices.Add(pD[j]);
+									vertex_info[i].bone_indices.Add(BoneMap[pD[j]]);
 								}
 
 								secondOffset += elemSize;
@@ -189,7 +190,7 @@ namespace mu
 
 								for (int j = 0; j < components; ++j)
 								{
-									vertex_info[i].bone_indices.Add(pD[j]);
+									vertex_info[i].bone_indices.Add(BoneMap[pD[j]]);
 								}
 
 								secondOffset += elemSize;
@@ -301,11 +302,20 @@ namespace mu
 							ConvertData(i, &vertex[0], MBF_FLOAT32, it.ptr(), it.GetFormat());
 						}
 
-						if (
-							(  BaseBoneIndex != INDEX_NONE && VertexIsAffectedByBone(v, AffectedBones, vertex_info) && VertexIsInMaxRadius(vertex, origin, vertexSelectionBoneMaxRadius))
-                            || (BaseBoneIndex == INDEX_NONE && selectionShape.type == (uint8_t)FShape::Type::None)
-                            || (selectionShape.type == (uint8_t)FShape::Type::AABox && PointInBoundingBox(vertex, selectionShape))
-							)
+						const bool bIsVertexAffectedBone = 
+								BaseBoneIndex != INDEX_NONE &&
+								VertexIsInMaxRadius(vertex, origin, vertexSelectionBoneMaxRadius) &&
+								VertexIsAffectedByBone(v, AffectedBones, vertex_info);
+						
+						const bool bIsVertexAffectedNoShape = 
+								BaseBoneIndex == INDEX_NONE && 
+								selectionShape.type == (uint8_t)FShape::Type::None;
+
+						const bool bIsVertexAffectedBoundingBox = 
+								(selectionShape.type == (uint8_t)FShape::Type::AABox && 
+								PointInBoundingBox(vertex, selectionShape));
+
+						if (bIsVertexAffectedBone || bIsVertexAffectedNoShape || bIsVertexAffectedBoundingBox)
 						{
 							vec3f morph_plane_center = origin;					// MORPH PLANE POS relative to root of the selected bone
 							vec3f clip_plane_center = origin + normal * dist;	// CPLIPPING PLANE POS
@@ -412,5 +422,4 @@ namespace mu
             Result->m_surfaces[0].m_indexCount -= (int32)removedIndices;
         }
 	}
-
 }
