[mutable] Fix race condition in object compilation.
#rnx

diff --git a/Engine/Plugins/Experimental/Mutable/Source/MutableTools/Private/MuT/CodeOptimiser.cpp b/Engine/Plugins/Experimental/Mutable/Source/MutableTools/Private/MuT/CodeOptimiser.cpp
--- a/Engine/Plugins/Experimental/Mutable/Source/MutableTools/Private/MuT/CodeOptimiser.cpp
+++ b/Engine/Plugins/Experimental/Mutable/Source/MutableTools/Private/MuT/CodeOptimiser.cpp
@@ -863,17 +863,14 @@ namespace mu
 					SubgraphCompletionEvent.Trigger();
 					EverythingComplete.AddPrerequisites(SubgraphCompletionEvent);
 				}
-
 			});
 
-		EverythingComplete.AddPrerequisites(LaunchTask);
-
 		// Wait for pending tasks
 		{
 			MUTABLE_CPUPROFILER_SCOPE(ConstantGenerator_WaitPending);
-			EverythingComplete.Trigger();
 
-			// Wait without executing tasks, since this is an independent thread.
+			LaunchTask.Wait();
+			EverythingComplete.Trigger();
 			EverythingComplete.Wait();
 		}
 
