[Backout] - CL31081738
Original CL Desc
-----------------------------------------------------------------
[mutable] Remove some of the unnecessary standard library header includes from core modules.
#rnx

diff --git a/Engine/Plugins/Experimental/Mutable/Source/CustomizableObject/Classes/MuCO/CustomizableObject.h b/Engine/Plugins/Experimental/Mutable/Source/CustomizableObject/Classes/MuCO/CustomizableObject.h
--- a/Engine/Plugins/Experimental/Mutable/Source/CustomizableObject/Classes/MuCO/CustomizableObject.h
+++ b/Engine/Plugins/Experimental/Mutable/Source/CustomizableObject/Classes/MuCO/CustomizableObject.h
@@ -137,9 +137,7 @@ struct FCompilationOptions
 {
 	GENERATED_USTRUCT_BODY()
 	
-	/** Enum to know what texture compression should be used. This compression is used on ly in manual compiles in editor. 
-	 *  When packaging, ECustomizableObjectTextureCompression::HighQuality is always used.
-	 */
+	// Enum to know what texture compression should be used
 	UPROPERTY()
 	ECustomizableObjectTextureCompression TextureCompression = ECustomizableObjectTextureCompression::Fast;
 
diff --git a/Engine/Plugins/Experimental/Mutable/Source/CustomizableObjectEditor/Private/MuCOE/Nodes/CustomizableObjectNodeStaticMesh.cpp b/Engine/Plugins/Experimental/Mutable/Source/CustomizableObjectEditor/Private/MuCOE/Nodes/CustomizableObjectNodeStaticMesh.cpp
--- a/Engine/Plugins/Experimental/Mutable/Source/CustomizableObjectEditor/Private/MuCOE/Nodes/CustomizableObjectNodeStaticMesh.cpp
+++ b/Engine/Plugins/Experimental/Mutable/Source/CustomizableObjectEditor/Private/MuCOE/Nodes/CustomizableObjectNodeStaticMesh.cpp
@@ -15,6 +15,8 @@
 #include "MuCOE/RemapPins/CustomizableObjectNodeRemapPinsByNameDefaultPin.h"
 #include "MuCOE/UnrealEditorPortabilityHelpers.h"
 #include "Widgets/Input/SCheckBox.h"
+#include <memory>
+#include <stdexcept>
 
 class UCustomizableObjectNodeRemapPinsByName;
 class UObject;
diff --git a/Engine/Plugins/Experimental/Mutable/Source/MutableRuntime/Private/MuR/Mesh.cpp b/Engine/Plugins/Experimental/Mutable/Source/MutableRuntime/Private/MuR/Mesh.cpp
--- a/Engine/Plugins/Experimental/Mutable/Source/MutableRuntime/Private/MuR/Mesh.cpp
+++ b/Engine/Plugins/Experimental/Mutable/Source/MutableRuntime/Private/MuR/Mesh.cpp
@@ -10,6 +10,8 @@
 #include "MuR/MeshPrivate.h"
 #include "MuR/MutableTrace.h"
 
+#include <limits>
+
 namespace mu
 {
 
@@ -924,8 +926,8 @@ void Mesh::GetVertexMap
 #define MUTABLE_NUM_BUCKETS 256
 #define MUTABLE_BUCKET_CHANNEL 0
 
-    float rangeMin = TNumericLimits<float>::Max();
-    float rangeMax = -TNumericLimits<float>::Max();
+    float rangeMin = std::numeric_limits<float>::max();
+    float rangeMax = -std::numeric_limits<float>::max();
     MeshBufferIteratorConst< MBF_FLOAT32, float, 3 >  itop = itopBegin;
     for ( int ov=0; ov<otherVertexCount; ++ov )
     {
diff --git a/Engine/Plugins/Experimental/Mutable/Source/MutableRuntime/Private/MuR/Serialisation.cpp b/Engine/Plugins/Experimental/Mutable/Source/MutableRuntime/Private/MuR/Serialisation.cpp
--- a/Engine/Plugins/Experimental/Mutable/Source/MutableRuntime/Private/MuR/Serialisation.cpp
+++ b/Engine/Plugins/Experimental/Mutable/Source/MutableRuntime/Private/MuR/Serialisation.cpp
@@ -10,6 +10,8 @@
 #include "MuR/Image.h"
 #include "MuR/SerialisationPrivate.h"
 
+#include <limits>
+
 
 namespace mu
 {
@@ -98,8 +100,11 @@ namespace mu
         {
             check( m_pD->m_pos + size <= m_pD->m_size );
 
-            const uint8* pSource = ((const uint8*)(m_pD->m_pBuffer))+m_pD->m_pos;
-            FMemory::Memcpy( pData, pSource, (SIZE_T)size );
+            // This could happen in 32-bit platforms
+            check( size<std::numeric_limits<size_t>::max() );
+
+            const uint8_t* pSource = ((const uint8_t*)(m_pD->m_pBuffer))+m_pD->m_pos;
+            FMemory::Memcpy( pData, pSource, (size_t)size );
             m_pD->m_pos += size;
         }
     }
@@ -113,7 +118,9 @@ namespace mu
         m_pD = new Private();
         if (reserve)
         {
-             m_pD->m_buffer.Reserve( reserve );
+            // This could happen in 32-bit platforms
+            check( reserve<std::numeric_limits<size_t>::max() );
+            m_pD->m_buffer.Reserve( reserve );
         }
     }
 
@@ -132,6 +139,9 @@ namespace mu
     {
         if (size)
         {
+            // This could happen in 32-bit platforms
+            check( size<std::numeric_limits<size_t>::max() );
+
             uint64 pos = m_pD->m_buffer.Num();
             m_pD->m_buffer.SetNum( pos + size, EAllowShrinking::No );
 			FMemory::Memcpy( &m_pD->m_buffer[pos], pData, size );
@@ -295,7 +305,7 @@ namespace mu
             }
             else
             {
-                if ( id < m_pD->m_proxyHistory.Num() )
+                if ( (std::size_t)id < m_pD->m_proxyHistory.Num() )
                 {
                     p = m_pD->m_proxyHistory[id];
 
diff --git a/Engine/Plugins/Experimental/Mutable/Source/MutableTools/Private/MuT/AST.h b/Engine/Plugins/Experimental/Mutable/Source/MutableTools/Private/MuT/AST.h
--- a/Engine/Plugins/Experimental/Mutable/Source/MutableTools/Private/MuT/AST.h
+++ b/Engine/Plugins/Experimental/Mutable/Source/MutableTools/Private/MuT/AST.h
@@ -26,8 +26,10 @@
 
 #include <array>
 #include <atomic>
+#include <memory>
 #include <unordered_map>
 #include <unordered_set>
+#include <utility>
 #include <set>
 
 
diff --git a/Engine/Plugins/Experimental/Mutable/Source/MutableTools/Private/MuT/ASTOpImageMultiLayer.cpp b/Engine/Plugins/Experimental/Mutable/Source/MutableTools/Private/MuT/ASTOpImageMultiLayer.cpp
--- a/Engine/Plugins/Experimental/Mutable/Source/MutableTools/Private/MuT/ASTOpImageMultiLayer.cpp
+++ b/Engine/Plugins/Experimental/Mutable/Source/MutableTools/Private/MuT/ASTOpImageMultiLayer.cpp
@@ -9,6 +9,9 @@
 #include "MuR/Types.h"
 #include "MuT/StreamsPrivate.h"
 
+#include <memory>
+#include <utility>
+
 
 namespace mu
 {
diff --git a/Engine/Plugins/Experimental/Mutable/Source/MutableTools/Private/MuT/ASTOpImageSwizzle.cpp b/Engine/Plugins/Experimental/Mutable/Source/MutableTools/Private/MuT/ASTOpImageSwizzle.cpp
--- a/Engine/Plugins/Experimental/Mutable/Source/MutableTools/Private/MuT/ASTOpImageSwizzle.cpp
+++ b/Engine/Plugins/Experimental/Mutable/Source/MutableTools/Private/MuT/ASTOpImageSwizzle.cpp
@@ -23,6 +23,9 @@
 #include "MuT/ASTOpSwitch.h"
 #include "MuT/StreamsPrivate.h"
 
+#include <memory>
+#include <utility>
+
 namespace mu
 {
 
diff --git a/Engine/Plugins/Experimental/Mutable/Source/MutableTools/Private/MuT/ASTOpImageTransform.cpp b/Engine/Plugins/Experimental/Mutable/Source/MutableTools/Private/MuT/ASTOpImageTransform.cpp
--- a/Engine/Plugins/Experimental/Mutable/Source/MutableTools/Private/MuT/ASTOpImageTransform.cpp
+++ b/Engine/Plugins/Experimental/Mutable/Source/MutableTools/Private/MuT/ASTOpImageTransform.cpp
@@ -9,6 +9,9 @@
 #include "MuR/Types.h"
 #include "MuT/StreamsPrivate.h"
 
+#include <memory>
+#include <utility>
+
 
 namespace mu
 {
diff --git a/Engine/Plugins/Experimental/Mutable/Source/MutableTools/Private/MuT/ASTOpInstanceAdd.cpp b/Engine/Plugins/Experimental/Mutable/Source/MutableTools/Private/MuT/ASTOpInstanceAdd.cpp
--- a/Engine/Plugins/Experimental/Mutable/Source/MutableTools/Private/MuT/ASTOpInstanceAdd.cpp
+++ b/Engine/Plugins/Experimental/Mutable/Source/MutableTools/Private/MuT/ASTOpInstanceAdd.cpp
@@ -10,6 +10,10 @@
 #include "MuT/CodeOptimiser.h"
 #include "MuT/StreamsPrivate.h"
 
+#include <algorithm>
+#include <memory>
+#include <utility>
+
 
 namespace mu
 {
diff --git a/Engine/Plugins/Experimental/Mutable/Source/MutableTools/Private/MuT/ASTOpLayoutFromMesh.cpp b/Engine/Plugins/Experimental/Mutable/Source/MutableTools/Private/MuT/ASTOpLayoutFromMesh.cpp
--- a/Engine/Plugins/Experimental/Mutable/Source/MutableTools/Private/MuT/ASTOpLayoutFromMesh.cpp
+++ b/Engine/Plugins/Experimental/Mutable/Source/MutableTools/Private/MuT/ASTOpLayoutFromMesh.cpp
@@ -24,6 +24,8 @@
 #include "MuT/ASTOpLayoutMerge.h"
 #include "MuT/StreamsPrivate.h"
 
+#include <functional>
+
 
 namespace mu
 {
diff --git a/Engine/Plugins/Experimental/Mutable/Source/MutableTools/Private/MuT/ASTOpMeshApplyPose.cpp b/Engine/Plugins/Experimental/Mutable/Source/MutableTools/Private/MuT/ASTOpMeshApplyPose.cpp
--- a/Engine/Plugins/Experimental/Mutable/Source/MutableTools/Private/MuT/ASTOpMeshApplyPose.cpp
+++ b/Engine/Plugins/Experimental/Mutable/Source/MutableTools/Private/MuT/ASTOpMeshApplyPose.cpp
@@ -8,6 +8,9 @@
 #include "MuR/Types.h"
 #include "MuT/StreamsPrivate.h"
 
+#include <memory>
+#include <utility>
+
 
 namespace mu
 {
diff --git a/Engine/Plugins/Experimental/Mutable/Source/MutableTools/Private/MuT/ASTOpMeshApplyShape.cpp b/Engine/Plugins/Experimental/Mutable/Source/MutableTools/Private/MuT/ASTOpMeshApplyShape.cpp
--- a/Engine/Plugins/Experimental/Mutable/Source/MutableTools/Private/MuT/ASTOpMeshApplyShape.cpp
+++ b/Engine/Plugins/Experimental/Mutable/Source/MutableTools/Private/MuT/ASTOpMeshApplyShape.cpp
@@ -8,6 +8,9 @@
 #include "MuR/Types.h"
 #include "MuT/StreamsPrivate.h"
 
+#include <memory>
+#include <utility>
+
 
 namespace mu
 {
diff --git a/Engine/Plugins/Experimental/Mutable/Source/MutableTools/Private/MuT/ASTOpMeshBindShape.cpp b/Engine/Plugins/Experimental/Mutable/Source/MutableTools/Private/MuT/ASTOpMeshBindShape.cpp
--- a/Engine/Plugins/Experimental/Mutable/Source/MutableTools/Private/MuT/ASTOpMeshBindShape.cpp
+++ b/Engine/Plugins/Experimental/Mutable/Source/MutableTools/Private/MuT/ASTOpMeshBindShape.cpp
@@ -13,6 +13,9 @@
 #include "MuT/ASTOpMeshAddTags.h"
 #include "MuT/StreamsPrivate.h"
 
+#include <memory>
+#include <utility>
+
 
 namespace mu
 {
diff --git a/Engine/Plugins/Experimental/Mutable/Source/MutableTools/Private/MuT/ASTOpMeshClipDeform.cpp b/Engine/Plugins/Experimental/Mutable/Source/MutableTools/Private/MuT/ASTOpMeshClipDeform.cpp
--- a/Engine/Plugins/Experimental/Mutable/Source/MutableTools/Private/MuT/ASTOpMeshClipDeform.cpp
+++ b/Engine/Plugins/Experimental/Mutable/Source/MutableTools/Private/MuT/ASTOpMeshClipDeform.cpp
@@ -8,6 +8,9 @@
 #include "MuR/Types.h"
 #include "MuT/StreamsPrivate.h"
 
+#include <memory>
+#include <utility>
+
 
 namespace mu
 {
diff --git a/Engine/Plugins/Experimental/Mutable/Source/MutableTools/Private/MuT/ASTOpMeshDifference.cpp b/Engine/Plugins/Experimental/Mutable/Source/MutableTools/Private/MuT/ASTOpMeshDifference.cpp
--- a/Engine/Plugins/Experimental/Mutable/Source/MutableTools/Private/MuT/ASTOpMeshDifference.cpp
+++ b/Engine/Plugins/Experimental/Mutable/Source/MutableTools/Private/MuT/ASTOpMeshDifference.cpp
@@ -12,6 +12,8 @@
 #include "MuT/ASTOpMeshRemoveMask.h"
 #include "MuT/ASTOpMeshAddTags.h"
 
+#include <memory>
+
 
 namespace mu
 {
diff --git a/Engine/Plugins/Experimental/Mutable/Source/MutableTools/Private/MuT/ASTOpMeshFormat.cpp b/Engine/Plugins/Experimental/Mutable/Source/MutableTools/Private/MuT/ASTOpMeshFormat.cpp
--- a/Engine/Plugins/Experimental/Mutable/Source/MutableTools/Private/MuT/ASTOpMeshFormat.cpp
+++ b/Engine/Plugins/Experimental/Mutable/Source/MutableTools/Private/MuT/ASTOpMeshFormat.cpp
@@ -20,6 +20,10 @@
 #include "MuT/ASTOpSwitch.h"
 #include "MuT/StreamsPrivate.h"
 
+#include <memory>
+#include <utility>
+
+
 
 namespace mu
 {
diff --git a/Engine/Plugins/Experimental/Mutable/Source/MutableTools/Private/MuT/ASTOpMeshGeometryOperation.cpp b/Engine/Plugins/Experimental/Mutable/Source/MutableTools/Private/MuT/ASTOpMeshGeometryOperation.cpp
--- a/Engine/Plugins/Experimental/Mutable/Source/MutableTools/Private/MuT/ASTOpMeshGeometryOperation.cpp
+++ b/Engine/Plugins/Experimental/Mutable/Source/MutableTools/Private/MuT/ASTOpMeshGeometryOperation.cpp
@@ -7,6 +7,9 @@
 #include "MuR/RefCounted.h"
 #include "MuT/StreamsPrivate.h"
 
+#include <memory>
+#include <utility>
+
 namespace mu
 {
 
diff --git a/Engine/Plugins/Experimental/Mutable/Source/MutableTools/Private/MuT/ASTOpMeshMorphReshape.cpp b/Engine/Plugins/Experimental/Mutable/Source/MutableTools/Private/MuT/ASTOpMeshMorphReshape.cpp
--- a/Engine/Plugins/Experimental/Mutable/Source/MutableTools/Private/MuT/ASTOpMeshMorphReshape.cpp
+++ b/Engine/Plugins/Experimental/Mutable/Source/MutableTools/Private/MuT/ASTOpMeshMorphReshape.cpp
@@ -8,6 +8,10 @@
 #include "MuR/Types.h"
 #include "MuT/StreamsPrivate.h"
 
+#include <memory>
+#include <utility>
+
+
 namespace mu
 {
 
diff --git a/Engine/Plugins/Experimental/Mutable/Source/MutableTools/Private/MuT/ASTOpMeshOptimizeSkinning.cpp b/Engine/Plugins/Experimental/Mutable/Source/MutableTools/Private/MuT/ASTOpMeshOptimizeSkinning.cpp
--- a/Engine/Plugins/Experimental/Mutable/Source/MutableTools/Private/MuT/ASTOpMeshOptimizeSkinning.cpp
+++ b/Engine/Plugins/Experimental/Mutable/Source/MutableTools/Private/MuT/ASTOpMeshOptimizeSkinning.cpp
@@ -8,6 +8,8 @@
 #include "MuR/Types.h"
 #include "MuT/StreamsPrivate.h"
 
+#include <utility>
+
 
 namespace mu
 {
diff --git a/Engine/Plugins/Experimental/Mutable/Source/MutableTools/Private/MuT/ASTOpMeshTransform.cpp b/Engine/Plugins/Experimental/Mutable/Source/MutableTools/Private/MuT/ASTOpMeshTransform.cpp
--- a/Engine/Plugins/Experimental/Mutable/Source/MutableTools/Private/MuT/ASTOpMeshTransform.cpp
+++ b/Engine/Plugins/Experimental/Mutable/Source/MutableTools/Private/MuT/ASTOpMeshTransform.cpp
@@ -8,6 +8,9 @@
 #include "MuR/Types.h"
 #include "MuT/StreamsPrivate.h"
 
+#include <memory>
+#include <utility>
+
 
 namespace mu
 {
diff --git a/Engine/Plugins/Experimental/Mutable/Source/MutableTools/Private/MuT/ASTOpParameter.cpp b/Engine/Plugins/Experimental/Mutable/Source/MutableTools/Private/MuT/ASTOpParameter.cpp
--- a/Engine/Plugins/Experimental/Mutable/Source/MutableTools/Private/MuT/ASTOpParameter.cpp
+++ b/Engine/Plugins/Experimental/Mutable/Source/MutableTools/Private/MuT/ASTOpParameter.cpp
@@ -10,6 +10,9 @@
 #include "MuR/Types.h"
 #include "MuT/StreamsPrivate.h"
 
+#include <memory>
+#include <utility>
+
 
 namespace mu
 {
diff --git a/Engine/Plugins/Experimental/Mutable/Source/MutableTools/Private/MuT/ASTOpScalarCurve.cpp b/Engine/Plugins/Experimental/Mutable/Source/MutableTools/Private/MuT/ASTOpScalarCurve.cpp
--- a/Engine/Plugins/Experimental/Mutable/Source/MutableTools/Private/MuT/ASTOpScalarCurve.cpp
+++ b/Engine/Plugins/Experimental/Mutable/Source/MutableTools/Private/MuT/ASTOpScalarCurve.cpp
@@ -8,6 +8,9 @@
 #include "MuR/Types.h"
 #include "MuT/StreamsPrivate.h"
 
+#include <memory>
+#include <utility>
+
 
 namespace mu
 {
diff --git a/Engine/Plugins/Experimental/Mutable/Source/MutableTools/Private/MuT/ASTOpSwitch.cpp b/Engine/Plugins/Experimental/Mutable/Source/MutableTools/Private/MuT/ASTOpSwitch.cpp
--- a/Engine/Plugins/Experimental/Mutable/Source/MutableTools/Private/MuT/ASTOpSwitch.cpp
+++ b/Engine/Plugins/Experimental/Mutable/Source/MutableTools/Private/MuT/ASTOpSwitch.cpp
@@ -12,6 +12,9 @@
 #include "MuR/Types.h"
 #include "MuT/StreamsPrivate.h"
 
+#include <memory>
+#include <utility>
+
 
 namespace mu
 {
diff --git a/Engine/Plugins/Experimental/Mutable/Source/MutableTools/Private/MuT/CodeOptimiser_Logic.cpp b/Engine/Plugins/Experimental/Mutable/Source/MutableTools/Private/MuT/CodeOptimiser_Logic.cpp
--- a/Engine/Plugins/Experimental/Mutable/Source/MutableTools/Private/MuT/CodeOptimiser_Logic.cpp
+++ b/Engine/Plugins/Experimental/Mutable/Source/MutableTools/Private/MuT/CodeOptimiser_Logic.cpp
@@ -11,6 +11,10 @@
 #include "MuT/ASTOpSwitch.h"
 #include "MuT/CodeOptimiser.h"
 
+#include <memory>
+#include <utility>
+
+
 
 namespace mu
 {
diff --git a/Engine/Plugins/Experimental/Mutable/Source/MutableTools/Private/MuT/CodeOptimiser_Parameters.cpp b/Engine/Plugins/Experimental/Mutable/Source/MutableTools/Private/MuT/CodeOptimiser_Parameters.cpp
--- a/Engine/Plugins/Experimental/Mutable/Source/MutableTools/Private/MuT/CodeOptimiser_Parameters.cpp
+++ b/Engine/Plugins/Experimental/Mutable/Source/MutableTools/Private/MuT/CodeOptimiser_Parameters.cpp
@@ -1,5 +1,18 @@
 // Copyright Epic Games, Inc. All Rights Reserved.
 
+#include "HAL/PlatformCrt.h"
+#include "HAL/PlatformMath.h"
+#include "Logging/LogCategory.h"
+#include "Logging/LogMacros.h"
+#include "Misc/AssertionMacros.h"
+#include "MuR/Image.h"
+#include "MuR/ImagePrivate.h"
+#include "MuR/MutableTrace.h"
+#include "MuR/Operations.h"
+#include "MuR/ParametersPrivate.h"
+#include "MuR/Platform.h"
+#include "MuR/Ptr.h"
+#include "MuR/RefCounted.h"
 #include "MuT/AST.h"
 #include "MuT/ASTOpAddLOD.h"
 #include "MuT/ASTOpConditional.h"
@@ -17,24 +30,13 @@
 #include "MuT/CompilerPrivate.h"
 #include "MuT/DataPacker.h"
 #include "MuT/Platform.h"
-#include "MuR/Image.h"
-#include "MuR/ImagePrivate.h"
-#include "MuR/MutableTrace.h"
-#include "MuR/Operations.h"
-#include "MuR/ParametersPrivate.h"
-#include "MuR/Platform.h"
-#include "MuR/Ptr.h"
-#include "MuR/RefCounted.h"
-#include "MuR/MutableRuntimeModule.h"
-
-#include "HAL/PlatformCrt.h"
-#include "HAL/PlatformMath.h"
-#include "Logging/LogCategory.h"
-#include "Logging/LogMacros.h"
-#include "Misc/AssertionMacros.h"
 #include "Trace/Detail/Channel.h"
 
 #include <array>
+#include <memory>
+#include <utility>
+
+#include "MuR/MutableRuntimeModule.h"
 
 
 namespace mu
diff --git a/Engine/Plugins/Experimental/Mutable/Source/MutableTools/Private/MuT/Compiler.cpp b/Engine/Plugins/Experimental/Mutable/Source/MutableTools/Private/MuT/Compiler.cpp
--- a/Engine/Plugins/Experimental/Mutable/Source/MutableTools/Private/MuT/Compiler.cpp
+++ b/Engine/Plugins/Experimental/Mutable/Source/MutableTools/Private/MuT/Compiler.cpp
@@ -34,6 +34,10 @@
 #include "MuT/Table.h"
 #include "Trace/Detail/Channel.h"
 
+#include <algorithm>
+#include <memory>
+#include <utility>
+
 
 namespace mu
 {
diff --git a/Engine/Plugins/Experimental/Mutable/Source/MutableTools/Private/MuT/DataPacker.h b/Engine/Plugins/Experimental/Mutable/Source/MutableTools/Private/MuT/DataPacker.h
--- a/Engine/Plugins/Experimental/Mutable/Source/MutableTools/Private/MuT/DataPacker.h
+++ b/Engine/Plugins/Experimental/Mutable/Source/MutableTools/Private/MuT/DataPacker.h
@@ -5,10 +5,12 @@
 #include "MuR/Operations.h"
 #include "MuT/AST.h"
 
+#include <stdint.h>
+
 namespace mu { struct FProgram; }
 
 
-#define MUTABLE_HASH_SEED					((uint32)0xcadababa)
+#define MUTABLE_HASH_SEED					((uint32_t)0xcadababa)
 
 
 namespace mu
diff --git a/Engine/Plugins/Experimental/Mutable/Source/MutableTools/Private/MuT/ErrorLog.cpp b/Engine/Plugins/Experimental/Mutable/Source/MutableTools/Private/MuT/ErrorLog.cpp
--- a/Engine/Plugins/Experimental/Mutable/Source/MutableTools/Private/MuT/ErrorLog.cpp
+++ b/Engine/Plugins/Experimental/Mutable/Source/MutableTools/Private/MuT/ErrorLog.cpp
@@ -25,6 +25,9 @@
 #include "Templates/Tuple.h"
 #include "Trace/Detail/Channel.h"
 
+#include <memory>
+#include <utility>
+
 namespace mu
 {
 
diff --git a/Engine/Plugins/Experimental/Mutable/Source/MutableTools/Private/MuT/NodeBool.cpp b/Engine/Plugins/Experimental/Mutable/Source/MutableTools/Private/MuT/NodeBool.cpp
--- a/Engine/Plugins/Experimental/Mutable/Source/MutableTools/Private/MuT/NodeBool.cpp
+++ b/Engine/Plugins/Experimental/Mutable/Source/MutableTools/Private/MuT/NodeBool.cpp
@@ -7,6 +7,8 @@
 #include "MuR/Serialisation.h"
 #include "MuR/SerialisationPrivate.h"
 
+#include <stdint.h>
+
 
 namespace mu
 {
diff --git a/Engine/Plugins/Experimental/Mutable/Source/MutableTools/Private/MuT/NodeBoolParameter.cpp b/Engine/Plugins/Experimental/Mutable/Source/MutableTools/Private/MuT/NodeBoolParameter.cpp
--- a/Engine/Plugins/Experimental/Mutable/Source/MutableTools/Private/MuT/NodeBoolParameter.cpp
+++ b/Engine/Plugins/Experimental/Mutable/Source/MutableTools/Private/MuT/NodeBoolParameter.cpp
@@ -10,6 +10,9 @@
 #include "MuT/NodePrivate.h"
 #include "MuT/NodeRange.h"
 
+#include <memory>
+#include <utility>
+
 
 namespace mu
 {
diff --git a/Engine/Plugins/Experimental/Mutable/Source/MutableTools/Private/MuT/NodeColour.cpp b/Engine/Plugins/Experimental/Mutable/Source/MutableTools/Private/MuT/NodeColour.cpp
--- a/Engine/Plugins/Experimental/Mutable/Source/MutableTools/Private/MuT/NodeColour.cpp
+++ b/Engine/Plugins/Experimental/Mutable/Source/MutableTools/Private/MuT/NodeColour.cpp
@@ -15,6 +15,8 @@
 #include "MuT/NodeColourTable.h"
 #include "MuT/NodeColourVariation.h"
 
+#include <stdint.h>
+
 
 namespace mu
 {
diff --git a/Engine/Plugins/Experimental/Mutable/Source/MutableTools/Private/MuT/NodeComponent.cpp b/Engine/Plugins/Experimental/Mutable/Source/MutableTools/Private/MuT/NodeComponent.cpp
--- a/Engine/Plugins/Experimental/Mutable/Source/MutableTools/Private/MuT/NodeComponent.cpp
+++ b/Engine/Plugins/Experimental/Mutable/Source/MutableTools/Private/MuT/NodeComponent.cpp
@@ -9,6 +9,8 @@
 #include "MuT/NodeComponentEdit.h"
 #include "MuT/NodeComponentNew.h"
 
+#include <stdint.h>
+
 
 namespace mu
 {
diff --git a/Engine/Plugins/Experimental/Mutable/Source/MutableTools/Private/MuT/NodeRange.cpp b/Engine/Plugins/Experimental/Mutable/Source/MutableTools/Private/MuT/NodeRange.cpp
--- a/Engine/Plugins/Experimental/Mutable/Source/MutableTools/Private/MuT/NodeRange.cpp
+++ b/Engine/Plugins/Experimental/Mutable/Source/MutableTools/Private/MuT/NodeRange.cpp
@@ -9,6 +9,8 @@
 #include "MuT/Node.h"
 #include "MuT/NodeRangeFromScalar.h"
 
+#include <stdint.h>
+
 
 namespace mu
 {
diff --git a/Engine/Plugins/Experimental/Mutable/Source/MutableTools/Private/MuT/NodeScalar.cpp b/Engine/Plugins/Experimental/Mutable/Source/MutableTools/Private/MuT/NodeScalar.cpp
--- a/Engine/Plugins/Experimental/Mutable/Source/MutableTools/Private/MuT/NodeScalar.cpp
+++ b/Engine/Plugins/Experimental/Mutable/Source/MutableTools/Private/MuT/NodeScalar.cpp
@@ -16,6 +16,8 @@
 #include "MuT/NodeScalarTable.h"
 #include "MuT/NodeScalarVariation.h"
 
+#include <stdint.h>
+
 
 namespace mu
 {
diff --git a/Engine/Plugins/Experimental/Mutable/Source/MutableTools/Private/MuT/NodeScalarSwitch.cpp b/Engine/Plugins/Experimental/Mutable/Source/MutableTools/Private/MuT/NodeScalarSwitch.cpp
--- a/Engine/Plugins/Experimental/Mutable/Source/MutableTools/Private/MuT/NodeScalarSwitch.cpp
+++ b/Engine/Plugins/Experimental/Mutable/Source/MutableTools/Private/MuT/NodeScalarSwitch.cpp
@@ -8,6 +8,9 @@
 #include "MuT/NodeScalar.h"
 #include "MuT/NodeScalarSwitchPrivate.h"
 
+#include <memory>
+#include <utility>
+
 
 namespace mu
 {
diff --git a/Engine/Plugins/Experimental/Mutable/Source/MutableTools/Private/MuT/NodeString.cpp b/Engine/Plugins/Experimental/Mutable/Source/MutableTools/Private/MuT/NodeString.cpp
--- a/Engine/Plugins/Experimental/Mutable/Source/MutableTools/Private/MuT/NodeString.cpp
+++ b/Engine/Plugins/Experimental/Mutable/Source/MutableTools/Private/MuT/NodeString.cpp
@@ -10,6 +10,8 @@
 #include "MuT/NodeStringConstant.h"
 #include "MuT/NodeStringParameter.h"
 
+#include <stdint.h>
+
 
 namespace mu
 {
diff --git a/Engine/Plugins/Experimental/Mutable/Source/MutableTools/Private/MuT/Platform.h b/Engine/Plugins/Experimental/Mutable/Source/MutableTools/Private/MuT/Platform.h
--- a/Engine/Plugins/Experimental/Mutable/Source/MutableTools/Private/MuT/Platform.h
+++ b/Engine/Plugins/Experimental/Mutable/Source/MutableTools/Private/MuT/Platform.h
@@ -7,6 +7,11 @@
 #include <set>
 #include <string>
 #include <unordered_map>
+#include <sys/stat.h>
+
+
+// This file contains some hacks to solve differences between platforms necessary for the tools
+// library and not present in the run-time library.
 
 
 namespace mu
diff --git a/Engine/Plugins/Experimental/Mutable/Source/MutableTools/Private/MuT/Table.cpp b/Engine/Plugins/Experimental/Mutable/Source/MutableTools/Private/MuT/Table.cpp
--- a/Engine/Plugins/Experimental/Mutable/Source/MutableTools/Private/MuT/Table.cpp
+++ b/Engine/Plugins/Experimental/Mutable/Source/MutableTools/Private/MuT/Table.cpp
@@ -12,6 +12,10 @@
 #include "MuR/SerialisationPrivate.h"
 #include "MuT/TablePrivate.h"
 
+#include <memory>
+#include <string>
+#include <utility>
+
 
 namespace mu
 {
