[Backout] - CL31041669
Original CL Desc
-----------------------------------------------------------------
[Mutable] UCustomizableObject Moved non-serialized members to private.

- Moved all non-serialized members to private.
- Moved non-public classes and structs to private.
- Moved the following functions to private:
    - IsLocked
    - GetModel

#rnx

diff --git a/Engine/Plugins/Experimental/Mutable/Source/CustomizableObject/Classes/MuCO/CustomizableObject.h b/Engine/Plugins/Experimental/Mutable/Source/CustomizableObject/Classes/MuCO/CustomizableObject.h
--- a/Engine/Plugins/Experimental/Mutable/Source/CustomizableObject/Classes/MuCO/CustomizableObject.h
+++ b/Engine/Plugins/Experimental/Mutable/Source/CustomizableObject/Classes/MuCO/CustomizableObject.h
@@ -7,6 +7,7 @@
 #include "MuCO/CustomizableObjectResourceData.h"
 #include "MuCO/CustomizableObjectResourceDataTypes.h"
 #include "MuCO/CustomizableObjectStreamedResourceData.h"
+#include "MuCO/CustomizableObjectIdentifier.h"
 #include "MuCO/CustomizableObjectParameterTypeDefinitions.h"
 #include "MuCO/CustomizableObjectUIData.h"
 
@@ -23,11 +24,9 @@ class ITargetPlatform;
 class UCustomizableObject;
 class UEdGraph;
 class USkeletalMesh;
-class UCustomizableObjectPrivate;
 struct FFrame;
 struct FStreamableHandle;
 template <typename FuncType> class TFunctionRef;
-enum class ECustomizableObjectCompilationState : uint8;
 
 DECLARE_MULTICAST_DELEGATE(FPostCompileDelegate)
  
@@ -56,6 +55,8 @@ struct FFParameterOptionsTags
 USTRUCT()
 struct FParameterTags
 {
+public:
+
 	GENERATED_USTRUCT_BODY()
 
 	/** List of tags of a parameter */
@@ -78,6 +79,9 @@ enum class ECustomizableObjectRelevancy : uint8
 };
 
 
+// This is used to hide Mutable SDK members in the public headers.
+class UCustomizableObjectPrivate;
+
 USTRUCT()
 struct FProfileParameterDat
 {
@@ -196,6 +200,97 @@ struct FCompilationOptions
 };
 
 
+UENUM()
+enum class ECustomizableObjectCompilationState : uint8
+{
+	//
+	None,
+	// 
+	InProgress,
+	//
+	Completed,
+	//
+	Failed
+};
+
+
+class FCustomizableObjectCompilerBase
+{
+public:
+
+	FCustomizableObjectCompilerBase() {};
+
+	// Ensure virtual destruction
+	virtual ~FCustomizableObjectCompilerBase() {};
+	
+	virtual void Compile(UCustomizableObject& Object, const FCompilationOptions& Options, bool bAsync) {};
+
+	virtual bool Tick() { return false; }
+	virtual void ForceFinishCompilation() {};
+
+	// Return true if this object doesn't reference a parent object.
+	virtual bool IsRootObject(const class UCustomizableObject* Object) const { return true; }
+
+	/** Returns the Customizable Object that does start the CO tree */
+	virtual UCustomizableObject* GetRootObject( class UCustomizableObject* Object) = 0;
+
+	/** Provides the caller with the warning and error messages produced during compilation */
+	virtual void GetCompilationMessages(TArray<FText>& OutWarningMessages, TArray<FText>& OutErrorMessages) const = 0;
+	
+	virtual ECustomizableObjectCompilationState GetCompilationState() const { return ECustomizableObjectCompilationState::None; }
+};
+
+
+// Warning! MutableCompiledDataHeader must be the first data serialized in a stream
+struct MutableCompiledDataStreamHeader
+{
+	int32 InternalVersion;
+	FGuid VersionId;
+
+	MutableCompiledDataStreamHeader() { }
+	MutableCompiledDataStreamHeader(int32 InInternalVersion, FGuid InVersionId) : InternalVersion(InInternalVersion), VersionId(InVersionId) { }
+
+	friend FArchive& operator<<(FArchive& Ar, MutableCompiledDataStreamHeader& Header)
+	{
+		Ar << Header.InternalVersion;
+		Ar << Header.VersionId;
+
+		return Ar;
+	}
+};
+
+
+USTRUCT()
+struct FMutableModelParameterValue
+{
+	GENERATED_USTRUCT_BODY()
+
+	FMutableModelParameterValue() = default;
+
+	UPROPERTY()
+	FString Name;
+
+	UPROPERTY()
+	int Value = 0;
+};
+
+
+USTRUCT()
+struct FMutableModelParameterProperties
+{
+	GENERATED_USTRUCT_BODY()
+
+	FMutableModelParameterProperties() = default;
+	FString Name;
+
+	UPROPERTY()
+	EMutableParameterType Type = EMutableParameterType::None;
+
+	UPROPERTY()
+	TArray<FMutableModelParameterValue> PossibleValues;
+};
+
+
 USTRUCT()
 struct FMorphTargetInfo
 {
@@ -382,6 +477,17 @@ struct CUSTOMIZABLEOBJECT_API FMutableStreamableBlock
 template<> struct TCanBulkSerialize<FMutableStreamableBlock> { enum { Value = true }; };
 
 
+struct CUSTOMIZABLEOBJECT_API FMutableCachedPlatformData
+{
+	/** */
+	TArray64<uint8> ModelData;
+
+	/** */
+	TArray64<uint8> StreamableData;
+};
+
+
+
 USTRUCT()
 struct FMutableLODSettings
 {
@@ -519,7 +625,10 @@ public:
 
 	UPROPERTY(EditAnywhere, Category = CustomizableObject, meta = (DisplayName = "LOD Settings"))
 	FMutableLODSettings LODSettings;
-	
+
+	/** This is a non-user-controlled flag to disable streaming (set at object compilation time, depending on optimization). */
+	bool bDisableTextureStreaming = false;
+
 	UPROPERTY(Transient)
 	TArray<FMorphTargetInfo> ContributingMorphTargetsInfo;
 	
@@ -616,6 +725,14 @@ public:
 	UPROPERTY()
 	FGuid VersionId;
 
+	/* Map to identify what CustomizableObject owns a parameter. Used to display a tooltip when hovering a parameter
+	   in the Prev. instance panel */
+	UPROPERTY(Transient)
+	TMap<FString, FString> CustomizableObjectPathMap;
+
+	UPROPERTY(Transient)
+	TMap<FString, FCustomizableObjectIdPair> GroupNodeMap; 
+
 	// 
 	UPROPERTY()
 	TArray<FProfileParameterDat>  InstancePropertiesProfiles;
@@ -756,9 +873,23 @@ public:
 	bool IsParameterMultidimensional ( UPARAM(DisplayName = "Parameter Name") const FString& InParameterName) const;
 	
 private: 
+
+	// This is information about the parameters in the model that is generated at model compile time.
+	UPROPERTY(Transient)
+	TArray<FMutableModelParameterProperties> ParameterProperties;
+
+	// Map of name to index of ParameterProperties.
+	// use this to lookup fast by Name
+	TMap<FString, int32> ParameterPropertiesLookupTable;
+
 	// Rebuild ParameterProperties from the current compiled model.
 	void UpdateParameterPropertiesFromModel(const TSharedPtr<mu::Model>& Model);
 
+	// This is a manual version number for the binary blobs in this asset.
+	// Increasing it invalidates all the previously compiled models.
+	// Warning: If while merging code both versions have changed, take the highest+1.
+	static const int32 CurrentSupportedVersion = 421;
+
 public:
 
 	void AddUncompiledCOWarning(const FString& AdditionalLoggingInfo);
@@ -783,6 +914,8 @@ public:
 	void UpdateVersionId();
 	FGuid GetVersionId() const { return VersionId; }
 
+	int32 GetCurrentSupportedVersion() const { return CurrentSupportedVersion; };
+
 	// Compose folder name where the data is stored
 	FString GetCompiledDataFolderPath() const;
 
@@ -818,13 +951,19 @@ public:
 	/** Cache platform data for cook */
 	void CachePlatformData(const ITargetPlatform* InTargetPlatform, const TArray64<uint8>& InObjectBytes, const TArray64<uint8>& InBulkBytes);
 
+#if WITH_EDITORONLY_DATA
+
+	/** Map of PlatformName to CachedPlatformData. Only valid while cooking. */
+	TMap<FString, FMutableCachedPlatformData> CachedPlatformsData;
+	
+#endif
 	/**
 	 * Loads data previously compiled in BeginCacheForCookedPlatformData onto the UProperties in *this,
 	 * in preparation for saving the cooked package for *this or for a CustomizableObjectInstance using *this.
 	 * Returns whether the data was successfully loaded.
 	 */
 	bool TryLoadCompiledCookDataForPlatform(const ITargetPlatform* TargetPlatform);
-
+	
 	void SaveEmbeddedData(FArchive& Ar);
 
 #endif
@@ -912,12 +1051,17 @@ public:
 	UPROPERTY()
 	bool bIsChildObject = false;
 
+	/** If the object is compiled, this flag is true if it was compiled with maximum optimizations. If the object is not compiled, its value is meaningless. */
+	bool bIsCompiledWithOptimization = true;
+
+	ECustomizableObjectCompilationState CompilationState = ECustomizableObjectCompilationState::None;
+
+	FPostCompileDelegate PostCompileDelegate;
+
 	void PostCompile();
 
 #endif
 
-	FPostCompileDelegate& GetPostCompileDelegate() const;
-	
 	/** Create a new instance of this object. The instance parameters will be initialized with the object default values. */
 	UFUNCTION(BlueprintCallable, Category = CustomizableObject)
 	class UCustomizableObjectInstance* CreateInstance();
@@ -931,6 +1075,11 @@ public:
 	UFUNCTION(BlueprintCallable, Category = CustomizableObject)
 	bool IsCompiled() const;
 
+	// See UCustomizableObjectSystem::LockObject()
+	bool IsLocked() const;
+
+	TSharedPtr<mu::Model, ESPMode::ThreadSafe> GetModel() const;
+
 #if WITH_EDITOR
 	void SetModel(TSharedPtr<mu::Model, ESPMode::ThreadSafe> Model);
 #endif
diff --git a/Engine/Plugins/Experimental/Mutable/Source/CustomizableObject/Internal/MuCO/CustomizableObjectPrivate.h b/Engine/Plugins/Experimental/Mutable/Source/CustomizableObject/Internal/MuCO/CustomizableObjectPrivate.h
--- a/Engine/Plugins/Experimental/Mutable/Source/CustomizableObject/Internal/MuCO/CustomizableObjectPrivate.h
+++ b/Engine/Plugins/Experimental/Mutable/Source/CustomizableObject/Internal/MuCO/CustomizableObjectPrivate.h
@@ -2,10 +2,8 @@
 
 #pragma once
 
-#include "MuCO/CustomizableObject.h"
 #include "MuCO/StateMachine.h"
 #include "MuCO/CustomizableObjectUIData.h"
-#include "MuCO/CustomizableObjectIdentifier.h"
 #include "Templates/SharedPointer.h"
 #include "UObject/WeakObjectPtr.h"
 #include "UObject/SoftObjectPtr.h"
@@ -18,6 +16,7 @@
 #include "CustomizableObjectPrivate.generated.h"
 
 namespace mu { class Model; }
+class UCustomizableObject;
 class USkeletalMesh;
 class USkeleton;
 class UPhysicsAsset;
@@ -25,85 +24,6 @@ class UMaterialInterface;
 class UTexture;
 class UAnimInstance;
 class UAssetUserData;
-class UCustomizableObject;
-
-
-class CUSTOMIZABLEOBJECT_API FCustomizableObjectCompilerBase
-{
-public:
-
-	FCustomizableObjectCompilerBase() {};
-
-	// Ensure virtual destruction
-	virtual ~FCustomizableObjectCompilerBase() {};
-	
-	virtual void Compile(UCustomizableObject& Object, const FCompilationOptions& Options, bool bAsync) {};
-
-	virtual bool Tick() { return false; }
-	virtual void ForceFinishCompilation() {};
-
-	// Return true if this object doesn't reference a parent object.
-	virtual bool IsRootObject(const class UCustomizableObject* Object) const { return true; }
-
-	/** Returns the Customizable Object that does start the CO tree */
-	virtual UCustomizableObject* GetRootObject( class UCustomizableObject* Object) = 0;
-
-	/** Provides the caller with the warning and error messages produced during compilation */
-	virtual void GetCompilationMessages(TArray<FText>& OutWarningMessages, TArray<FText>& OutErrorMessages) const = 0;
-	
-	virtual ECustomizableObjectCompilationState GetCompilationState() const;
-};
-
-
-// Warning! MutableCompiledDataHeader must be the first data serialized in a stream
-struct MutableCompiledDataStreamHeader
-{
-	int32 InternalVersion;
-	FGuid VersionId;
-
-	MutableCompiledDataStreamHeader() { }
-	MutableCompiledDataStreamHeader(int32 InInternalVersion, FGuid InVersionId) : InternalVersion(InInternalVersion), VersionId(InVersionId) { }
-
-	friend FArchive& operator<<(FArchive& Ar, MutableCompiledDataStreamHeader& Header)
-	{
-		Ar << Header.InternalVersion;
-		Ar << Header.VersionId;
-
-		return Ar;
-	}
-};
-
-
-USTRUCT()
-struct FMutableModelParameterValue
-{
-	GENERATED_USTRUCT_BODY()
-
-	FMutableModelParameterValue() = default;
-
-	UPROPERTY()
-	FString Name;
-
-	UPROPERTY()
-	int Value = 0;
-};
-
-
-USTRUCT()
-struct FMutableModelParameterProperties
-{
-	GENERATED_USTRUCT_BODY()
-
-	FMutableModelParameterProperties() = default;
-	FString Name;
-
-	UPROPERTY()
-	EMutableParameterType Type = EMutableParameterType::None;
-
-	UPROPERTY()
-	TArray<FMutableModelParameterValue> PossibleValues;
-};
-
 
 class FMeshCache
 {
@@ -156,20 +76,6 @@ struct FCustomizableObjectStatusTypes
 using FCustomizableObjectStatus = FStateMachine<FCustomizableObjectStatusTypes>;
 
 
-UENUM()
-enum class ECustomizableObjectCompilationState : uint8
-{
-	//
-	None,
-	// 
-	InProgress,
-	//
-	Completed,
-	//
-	Failed
-};
-
-
 USTRUCT()
 struct CUSTOMIZABLEOBJECT_API FMutableModelImageProperties
 {
@@ -535,21 +441,13 @@ struct FModelResources
 };
 
 
-struct CUSTOMIZABLEOBJECT_API FMutableCachedPlatformData
-{
-	/** */
-	TArray64<uint8> ModelData;
-
-	/** */
-	TArray64<uint8> StreamableData;
-};
-
-
 UCLASS()
-class CUSTOMIZABLEOBJECT_API UCustomizableObjectPrivate : public UObject
+class UCustomizableObjectPrivate : public UObject
 {
 	GENERATED_BODY()
 
+private:
+
 	TSharedPtr<mu::Model, ESPMode::ThreadSafe> MutableModel;
 
 	/** Stores resources to be used by MutableModel In-Game. Cooked resources. */
@@ -571,12 +469,9 @@ public:
 	const FModelResources& GetModelResources() const;
 
 #if WITH_EDITORONLY_DATA
-	FModelResources& GetModelResources(bool bIsCooking);
+	CUSTOMIZABLEOBJECT_API FModelResources& GetModelResources(bool bIsCooking);
 #endif
 
-	// See UCustomizableObjectSystem::LockObject()
-	bool IsLocked() const;
-
 	/** Cache of generated SkeletalMeshes */
 	FMeshCache MeshCache;
 
@@ -599,44 +494,8 @@ public:
 	 *
 	 * Updated each time the CO is compiled and saved in the Derived Data. */
 	TMap<FName, FGuid> ParticipatingObjects;
-
-	/** If the object is compiled, this flag is true if it was compiled with maximum optimizations. If the object is not compiled, its value is meaningless. */
-	bool bIsCompiledWithOptimization = true;
-
-	ECustomizableObjectCompilationState CompilationState = ECustomizableObjectCompilationState::None;
-	
-#if WITH_EDITOR
-	/** Map of PlatformName to CachedPlatformData. Only valid while cooking. */
-	TMap<FString, FMutableCachedPlatformData> CachedPlatformsData;
-#endif
 #endif
 
 	FCustomizableObjectStatus Status;
-
-	/** Map to identify what CustomizableObject owns a parameter. Used to display a tooltip when hovering a parameter
-	 * in the Prev. instance panel */
-	UPROPERTY(Transient)
-	TMap<FString, FString> CustomizableObjectPathMap;
-
-	UPROPERTY(Transient)
-	TMap<FString, FCustomizableObjectIdPair> GroupNodeMap;
-
-	FPostCompileDelegate PostCompileDelegate;
-
-	// This is information about the parameters in the model that is generated at model compile time.
-	UPROPERTY(Transient)
-	TArray<FMutableModelParameterProperties> ParameterProperties;
-
-	// Map of name to index of ParameterProperties.
-	// use this to lookup fast by Name
-	TMap<FString, int32> ParameterPropertiesLookupTable;
-
-	// This is a manual version number for the binary blobs in this asset.
-	// Increasing it invalidates all the previously compiled models.
-	// Warning: If while merging code both versions have changed, take the highest+1.
-	static constexpr int32 CurrentSupportedVersion = 423;
-
-	/** This is a non-user-controlled flag to disable streaming (set at object compilation time, depending on optimization). */
-	bool bDisableTextureStreaming = false;
 };
 
diff --git a/Engine/Plugins/Experimental/Mutable/Source/CustomizableObject/Internal/MuCO/CustomizableObjectSystemPrivate.h b/Engine/Plugins/Experimental/Mutable/Source/CustomizableObject/Internal/MuCO/CustomizableObjectSystemPrivate.h
--- a/Engine/Plugins/Experimental/Mutable/Source/CustomizableObject/Internal/MuCO/CustomizableObjectSystemPrivate.h
+++ b/Engine/Plugins/Experimental/Mutable/Source/CustomizableObject/Internal/MuCO/CustomizableObjectSystemPrivate.h
@@ -6,7 +6,6 @@
 #include "Containers/Queue.h"
 #include "MuCO/DescriptorHash.h"
 #include "MuCO/CustomizableObject.h"
-#include "MuCO/CustomizableObjectPrivate.h"
 #include "MuCO/CustomizableObjectInstance.h"
 #include "MuCO/CustomizableObjectExtension.h"
 #include "Containers/Ticker.h"
diff --git a/Engine/Plugins/Experimental/Mutable/Source/CustomizableObject/Private/MuCO/CustomizableInstanceLODManagement.cpp b/Engine/Plugins/Experimental/Mutable/Source/CustomizableObject/Private/MuCO/CustomizableInstanceLODManagement.cpp
--- a/Engine/Plugins/Experimental/Mutable/Source/CustomizableObject/Private/MuCO/CustomizableInstanceLODManagement.cpp
+++ b/Engine/Plugins/Experimental/Mutable/Source/CustomizableObject/Private/MuCO/CustomizableInstanceLODManagement.cpp
@@ -15,7 +15,6 @@
 #include "MuCO/CustomizableObjectInstanceUsage.h"
 #include "UObject/UObjectIterator.h"
 #include "Components/SkeletalMeshComponent.h"
-#include "MuCO/CustomizableObjectPrivate.h"
 
 #include UE_INLINE_GENERATED_CPP_BY_NAME(CustomizableInstanceLODManagement)
 
@@ -207,7 +206,7 @@ void UCustomizableInstanceLODManagement::UpdateInstanceDistsAndLODs(FMutableInst
 			{
 				if (const UCustomizableObject* CustomizableObject = CustomizableObjectInstance->GetCustomizableObject();
 					CustomizableObject &&
-					!CustomizableObject->GetPrivate()->IsLocked())
+					!CustomizableObject->IsLocked())
 				{
 					UCustomizableObjectInstance* Ptr = *CustomizableObjectInstance;
 					Ptr->SetIsDiscardedBecauseOfTooManyInstances(false);
@@ -404,7 +403,7 @@ void UCustomizableInstanceLODManagement::UpdateInstanceDistsAndLODs(FMutableInst
 			if (IsValidChecked(It.Key) && It.Key->GetPrivate())
 			{
 				const UCustomizableObject* CustomizableObject = It.Key->GetCustomizableObject();
-				if (!CustomizableObject || CustomizableObject->GetPrivate()->IsLocked())
+				if (!CustomizableObject || CustomizableObject->IsLocked())
 				{
 					continue;
 				}
diff --git a/Engine/Plugins/Experimental/Mutable/Source/CustomizableObject/Private/MuCO/CustomizableObject.cpp b/Engine/Plugins/Experimental/Mutable/Source/CustomizableObject/Private/MuCO/CustomizableObject.cpp
--- a/Engine/Plugins/Experimental/Mutable/Source/CustomizableObject/Private/MuCO/CustomizableObject.cpp
+++ b/Engine/Plugins/Experimental/Mutable/Source/CustomizableObject/Private/MuCO/CustomizableObject.cpp
@@ -180,7 +180,7 @@ bool UCustomizableObject::TryUpdateIsChildObject()
 #if WITH_EDITOR
 bool UCustomizableObject::TryLoadCompiledCookDataForPlatform(const ITargetPlatform* TargetPlatform)
 {
-	const FMutableCachedPlatformData* PlatformData = GetPrivate()->CachedPlatformsData.Find(TargetPlatform->PlatformName());
+	const FMutableCachedPlatformData* PlatformData = CachedPlatformsData.Find(TargetPlatform->PlatformName());
 	if (!PlatformData)
 	{
 		return false;
@@ -241,9 +241,14 @@ void UCustomizableObject::PostLoad()
 }
 
 
-bool UCustomizableObjectPrivate::IsLocked() const
+bool UCustomizableObject::IsLocked() const
 {
-	return bLocked;
+	if (Private)
+	{
+		return Private->bLocked;
+	}
+
+	return false;
 }
 
 
@@ -265,7 +270,7 @@ void UCustomizableObject::Serialize(FArchive& Ar_Asset)
 	else
 	{
 		// Can't remove this or saved customizable objects will fail to load
-		int64 InternalVersion = UCustomizableObjectPrivate::CurrentSupportedVersion;
+		int64 InternalVersion = CurrentSupportedVersion;
 		Ar_Asset << InternalVersion;
 	}
 #else
@@ -332,8 +337,8 @@ void UCustomizableObject::ClearCompiledData(bool bIsCooking)
 	ClothSharedConfigsData.Empty();
 
 #if WITH_EDITORONLY_DATA
-	GetPrivate()->CustomizableObjectPathMap.Empty();
-	GetPrivate()->GroupNodeMap.Empty();
+	CustomizableObjectPathMap.Empty();
+	GroupNodeMap.Empty();
 	Private->ParticipatingObjects.Empty();
 #endif
 
@@ -455,7 +460,7 @@ void SerializeStreamedResources(FArchive& Ar, UObject* Object, TArray<FCustomiza
 
 void UCustomizableObject::SaveCompiledData(FArchive& MemoryWriter, bool bIsCooking)
 {
-	int32 InternalVersion = UCustomizableObjectPrivate::CurrentSupportedVersion;
+	int32 InternalVersion = CurrentSupportedVersion;
 	MutableCompiledDataStreamHeader Header(InternalVersion, VersionId);
 	MemoryWriter << Header;
 
@@ -540,9 +545,9 @@ void UCustomizableObject::SaveCompiledData(FArchive& MemoryWriter, bool bIsCooki
 	MemoryWriter << LODSettings.bLODStreamingEnabled;
 
 	// Editor Only data
-	MemoryWriter << GetPrivate()->bIsCompiledWithOptimization;
-	MemoryWriter << GetPrivate()->CustomizableObjectPathMap;
-	MemoryWriter << GetPrivate()->GroupNodeMap;
+	MemoryWriter << bIsCompiledWithOptimization;
+	MemoryWriter << CustomizableObjectPathMap;
+	MemoryWriter << GroupNodeMap;
 	MemoryWriter << GetPrivate()->ParticipatingObjects;
 }
 
@@ -554,7 +559,7 @@ void UCustomizableObject::LoadCompiledData(FArchive& MemoryReader, const ITarget
 	MutableCompiledDataStreamHeader Header;
 	MemoryReader << Header;
 
-	if (UCustomizableObjectPrivate::CurrentSupportedVersion == Header.InternalVersion)
+	if (CurrentSupportedVersion == Header.InternalVersion)
 	{
 		// Make sure mutable has been initialised.
 		UCustomizableObjectSystem::GetInstance();
@@ -662,9 +667,9 @@ void UCustomizableObject::LoadCompiledData(FArchive& MemoryReader, const ITarget
 
 		// Editor Only data
 		{
-			MemoryReader << GetPrivate()->bIsCompiledWithOptimization;
-			MemoryReader << GetPrivate()->CustomizableObjectPathMap;
-			MemoryReader << GetPrivate()->GroupNodeMap;
+			MemoryReader << bIsCompiledWithOptimization;
+			MemoryReader << CustomizableObjectPathMap;
+			MemoryReader << GroupNodeMap;
 
 			TMap<FName, FGuid>& ParticipatingObjects = GetPrivate()->ParticipatingObjects;
 			MemoryReader << ParticipatingObjects;
@@ -753,7 +758,7 @@ void UCustomizableObject::LoadCompiledDataFromDisk()
 			AuxMemoryReader << StreamableDataHeader;
 		}
 
-		if (CompiledDataHeader.InternalVersion == UCustomizableObjectPrivate::CurrentSupportedVersion
+		if (CompiledDataHeader.InternalVersion == CurrentSupportedVersion
 			&&
 			CompiledDataHeader.InternalVersion == StreamableDataHeader.InternalVersion 
 			&&
@@ -790,9 +795,9 @@ void UCustomizableObject::CachePlatformData(const ITargetPlatform* InTargetPlatf
 {
 	FString PlatformName = InTargetPlatform ? InTargetPlatform->PlatformName() : FPlatformProperties::PlatformName();
 
-	check(!GetPrivate()->CachedPlatformsData.Find(PlatformName));
+	check(!CachedPlatformsData.Find(PlatformName));
 
-	FMutableCachedPlatformData& Data = GetPrivate()->CachedPlatformsData.Add(PlatformName);
+	FMutableCachedPlatformData& Data = CachedPlatformsData.Add(PlatformName);
 
 	// Cache CO data and mu::Model
 	FMemoryWriter64 MemoryWriter(Data.ModelData);
@@ -842,7 +847,7 @@ bool UCustomizableObject::ConditionalAutoCompile()
 	check(IsInGameThread());
 
 	// Don't compile objects being compiled
-	if (GetPrivate()->IsLocked())
+	if (IsLocked())
 	{
 		return false;
 	}
@@ -851,7 +856,7 @@ bool UCustomizableObject::ConditionalAutoCompile()
 	if (IsCompiled())
 	{
 		// Show a warning if the compilation was not done with optimizations.
-		if (GetPrivate()->bIsCompiledWithOptimization)
+		if (bIsCompiledWithOptimization)
 		{
 			FString Msg = FString::Printf(TEXT("Warning: Customizable Object [%s] was compiled without optimization."), *GetName());
 			GEngine->AddOnScreenDebugMessage((uint64)((PTRINT)this), 10.0f, FColor::Red, Msg);
@@ -873,7 +878,7 @@ bool UCustomizableObject::ConditionalAutoCompile()
 	}
 
 	// Don't re-compile objects if they failed to compile. 
-	if (GetPrivate()->CompilationState == ECustomizableObjectCompilationState::Failed)
+	if (CompilationState == ECustomizableObjectCompilationState::Failed)
 	{
 		return false;
 	}
@@ -973,13 +978,13 @@ void UCustomizableObject::SaveEmbeddedData(FArchive& Ar)
 {
 	UE_LOG(LogMutable, Verbose, TEXT("Saving embedded data for Customizable Object [%s] now at position %d."), *GetName(), int(Ar.Tell()));
 
-	int32 InternalVersion = Private->GetModel() ? UCustomizableObjectPrivate::CurrentSupportedVersion : -1;
+	int32 InternalVersion = Private->GetModel() ? CurrentSupportedVersion : -1;
 	Ar << InternalVersion;
 
 	if (Private->GetModel())
 	{
 		// General derived flags
-		Ar << GetPrivate()->bDisableTextureStreaming;
+		Ar << bDisableTextureStreaming;
 
 		// Serialize morph data
 		{
@@ -1017,12 +1022,12 @@ void UCustomizableObject::LoadEmbeddedData(FArchive& Ar)
 
 	// If this fails, something went wrong with the packaging: we have data that belongs
 	// to a different version than the code.
-	check(UCustomizableObjectPrivate::CurrentSupportedVersion==InternalVersion);
+	check(CurrentSupportedVersion==InternalVersion);
 
-	if(UCustomizableObjectPrivate::CurrentSupportedVersion == InternalVersion)
+	if(CurrentSupportedVersion == InternalVersion)
 	{
 		// General derived flags
-		Ar << GetPrivate()->bDisableTextureStreaming;
+		Ar << bDisableTextureStreaming;
 
 		// Load morph data
 		{
@@ -1183,17 +1188,11 @@ FString UCustomizableObject::GetStateParameterName(int32 StateIndex, int32 Param
 #if WITH_EDITORONLY_DATA
 void UCustomizableObject::PostCompile()
 {
-	GetPrivate()->PostCompileDelegate.Broadcast();
+	PostCompileDelegate.Broadcast();
 }
 #endif
 
 
-FPostCompileDelegate& UCustomizableObject::GetPostCompileDelegate() const
-{
-	return GetPrivate()->PostCompileDelegate;
-}
-
-
 UCustomizableObjectInstance* UCustomizableObject::CreateInstance()
 {
 	MUTABLE_CPUPROFILER_SCOPE(UCustomizableObject::CreateInstance)
@@ -1208,6 +1207,11 @@ UCustomizableObjectInstance* UCustomizableObject::CreateInstance()
 }
 
 
+TSharedPtr<mu::Model, ESPMode::ThreadSafe> UCustomizableObject::GetModel() const
+{
+	return Private->GetModel();
+}
+
 #if WITH_EDITOR
 void UCustomizableObject::SetModel(TSharedPtr<mu::Model, ESPMode::ThreadSafe> Model)
 {
@@ -1249,15 +1253,15 @@ int32 UCustomizableObject::GetComponentCount() const
 
 int32 UCustomizableObject::GetParameterCount() const
 {
-	return GetPrivate()->ParameterProperties.Num();
+	return ParameterProperties.Num();
 }
 
 
 EMutableParameterType UCustomizableObject::GetParameterType(int32 ParamIndex) const
 {
-	if (ParamIndex >= 0 && ParamIndex < GetPrivate()->ParameterProperties.Num())
+	if (ParamIndex >= 0 && ParamIndex < ParameterProperties.Num())
 	{
-		return GetPrivate()->ParameterProperties[ParamIndex].Type;
+		return ParameterProperties[ParamIndex].Type;
 	}
 	else
 	{
@@ -1270,21 +1274,21 @@ EMutableParameterType UCustomizableObject::GetParameterType(int32 ParamIndex) co
 
 EMutableParameterType UCustomizableObject::GetParameterTypeByName(const FString& Name) const
 {
-	const int* IndexPtr = GetPrivate()->ParameterPropertiesLookupTable.Find(Name);
+	const int* IndexPtr = ParameterPropertiesLookupTable.Find(Name);
 	int Index = IndexPtr ? *IndexPtr : INDEX_NONE;
 
-	if (GetPrivate()->ParameterProperties.IsValidIndex(Index))
+	if (ParameterProperties.IsValidIndex(Index))
 	{
-		return GetPrivate()->ParameterProperties[Index].Type;
+		return ParameterProperties[Index].Type;
 	}
 
 	UE_LOG(LogMutable, Warning, TEXT("Name '%s' does not exist in ParameterProperties lookup table at GetParameterTypeByName at CO %s."), *Name, *GetName());
 
-	for (int32 ParamIndex = 0; ParamIndex < GetPrivate()->ParameterProperties.Num(); ++ParamIndex)
+	for (int32 ParamIndex = 0; ParamIndex < ParameterProperties.Num(); ++ParamIndex)
 	{
-		if (GetPrivate()->ParameterProperties[ParamIndex].Name == Name)
+		if (ParameterProperties[ParamIndex].Name == Name)
 		{
-			return GetPrivate()->ParameterProperties[ParamIndex].Type;
+			return ParameterProperties[ParamIndex].Type;
 		}
 	}
 
@@ -1298,9 +1302,9 @@ static const FString s_EmptyString;
 
 const FString & UCustomizableObject::GetParameterName(int32 ParamIndex) const
 {
-	if (ParamIndex >= 0 && ParamIndex < GetPrivate()->ParameterProperties.Num())
+	if (ParamIndex >= 0 && ParamIndex < ParameterProperties.Num())
 	{
-		return GetPrivate()->ParameterProperties[ParamIndex].Name;
+		return ParameterProperties[ParamIndex].Name;
 	}
 	else
 	{
@@ -1318,8 +1322,8 @@ void UCustomizableObject::UpdateParameterPropertiesFromModel(const TSharedPtr<mu
 		mu::ParametersPtr MutableParameters = mu::Model::NewParameters(Model);
 		int paramCount = MutableParameters->GetCount();
 
-		GetPrivate()->ParameterProperties.Reset(paramCount);
-		GetPrivate()->ParameterPropertiesLookupTable.Empty(paramCount);
+		ParameterProperties.Reset(paramCount);
+		ParameterPropertiesLookupTable.Empty(paramCount);
 		for (int paramIndex = 0; paramIndex<paramCount; ++paramIndex)
 		{
 			FMutableModelParameterProperties Data;
@@ -1381,14 +1385,14 @@ void UCustomizableObject::UpdateParameterPropertiesFromModel(const TSharedPtr<mu
 				break;
 			}
 
-			GetPrivate()->ParameterProperties.Add(Data);
-			GetPrivate()->ParameterPropertiesLookupTable.Add(Data.Name, paramIndex);
+			ParameterProperties.Add(Data);
+			ParameterPropertiesLookupTable.Add(Data.Name, paramIndex);
 		}
 	}
 	else
 	{
-		GetPrivate()->ParameterProperties.Empty();
-		GetPrivate()->ParameterPropertiesLookupTable.Empty();
+		ParameterProperties.Empty();
+		ParameterPropertiesLookupTable.Empty();
 	}
 }
 
@@ -1401,9 +1405,9 @@ int32 UCustomizableObject::GetParameterDescriptionCount(const FString& ParamName
 
 int32 UCustomizableObject::GetIntParameterNumOptions(int32 ParamIndex) const
 {
-	if (ParamIndex >= 0 && ParamIndex < GetPrivate()->ParameterProperties.Num())
+	if (ParamIndex >= 0 && ParamIndex < ParameterProperties.Num())
 	{
-		return GetPrivate()->ParameterProperties[ParamIndex].PossibleValues.Num();
+		return ParameterProperties[ParamIndex].PossibleValues.Num();
 	}
 	else
 	{
@@ -1416,11 +1420,11 @@ int32 UCustomizableObject::GetIntParameterNumOptions(int32 ParamIndex) const
 
 const FString& UCustomizableObject::GetIntParameterAvailableOption(int32 ParamIndex, int32 K) const
 {
-	if (ParamIndex >= 0 && ParamIndex < GetPrivate()->ParameterProperties.Num())
+	if (ParamIndex >= 0 && ParamIndex < ParameterProperties.Num())
 	{
 		if (K >= 0 && K < GetIntParameterNumOptions(ParamIndex))
 		{
-			return GetPrivate()->ParameterProperties[ParamIndex].PossibleValues[K].Name;
+			return ParameterProperties[ParamIndex].PossibleValues[K].Name;
 		}
 		else
 		{
@@ -1438,7 +1442,7 @@ const FString& UCustomizableObject::GetIntParameterAvailableOption(int32 ParamIn
 
 int32 UCustomizableObject::FindParameter(const FString& Name) const
 {
-	const int32 * Found = GetPrivate()->ParameterPropertiesLookupTable.Find(Name);
+	const int32 * Found = ParameterPropertiesLookupTable.Find(Name);
 	if (Found == nullptr)
 	{
 		return INDEX_NONE;
@@ -1450,9 +1454,9 @@ int32 UCustomizableObject::FindParameter(const FString& Name) const
 int32 UCustomizableObject::FindIntParameterValue(int32 ParamIndex, const FString& Value) const
 {
 	int32 MinValueIndex = 0;
-	if (ParamIndex >= 0 && ParamIndex < GetPrivate()->ParameterProperties.Num())
+	if (ParamIndex >= 0 && ParamIndex < ParameterProperties.Num())
 	{
-		const TArray<FMutableModelParameterValue>& PossibleValues = GetPrivate()->ParameterProperties[ParamIndex].PossibleValues;
+		const TArray<FMutableModelParameterValue>& PossibleValues = ParameterProperties[ParamIndex].PossibleValues;
 		if (PossibleValues.Num())
 		{
 			MinValueIndex = PossibleValues[0].Value;
@@ -1484,9 +1488,9 @@ int32 UCustomizableObject::FindIntParameterValue(int32 ParamIndex, const FString
 
 FString UCustomizableObject::FindIntParameterValueName(int32 ParamIndex, int32 ParamValue) const
 {
-	if (ParamIndex >= 0 && ParamIndex < GetPrivate()->ParameterProperties.Num())
+	if (ParamIndex >= 0 && ParamIndex < ParameterProperties.Num())
 	{
-		const TArray<FMutableModelParameterValue> & PossibleValues = GetPrivate()->ParameterProperties[ParamIndex].PossibleValues;
+		const TArray<FMutableModelParameterValue> & PossibleValues = ParameterProperties[ParamIndex].PossibleValues;
 
 		const int32 MinValueIndex = !PossibleValues.IsEmpty() ? PossibleValues[0].Value : 0;
 		ParamValue = ParamValue - MinValueIndex;
@@ -1542,7 +1546,7 @@ float UCustomizableObject::GetFloatParameterDefaultValue(const FString& InParame
 		return FCustomizableObjectFloatParameterValue::DEFAULT_PARAMETER_VALUE;
 	}
 
-	const TSharedPtr<mu::Model> Model = GetPrivate()->GetModel();
+	const TSharedPtr<mu::Model> Model = GetModel();
 	if (!Model)
 	{
 		checkNoEntry();
@@ -1562,7 +1566,7 @@ int32 UCustomizableObject::GetIntParameterDefaultValue(const FString& InParamete
 		return FCustomizableObjectIntParameterValue::DEFAULT_PARAMETER_VALUE;
 	}
 
-	const TSharedPtr<mu::Model> Model = GetPrivate()->GetModel();
+	const TSharedPtr<mu::Model> Model = GetModel();
 	if (!Model)
 	{
 		checkNoEntry();
@@ -1582,7 +1586,7 @@ bool UCustomizableObject::GetBoolParameterDefaultValue(const FString& InParamete
 		return FCustomizableObjectBoolParameterValue::DEFAULT_PARAMETER_VALUE;
 	}
 
-	const TSharedPtr<mu::Model> Model = GetPrivate()->GetModel();
+	const TSharedPtr<mu::Model> Model = GetModel();
 	if (!Model)
 	{
 		checkNoEntry();
@@ -1602,7 +1606,7 @@ FLinearColor UCustomizableObject::GetColorParameterDefaultValue(const FString& I
 		return FCustomizableObjectVectorParameterValue::DEFAULT_PARAMETER_VALUE;
 	}
 
-	const TSharedPtr<mu::Model> Model = GetPrivate()->GetModel();
+	const TSharedPtr<mu::Model> Model = GetModel();
 	if (!Model)
 	{
 		checkNoEntry();
@@ -1640,7 +1644,7 @@ FCustomizableObjectProjector UCustomizableObject::GetProjectorParameterDefaultVa
 		return FCustomizableObjectProjectorParameterValue::DEFAULT_PARAMETER_VALUE;
 	}
 
-	const TSharedPtr<mu::Model> Model = GetPrivate()->GetModel();
+	const TSharedPtr<mu::Model> Model = GetModel();
 	if (!Model)
 	{
 		checkNoEntry();
@@ -1665,7 +1669,7 @@ FName UCustomizableObject::GetTextureParameterDefaultValue(const FString& InPara
 		return FCustomizableObjectTextureParameterValue::DEFAULT_PARAMETER_VALUE;
 	}
 
-	const TSharedPtr<mu::Model> Model = GetPrivate()->GetModel();
+	const TSharedPtr<mu::Model> Model = GetModel();
 	if (!Model)
 	{
 		checkNoEntry();
@@ -1917,12 +1921,6 @@ FModelResources& UCustomizableObjectPrivate::GetModelResources(bool bIsCooking)
 //-------------------------------------------------------------------------------------------------
 //-------------------------------------------------------------------------------------------------
 
-ECustomizableObjectCompilationState FCustomizableObjectCompilerBase::GetCompilationState() const
-{
-	return ECustomizableObjectCompilationState::None;
-}
-
-
 void UCustomizableObjectBulk::PostLoad()
 {
 	UObject::PostLoad();
@@ -1985,7 +1983,7 @@ void UCustomizableObjectBulk::CookAdditionalFilesOverride(const TCHAR* PackageFi
 	
 	check(CustomizableObject);
 	
-	FMutableCachedPlatformData* PlatformData = CustomizableObject->GetPrivate()->CachedPlatformsData.Find(TargetPlatform->PlatformName());
+	FMutableCachedPlatformData* PlatformData = CustomizableObject->CachedPlatformsData.Find(TargetPlatform->PlatformName());
 	check(PlatformData);
 
 	// Data to serialize in separate files
@@ -2015,7 +2013,7 @@ void UCustomizableObjectBulk::PrepareBulkData(UCustomizableObject* InOuter, cons
 	BulkDataFileNames.Empty();
 	
 	// Split the Streamable data into several separate files and fix up FileIndex and Offset of each StreamableBlock
-	if (TSharedPtr<const mu::Model, ESPMode::ThreadSafe> Model = CustomizableObject->GetPrivate()->GetModel())
+	if (TSharedPtr<const mu::Model, ESPMode::ThreadSafe> Model = CustomizableObject->GetModel())
 	{
 		const uint64 MaxChunkSize = UCustomizableObjectSystem::GetInstance()->GetMaxChunkSizeForPlatform(TargetPlatform);
 
diff --git a/Engine/Plugins/Experimental/Mutable/Source/CustomizableObject/Private/MuCO/CustomizableObjectInstance.cpp b/Engine/Plugins/Experimental/Mutable/Source/CustomizableObject/Private/MuCO/CustomizableObjectInstance.cpp
--- a/Engine/Plugins/Experimental/Mutable/Source/CustomizableObject/Private/MuCO/CustomizableObjectInstance.cpp
+++ b/Engine/Plugins/Experimental/Mutable/Source/CustomizableObject/Private/MuCO/CustomizableObjectInstance.cpp
@@ -347,14 +347,14 @@ void UCustomizableInstancePrivate::BindObjectDelegates(UCustomizableObject*  Cur
 	// Unbind callback from the previous CO
 	if (CurrentCustomizableObject)
 	{
-		CurrentCustomizableObject->GetPrivate()->PostCompileDelegate.RemoveAll(this);
+		CurrentCustomizableObject->PostCompileDelegate.RemoveAll(this);
 		CurrentCustomizableObject->GetPrivate()->Status.GetOnStateChangedDelegate().RemoveAll(this);
 	}
 
 	// Bind callback to the new CO
 	if (NewCustomizableObject)
 	{
-		NewCustomizableObject->GetPrivate()->PostCompileDelegate.AddUObject(this, &UCustomizableInstancePrivate::OnPostCompile);
+		NewCustomizableObject->PostCompileDelegate.AddUObject(this, &UCustomizableInstancePrivate::OnPostCompile);
 		NewCustomizableObject->GetPrivate()->Status.GetOnStateChangedDelegate().AddUObject(this, &UCustomizableInstancePrivate::OnObjectStatusChanged);
 	}
 }
@@ -5397,7 +5397,7 @@ void UCustomizableInstancePrivate::BuildMaterials(const TSharedRef<FUpdateContex
 		CustomizableObject->GetPathName(),
 		Public->GetPathName(),
 		UCustomizableObjectSystem::GetInstance()->GetPrivate()->MutableSystem,
-		CustomizableObject->GetPrivate()->GetModel(),
+		CustomizableObject->GetModel(),
 		OperationData->Parameters,
 	    OperationData->GetCapturedDescriptor().GetState());
 	
diff --git a/Engine/Plugins/Experimental/Mutable/Source/CustomizableObject/Private/MuCO/CustomizableObjectMipDataProvider.cpp b/Engine/Plugins/Experimental/Mutable/Source/CustomizableObject/Private/MuCO/CustomizableObjectMipDataProvider.cpp
--- a/Engine/Plugins/Experimental/Mutable/Source/CustomizableObject/Private/MuCO/CustomizableObjectMipDataProvider.cpp
+++ b/Engine/Plugins/Experimental/Mutable/Source/CustomizableObject/Private/MuCO/CustomizableObjectMipDataProvider.cpp
@@ -112,7 +112,7 @@ void FMutableTextureMipDataProvider::Init(const FTextureUpdateContext& Context,
 #if WITH_EDITOR
 	check(Context.Texture->HasPendingInitOrStreaming());
 	check(CustomizableObjectInstance->GetCustomizableObject());
-	if (CustomizableObjectInstance->GetCustomizableObject()->GetPrivate()->IsLocked())
+	if (CustomizableObjectInstance->GetCustomizableObject()->IsLocked())
 	{
 		PrintWarningAndAdvanceToCleanup();
 
@@ -300,7 +300,7 @@ int32 FMutableTextureMipDataProvider::GetMips(const FTextureUpdateContext& Conte
 #if WITH_EDITOR
 	check(Context.Texture->HasPendingInitOrStreaming());
 	check(CustomizableObjectInstance->GetCustomizableObject());
-	if (CustomizableObjectInstance->GetCustomizableObject()->GetPrivate()->IsLocked())
+	if (CustomizableObjectInstance->GetCustomizableObject()->IsLocked())
 	{
 		PrintWarningAndAdvanceToCleanup();
 
@@ -398,7 +398,7 @@ bool FMutableTextureMipDataProvider::PollMips(const FTextureUpdateSyncOptions& S
 	
 #if WITH_EDITOR
 	check(CustomizableObjectInstance->GetCustomizableObject());
-	if (CustomizableObjectInstance->GetCustomizableObject()->GetPrivate()->IsLocked())
+	if (CustomizableObjectInstance->GetCustomizableObject()->IsLocked())
 	{
 		PrintWarningAndAdvanceToCleanup();
 
diff --git a/Engine/Plugins/Experimental/Mutable/Source/CustomizableObject/Private/MuCO/CustomizableObjectSystem.cpp b/Engine/Plugins/Experimental/Mutable/Source/CustomizableObject/Private/MuCO/CustomizableObjectSystem.cpp
--- a/Engine/Plugins/Experimental/Mutable/Source/CustomizableObject/Private/MuCO/CustomizableObjectSystem.cpp
+++ b/Engine/Plugins/Experimental/Mutable/Source/CustomizableObject/Private/MuCO/CustomizableObjectSystem.cpp
@@ -985,7 +985,7 @@ void UCustomizableObjectSystemPrivate::GetMipStreamingConfig(const UCustomizable
 		bOutNeverStream = State->bDisableTextureStreaming;
 
 		// Was streaming disabled at object-compilation time? 
-		if (Instance.GetCustomizableObject()->GetPrivate()->bDisableTextureStreaming)
+		if (Instance.GetCustomizableObject()->bDisableTextureStreaming)
 		{
 			bOutNeverStream = true;
 		}
@@ -3453,7 +3453,7 @@ void UCustomizableObjectSystem::OnPreBeginPIE(const bool bIsSimulatingInEditor)
 		}
 		
 		const UCustomizableObject* Object = Cast<UCustomizableObject>(Asset.GetAsset());
-		if (!Object || Object->IsCompiled() || Object->GetPrivate()->IsLocked() || Object->bIsChildObject)
+		if (!Object || Object->IsCompiled() || Object->IsLocked() || Object->bIsChildObject)
 		{
 			continue;
 		}
@@ -3503,7 +3503,7 @@ void UCustomizableObjectSystem::RecompileCustomizableObjectAsync(const FAssetDat
 		return;
 	}
 	
-	if ((InObject && InObject->GetPrivate()->IsLocked()) || ObjectsToRecompile.Find((InAssetData)) != INDEX_NONE)
+	if ((InObject && InObject->IsLocked()) || ObjectsToRecompile.Find((InAssetData)) != INDEX_NONE)
 	{
 		return;
 	}
diff --git a/Engine/Plugins/Experimental/Mutable/Source/CustomizableObject/Private/MuCO/UnrealMutableImageProvider.cpp b/Engine/Plugins/Experimental/Mutable/Source/CustomizableObject/Private/MuCO/UnrealMutableImageProvider.cpp
--- a/Engine/Plugins/Experimental/Mutable/Source/CustomizableObject/Private/MuCO/UnrealMutableImageProvider.cpp
+++ b/Engine/Plugins/Experimental/Mutable/Source/CustomizableObject/Private/MuCO/UnrealMutableImageProvider.cpp
@@ -110,7 +110,7 @@ bool FUnrealMutableImageProvider::Tick(float DeltaTime)
 	UCustomizableObject* CO = nullptr;
 	for (TObjectIterator<UCustomizableObject> It; It; ++It)
 	{
-		if (IsValid(*It) && It->GetPrivate()->GetModel().Get() == Request->ModelPtr)
+		if (IsValid(*It) && It->GetModel().Get() == Request->ModelPtr)
 		{
 			CO = *It;
 			break;
diff --git a/Engine/Plugins/Experimental/Mutable/Source/CustomizableObject/Private/MuCO/UnrealMutableModelDiskStreamer.cpp b/Engine/Plugins/Experimental/Mutable/Source/CustomizableObject/Private/MuCO/UnrealMutableModelDiskStreamer.cpp
--- a/Engine/Plugins/Experimental/Mutable/Source/CustomizableObject/Private/MuCO/UnrealMutableModelDiskStreamer.cpp
+++ b/Engine/Plugins/Experimental/Mutable/Source/CustomizableObject/Private/MuCO/UnrealMutableModelDiskStreamer.cpp
@@ -4,7 +4,6 @@
 
 #include "Async/AsyncFileHandle.h"
 #include "HAL/PlatformFile.h"
-#include "MuCO/CustomizableObjectPrivate.h"
 #include "MuCO/CustomizableObjectSystem.h"
 #include "MuR/Model.h"
 #include "MuR/MutableTrace.h"
@@ -77,14 +76,14 @@ bool FUnrealMutableModelBulkReader::PrepareStreamingForObject(UCustomizableObjec
 	// Is the object already prepared for streaming?
 	bool bAlreadyStreaming = Objects.FindByPredicate(
 		[CustomizableObject](const FObjectData& d)
-		{ return d.Model.Pin().Get() == CustomizableObject->GetPrivate()->GetModel().Get(); })
+		{ return d.Model.Pin().Get() == CustomizableObject->GetModel().Get(); })
 		!=
 		nullptr;
 
 	if (!bAlreadyStreaming)
 	{
 		FObjectData NewData;
-		NewData.Model = TWeakPtr<const mu::Model>(CustomizableObject->GetPrivate()->GetModel());
+		NewData.Model = TWeakPtr<const mu::Model>(CustomizableObject->GetModel());
 
 #if WITH_EDITOR
 		FString FolderPath = CustomizableObject->GetCompiledDataFolderPath();
@@ -147,7 +146,7 @@ void FUnrealMutableModelBulkReader::CancelStreamingForObject(const UCustomizable
 	// See if we can free previuously allocated resources
 	for (int32 ObjectIndex = 0; ObjectIndex < Objects.Num(); ++ObjectIndex)
 	{
-		if (Objects[ObjectIndex].Model.Pin() == CustomizableObject->GetPrivate()->GetModel())
+		if (Objects[ObjectIndex].Model.Pin() == CustomizableObject->GetModel())
 		{
 			check(Objects[ObjectIndex].CurrentReadRequests.IsEmpty());
 
@@ -171,7 +170,7 @@ bool FUnrealMutableModelBulkReader::AreTherePendingStreamingOperationsForObject(
 
 	for (int32 ObjectIndex = 0; ObjectIndex < Objects.Num(); ++ObjectIndex)
 	{
-		if (Objects[ObjectIndex].Model.Pin() == CustomizableObject->GetPrivate()->GetModel())
+		if (Objects[ObjectIndex].Model.Pin() == CustomizableObject->GetModel())
 		{
 			if (!Objects[ObjectIndex].CurrentReadRequests.IsEmpty())
 			{
diff --git a/Engine/Plugins/Experimental/Mutable/Source/CustomizableObjectEditor/Private/MuCOE/CustomizableInstanceDetails.cpp b/Engine/Plugins/Experimental/Mutable/Source/CustomizableObjectEditor/Private/MuCOE/CustomizableInstanceDetails.cpp
--- a/Engine/Plugins/Experimental/Mutable/Source/CustomizableObjectEditor/Private/MuCOE/CustomizableInstanceDetails.cpp
+++ b/Engine/Plugins/Experimental/Mutable/Source/CustomizableObjectEditor/Private/MuCOE/CustomizableInstanceDetails.cpp
@@ -884,10 +884,10 @@ TSharedRef<SWidget> FCustomizableInstanceDetails::GenerateIntWidget(const int32
 			{
 				ValueIndex = i;
 
-				const FString* Identifier = CustomizableObject->GetPrivate()->GroupNodeMap.FindKey(FCustomizableObjectIdPair(ParamName, PossibleValue));
+				const FString* Identifier = CustomizableObject->GroupNodeMap.FindKey(FCustomizableObjectIdPair(ParamName, PossibleValue));
 				if (Identifier)
 				{
-					if (FString* CustomizableObjectPath = CustomizableObject->GetPrivate()->CustomizableObjectPathMap.Find(*Identifier))
+					if (FString* CustomizableObjectPath = CustomizableObject->CustomizableObjectPathMap.Find(*Identifier))
 					{
 						ToolTipText = *CustomizableObjectPath;
 					}
diff --git a/Engine/Plugins/Experimental/Mutable/Source/CustomizableObjectEditor/Private/MuCOE/CustomizableObjectCompileRunnable.cpp b/Engine/Plugins/Experimental/Mutable/Source/CustomizableObjectEditor/Private/MuCOE/CustomizableObjectCompileRunnable.cpp
--- a/Engine/Plugins/Experimental/Mutable/Source/CustomizableObjectEditor/Private/MuCOE/CustomizableObjectCompileRunnable.cpp
+++ b/Engine/Plugins/Experimental/Mutable/Source/CustomizableObjectEditor/Private/MuCOE/CustomizableObjectCompileRunnable.cpp
@@ -5,7 +5,6 @@
 #include "HAL/FileManager.h"
 #include "MuCO/UnrealMutableModelDiskStreamer.h"
 #include "MuCO/UnrealToMutableTextureConversionUtils.h"
-#include "MuCO/CustomizableObjectPrivate.h"
 #include "MuR/Model.h"
 #include "MuT/Compiler.h"
 #include "MuT/ErrorLog.h"
@@ -246,10 +245,10 @@ const TArray<FCustomizableObjectCompileRunnable::FError>& FCustomizableObjectCom
 
 FCustomizableObjectSaveDDRunnable::FCustomizableObjectSaveDDRunnable(UCustomizableObject* CustomizableObject, const FCompilationOptions& InOptions)
 {
-	Model = CustomizableObject->GetPrivate()->GetModel();
+	Model = CustomizableObject->GetModel();
 	Options = InOptions;
 	
-	CustomizableObjectHeader.InternalVersion = CustomizableObject->GetPrivate()->CurrentSupportedVersion;
+	CustomizableObjectHeader.InternalVersion = CustomizableObject->GetCurrentSupportedVersion();
 	CustomizableObjectHeader.VersionId = Options.bIsCooking? FGuid::NewGuid() : CustomizableObject->GetVersionId();
 
 	if (!Options.bIsCooking)
diff --git a/Engine/Plugins/Experimental/Mutable/Source/CustomizableObjectEditor/Private/MuCOE/CustomizableObjectCompileRunnable.h b/Engine/Plugins/Experimental/Mutable/Source/CustomizableObjectEditor/Private/MuCOE/CustomizableObjectCompileRunnable.h
--- a/Engine/Plugins/Experimental/Mutable/Source/CustomizableObjectEditor/Private/MuCOE/CustomizableObjectCompileRunnable.h
+++ b/Engine/Plugins/Experimental/Mutable/Source/CustomizableObjectEditor/Private/MuCOE/CustomizableObjectCompileRunnable.h
@@ -4,7 +4,6 @@
 
 #include "HAL/Runnable.h"
 #include "MuCO/CustomizableObject.h"
-#include "MuCO/CustomizableObjectPrivate.h"
 #include "MuR/Ptr.h"
 #include "MuT/Node.h"
 #include "MuCOE/CustomizableObjectEditorLogger.h"
diff --git a/Engine/Plugins/Experimental/Mutable/Source/CustomizableObjectEditor/Private/MuCOE/CustomizableObjectCompiler.cpp b/Engine/Plugins/Experimental/Mutable/Source/CustomizableObjectEditor/Private/MuCOE/CustomizableObjectCompiler.cpp
--- a/Engine/Plugins/Experimental/Mutable/Source/CustomizableObjectEditor/Private/MuCOE/CustomizableObjectCompiler.cpp
+++ b/Engine/Plugins/Experimental/Mutable/Source/CustomizableObjectEditor/Private/MuCOE/CustomizableObjectCompiler.cpp
@@ -185,7 +185,7 @@ void FCustomizableObjectCompiler::Compile(UCustomizableObject& Object, const FCo
 	Info.FadeOutDuration = 1.0f;
 	Info.ExpireDuration = 1.0f;
 
-	if (PreloadingReferencerAssets || CompilationLaunchPending || (Object.GetPrivate()->IsLocked()))
+	if (PreloadingReferencerAssets || CompilationLaunchPending || (Object.IsLocked()))
 	{
 		UE_LOG(LogMutable, Warning, TEXT("%s"), *Message);
 		FSlateNotificationManager::Get().AddNotification(Info);
@@ -822,7 +822,7 @@ void FCustomizableObjectCompiler::SetCompilationState(ECustomizableObjectCompila
 	
 	if (CurrentObject)
 	{
-		CurrentObject->GetPrivate()->CompilationState = InState;
+		CurrentObject->CompilationState = InState;
 	}
 }
 
@@ -864,7 +864,7 @@ void FCustomizableObjectCompiler::CompileInternal(UCustomizableObject* Object, c
 		return;
 	}
 
-	if (Object->GetPrivate()->IsLocked() || !UCustomizableObjectSystem::GetInstance()->LockObject(Object))
+	if (Object->IsLocked() || !UCustomizableObjectSystem::GetInstance()->LockObject(Object))
 	{
 		UE_LOG(LogMutable, Display, TEXT("Customizable Object is already being compiled or updated %s. Please wait a few seconds and try again."), *Object->GetName());
 		return;
@@ -921,7 +921,7 @@ void FCustomizableObjectCompiler::CompileInternal(UCustomizableObject* Object, c
 			CompilerLog(FText::FromString(TEXT("Failed to generate the mutable node graph. Object not built.")), nullptr);
 		}
 
-		if (Object->GetPrivate()->IsLocked())
+		if (Object->IsLocked())
 		{
 			UCustomizableObjectSystem::GetInstance()->UnlockObject(Object);
 		}
@@ -1033,16 +1033,16 @@ void FCustomizableObjectCompiler::CompileInternal(UCustomizableObject* Object, c
 			 ClothingAssetData.ConfigsData.RemoveAllSwap(IsSharedConfigData);
 		}
 
-		Object->GetPrivate()->GroupNodeMap = GenerationContext.GroupNodeMap;
+		Object->GroupNodeMap = GenerationContext.GroupNodeMap;
 
 		if (GenerationContext.Options.OptimizationLevel == 0)
 		{
 			// If the optimization level is "none" disable texture streaming, because textures are all referenced
 			// unreal assets and progressive generation is not supported.
-			Object->GetPrivate()->bDisableTextureStreaming = true;
+			Object->bDisableTextureStreaming = true;
 		}
 		
-		Object->GetPrivate()->bIsCompiledWithOptimization = GenerationContext.Options.OptimizationLevel < UE_MUTABLE_MAX_OPTIMIZATION;
+		Object->bIsCompiledWithOptimization = GenerationContext.Options.OptimizationLevel < UE_MUTABLE_MAX_OPTIMIZATION;
 
 		Object->AlwaysLoadedExtensionData = MoveTemp(GenerationContext.AlwaysLoadedExtensionData);
 
@@ -1053,7 +1053,7 @@ void FCustomizableObjectCompiler::CompileInternal(UCustomizableObject* Object, c
 		}
 
 #if WITH_EDITORONLY_DATA
-		Object->GetPrivate()->CustomizableObjectPathMap = GenerationContext.CustomizableObjectPathMap;
+		Object->CustomizableObjectPathMap = GenerationContext.CustomizableObjectPathMap;
 #endif
 
 		Object->LODSettings.NumLODsInRoot = GenerationContext.NumLODsInRoot;
@@ -1092,7 +1092,7 @@ void FCustomizableObjectCompiler::CompileInternal(UCustomizableObject* Object, c
 
 		// Lock the object to prevent instance updates while compiling. It's ignored and returns false if it's already locked
 		// Will unlock in the FinishCompilation call.
-		if (!Object->GetPrivate()->IsLocked())
+		if (!Object->IsLocked())
 		{
 			UCustomizableObjectSystem::GetInstance()->LockObject(Object);
 		}
@@ -1245,7 +1245,7 @@ void FCustomizableObjectCompiler::FinishCompilation()
 	}
 	else
 	{
-		CurrentObject->GetPrivate()->GetModel()->GetPrivate()->UnloadRoms();
+		CurrentObject->GetModel()->GetPrivate()->UnloadRoms();
 
 		// when skipping the SaveDerivedData task unlock the object so that instances can be updated
 		UCustomizableObjectSystem::GetInstance()->UnlockObject(CurrentObject);
@@ -1267,7 +1267,7 @@ void FCustomizableObjectCompiler::FinishSavingDerivedData()
 	SaveDDThread.Reset();
 	SaveDDTask.Reset();
 
-	CurrentObject->GetPrivate()->GetModel()->GetPrivate()->UnloadRoms();
+	CurrentObject->GetModel()->GetPrivate()->UnloadRoms();
 	
 	// Unlock the object so that instances can be updated
 	UCustomizableObjectSystem::GetInstance()->UnlockObject(CurrentObject);
@@ -1323,7 +1323,7 @@ void FCustomizableObjectCompiler::ForceFinishBeforeStartCompilation(UCustomizabl
 	CleanCachedReferencers();
 	UpdateArrayGCProtect();
 
-	if (Object && Object->GetPrivate()->IsLocked())
+	if (Object && Object->IsLocked())
 	{
 		UCustomizableObjectSystem::GetInstance()->UnlockObject(Object);
 	}
diff --git a/Engine/Plugins/Experimental/Mutable/Source/CustomizableObjectEditor/Private/MuCOE/CustomizableObjectCompiler.h b/Engine/Plugins/Experimental/Mutable/Source/CustomizableObjectEditor/Private/MuCOE/CustomizableObjectCompiler.h
--- a/Engine/Plugins/Experimental/Mutable/Source/CustomizableObjectEditor/Private/MuCOE/CustomizableObjectCompiler.h
+++ b/Engine/Plugins/Experimental/Mutable/Source/CustomizableObjectEditor/Private/MuCOE/CustomizableObjectCompiler.h
@@ -3,7 +3,6 @@
 #pragma once
 
 #include "AssetRegistry/AssetData.h"
-#include "MuCO/CustomizableObjectPrivate.h"
 #include "MuCOE/CompilationMessageCache.h"
 #include "MuCOE/GenerateMutableSource/GenerateMutableSource.h"
 #include "MuCOE/CustomizableObjectEditorLogger.h"
diff --git a/Engine/Plugins/Experimental/Mutable/Source/CustomizableObjectEditor/Private/MuCOE/CustomizableObjectDetails.cpp b/Engine/Plugins/Experimental/Mutable/Source/CustomizableObjectEditor/Private/MuCOE/CustomizableObjectDetails.cpp
--- a/Engine/Plugins/Experimental/Mutable/Source/CustomizableObjectEditor/Private/MuCOE/CustomizableObjectDetails.cpp
+++ b/Engine/Plugins/Experimental/Mutable/Source/CustomizableObjectEditor/Private/MuCOE/CustomizableObjectDetails.cpp
@@ -6,7 +6,6 @@
 #include "DetailWidgetRow.h"
 #include "IDetailsView.h"
 #include "MuCO/CustomizableObject.h"
-#include "MuCO/CustomizableObjectPrivate.h"
 #include "MuR/Model.h"
 #include "MuR/Parameters.h"
 #include "MuR/Ptr.h"
@@ -97,9 +96,9 @@ void FCustomizableObjectDetails::CustomizeDetails( IDetailLayoutBuilder& DetailB
 	DetailBuilder.HideProperty("PopulationClassTags");
 
 	// Make the tree get automatically updated each time we compile the host CO
-	if (CustomizableObject && !CustomizableObject->GetPrivate()->PostCompileDelegate.IsBoundToObject(this))
+	if (CustomizableObject && !CustomizableObject->PostCompileDelegate.IsBoundToObject(this))
 	{
-		CustomizableObject->GetPrivate()->PostCompileDelegate.AddSP(this, &FCustomizableObjectDetails::UpdateTree);
+		CustomizableObject->PostCompileDelegate.AddSP(this, &FCustomizableObjectDetails::UpdateTree);
 	}
 	
 	// Cache the states defined in the CO
@@ -150,15 +149,15 @@ void FCustomizableObjectDetails::CustomizeDetails( IDetailLayoutBuilder& DetailB
 
 void FCustomizableObjectDetails::UpdateTree()
 {
-	if ( CustomizableObject &&  CustomizableObject->GetPrivate()->GetModel() )
+	if ( CustomizableObject &&  CustomizableObject->GetModel() )
 	{
 		RootTreeItems.SetNumUninitialized(0);
-		const uint32 NumElements = CustomizableObject->GetPrivate()->GetModel()->GetStateCount();
+		const uint32 NumElements = CustomizableObject->GetModel()->GetStateCount();
 		for ( uint32 i=0; i<NumElements; ++i )
 		{
 			TSharedPtr<FStateDetailsNode> SlateDetailsNode = MakeShareable( new FStateDetailsNode );
 
-			SlateDetailsNode->Model = CustomizableObject->GetPrivate()->GetModel();
+			SlateDetailsNode->Model = CustomizableObject->GetModel();
 			SlateDetailsNode->StateIndex = i;
 
 			RootTreeItems.Add( SlateDetailsNode );
diff --git a/Engine/Plugins/Experimental/Mutable/Source/CustomizableObjectEditor/Private/MuCOE/CustomizableObjectEditorFunctionLibrary.cpp b/Engine/Plugins/Experimental/Mutable/Source/CustomizableObjectEditor/Private/MuCOE/CustomizableObjectEditorFunctionLibrary.cpp
--- a/Engine/Plugins/Experimental/Mutable/Source/CustomizableObjectEditor/Private/MuCOE/CustomizableObjectEditorFunctionLibrary.cpp
+++ b/Engine/Plugins/Experimental/Mutable/Source/CustomizableObjectEditor/Private/MuCOE/CustomizableObjectEditorFunctionLibrary.cpp
@@ -17,7 +17,7 @@ ECustomizableObjectCompilationState UCustomizableObjectEditorFunctionLibrary::Co
 	}
 
 	FString ObjectPath = CustomizableObject->GetPathName();
-	if (CustomizableObject->GetPrivate()->IsLocked())
+	if (CustomizableObject->IsLocked())
 	{
 		// Take this if you need a hack:
 		// UCustomizableObjectSystem::GetInstance()->UnlockObject(CustomizableObject);
diff --git a/Engine/Plugins/Experimental/Mutable/Source/CustomizableObjectEditor/Private/MuCOE/GenerateMutableSource/GenerateMutableSource.h b/Engine/Plugins/Experimental/Mutable/Source/CustomizableObjectEditor/Private/MuCOE/GenerateMutableSource/GenerateMutableSource.h
--- a/Engine/Plugins/Experimental/Mutable/Source/CustomizableObjectEditor/Private/MuCOE/GenerateMutableSource/GenerateMutableSource.h
+++ b/Engine/Plugins/Experimental/Mutable/Source/CustomizableObjectEditor/Private/MuCOE/GenerateMutableSource/GenerateMutableSource.h
@@ -7,7 +7,6 @@
 #include "Engine/DataTable.h"
 #include "Engine/SkeletalMesh.h"
 #include "MuCO/CustomizableObject.h"
-#include "MuCO/CustomizableObjectIdentifier.h"
 #include "MuCOE/ExtensionDataCompilerInterface.h"
 #include "MuCOE/Nodes/CustomizableObjectNodeMaterialBase.h"
 #include "MuCOE/Nodes/CustomizableObjectNodeObject.h"
@@ -41,6 +40,7 @@ class UMaterialInterface;
 class UObject;
 class UPhysicsAsset;
 class UTexture2D;
+struct FCustomizableObjectIdPair;
 struct FAnimBpOverridePhysicsAssetsInfo;
 struct FMutableGraphGenerationContext;
 struct FParameterUIData;
diff --git a/Engine/Plugins/Experimental/Mutable/Source/CustomizableObjectEditor/Private/MuCOE/SCustomizableInstanceProperties.cpp b/Engine/Plugins/Experimental/Mutable/Source/CustomizableObjectEditor/Private/MuCOE/SCustomizableInstanceProperties.cpp
--- a/Engine/Plugins/Experimental/Mutable/Source/CustomizableObjectEditor/Private/MuCOE/SCustomizableInstanceProperties.cpp
+++ b/Engine/Plugins/Experimental/Mutable/Source/CustomizableObjectEditor/Private/MuCOE/SCustomizableInstanceProperties.cpp
@@ -762,10 +762,10 @@ void SCustomizableInstanceProperties::AddParameter(int32 ParamIndexInObject)
 				{
 					ValueIndex = i;
 
-					const FString* Identifier = CustomizableObject->GetPrivate()->GroupNodeMap.FindKey(FCustomizableObjectIdPair(ParamName, PossibleValue));
+					const FString* Identifier = CustomizableObject->GroupNodeMap.FindKey(FCustomizableObjectIdPair(ParamName, PossibleValue));
 					if (Identifier)
 					{
-						if (FString* CustomizableObjectPath = CustomizableObject->GetPrivate()->CustomizableObjectPathMap.Find(*Identifier))
+						if (FString* CustomizableObjectPath = CustomizableObject->CustomizableObjectPathMap.Find(*Identifier))
 						{
 							ToolTipText = *CustomizableObjectPath;
 						}
diff --git a/Engine/Plugins/Experimental/Mutable/Source/CustomizableObjectEditor/Private/MuCOE/SCustomizableObjectEditorPerformanceReport.cpp b/Engine/Plugins/Experimental/Mutable/Source/CustomizableObjectEditor/Private/MuCOE/SCustomizableObjectEditorPerformanceReport.cpp
--- a/Engine/Plugins/Experimental/Mutable/Source/CustomizableObjectEditor/Private/MuCOE/SCustomizableObjectEditorPerformanceReport.cpp
+++ b/Engine/Plugins/Experimental/Mutable/Source/CustomizableObjectEditor/Private/MuCOE/SCustomizableObjectEditorPerformanceReport.cpp
@@ -15,7 +15,6 @@
 #include "MuCOE/SCustomizableObjectEditorTextureAnalyzer.h"
 #include "MuCOE/UnrealEditorPortabilityHelpers.h"
 #include "PropertyCustomizationHelpers.h"
-#include "MuCO/CustomizableObjectPrivate.h"
 #include "Widgets/Input/SComboButton.h"
 #include "Widgets/Input/SNumericEntryBox.h"
 #include "Widgets/Layout/SScrollBox.h"
@@ -1344,7 +1343,7 @@ TSharedRef<SWidget> SPerformanceReportWorstCaseRow::GenerateWidgetForColumn(cons
 		{
 			if (const UCustomizableObject* CustomizableObject = ReportInstance->WorstCaseInstance->GetCustomizableObject();
 				CustomizableObject &&
-				!CustomizableObject->GetPrivate()->IsLocked())
+				!CustomizableObject->IsLocked())
 			{
 				const bool bIsDetailParam = ReportInstance->WorstCaseType == UWorstCasePerformanceReportInstance::EPerformanceReportWorstCaseType::SLOWEST_UPDATE_PARAM_IN_RELEVANT_STATE;
 				const FText TextState = ReportInstance->WorstCaseType == UWorstCasePerformanceReportInstance::EPerformanceReportWorstCaseType::SLOWEST_UPDATE_ENTER_STATE ?
diff --git a/Engine/Plugins/Experimental/Mutable/Source/CustomizableObjectEditor/Public/MuCOE/CustomizableObjectInstanceBakingUtils.cpp b/Engine/Plugins/Experimental/Mutable/Source/CustomizableObjectEditor/Public/MuCOE/CustomizableObjectInstanceBakingUtils.cpp
--- a/Engine/Plugins/Experimental/Mutable/Source/CustomizableObjectEditor/Public/MuCOE/CustomizableObjectInstanceBakingUtils.cpp
+++ b/Engine/Plugins/Experimental/Mutable/Source/CustomizableObjectEditor/Public/MuCOE/CustomizableObjectInstanceBakingUtils.cpp
@@ -310,7 +310,7 @@ void BakeCustomizableObjectInstance(
 	check (InstanceCO);
 
 	// Ensure the CO of the COI is accessible 
-	if (!InstanceCO || InstanceCO->GetPrivate()->IsLocked())
+	if (!InstanceCO || InstanceCO->IsLocked())
 	{
 		FCustomizableObjectEditorLogger::CreateLog(
 		LOCTEXT("CustomizableObjectCompilingTryLater_Baking", "Please wait until the Customizable Object is compiled"))
diff --git a/Engine/Plugins/Experimental/Mutable/Source/MutableValidation/Private/MuV/AssetValidator_CustomizableObjects.cpp b/Engine/Plugins/Experimental/Mutable/Source/MutableValidation/Private/MuV/AssetValidator_CustomizableObjects.cpp
--- a/Engine/Plugins/Experimental/Mutable/Source/MutableValidation/Private/MuV/AssetValidator_CustomizableObjects.cpp
+++ b/Engine/Plugins/Experimental/Mutable/Source/MutableValidation/Private/MuV/AssetValidator_CustomizableObjects.cpp
@@ -11,7 +11,6 @@
 #include "Engine/Texture.h"
 #include "Materials/Material.h"
 #include "MuCO/CustomizableObject.h"
-#include "MuCO/CustomizableObjectPrivate.h"
 #include "MuCO/CustomizableObjectSystem.h"
 #include "UObject/NameTypes.h"
 #include "UObject/Object.h"
diff --git a/Engine/Plugins/Experimental/Mutable/Source/MutableValidation/Private/MuV/COIBakingTestCommandlet.cpp b/Engine/Plugins/Experimental/Mutable/Source/MutableValidation/Private/MuV/COIBakingTestCommandlet.cpp
--- a/Engine/Plugins/Experimental/Mutable/Source/MutableValidation/Private/MuV/COIBakingTestCommandlet.cpp
+++ b/Engine/Plugins/Experimental/Mutable/Source/MutableValidation/Private/MuV/COIBakingTestCommandlet.cpp
@@ -8,7 +8,6 @@
 #include "MuCO/CustomizableObjectSystem.h"
 #include "MuCOE/CustomizableObjectInstanceBakingUtils.h"
 #include "HAL/FileManager.h"
-#include "MuCO/CustomizableObjectPrivate.h"
 #include "MuT/UnrealPixelFormatOverride.h"
 
 
diff --git a/Engine/Plugins/Experimental/Mutable/Source/MutableValidation/Private/MuV/CustomizableObjectValidationCommandlet.cpp b/Engine/Plugins/Experimental/Mutable/Source/MutableValidation/Private/MuV/CustomizableObjectValidationCommandlet.cpp
--- a/Engine/Plugins/Experimental/Mutable/Source/MutableValidation/Private/MuV/CustomizableObjectValidationCommandlet.cpp
+++ b/Engine/Plugins/Experimental/Mutable/Source/MutableValidation/Private/MuV/CustomizableObjectValidationCommandlet.cpp
@@ -6,7 +6,6 @@
 #include "AssetRegistry/AssetRegistryModule.h"
 #include "Containers/Array.h"
 #include "MuCO/CustomizableObject.h"
-#include "MuCO/CustomizableObjectPrivate.h"
 #include "MuCO/CustomizableObjectSystem.h"
 #include "MuR/Model.h"
 
@@ -103,7 +102,7 @@ int32 UCustomizableObjectValidationCommandlet::Main(const FString& Params)
 		UE_LOG(LogMutable,Display,TEXT("Customizable Object was compiled succesfully."));
 		
 		// GHet the total size of the streaming data of the model ---------------------------------------------- //
-		const TSharedPtr<const mu::Model> MutableModel = ToTestCustomizableObject->GetPrivate()->GetModel();
+		const TSharedPtr<const mu::Model> MutableModel = ToTestCustomizableObject->GetModel();
 		check (MutableModel);
 
 		// Roms ---------------------- //
