[mutable] Bulk instance updater implementation
- Test commandlet that updates all instances found in a given path.
- Creted supporting classes/files to handle compilation and update of COs and COIs
- Refactored other testing commandlets to make use of the shared methods and objects in the "ValidationUtils" file

diff --git a/Engine/Plugins/Experimental/Mutable/Source/MutableValidation/Private/MuV/COIBakingTestCommandlet.cpp b/Engine/Plugins/Experimental/Mutable/Source/MutableValidation/Private/MuV/COIBakingTestCommandlet.cpp
--- a/Engine/Plugins/Experimental/Mutable/Source/MutableValidation/Private/MuV/COIBakingTestCommandlet.cpp
+++ b/Engine/Plugins/Experimental/Mutable/Source/MutableValidation/Private/MuV/COIBakingTestCommandlet.cpp
@@ -2,7 +2,9 @@
 
 #include "../../Public/MuV/COIBakingTestCommandlet.h"
 
-#include "AssetRegistry/AssetRegistryModule.h"
+#include "ValidationUtils.h"
+#include "AssetRegistry/AssetData.h"
+#include "HAL/FileManager.h"
 #include "MuCO/CustomizableObject.h"
 #include "MuCO/CustomizableObjectInstance.h"
 #include "MuCO/CustomizableObjectSystem.h"
@@ -64,42 +66,20 @@ int32 UCOIBakingTestCommandlet::Main(const FString& Params)
 	}
 
 	// Perform a blocking search to ensure all assets used by mutable are reachable using the AssetRegistry
-	{
-		UE_LOG(LogMutable,Display,TEXT("Searching all assets (this will take some time)..."));
-		FAssetRegistryModule& AssetRegistryModule = FModuleManager::LoadModuleChecked<FAssetRegistryModule>(AssetRegistryConstants::ModuleName);
-		AssetRegistryModule.Get().SearchAllAssets(true /* bSynchronousSearch */);
-	}
+	PrepareAssetRegistry();
 	
 	// Compile it's CO (using current config)
 	UCustomizableObject* InstanceCustomizableObject = TargetInstance->GetCustomizableObject();
-	UCustomizableObjectSystem* System =  UCustomizableObjectSystem::GetInstanceChecked();
-	check (InstanceCustomizableObject);
+	if (!InstanceCustomizableObject)
 	{
-		UE_LOG(LogMutable,Display,TEXT("Compiling Customizable Object..."));
-    	
-		// Request a compiler to be able to locate the root and to compile it
-		const TUniquePtr<FCustomizableObjectCompilerBase> Compiler =
-			TUniquePtr<FCustomizableObjectCompilerBase>(System->GetNewCompiler());
-
-		// Compile (sync)
-		// Compile the CO with the provided compilation options
-		// Run Sync compilation -> Warning : Potentially long operation -------------
-		Compiler->Compile(*InstanceCustomizableObject, InstanceCustomizableObject->CompileOptions, false);
-		// --------------------------------------------------------------------------
-		
-		// Get the compilation result
-		const ECustomizableObjectCompilationState CompilationEndResult = Compiler->GetCompilationState();
-		check(CompilationEndResult != ECustomizableObjectCompilationState::None);
-		check(CompilationEndResult != ECustomizableObjectCompilationState::InProgress);
-
-		UE_LOG(LogMutable, Display, TEXT("CO compilation completed with end status : %s "), *UEnum::GetValueAsString(CompilationEndResult));
-		
-		const bool bWasCompilationSuccessful = CompilationEndResult == ECustomizableObjectCompilationState::Completed;
-		if (!bWasCompilationSuccessful)
-		{
-			UE_LOG(LogMutable,Error,TEXT("Failed to compile the target CO."));
-			return 1;
-		}
+		UE_LOG(LogMutable,Error,TEXT("The instance %s does not have a CO to compile : Exitting commandlet."), *TargetInstance->GetName());
+		return 1;
+	}
+	
+	if (!CompileCustomizableObject(InstanceCustomizableObject))
+	{
+		UE_LOG(LogMutable,Error,TEXT("Failed to compile the target CO. Exitting commandlet."));
+		return 1;
 	}
 	
 	// Update the instance
@@ -131,7 +111,7 @@ int32 UCOIBakingTestCommandlet::Main(const FString& Params)
 		// Check the end status of the instance update
 		if (!bWasInstanceUpdateSuccessful)
 		{
-			UE_LOG(LogMutable,Error,TEXT("Failed to succesfully update the target COI."));
+			UE_LOG(LogMutable,Error,TEXT("Failed to succesfully update the target COI. Exitting commandlet."));
 			return 1;
 		}
 	}
diff --git a/Engine/Plugins/Experimental/Mutable/Source/MutableValidation/Private/MuV/COIBulkUpdateTestCommandlet.cpp b/Engine/Plugins/Experimental/Mutable/Source/MutableValidation/Private/MuV/COIBulkUpdateTestCommandlet.cpp
new file mode 100644
--- /dev/null
+++ b/Engine/Plugins/Experimental/Mutable/Source/MutableValidation/Private/MuV/COIBulkUpdateTestCommandlet.cpp
@@ -0,0 +1,151 @@
+﻿// Copyright Epic Games, Inc. All Rights Reserved.
+
+#include "MuV/COIBulkUpdateTestCommandlet.h"
+
+#include "ValidationUtils.h"
+#include "AssetRegistry/AssetData.h"
+#include "AssetRegistry/AssetRegistryModule.h"
+#include "MuCO/CustomizableObject.h"
+#include "MuCO/CustomizableObjectInstance.h"
+
+
+int32 UCOIBulkUpdateTestCommandlet::Main(const FString& Params)
+{
+	// Get the path where to look for the Customizable Object Instances we want to validate
+	FString InstancesPackagePath;
+	{
+		if (!FParse::Value(*Params, TEXT("InstancesPackagePath="), InstancesPackagePath))
+		{
+			UE_LOG(LogMutable,Error,TEXT("Failed to parse path where to find the Customizable Object Instances to update : %s"),*Params)
+			return 1;
+		}
+
+		if (InstancesPackagePath.IsEmpty())
+		{
+			UE_LOG(LogMutable,Error,TEXT("The path to scan can not be empty"))
+			return 1;
+		}
+	}
+
+	UE_LOG(LogMutable, Log, TEXT("(string) instance_search_path : %s "), *InstancesPackagePath);
+	
+	// Load the asset registry system so we can proceed without issues
+	PrepareAssetRegistry();
+	
+	// Cache all UAssets (find a way to not have them in memory yet, not until we need them)
+	TArray<FAssetData> FoundAssetData;
+	{
+		FARFilter Filter;
+		Filter.ClassPaths.Add(UCustomizableObjectInstance::StaticClass()->GetClassPathName());
+		Filter.PackagePaths.Add( FName(*InstancesPackagePath) );
+		Filter.bRecursivePaths = true;
+		
+		const FAssetRegistryModule& AssetRegistryModule = FModuleManager::LoadModuleChecked<FAssetRegistryModule>(AssetRegistryConstants::ModuleName);
+		
+		UE_LOG(LogMutable,Display,TEXT("Searching for COIs to test at path : %s ."), *InstancesPackagePath);
+		AssetRegistryModule.Get().GetAssets(Filter, FoundAssetData);
+		UE_LOG(LogMutable,Display,TEXT("Search of instances completed. Found %i mutable instances to validate."), FoundAssetData.Num());
+	}
+
+	// Early exit if no instances could be found for testing.
+	if (FoundAssetData.IsEmpty())
+	{
+		UE_LOG(LogMutable,Error,TEXT("Aborting Bulk Instance Update Test: No assets could be found at  the provided package path"));
+		return 1;
+	}
+
+	// Get all the CO's that need compilation before proceeding
+	TMap<UCustomizableObject*, TSet<UCustomizableObjectInstance*>>  MutableResources;
+	for (FAssetData& Data : FoundAssetData)
+	{
+		UObject* LoadedAsset =  Data.GetAsset();
+		check (LoadedAsset);
+		UCustomizableObjectInstance* LoadedInstance =  Cast<UCustomizableObjectInstance>(LoadedAsset);
+		check (LoadedInstance)
+
+		// Get the COI CO and cache it for later compilation
+		UCustomizableObject* InstanceCO = LoadedInstance->GetCustomizableObject();
+		if (!InstanceCO)
+		{
+			UE_LOG(LogMutable,Error,TEXT("The instance %s does not have a CO. This instance will not get tested."), *LoadedInstance->GetName());
+			continue;
+		}
+		
+		// Add/update an entry with the new instance for a given CO. 
+		TSet<UCustomizableObjectInstance*>& Entry = MutableResources.FindOrAdd(InstanceCO);
+		Entry.Add(LoadedInstance);
+	}
+
+	// At this point it is safe to assume that all keys are valid COs and all instances are also valid.
+	
+	// Report the objects found (COs and COIs) -------------------------------
+	
+	const uint32 TotalAmountOfCustomizableObjects = MutableResources.Num();
+	UE_LOG(LogMutable,Display,TEXT("Customizable Objects to compile : %u"), TotalAmountOfCustomizableObjects);
+	UE_LOG(LogMutable, Log, TEXT("(int) customizable_objects_count : %u "), TotalAmountOfCustomizableObjects);
+
+	const uint32 TotalAmountOfInstances = FoundAssetData.Num();
+	UE_LOG(LogMutable,Display,TEXT("Customizable Object Instances to update : %u"), TotalAmountOfInstances);
+	UE_LOG(LogMutable, Log, TEXT("(int) customizable_object_instances_count : %u "), TotalAmountOfInstances);
+	
+	// Report the amount of instances for each of the COs found as parents of the instances in the target path
+	for (TTuple<UCustomizableObject*, TSet<UCustomizableObjectInstance*>>& MutableResourceTuple : MutableResources)
+	{
+		const UCustomizableObject* CustomizableObjectToCompile = MutableResourceTuple.Key;
+		check (CustomizableObjectToCompile);
+		const FString CustomizableObjectName = CustomizableObjectToCompile->GetName();
+		
+		const uint32 COInstancesCount = MutableResourceTuple.Value.Num();
+		UE_LOG(LogMutable,Display,TEXT("The CO \"%s\" has in total \"%u\" instances."), *CustomizableObjectName, COInstancesCount );
+		UE_LOG(LogMutable, Log, TEXT("(int) %s_instance_count : %u "), *CustomizableObjectName, COInstancesCount);
+		
+		// print the name of the instances+
+		uint32 InstanceIndex = 0;
+		for (const UCustomizableObjectInstance* Instance : MutableResourceTuple.Value)
+		{
+			check(Instance);
+			const FString CustomizableObjectInstanceName = Instance->GetName();
+			UE_LOG(LogMutable, Log, TEXT("%u : \"%s\""), InstanceIndex++, *CustomizableObjectInstanceName);
+		}
+	}
+	
+	// ------ Execution of the actual mutable operations ------ 
+	
+	// Create the updater object so we can later call for the update a target instance
+	InstanceUpdater = NewObject<UCOIUpdater>();
+	
+	// Compile all found COs one by one
+	uint32 CurrentInstanceIndex = 1;
+	for (TTuple<UCustomizableObject*, TSet<UCustomizableObjectInstance*>>& MutableResourceTuple : MutableResources)
+	{
+		UCustomizableObject* CustomizableObjectToCompile = MutableResourceTuple.Key;
+		check (CustomizableObjectToCompile);
+		
+		// Compile the current CO object
+		const FString CustomizableObjectName = CustomizableObjectToCompile->GetName();
+		if (!CompileCustomizableObject(CustomizableObjectToCompile))
+		{
+			UE_LOG(LogMutable,Error,TEXT("The CO %s could not be compiled succesfully. Skipping the update of all COIs that use it."), *CustomizableObjectName )
+			continue;
+		}
+
+		// Now that CO has been compiled proceed with the update of the instances that use it
+		const uint32 CustomizableObjectInstanceCount = MutableResourceTuple.Value.Num();
+		UE_LOG(LogMutable,Display,TEXT("Starting update of the \"%u\" instances with CO : \"%s\"."), CustomizableObjectInstanceCount, *CustomizableObjectName);
+		
+		// Iterate over all the COIs of the CO and update them
+		for (UCustomizableObjectInstance* Instance : MutableResourceTuple.Value)
+		{
+			UE_LOG(LogMutable,Display,TEXT("\t( %u / %u ) Processing instance : \"%s\" ."),CurrentInstanceIndex++, TotalAmountOfInstances ,*Instance->GetName());
+			
+			// Update each one of the instances and notify if the update failed in any manner
+			if (InstanceUpdater != nullptr)
+			{
+				InstanceUpdater->UpdateInstance(Instance);
+			}
+		}
+	}
+
+	UE_LOG(LogMutable,Display,TEXT("Mutable commandlet finished."));
+	return 0;
+}
diff --git a/Engine/Plugins/Experimental/Mutable/Source/MutableValidation/Private/MuV/CustomizableObjectValidationCommandlet.cpp b/Engine/Plugins/Experimental/Mutable/Source/MutableValidation/Private/MuV/CustomizableObjectValidationCommandlet.cpp
--- a/Engine/Plugins/Experimental/Mutable/Source/MutableValidation/Private/MuV/CustomizableObjectValidationCommandlet.cpp
+++ b/Engine/Plugins/Experimental/Mutable/Source/MutableValidation/Private/MuV/CustomizableObjectValidationCommandlet.cpp
@@ -2,8 +2,8 @@
 
 #include "MuV/CustomizableObjectValidationCommandlet.h"
 
+#include "ValidationUtils.h"
 #include "Components/SkeletalMeshComponent.h"
-#include "AssetRegistry/AssetRegistryModule.h"
 #include "Containers/Array.h"
 #include "MuCO/CustomizableObject.h"
 #include "MuCO/CustomizableObjectSystem.h"
@@ -13,7 +13,7 @@
 int32 UCustomizableObjectValidationCommandlet::Main(const FString& Params)
 {
 	// Execution arguments for commandlet from IDE
-	// $(LocalDebuggerCommandArguments) -run=CustomizableObjectValidation -CustomizableObject=(PathToCO)
+	// -run=CustomizableObjectValidation -CustomizableObject=(PathToCO)
 	
 	// Get the package name of the CO to test
 	FString CustomizableObjectAssetPath = "";
@@ -29,6 +29,7 @@ int32 UCustomizableObjectValidationCommandlet::Main(const FString& Params)
 	{
 		UE_LOG(LogMutable,Display,TEXT("Instance generation count not specified. Using default value : %u"),InstancesToGenerate);
 	}
+	UE_LOG(LogMutable, Log,TEXT("(int) instances_to_generate_count : %u "), InstancesToGenerate);
 	
 	// TODO: Detect target compilation platform based on argument value (string)
 	
@@ -49,51 +50,18 @@ int32 UCustomizableObjectValidationCommandlet::Main(const FString& Params)
 	}
 	
 	// Perform a blocking search to ensure all assets used by mutable are reachable using the AssetRegistry
-	{
-		UE_LOG(LogMutable,Display,TEXT("Searching all assets (this will take some time)..."));
-		FAssetRegistryModule& AssetRegistryModule = FModuleManager::LoadModuleChecked<FAssetRegistryModule>(AssetRegistryConstants::ModuleName);
-		AssetRegistryModule.Get().SearchAllAssets(true /* bSynchronousSearch */);
-	}
+	PrepareAssetRegistry();
 	
 	// Compile the Customizable Object ------------------------------------------------------------------------------ //
 	bool bWasCoCompilationSuccessful = false;
 	{
-		UE_LOG(LogMutable,Display,TEXT("Compiling Customizable Object..."));
-    	
-    	// Request a compiler to be able to locate the root and to compile it
-    	const TUniquePtr<FCustomizableObjectCompilerBase> Compiler =
-    		TUniquePtr<FCustomizableObjectCompilerBase>(UCustomizableObjectSystem::GetInstanceChecked()->GetNewCompiler());
-		
 		// Override some configurations that may have been changed by the user
 		FCompilationOptions CompilationOptions = ToTestCustomizableObject->CompileOptions;
 		CompilationOptions.bSilentCompilation = false;
 		CompilationOptions.OptimizationLevel = 2;			// Set the optimization level to the max
 		CompilationOptions.TextureCompression = ECustomizableObjectTextureCompression::Fast;
-		
-		// TODO: Add logs for the other relevant configs of the model being compiled
-		// Print MTU parseable logs
-		UE_LOG(LogMutable, Log, TEXT("(int) model_optimization_level : %d "), CompilationOptions.OptimizationLevel);
-		UE_LOG(LogMutable, Log, TEXT("(string) model_texture_compression : %s "), *UEnum::GetValueAsString(CompilationOptions.TextureCompression));
-		UE_LOG(LogMutable, Log, TEXT("(string) model_disk_compilation : %s "), CompilationOptions.bUseDiskCompilation ? TEXT("true") : TEXT("false"));
 
-    	// Compile the CO with the provided compilation options
-    	// Run Sync compilation -> Warning : Potentially long operation -------------
-		const double CompilationStartSeconds = FPlatformTime::Seconds();
-		{
-			Compiler->Compile(*ToTestCustomizableObject, CompilationOptions, false);
-		}
-		const double CompilationEndSeconds = FPlatformTime::Seconds() - CompilationStartSeconds;
-		UE_LOG(LogMutable, Log, TEXT("(double) model_compile_time_ms : %f "), CompilationEndSeconds * 1000);
-		UE_LOG(LogMutable, Display, TEXT("The compilation of the %s model took %f seconds."), *ToTestCustomizableObject->GetName(), CompilationEndSeconds);
-
-    	// --------------------------------------------------------------------------
-		
-    	// Get the compilation result
-    	const ECustomizableObjectCompilationState CompilationEndResult = Compiler->GetCompilationState();
-    	check(CompilationEndResult != ECustomizableObjectCompilationState::None);
-    	check(CompilationEndResult != ECustomizableObjectCompilationState::InProgress);
-    	
-    	bWasCoCompilationSuccessful = CompilationEndResult == ECustomizableObjectCompilationState::Completed;
+		bWasCoCompilationSuccessful = CompileCustomizableObject(ToTestCustomizableObject, &CompilationOptions);
 	}
 	// -------------------------------------------------------------------------------------------------------------- //
 	
@@ -165,73 +133,22 @@ int32 UCustomizableObjectValidationCommandlet::Main(const FString& Params)
 		
 		// Update the instances generated --------------------------------------------------------------------------- //
 		UE_LOG(LogMutable,Display,TEXT("Updating generated instances..."));
-		const double InstanceUpdateStartSeconds = FPlatformTime::Seconds();
+		bool bInstanceFailedUpdate = false;
+		const double InstancesUpdateStartSeconds = FPlatformTime::Seconds();
 		{
-            // Now update the instances one by one
-            while (!InstancesToProcess.IsEmpty() || InstanceBeingUpdated)
-            {
-            	// Tick the engine
-            	CommandletHelpers::TickEngine();
-    
-            	// Stop if exit was requested
-            	if (IsEngineExitRequested())
-            	{
-            		break;
-            	}
-            
-            	// Wait until current instance turns invalid
-            	if (InstanceBeingUpdated)
-            	{
-            		// Wait until all MIPs gets streamed
-            		if (!ComponentsBeingUpdated.IsEmpty())
-            		{
-            			bool bFullyStreamed = true;
-            			for (auto It = ComponentsBeingUpdated.CreateIterator(); It && bFullyStreamed; ++It)
-            			{
-            				TObjectPtr<USkeletalMeshComponent>& ComponentBeingUpdated = *It;
-            		
-            				FStreamingTextureLevelContext LevelContext(EMaterialQualityLevel::Num, ComponentBeingUpdated);
-            				TArray<FStreamingRenderAssetPrimitiveInfo> RenderAssetInfoArray;
-            				ComponentBeingUpdated->GetStreamingRenderAssetInfo(LevelContext, RenderAssetInfoArray);
-
-            				for (auto ItAsset = RenderAssetInfoArray.CreateIterator(); ItAsset && bFullyStreamed; ++ItAsset)
-            				{
-            					bFullyStreamed = ItAsset->RenderAsset->IsFullyStreamedIn();
-            				}
-            			}
-
-            			if (bFullyStreamed)
-            			{
-            				UE_LOG(LogMutable,Display,TEXT("Instance %s finished streaming all MIPs."), *InstanceBeingUpdated->GetName());
-            				ComponentsBeingUpdated.Reset();
-            				InstanceBeingUpdated = nullptr;
-            			}
-            		}
-            		
-            		continue;
-            	}
-           
-            	// We may have already updated the instance so ensure we have more instances to work with
-            	if (InstancesToProcess.IsEmpty())
-            	{
-            		continue;
-            	}
-            	
-            	InstanceBeingUpdated = InstancesToProcess[0];
-            	InstancesToProcess.RemoveAt(0);
-            	if (InstanceBeingUpdated)
-            	{
-            		UE_LOG(LogMutable,Display,TEXT("Invoking update for %s instance."),*InstanceBeingUpdated->GetName());
-            		// Instance update delegate
-            		FInstanceUpdateDelegate InstanceUpdateDelegate;
-            		InstanceUpdateDelegate.BindDynamic(this, &UCustomizableObjectValidationCommandlet::OnInstanceUpdate);
-            		InstanceBeingUpdated->UpdateSkeletalMeshAsyncResult(InstanceUpdateDelegate);
-            	}
-            }
-		}	
+			InstanceUpdater = NewObject<UCOIUpdater>();
+			for (UCustomizableObjectInstance* InstanceToUpdate : InstancesToProcess)
+			{
+				if (InstanceUpdater && !InstanceUpdater->UpdateInstance(InstanceToUpdate))
+				{
+					bInstanceFailedUpdate = true;
+				}
+			}
+		}
+		const double InstancesUpdateEndSeconds = FPlatformTime::Seconds();
 		
 		// Notify and log time required by the instances to get updated
-		const double CombinedInstanceUpdateSeconds = FPlatformTime::Seconds() - InstanceUpdateStartSeconds;
+		const double CombinedInstanceUpdateSeconds = InstancesUpdateEndSeconds - InstancesUpdateStartSeconds;
 		UE_LOG(LogMutable, Log,TEXT("(double) combined_update_time_ms : %f "), CombinedInstanceUpdateSeconds * 1000);
 
 		check(InstancesToGenerate > 0);
@@ -254,7 +171,7 @@ int32 UCustomizableObjectValidationCommandlet::Main(const FString& Params)
 	}
 	else
 	{
-		UE_LOG(LogMutable,Error,TEXT("The compilation of the Customizable object was not succesfull."));
+		UE_LOG(LogMutable,Error,TEXT("The compilation of the Customizable object was not succesfull : No instances will be generated."));
 	}
 	
 	
@@ -262,47 +179,3 @@ int32 UCustomizableObjectValidationCommandlet::Main(const FString& Params)
 	UE_LOG(LogMutable,Display,TEXT("Mutable commandlet finished."));
 	return 0;
 }
-
-
-void UCustomizableObjectValidationCommandlet::OnInstanceUpdate(const FUpdateContext& Result)
-{
-	const FString InstanceName = InstanceBeingUpdated->GetName();
-	const EUpdateResult InstanceUpdateResult = Result.UpdateResult;
-	if (UCustomizableObjectSystem::IsUpdateResultValid(InstanceUpdateResult))
-	{
-		UE_LOG(LogMutable,Display,TEXT("Instance %s finished update succesfully."),*InstanceName);
-
-		// Request load all MIPs
-		UE_LOG(LogMutable,Display,TEXT("Instance %s rquesting streaming all MIPs."), *InstanceBeingUpdated->GetName());
-
-		check(ComponentsBeingUpdated.IsEmpty());
-		for (int32 Index = 0; Index < InstanceBeingUpdated->GetNumComponents(); ++Index)
-		{
-			USkeletalMeshComponent* SkeletalComponent = NewObject<USkeletalMeshComponent>();
-			SkeletalComponent->SetSkeletalMesh(InstanceBeingUpdated->GetSkeletalMesh(Index));
-            
-			ComponentsBeingUpdated.Add(SkeletalComponent);            			
-		}
-		
-		for (TObjectPtr<USkeletalMeshComponent>& ComponentBeingUpdated : ComponentsBeingUpdated)
-		{
-			FStreamingTextureLevelContext LevelContext(EMaterialQualityLevel::Num, ComponentBeingUpdated);
-			TArray<FStreamingRenderAssetPrimitiveInfo> RenderAssetInfoArray;
-			ComponentBeingUpdated->GetStreamingRenderAssetInfo(LevelContext, RenderAssetInfoArray);
-
-			for (const FStreamingRenderAssetPrimitiveInfo& Info : RenderAssetInfoArray)
-			{
-				Info.RenderAsset->StreamIn(MAX_int32, true);
-			}
-		}
-	}
-	else
-	{
-		const FString OutputStatus = UEnum::GetValueAsString(Result.UpdateResult);
-		UE_LOG(LogMutable,Error,TEXT("Instance %s finished update with anomalous state : %s."), *InstanceName, *OutputStatus);
-		bInstanceFailedUpdate = true;
-
-		InstanceBeingUpdated = nullptr;
-	}	
-}
-
diff --git a/Engine/Plugins/Experimental/Mutable/Source/MutableValidation/Private/MuV/ValidationUtils.cpp b/Engine/Plugins/Experimental/Mutable/Source/MutableValidation/Private/MuV/ValidationUtils.cpp
new file mode 100644
--- /dev/null
+++ b/Engine/Plugins/Experimental/Mutable/Source/MutableValidation/Private/MuV/ValidationUtils.cpp
@@ -0,0 +1,208 @@
+﻿// Copyright Epic Games, Inc. All Rights Reserved.
+
+#include "ValidationUtils.h"
+
+#include "AssetRegistry/AssetRegistryModule.h"
+#include "Commandlets/Commandlet.h"
+#include "Components/SkeletalMeshComponent.h"
+#include "MuCO/CustomizableObject.h"
+#include "MuCO/CustomizableObjectInstance.h"
+#include "MuCO/CustomizableObjectPrivate.h"
+#include "MuCO/CustomizableObjectSystem.h"
+
+void PrepareAssetRegistry()
+{
+	FAssetRegistryModule& AssetRegistryModule = FModuleManager::LoadModuleChecked<FAssetRegistryModule>(AssetRegistryConstants::ModuleName);
+	if (AssetRegistryModule.Get().IsSearchAllAssets())
+	{
+		UE_LOG(LogMutable,Display,TEXT("Asset search already in progress. Waitting for it to complete."));
+		AssetRegistryModule.Get().WaitForCompletion();
+	}
+	else
+	{
+		UE_LOG(LogMutable,Display,TEXT("Searching all assets (this will take some time)..."));
+		AssetRegistryModule.Get().SearchAllAssets(true /* bSynchronousSearch */);
+	}
+
+	UE_LOG(LogMutable,Display,TEXT("Asset searching completed!"));
+}
+
+
+bool CompileCustomizableObject(UCustomizableObject* InCustomizableObject, const FCompilationOptions* InCompilationOptionsOverride  /* nullptr */)
+{
+	if (!InCustomizableObject)
+	{
+		UE_LOG(LogMutable, Error, TEXT("Unable to compile CO : CO pointer is null!"));
+		return false;
+	}
+	
+	// Override the compilation options if an override has been provided by the user
+	FCompilationOptions CompilationOptions;
+	const bool bOverrideCompilationOptions = InCompilationOptionsOverride != nullptr;
+	if (bOverrideCompilationOptions)
+	{
+		CompilationOptions = *InCompilationOptionsOverride;
+		UE_LOG(LogMutable,Display,TEXT("CO Compilation options overriden by the user defined ones."));
+	}
+	else
+	{
+		CompilationOptions = InCustomizableObject->CompileOptions;
+		UE_LOG(LogMutable,Display,TEXT("Compiling CO using it's own compilation options."));
+	}
+
+	UE_LOG(LogMutable, Log, TEXT("(string) model_compile_options_overriden : %s "), bOverrideCompilationOptions ? TEXT("true") : TEXT("false"));
+	
+	// TODO: Add logs for the other relevant configs of the model being compiled
+	// Print MTU parseable logs
+	UE_LOG(LogMutable, Log, TEXT("(int) model_optimization_level : %d "), CompilationOptions.OptimizationLevel);
+	UE_LOG(LogMutable, Log, TEXT("(string) model_texture_compression : %s "), *UEnum::GetValueAsString(CompilationOptions.TextureCompression));
+	UE_LOG(LogMutable, Log, TEXT("(string) model_disk_compilation : %s "), CompilationOptions.bUseDiskCompilation ? TEXT("true") : TEXT("false"));
+
+	// Request a compiler to be able to locate the root and to compile it
+	const TUniquePtr<FCustomizableObjectCompilerBase> Compiler =
+		TUniquePtr<FCustomizableObjectCompilerBase>(UCustomizableObjectSystem::GetInstanceChecked()->GetNewCompiler());
+	
+    // Compile the CO with the selected compilation options
+    // Run Sync compilation -> Warning : Potentially long operation -------------
+	UE_LOG(LogMutable,Display,TEXT("Compiling Customizable Object..."));
+	const double CompilationStartSeconds = FPlatformTime::Seconds();
+	{
+		Compiler->Compile(*InCustomizableObject, CompilationOptions, false);
+	}
+	const double CompilationEndSeconds = FPlatformTime::Seconds() - CompilationStartSeconds;
+    // --------------------------------------------------------------------------
+	UE_LOG(LogMutable,Display,TEXT("Compilation of CO completed!"));
+
+	UE_LOG(LogMutable, Log, TEXT("(double) model_compile_time_ms : %f "), CompilationEndSeconds * 1000);
+	UE_LOG(LogMutable, Display, TEXT("The compilation of the %s model took %f seconds."), *InCustomizableObject->GetName(), CompilationEndSeconds);
+
+	
+    // Get the compilation result
+    const ECustomizableObjectCompilationState CompilationEndResult = Compiler->GetCompilationState();
+	UE_LOG(LogMutable, Log, TEXT("(string) model_compile_end_state : %s "), *UEnum::GetValueAsString(CompilationEndResult));
+	
+    const bool bWasCoCompilationSuccessful = CompilationEndResult == ECustomizableObjectCompilationState::Completed;
+    if (bWasCoCompilationSuccessful)
+    {
+    	UE_LOG(LogMutable, Display, TEXT("The compilation of the %s model was succesfull."), *InCustomizableObject->GetName());
+    }
+    else
+    {
+    	UE_LOG(LogMutable, Error, TEXT("The compilation of the %s model failed."), *InCustomizableObject->GetName());
+    }
+    
+	return bWasCoCompilationSuccessful;
+}
+
+
+
+bool UCOIUpdater::UpdateInstance(UCustomizableObjectInstance* InInstance)
+{
+	check (InInstance);
+	check (ComponentsBeingUpdated.IsEmpty());
+	
+	// Cache the instance being updated for reference once in the update end callback
+	Instance = InInstance;
+	
+	// Schedule the update of the COI
+	{
+		UE_LOG(LogMutable,Display,TEXT("Invoking update for %s instance."),*InInstance->GetName());
+
+		// Instance update delegate	
+		FInstanceUpdateDelegate InstanceUpdateDelegate;
+		InstanceUpdateDelegate.BindDynamic(this, &UCOIUpdater::OnInstanceUpdateResult);
+
+		bIsInstanceBeingUpdated = true;
+		InInstance->UpdateSkeletalMeshAsyncResult(InstanceUpdateDelegate,true, true);
+	}
+	
+	
+	// Wait until the update has been completed and the mips streamed
+	while (bIsInstanceBeingUpdated)
+	{
+		// Tick the engine
+		CommandletHelpers::TickEngine();
+
+		// Stop if exit was requested
+		if (IsEngineExitRequested())
+		{
+			break;
+		}
+
+		// Wait until all MIPs gets streamed 
+		if (!ComponentsBeingUpdated.IsEmpty())
+		{
+			bool bFullyStreamed = true;
+			for (auto It = ComponentsBeingUpdated.CreateIterator(); It && bFullyStreamed; ++It)
+			{
+				TObjectPtr<USkeletalMeshComponent>& ComponentBeingUpdated = *It;
+            		
+				FStreamingTextureLevelContext LevelContext(EMaterialQualityLevel::Num, ComponentBeingUpdated);
+				TArray<FStreamingRenderAssetPrimitiveInfo> RenderAssetInfoArray;
+				ComponentBeingUpdated->GetStreamingRenderAssetInfo(LevelContext, RenderAssetInfoArray);
+
+				for (auto ItAsset = RenderAssetInfoArray.CreateIterator(); ItAsset && bFullyStreamed; ++ItAsset)
+				{
+					bFullyStreamed = ItAsset->RenderAsset->IsFullyStreamedIn();
+				}
+			}
+
+			if (bFullyStreamed)
+			{
+				UE_LOG(LogMutable,Display,TEXT("Instance %s finished streaming all MIPs."), *InInstance->GetName());
+				ComponentsBeingUpdated.Reset();
+				
+				bIsInstanceBeingUpdated = false;		// Exit the while loop
+			}
+		}
+	}
+	
+	// Return true if the update was successful and false if not
+	return !bInstanceFailedUpdate;
+}
+
+
+void UCOIUpdater::OnInstanceUpdateResult(const FUpdateContext& Result)
+{
+	const FString InstanceName = Instance->GetName();
+	
+	if (UCustomizableObjectSystem::IsUpdateResultValid(Result.UpdateResult))
+	{
+		UE_LOG(LogMutable,Display,TEXT("Instance %s finished update succesfully."),*InstanceName);
+		bInstanceFailedUpdate = false;
+		
+		// Request load all MIPs
+		UE_LOG(LogMutable,Display,TEXT("Instance %s rquesting streaming all MIPs."), *InstanceName);
+
+		check(ComponentsBeingUpdated.IsEmpty());
+		for (int32 Index = 0; Index < Instance->GetNumComponents(); ++Index)
+		{
+			USkeletalMeshComponent* SkeletalComponent = NewObject<USkeletalMeshComponent>();
+			SkeletalComponent->SetSkeletalMesh(Instance->GetSkeletalMesh(Index));
+			
+			ComponentsBeingUpdated.Add(SkeletalComponent);            			
+		}
+
+		// Request the streaming in of all the components affected by the update
+		for (TObjectPtr<USkeletalMeshComponent>& ComponentBeingUpdated : ComponentsBeingUpdated)
+		{
+			FStreamingTextureLevelContext LevelContext(EMaterialQualityLevel::Num, ComponentBeingUpdated);
+			TArray<FStreamingRenderAssetPrimitiveInfo> RenderAssetInfoArray;
+			ComponentBeingUpdated->GetStreamingRenderAssetInfo(LevelContext, RenderAssetInfoArray);
+
+			for (const FStreamingRenderAssetPrimitiveInfo& Info : RenderAssetInfoArray)
+			{
+				Info.RenderAsset->StreamIn(MAX_int32, true);
+			}
+		}
+	}
+	else
+	{
+		const FString OutputStatus = UEnum::GetValueAsString(Result.UpdateResult);
+		UE_LOG(LogMutable,Error,TEXT("Instance %s finished update with anomalous state : %s."), *InstanceName, *OutputStatus);
+		bInstanceFailedUpdate = true;
+		
+		// Tell the system the instance finished it's update so we can continue the execution without waitting for the mips to stream in
+		bIsInstanceBeingUpdated = false;
+	}	
+}
\ No newline at end of file
diff --git a/Engine/Plugins/Experimental/Mutable/Source/MutableValidation/Private/MuV/ValidationUtils.h b/Engine/Plugins/Experimental/Mutable/Source/MutableValidation/Private/MuV/ValidationUtils.h
new file mode 100644
--- /dev/null
+++ b/Engine/Plugins/Experimental/Mutable/Source/MutableValidation/Private/MuV/ValidationUtils.h
@@ -0,0 +1,71 @@
+﻿// Copyright Epic Games, Inc. All Rights Reserved.
+#pragma once
+
+#include "MuCO/CustomizableObjectInstance.h"
+#include "UObject/Object.h"
+
+#include "ValidationUtils.generated.h"
+
+// Forward declarations
+class UCustomizableObject;
+class UCustomizableObjectInstance;
+class USkeletalMeshComponent;
+struct FCompilationOptions;
+
+
+/**
+ * Prepare the asset registry so we can later use it to search assets. It is required by Mutable to compile.
+ */
+void PrepareAssetRegistry();
+
+
+/**
+ * Compiles a CO synchronously
+ * @param InCustomizableObject The Customizable Object we want to synchronously compile
+ * @param InCompilationOptionsOverride The configuration for the compilation of the CO we want to use instead of the one part of the CO.
+ * @return True if the compilation was successful and false if it failed.
+ */
+bool CompileCustomizableObject(UCustomizableObject* InCustomizableObject, const FCompilationOptions* InCompilationOptionsOverride = nullptr);
+
+
+/**
+ * Helping class that handles the async update of the provided instance. It will also wait for the mips of it so they get streamed.
+ */
+UCLASS()
+class UCOIUpdater : public UObject
+{
+	GENERATED_BODY()
+	
+public:
+	
+	/**
+	 * Updates the provided customizable object instance.
+	 * @param InInstance Instance to be compiled
+	 * @return True if the update was successful and false otherwise
+	 */
+	bool UpdateInstance(UCustomizableObjectInstance* InInstance);
+	
+
+private:
+	
+	/**
+	 * Callback executed when the instance being updated finishes it's mesh update.
+	 * @param Result The result of the updating operation
+	 */
+	UFUNCTION()
+	void OnInstanceUpdateResult(const FUpdateContext& Result);
+
+	/** The instance that is currently being handled by this class. */
+	UPROPERTY()
+	TObjectPtr<UCustomizableObjectInstance> Instance;
+
+	/** The components of the Instance that we are currently waiting for their mips to be streamed in */
+	UPROPERTY()
+	TArray<TObjectPtr<USkeletalMeshComponent>> ComponentsBeingUpdated;
+
+	/** Flag used to control if we are updating and instance or not. Once it gets set to false then the update gets halted and the program continues */
+	bool bIsInstanceBeingUpdated = false;
+
+	/** False if the instance did update successfully, true if it failed */
+	bool bInstanceFailedUpdate = false;
+};
diff --git a/Engine/Plugins/Experimental/Mutable/Source/MutableValidation/Public/MuV/COIBulkUpdateTestCommandlet.h b/Engine/Plugins/Experimental/Mutable/Source/MutableValidation/Public/MuV/COIBulkUpdateTestCommandlet.h
new file mode 100644
--- /dev/null
+++ b/Engine/Plugins/Experimental/Mutable/Source/MutableValidation/Public/MuV/COIBulkUpdateTestCommandlet.h
@@ -0,0 +1,30 @@
+﻿// Copyright Epic Games, Inc. All Rights Reserved.
+#pragma once
+
+#include "Commandlets/Commandlet.h"
+#include "MuV/ValidationUtils.h"
+
+#include "COIBulkUpdateTestCommandlet.generated.h"
+
+/**
+ * Testing commandlet designed to allow the mass update of instances in a given project path.
+ * It will run in sequence where each of the instances will get updated after their CO gets compiled.
+ *
+ * Example commandline that updates all instances found in the plugin content directory (and all subdirectories):
+ * -run=COIBulkUpdateTest -EnablePlugins=MutableTesting -InstancesPackagePath="/MutableTesting" -AllowCommandletRendering
+ */
+UCLASS()
+class UCOIBulkUpdateTestCommandlet : public UCommandlet 
+{
+	GENERATED_BODY()
+
+public:
+	virtual int32 Main(const FString& Params) override;
+
+private:
+
+	/** Helper object that handles the update of the instance for us.
+	 * It allows us to work with the instance update as if it was a sync operation */
+	UPROPERTY()
+	TObjectPtr<UCOIUpdater> InstanceUpdater = nullptr;
+};
diff --git a/Engine/Plugins/Experimental/Mutable/Source/MutableValidation/Public/MuV/CustomizableObjectValidationCommandlet.h b/Engine/Plugins/Experimental/Mutable/Source/MutableValidation/Public/MuV/CustomizableObjectValidationCommandlet.h
--- a/Engine/Plugins/Experimental/Mutable/Source/MutableValidation/Public/MuV/CustomizableObjectValidationCommandlet.h
+++ b/Engine/Plugins/Experimental/Mutable/Source/MutableValidation/Public/MuV/CustomizableObjectValidationCommandlet.h
@@ -8,7 +8,9 @@
 
 #include "CustomizableObjectValidationCommandlet.generated.h"
 
+// Forward declarations
 class UCustomizableObject;
+class UCOIUpdater;
 
 UCLASS()
 class UCustomizableObjectValidationCommandlet : public UCommandlet 
@@ -22,28 +24,11 @@ private:
 	/** Customizable Object to be tested */
 	UPROPERTY()
 	TObjectPtr<UCustomizableObject> ToTestCustomizableObject = nullptr;
-	
-	/** Customizable Object Instance currently being updated */
-	UPROPERTY()
-	TObjectPtr<UCustomizableObjectInstance> InstanceBeingUpdated = nullptr;
-
-	/** Customizable Skeletal Component currently being updated */
-	UPROPERTY()
-	TArray<TObjectPtr<USkeletalMeshComponent>> ComponentsBeingUpdated;
 
 	/** Array of COI to be generated with randomized parameter values */
 	UPROPERTY()
 	TArray<TObjectPtr<UCustomizableObjectInstance>> InstancesToProcess;
 	
-	/** Handle to be able to unbind OnInstanceUpdated(...) from instance end of update delegate. */
-	FDelegateHandle OnInstanceUpdateHandle;
-
-	/** Did any of the instances fail the UpdateSkeletalMesh process? */
-	bool bInstanceFailedUpdate = false;
-
-	/** Callback invoked once the currently updating instance has done the update.
-	 * @paragm Result is a container that provides us with data related with the instance updating process.
-	 */
-	UFUNCTION()
-	void OnInstanceUpdate(const FUpdateContext& Result);
+	UPROPERTY()
+	TObjectPtr<UCOIUpdater> InstanceUpdater;
 };
\ No newline at end of file
