Fixup mutable plugin for platform extensions

#rnx

diff --git a/Engine/Plugins/Experimental/Mutable/Source/CustomizableObject/Private/MuCO/CustomizableObjectInstance.cpp b/Engine/Plugins/Experimental/Mutable/Source/CustomizableObject/Private/MuCO/CustomizableObjectInstance.cpp
--- a/Engine/Plugins/Experimental/Mutable/Source/CustomizableObject/Private/MuCO/CustomizableObjectInstance.cpp
+++ b/Engine/Plugins/Experimental/Mutable/Source/CustomizableObject/Private/MuCO/CustomizableObjectInstance.cpp
@@ -53,6 +53,10 @@
 #include "Subsystems/AssetEditorSubsystem.h"
 #endif
 
+#ifndef REQUIRES_SINGLEUSE_FLAG_FOR_RUNTIME_TEXTURES
+	#define REQUIRES_SINGLEUSE_FLAG_FOR_RUNTIME_TEXTURES !PLATFORM_DESKTOP
+#endif
+
 
 // Struct used by BuildMaterials() to identify common materials between LODs
 struct FMutableMaterialPlaceholder
@@ -5727,7 +5731,7 @@ void UCustomizableInstancePrivate::BuildMaterials(const TSharedRef<FUpdateContex
 
 									{
 										MUTABLE_CPUPROFILER_SCOPE(UpdateResource);
-#if !PLATFORM_DESKTOP && !PLATFORM_SWITCH // Switch does this in FTexture2DResource::InitRHI()
+#if REQUIRES_SINGLEUSE_FLAG_FOR_RUNTIME_TEXTURES
 										for (int32 i = 0; i < MutableTexture->GetPlatformData()->Mips.Num(); ++i)
 										{
 											uint32 DataFlags = MutableTexture->GetPlatformData()->Mips[i].BulkData.GetBulkDataFlags();
