[Mutable] Fix crashes after renaming the recently created AssetUserData object while loading the compiled data.

#rnx

diff --git a/Engine/Plugins/Experimental/Mutable/Source/CustomizableObject/Internal/MuCO/CustomizableObjectPrivate.h b/Engine/Plugins/Experimental/Mutable/Source/CustomizableObject/Internal/MuCO/CustomizableObjectPrivate.h
--- a/Engine/Plugins/Experimental/Mutable/Source/CustomizableObject/Internal/MuCO/CustomizableObjectPrivate.h
+++ b/Engine/Plugins/Experimental/Mutable/Source/CustomizableObject/Internal/MuCO/CustomizableObjectPrivate.h
@@ -351,12 +351,8 @@ struct FMutableRefAssetUserData
 	TObjectPtr<UAssetUserData> AssetUserData;
 
 #if WITH_EDITORONLY_DATA
-	FString ClassPath;
-	TArray<uint8> Bytes;
-
-	friend FArchive& operator<<(FArchive& Ar, FMutableRefAssetUserData& Data);
-
-	void InitResources(UCustomizableObject* InOuter);
+	UPROPERTY()
+	int32 AssetUserDataIndex = INDEX_NONE;
 #endif
 
 };
@@ -634,7 +630,7 @@ public:
 	// This is a manual version number for the binary blobs in this asset.
 	// Increasing it invalidates all the previously compiled models.
 	// Warning: If while merging code both versions have changed, take the highest+1.
-	static constexpr int32 CurrentSupportedVersion = 423;
+	static constexpr int32 CurrentSupportedVersion = 424;
 
 	/** This is a non-user-controlled flag to disable streaming (set at object compilation time, depending on optimization). */
 	bool bDisableTextureStreaming = false;
diff --git a/Engine/Plugins/Experimental/Mutable/Source/CustomizableObject/Private/MuCO/CustomizableObject.cpp b/Engine/Plugins/Experimental/Mutable/Source/CustomizableObject/Private/MuCO/CustomizableObject.cpp
--- a/Engine/Plugins/Experimental/Mutable/Source/CustomizableObject/Private/MuCO/CustomizableObject.cpp
+++ b/Engine/Plugins/Experimental/Mutable/Source/CustomizableObject/Private/MuCO/CustomizableObject.cpp
@@ -144,8 +144,11 @@ void UCustomizableObject::PreSave(FObjectPreSaveContext ObjectSaveContext)
 		if (TryLoadCompiledCookDataForPlatform(TargetPlatform))
 		{
 			// Create an export object to manage the streamable data
-			BulkData = UE::Mutable::Private::MoveOldObjectAndCreateNew<UCustomizableObjectBulk>(
-				UCustomizableObjectBulk::StaticClass(), this);
+			if (!BulkData)
+			{
+				BulkData = UE::Mutable::Private::MoveOldObjectAndCreateNew<UCustomizableObjectBulk>(
+					UCustomizableObjectBulk::StaticClass(), this);
+			}
 			BulkData->Mark(OBJECTMARK_TagExp);
 
 			// Split streamable data into smaller chunks and fix up the CO HashToStreamableBlock's FileIndex and Offset
@@ -564,15 +567,14 @@ void UCustomizableObject::LoadCompiledData(FArchive& MemoryReader, const ITarget
 
 		MemoryReader << ModelResources.ReferenceSkeletalMeshesData;
 
+		SerializeStreamedResources(MemoryReader, this, StreamedResourceData, bIsCooking);
+
 		// Initialize resources. 
 		for(FMutableRefSkeletalMeshData& ReferenceSkeletalMeshData : ModelResources.ReferenceSkeletalMeshesData)
 		{
 			ReferenceSkeletalMeshData.InitResources(this, InTargetPlatform);
 		}
 
-		SerializeStreamedResources(MemoryReader, this, StreamedResourceData, bIsCooking);
-
-
 		int32 NumReferencedMaterials = 0;
 		MemoryReader << NumReferencedMaterials;
 
@@ -2238,27 +2240,7 @@ FArchive& operator<<(FArchive& Ar, FMutableRefSkeletalMeshSettings& Data)
 
 FArchive& operator<<(FArchive& Ar, FMutableRefAssetUserData& Data)
 {
-	if (Ar.IsSaving())
-	{
-		if (Data.AssetUserData)
-		{
-			Data.ClassPath = Data.AssetUserData->GetClass()->GetPathName();
-
-			FMemoryWriter MemoryWriter(Data.Bytes);
-			Data.AssetUserData->Serialize(MemoryWriter);
-		}
-
-		Ar << Data.ClassPath;
-		Ar << Data.Bytes;
-
-		Data.ClassPath.Empty();
-		Data.Bytes.Empty();
-	}
-	else
-	{
-		Ar << Data.ClassPath;
-		Ar << Data.Bytes;
-	}
+	Ar << Data.AssetUserDataIndex;
 
 	return Ar;
 }
@@ -2274,27 +2256,6 @@ FArchive& operator<<(FArchive& Ar, FMutableSkinWeightProfileInfo& Info)
 }
 
 
-//-------------------------------------------------------------------------------------------------
-//-------------------------------------------------------------------------------------------------
-//-------------------------------------------------------------------------------------------------
-void FMutableRefAssetUserData::InitResources(UCustomizableObject* InOuter)
-{
-	UClass* AssetUserDataClass = FindObject<UClass>(nullptr, *ClassPath);
-	if (AssetUserDataClass)
-	{
-		AssetUserData = UE::Mutable::Private::MoveOldObjectAndCreateNew<UAssetUserData>(AssetUserDataClass, InOuter);
-		if (AssetUserData)
-		{
-			FMemoryReaderView MemoryReader(Bytes);
-			AssetUserData->Serialize(MemoryReader);
-
-			ClassPath.Empty();
-			Bytes.Empty();
-		}
-	}
-}
-
-
 //-------------------------------------------------------------------------------------------------
 //-------------------------------------------------------------------------------------------------
 //-------------------------------------------------------------------------------------------------
@@ -2302,7 +2263,7 @@ void FMutableRefSkeletalMeshData::InitResources(UCustomizableObject* InOuter, co
 {
 	check(InOuter);
 
-	const bool bHasServer = InTargetPlatform ? !InTargetPlatform->IsClientOnly() : false;	
+	const bool bHasServer = InTargetPlatform ? !InTargetPlatform->IsClientOnly() : false;
 	if (InOuter->IsEnableUseRefSkeletalMeshAsPlaceholder() || bHasServer)
 	{
 		SkeletalMesh = TSoftObjectPtr<USkeletalMesh>(SoftSkeletalMesh).LoadSynchronous();
@@ -2311,7 +2272,22 @@ void FMutableRefSkeletalMeshData::InitResources(UCustomizableObject* InOuter, co
 	// Initialize AssetUserData
 	for (FMutableRefAssetUserData& Data : AssetUserData)
 	{
-		Data.InitResources(InOuter);
+		if (!InOuter->StreamedResourceData.IsValidIndex(Data.AssetUserDataIndex))
+		{
+			check(false);
+			continue;
+		}
+
+		FCustomizableObjectStreamedResourceData& StreamedResource = InOuter->StreamedResourceData[Data.AssetUserDataIndex];
+		StreamedResource.NotifyLoaded(StreamedResource.GetPath().LoadSynchronous());
+
+		const FCustomizableObjectResourceData& ResourceData = StreamedResource.GetLoadedData();
+		check(ResourceData.Type == ECOResourceDataType::AssetUserData);
+
+		const FCustomizableObjectAssetUserData* AUDResource = ResourceData.Data.GetPtr<FCustomizableObjectAssetUserData>();
+		check(AUDResource);
+
+		Data.AssetUserData = AUDResource->AssetUserData ? AUDResource->AssetUserData : AUDResource->AssetUserDataEditor;
 	}
 }
 
diff --git a/Engine/Plugins/Experimental/Mutable/Source/CustomizableObjectEditor/Private/MuCOE/GenerateMutableSource/GenerateMutableSource.cpp b/Engine/Plugins/Experimental/Mutable/Source/CustomizableObjectEditor/Private/MuCOE/GenerateMutableSource/GenerateMutableSource.cpp
--- a/Engine/Plugins/Experimental/Mutable/Source/CustomizableObjectEditor/Private/MuCOE/GenerateMutableSource/GenerateMutableSource.cpp
+++ b/Engine/Plugins/Experimental/Mutable/Source/CustomizableObjectEditor/Private/MuCOE/GenerateMutableSource/GenerateMutableSource.cpp
@@ -1612,9 +1612,8 @@ void PopulateReferenceSkeletalMeshesData(FMutableGraphGenerationContext& Generat
 				if (AssetUserData)
 				{
 					FMutableRefAssetUserData& MutAssetUserData = Data.AssetUserData.AddDefaulted_GetRef();
-
-					// Duplicate the AssetUserData using the CO as an outer, so we can save it.
-					MutAssetUserData.AssetUserData = Cast<UAssetUserData>(StaticDuplicateObject(AssetUserData, GenerationContext.Object));
+					MutAssetUserData.AssetUserDataIndex = GenerationContext.AddAssetUserDataToStreamedResources(AssetUserData);
+					MutAssetUserData.AssetUserData = AssetUserData;
 				}
 			}
 		}
