[Mutable] Show logs and errors in the Mutable Log category when compiling CO through the context menu.

#rnx

diff --git a/Engine/Plugins/Experimental/Mutable/Source/CustomizableObjectEditor/Private/MuCOE/CustomizableObjectCompiler.cpp b/Engine/Plugins/Experimental/Mutable/Source/CustomizableObjectEditor/Private/MuCOE/CustomizableObjectCompiler.cpp
--- a/Engine/Plugins/Experimental/Mutable/Source/CustomizableObjectEditor/Private/MuCOE/CustomizableObjectCompiler.cpp
+++ b/Engine/Plugins/Experimental/Mutable/Source/CustomizableObjectEditor/Private/MuCOE/CustomizableObjectCompiler.cpp
@@ -64,7 +64,7 @@ bool FCustomizableObjectCompiler::Tick()
 
 		if (SaveDDTask.IsValid())
 		{
-			SaveCODerivedData(true);
+			SaveCODerivedData();
 		}
 		else
 		{
@@ -107,7 +107,7 @@ bool FCustomizableObjectCompiler::Tick()
 	if (CompilationLaunchPending)
 	{
 		CompilationLaunchPending = false;
-		LaunchMutableCompile(true);
+		LaunchMutableCompile();
 	}
 
 	return bUpdated;
@@ -723,12 +723,9 @@ void FCustomizableObjectCompiler::AddCachedReferencers(const FName& PathName, TA
 }
 
 
-void FCustomizableObjectCompiler::LaunchMutableCompile(bool bShowNotification)
+void FCustomizableObjectCompiler::LaunchMutableCompile()
 {
-	if(bShowNotification)
-	{
-		AddCompileNotification(LOCTEXT("CustomizableObjectCompileInProgress", "Compiling"));
-	}
+	AddCompileNotification(LOCTEXT("CustomizableObjectCompileInProgress", "Compiling"));
 
 	// Even for async build, we spawn a thread, so that we can set a large stack. 
 	// Thread names need to be unique, apparently.
@@ -738,17 +735,14 @@ void FCustomizableObjectCompiler::LaunchMutableCompile(bool bShowNotification)
 }
 
 
-void FCustomizableObjectCompiler::SaveCODerivedData(bool bShowNotification)
+void FCustomizableObjectCompiler::SaveCODerivedData()
 {
 	if (!SaveDDTask.IsValid())
 	{
 		return;
 	}
 
-	if (bShowNotification)
-	{
-		AddCompileNotification(LOCTEXT("SavingCustomizableObjectDerivedData", "Saving Data"));
-	}
+	AddCompileNotification(LOCTEXT("SavingCustomizableObjectDerivedData", "Saving Data"));
 
 	// Even for async saving derived data.
 	static int SDDThreadCount = 0;
@@ -1110,7 +1104,7 @@ void FCustomizableObjectCompiler::CompileInternal(UCustomizableObject* Object, c
 
 			if (SaveDDTask.IsValid())
 			{
-				SaveCODerivedData(false);
+				SaveCODerivedData();
 				SaveDDThread->WaitForCompletion();
 				FinishSavingDerivedData();
 			}
@@ -1127,7 +1121,7 @@ void FCustomizableObjectCompiler::CompileInternal(UCustomizableObject* Object, c
 			// If packaging, convert textures and launch Mutable compile thread
 			if (IsRunningCommandlet())
 			{
-				LaunchMutableCompile(false);
+				LaunchMutableCompile();
 			}
 			else
 			{
@@ -1334,16 +1328,12 @@ void FCustomizableObjectCompiler::ForceFinishBeforeStartCompilation(UCustomizabl
 
 void FCustomizableObjectCompiler::AddCompileNotification(const FText& CompilationStep) const
 {
-	if (Options.bSilentCompilation)
-	{
-		return;
-	}
-
 	const FText Text = CurrentObject ? FText::FromString(FString::Printf(TEXT("Compiling %s"), *CurrentObject->GetName())) : LOCTEXT("CustomizableObjectCompileInProgressNotification", "Compiling Customizable Object");
 	
 	FCustomizableObjectEditorLogger::CreateLog(Text)
 	.SubText(CompilationStep)
 	.Category(ELoggerCategory::Compilation)
+	.Notification(!Options.bSilentCompilation)
 	.CustomNotification()
 	.FixNotification()
 	.Log();
@@ -1363,11 +1353,6 @@ void FCustomizableObjectCompiler::NotifyCompilationErrors() const
 	const uint32 NumIgnoreds = CompilationLogsContainer.GetIgnoredCount();
 	const bool NoWarningsOrErrors = !(NumWarnings || NumErrors);
 
-	if (Options.bSilentCompilation && NoWarningsOrErrors)
-	{
-		return;
-	}
-	
 	const EMessageSeverity::Type Severity = [&]
 	{
 		if (NumErrors)
@@ -1396,6 +1381,7 @@ void FCustomizableObjectCompiler::NotifyCompilationErrors() const
 	FCustomizableObjectEditorLogger::CreateLog(Message)
 	.Category(ELoggerCategory::Compilation)
 	.Severity(Severity)
+	.Notification(!Options.bSilentCompilation || !NoWarningsOrErrors)
 	.CustomNotification()
 	.Log();
 }
diff --git a/Engine/Plugins/Experimental/Mutable/Source/CustomizableObjectEditor/Private/MuCOE/CustomizableObjectCompiler.h b/Engine/Plugins/Experimental/Mutable/Source/CustomizableObjectEditor/Private/MuCOE/CustomizableObjectCompiler.h
--- a/Engine/Plugins/Experimental/Mutable/Source/CustomizableObjectEditor/Private/MuCOE/CustomizableObjectCompiler.h
+++ b/Engine/Plugins/Experimental/Mutable/Source/CustomizableObjectEditor/Private/MuCOE/CustomizableObjectCompiler.h
@@ -168,13 +168,13 @@ private:
 	/** Launches the compile task in another thread when compiling a CO in the editor
 	* @param bShowNotification [in] whether to show the compiling CO notification or not
 	* @return nothing */
-	void LaunchMutableCompile(bool bShowNotification);
+	void LaunchMutableCompile();
 
 	/** Launches the save derived data task in another thread after compiling a CO in the
 	* editor
 	* @param bShowNotification [in] whether to show the saving DD notification or not
 	* @return nothing */
-	void SaveCODerivedData(bool bShowNotification);
+	void SaveCODerivedData();
 
 	/** When compiling a CO in the editor, flag to know if there's a mutable task pending to be launched through LaunchMutableCompile */
 	bool CompilationLaunchPending;
