[Mutable] Moved all private UCustomizableObjectSystem members to private.

- Removed unused MUTABLE_SYSTEM_ASYNC_STREAMING_PRIORITY macro.
- Merged GetPrivateChecked with GetPrivate.
- Minor typos.
- Moved the following functions to private:
    - GetStreamableManager
    - FinishUpdateGlobal

#rnx

diff --git a/Engine/Plugins/Experimental/Mutable/Source/CustomizableObject/Classes/MuCO/CustomizableObjectSystem.h b/Engine/Plugins/Experimental/Mutable/Source/CustomizableObject/Classes/MuCO/CustomizableObjectSystem.h
--- a/Engine/Plugins/Experimental/Mutable/Source/CustomizableObject/Classes/MuCO/CustomizableObjectSystem.h
+++ b/Engine/Plugins/Experimental/Mutable/Source/CustomizableObject/Classes/MuCO/CustomizableObjectSystem.h
@@ -33,12 +33,6 @@ struct FFrame;
 struct FGuid;
 
 
-// Split StreamedBulkData into chunks smaller than MUTABLE_STREAMED_DATA_MAXCHUNKSIZE
-#define MUTABLE_STREAMED_DATA_MAXCHUNKSIZE		(512 * 1024 * 1024)
-
-// In case of async file operations, what priority to use
-#define MUTABLE_SYSTEM_ASYNC_STREAMING_PRIORITY			EAsyncIOPriorityAndFlags::AIOP_Normal
-
 extern TAutoConsoleVariable<bool> CVarClearWorkingMemoryOnUpdateEnd;
 
 extern TAutoConsoleVariable<bool> CVarReuseImagesBetweenInstances;
@@ -76,18 +70,6 @@ namespace EMutableProfileMetric
 
 };
 
-USTRUCT()
-struct FPendingReleaseSkeletalMeshInfo
-{
-	GENERATED_USTRUCT_BODY()
-
-	UPROPERTY()
-	TObjectPtr<USkeletalMesh> SkeletalMesh = nullptr;
-
-	UPROPERTY()
-	double TimeStamp = 0.0f;
-};
-
 
 USTRUCT()
 struct FPendingReleaseMaterialsInfo
@@ -177,13 +159,12 @@ public:
 };
 
 
-/** End a Customizable Object Instance Update. All code paths of an update have to end here. */
-void FinishUpdateGlobal(const TSharedRef<FUpdateContextPrivate>& Context);
-
-
 UCLASS(Blueprintable, BlueprintType)
 class CUSTOMIZABLEOBJECT_API UCustomizableObjectSystem : public UObject
 {
+	// Friends
+	friend class UCustomizableObjectSystemPrivate;
+
 public:
 	GENERATED_BODY()
 
@@ -198,8 +179,7 @@ public:
 	static UCustomizableObjectSystem* GetInstanceChecked();
 
 	/** Determines if the result of the instance update is valid or not.
-	* @return true if the result is successful or has warnings, false if the result is from the Error category
-	 */
+	 * @return true if the result is successful or has warnings, false if the result is from the Error category */
 	UFUNCTION(BlueprintCallable, Category = Status)
 	static bool IsUpdateResultValid(const EUpdateResult UpdateResult);
 	
@@ -252,9 +232,6 @@ public:
 	
 	// If true, uncompiled Customizable Objects will be compiled synchronously
 	bool IsAutoCompilationSync() const;
-
-	// Copy of the Mutable Editor Settings tied to CO compilation. They are updated whenever changed
-	FEditorCompileSettings EditorSettings;
 #endif
 	
 	// Return the current MinLodQualityLevel for skeletal meshes.
@@ -272,11 +249,6 @@ public:
 	void SetImagePixelFormatOverride(const mu::FImageOperator::FImagePixelFormatFunc&);
 #endif
 
-private:
-	UPROPERTY()
-	TArray<FPendingReleaseSkeletalMeshInfo> PendingReleaseSkeletalMesh;
-
-public:
 	/** [Texture Parameters] Get a list of all the possible values for external texture parameters according to the various providers registered with RegisterImageProvider. */
 	TArray<FCustomizableObjectExternalTexture> GetTextureParameterValues();
 
@@ -286,7 +258,7 @@ public:
 	/** [Texture Parameters] Remove a previously registered provider. */
 	void UnregisterImageProvider(UCustomizableSystemImageProvider* Provider);
 
-	/** [Texture Parameters] Interface to actually cache Images in the Mutable system and make them availabe at run-time.
+	/** [Texture Parameters] Interface to actually cache Images in the Mutable system and make them available at run-time.
 		Any cached image has to be registered by an Image provider before caching it.
 		Have in mind that once an image has been cached, it will spend memory according to its size, except in the case
 		of images of type UCustomizableSystemImageProvider::ValueType::Unreal_Deferred, where only a very small amount of
@@ -303,7 +275,7 @@ public:
 	// being used and it's not compiled.  Callers can add additional information to the error log.
 	void AddUncompiledCOWarning(const UCustomizableObject& InObject, FString const* OptionalLogInfo = nullptr);
 
-	// Enables the collection of internal mutabe performance data. It has a performance cost.
+	// Enables the collection of internal Mutable performance data. It has a performance cost.
 	void EnableBenchmark();
 	// Writes the benchmark results
 	void EndBenchmark();
@@ -315,11 +287,6 @@ public:
 	UCustomizableObjectSystemPrivate* GetPrivate();
 	const UCustomizableObjectSystemPrivate* GetPrivate() const;
 
-	UCustomizableObjectSystemPrivate* GetPrivateChecked();
-	const UCustomizableObjectSystemPrivate* GetPrivateChecked() const;
-	
-	FStreamableManager& GetStreamableManager();
-
 	UCustomizableInstanceLODManagementBase* GetInstanceLODManagement() const;
 
 	// Pass a null ptr to reset to the default InstanceLODManagement
@@ -368,24 +335,7 @@ public:
 
 	void ClearCurrentMutableOperation();
 
-	UPROPERTY(Transient)
-	TObjectPtr<UCustomizableInstanceLODManagementBase> DefaultInstanceLODManagement = nullptr;
-
-	UPROPERTY(Transient)
-	TObjectPtr<UCustomizableInstanceLODManagementBase> CurrentInstanceLODManagement = nullptr;
-
-	// Array where textures are added temporarily while the mutable thread may want to
-	// reused them for some instance under construction.
-	UPROPERTY(Transient)
-	TArray<TObjectPtr<UTexture2D>> ProtectedCachedTextures;
-
 private:
-	UPROPERTY(Transient)
-	TObjectPtr<UCustomizableObjectSystemPrivate> Private = nullptr;
-
-	// For async material loading
-	FStreamableManager StreamableManager;
-
 	// Most of the work in this plugin happens here.
 	bool Tick(float DeltaTime);
 
@@ -408,6 +358,7 @@ private:
 
 	void StartNextRecompile();
 	void TickRecompileCustomizableObjects();
+#endif
 
 public:
 	/** Set Mutable's working memory limit (bytes). Mutable will flush internal caches to try to keep its memory consumption below the WorkingMemory (i.e., it is not a hard limit).
@@ -418,29 +369,15 @@ public:
 
 	/** Get Mutable's working memory limit (bytes). See SetWorkingMemory(int32). */
 	int32 GetWorkingMemory() const;
-	
-private:
-	FCustomizableObjectCompilerBase* RecompileCustomizableObjectsCompiler = nullptr;
-	
-	TArray<FAssetData> ObjectsToRecompile;
-	uint32 TotalNumObjectsToRecompile = 0;
-	uint32 NumObjectsCompiled = 0;
-
-	/** Recompile progress bar handle */
-	FProgressNotificationHandle RecompileNotificationHandle;
-
-	// Array to keep track of cached objects
-	TArray<FGuid> UncompiledCustomizableObjectIds;
 
-	/** Weak pointer to the Uncompiled Customizable Objects notification */
-	TWeakPtr<SNotificationItem> UncompiledCustomizableObjectsNotificationPtr;
-
-	/** Map used to cache per platform MaxChunkSize. If MaxChunkSize > 0, streamed data will be split in multiple files */
-	TMap<FString, int64> PlatformMaxChunkSize;
+#if WITH_EDITOR
+	// Copy of the Mutable Editor Settings tied to CO compilation. They are updated whenever changed
+	FEditorCompileSettings EditorSettings;
 #endif
 	
-	// Friends
-	friend class UCustomizableObjectSystemPrivate;
+private:
+	UPROPERTY(Transient)
+	TObjectPtr<UCustomizableObjectSystemPrivate> Private = nullptr;
 };
 
 
diff --git a/Engine/Plugins/Experimental/Mutable/Source/CustomizableObject/Internal/MuCO/CustomizableObjectSystemPrivate.h b/Engine/Plugins/Experimental/Mutable/Source/CustomizableObject/Internal/MuCO/CustomizableObjectSystemPrivate.h
--- a/Engine/Plugins/Experimental/Mutable/Source/CustomizableObject/Internal/MuCO/CustomizableObjectSystemPrivate.h
+++ b/Engine/Plugins/Experimental/Mutable/Source/CustomizableObject/Internal/MuCO/CustomizableObjectSystemPrivate.h
@@ -9,6 +9,7 @@
 #include "MuCO/CustomizableObjectPrivate.h"
 #include "MuCO/CustomizableObjectInstance.h"
 #include "MuCO/CustomizableObjectExtension.h"
+#include "MuCO/CustomizableInstanceLODManagement.h"
 #include "Containers/Ticker.h"
 
 #include "MuCO/CustomizableObjectInstanceDescriptor.h"
@@ -18,7 +19,10 @@
 #include "MuR/Image.h"
 #include "UObject/GCObject.h"
 #include "WorldCollision.h"
+#include "Engine/StreamableManager.h"
+#include "Framework/Notifications/NotificationManager.h"
 #include "MuCO/FMutableTaskGraph.h"
+#include "AssetRegistry/AssetData.h"
 
 // This define could come from MuR/System.h
 #ifdef MUTABLE_USE_NEW_TASKGRAPH
@@ -38,6 +42,14 @@ namespace LowLevelTasks { enum class ETaskPriority : int8; }
 struct FTexturePlatformData;
 
 
+// Split StreamedBulkData into chunks smaller than MUTABLE_STREAMED_DATA_MAXCHUNKSIZE
+#define MUTABLE_STREAMED_DATA_MAXCHUNKSIZE		(512 * 1024 * 1024)
+
+
+/** End a Customizable Object Instance Update. All code paths of an update have to end here. */
+void FinishUpdateGlobal(const TSharedRef<FUpdateContextPrivate>& Context);
+
+
 struct FMutablePendingInstanceUpdate
 {
 	TSharedRef<FUpdateContextPrivate> Context;
@@ -536,6 +548,19 @@ struct FMutableReleasePlatformOperationData
 };
 
 
+USTRUCT()
+struct FPendingReleaseSkeletalMeshInfo
+{
+	GENERATED_USTRUCT_BODY()
+
+	UPROPERTY()
+	TObjectPtr<USkeletalMesh> SkeletalMesh = nullptr;
+
+	UPROPERTY()
+	double TimeStamp = 0.0f;
+};
+
+
 UCLASS()
 class UCustomizableObjectSystemPrivate : public UObject
 {
@@ -683,6 +708,43 @@ public:
 	int32 NumSkeletalMeshes = 0;
 
 	bool bAutoCompileCommandletEnabled = false;
+
+	UPROPERTY()
+	TArray<FPendingReleaseSkeletalMeshInfo> PendingReleaseSkeletalMesh;
+	
+	UPROPERTY(Transient)
+	TObjectPtr<UCustomizableInstanceLODManagementBase> DefaultInstanceLODManagement = nullptr;
+
+	UPROPERTY(Transient)
+	TObjectPtr<UCustomizableInstanceLODManagementBase> CurrentInstanceLODManagement = nullptr;
+
+	// Array where textures are added temporarily while the mutable thread may want to
+	// reused them for some instance under construction.
+	UPROPERTY(Transient)
+	TArray<TObjectPtr<UTexture2D>> ProtectedCachedTextures;
+	
+	// For async material loading
+	FStreamableManager StreamableManager;
+	
+#if WITH_EDITOR
+	FCustomizableObjectCompilerBase* RecompileCustomizableObjectsCompiler = nullptr;
+	
+	TArray<FAssetData> ObjectsToRecompile;
+	uint32 TotalNumObjectsToRecompile = 0;
+	uint32 NumObjectsCompiled = 0;
+
+	/** Recompile progress bar handle */
+	FProgressNotificationHandle RecompileNotificationHandle;
+
+	// Array to keep track of cached objects
+	TArray<FGuid> UncompiledCustomizableObjectIds;
+
+	/** Weak pointer to the Uncompiled Customizable Objects notification */
+	TWeakPtr<SNotificationItem> UncompiledCustomizableObjectsNotificationPtr;
+
+	/** Map used to cache per platform MaxChunkSize. If MaxChunkSize > 0, streamed data will be split in multiple files */
+	TMap<FString, int64> PlatformMaxChunkSize;
+#endif
 };
 
 
diff --git a/Engine/Plugins/Experimental/Mutable/Source/CustomizableObject/Private/MuCO/CustomizableObjectMipDataProvider.cpp b/Engine/Plugins/Experimental/Mutable/Source/CustomizableObject/Private/MuCO/CustomizableObjectMipDataProvider.cpp
--- a/Engine/Plugins/Experimental/Mutable/Source/CustomizableObject/Private/MuCO/CustomizableObjectMipDataProvider.cpp
+++ b/Engine/Plugins/Experimental/Mutable/Source/CustomizableObject/Private/MuCO/CustomizableObjectMipDataProvider.cpp
@@ -259,7 +259,7 @@ namespace impl
 					return;
 				}
 
-				System->GetPrivateChecked()->LogBenchmarkUtil.FinishUpdateImage(CustomizableObjectPathName, InstancePathName, Time, PeakMemory, RealMemoryPeak);
+				System->GetPrivate()->LogBenchmarkUtil.FinishUpdateImage(CustomizableObjectPathName, InstancePathName, Time, PeakMemory, RealMemoryPeak);
 			},
 			TStatId{},
 			nullptr,
diff --git a/Engine/Plugins/Experimental/Mutable/Source/CustomizableObject/Private/MuCO/CustomizableObjectSystem.cpp b/Engine/Plugins/Experimental/Mutable/Source/CustomizableObject/Private/MuCO/CustomizableObjectSystem.cpp
--- a/Engine/Plugins/Experimental/Mutable/Source/CustomizableObject/Private/MuCO/CustomizableObjectSystem.cpp
+++ b/Engine/Plugins/Experimental/Mutable/Source/CustomizableObject/Private/MuCO/CustomizableObjectSystem.cpp
@@ -398,13 +398,13 @@ bool UCustomizableObjectSystem::IsUpdateResultValid(const EUpdateResult UpdateRe
 
 UCustomizableInstanceLODManagementBase* UCustomizableObjectSystem::GetInstanceLODManagement() const
 {
-	return CurrentInstanceLODManagement.Get();
+	return GetPrivate()->CurrentInstanceLODManagement.Get();
 }
 
 
 void UCustomizableObjectSystem::SetInstanceLODManagement(UCustomizableInstanceLODManagementBase* NewInstanceLODManagement)
 {
-	CurrentInstanceLODManagement = NewInstanceLODManagement ? NewInstanceLODManagement : ToRawPtr(DefaultInstanceLODManagement);
+	GetPrivate()->CurrentInstanceLODManagement = NewInstanceLODManagement ? NewInstanceLODManagement : ToRawPtr(GetPrivate()->DefaultInstanceLODManagement);
 }
 
 
@@ -486,33 +486,15 @@ void UCustomizableObjectSystem::LogShowData(bool bFullInfo, bool ShowMaterialInf
 
 UCustomizableObjectSystemPrivate* UCustomizableObjectSystem::GetPrivate()
 {
-	return Private.Get();
+	check(Private);
+	return Private;
 }
 
 
 const UCustomizableObjectSystemPrivate* UCustomizableObjectSystem::GetPrivate() const
 {
-	return Private.Get();
-}
-
-
-UCustomizableObjectSystemPrivate* UCustomizableObjectSystem::GetPrivateChecked()
-{
-	check(Private)
-	return Private.Get();
-}
-
-
-const UCustomizableObjectSystemPrivate* UCustomizableObjectSystem::GetPrivateChecked() const
-{
-	check(Private)
-	return Private.Get();
-}
-
-
-FStreamableManager& UCustomizableObjectSystem::GetStreamableManager()
-{
-	return StreamableManager;
+	check(Private);
+	return Private;
 }
 
 
@@ -566,9 +548,9 @@ void UCustomizableObjectSystem::InitSystem()
 	RegisterImageProvider(Private->EditorImageProvider);
 #endif
 	
-	DefaultInstanceLODManagement = NewObject<UCustomizableInstanceLODManagement>();
-	check(DefaultInstanceLODManagement != nullptr);
-	CurrentInstanceLODManagement = DefaultInstanceLODManagement;
+	GetPrivate()->DefaultInstanceLODManagement = NewObject<UCustomizableInstanceLODManagement>();
+	check(GetPrivate()->DefaultInstanceLODManagement != nullptr);
+	GetPrivate()->CurrentInstanceLODManagement = GetPrivate()->DefaultInstanceLODManagement;
 
 	// This CVar is constant for the lifespan of the program. Read its value once. 
 	const IConsoleVariable* CVarSupport16BitBoneIndex = IConsoleManager::Get().FindConsoleVariable(TEXT("r.GPUSkin.Support16BitBoneIndex"));
@@ -584,12 +566,6 @@ void UCustomizableObjectSystem::InitSystem()
 void UCustomizableObjectSystem::BeginDestroy()
 {
 #if WITH_EDITOR
-	if (RecompileCustomizableObjectsCompiler)
-	{
-		RecompileCustomizableObjectsCompiler->ForceFinishCompilation();
-		delete RecompileCustomizableObjectsCompiler;
-	}
-
 	if (!IsRunningGame())
 	{
 		FEditorDelegates::PreBeginPIE.RemoveAll(this);
@@ -600,6 +576,13 @@ void UCustomizableObjectSystem::BeginDestroy()
 	// It could be null, for the default object.
 	if (Private)
 	{
+#if WITH_EDITOR
+		if (Private->RecompileCustomizableObjectsCompiler)
+		{
+			Private->RecompileCustomizableObjectsCompiler->ForceFinishCompilation();
+			delete Private->RecompileCustomizableObjectsCompiler;
+		}
+#endif
 
 #if !UE_SERVER
 		FTSTicker::GetCoreTicker().RemoveTicker(Private->TickDelegateHandle);
@@ -655,7 +638,7 @@ FCustomizableObjectCompilerBase* UCustomizableObjectSystem::GetNewCompiler()
 
 void UCustomizableObjectSystem::SetNewCompilerFunc(FCustomizableObjectCompilerBase* (*InNewCompilerFunc)())
 {
-	GetPrivateChecked()->NewCompilerFunc = InNewCompilerFunc;
+	GetPrivate()->NewCompilerFunc = InNewCompilerFunc;
 }
 
 
@@ -914,7 +897,7 @@ void FinishUpdateGlobal(const TSharedRef<FUpdateContextPrivate>& Context)
 				return;
 			}
 
-			System->GetPrivateChecked()->LogBenchmarkUtil.FinishUpdateMesh(Context);
+			System->GetPrivate()->LogBenchmarkUtil.FinishUpdateMesh(Context);
 		},
 		TStatId{},
 		nullptr,
@@ -1324,7 +1307,7 @@ void UCustomizableObjectSystem::ClearResourceCacheProtected()
 {
 	check(IsInGameThread());
 
-	ProtectedCachedTextures.Reset(0);
+	GetPrivate()->ProtectedCachedTextures.Reset(0);
 	check(GetPrivate() != nullptr);
 	GetPrivate()->ProtectedObjectCachedImages.Reset(0);
 }
@@ -1559,7 +1542,7 @@ namespace impl
 	void CreateMutableInstance(const TSharedRef<FUpdateContextPrivate>& Operation)
 	{
 		UCustomizableObjectSystem* System = UCustomizableObjectSystem::GetInstanceChecked(); // Save since UCustomizableObjectSystem::BeginDestroy always waits for all tasks to finish
-		const UCustomizableObjectSystemPrivate* SystemPrivate = System->GetPrivateChecked();
+		const UCustomizableObjectSystemPrivate* SystemPrivate = System->GetPrivate();
 		
 		Operation->UpdateStartBytes = mu::FGlobalMemoryCounter::GetCounter();
 		mu::FGlobalMemoryCounter::Zero();
@@ -1880,7 +1863,7 @@ namespace impl
 	{
 		MUTABLE_CPUPROFILER_SCOPE(Subtask_Mutable_GetImages)
 
-		const UCustomizableObjectSystemPrivate* CustomizableObjectSystemPrivateData = UCustomizableObjectSystem::GetInstanceChecked()->GetPrivateChecked();
+		const UCustomizableObjectSystemPrivate* CustomizableObjectSystemPrivateData = UCustomizableObjectSystem::GetInstanceChecked()->GetPrivate();
 		mu::System* System = CustomizableObjectSystemPrivateData->MutableSystem.get();
 		check(System != nullptr);
 
@@ -2227,7 +2210,7 @@ namespace impl
 			return;
 		}
 
-		UCustomizableObjectSystemPrivate * CustomizableObjectSystemPrivateData = System->GetPrivateChecked();
+		UCustomizableObjectSystemPrivate * CustomizableObjectSystemPrivateData = System->GetPrivate();
 
 		// Actual work
 		// TODO MTBL-391: Review This hotfix
@@ -2313,7 +2296,7 @@ namespace impl
 		// Memory used in the context of the mesh update + the baseline memory already in use by mutable
 		OperationData->UpdateEndRealPeakBytes = OperationData->UpdateEndPeakBytes + OperationData->UpdateStartBytes;
 		
-		UCustomizableObjectSystemPrivate* CustomizableObjectSystemPrivateData = System->GetPrivateChecked();
+		UCustomizableObjectSystemPrivate* CustomizableObjectSystemPrivateData = System->GetPrivate();
 
 		// Next Task: Release Mutable. We need this regardless if we cancel or not
 		//-------------------------------------------------------------		
@@ -2427,9 +2410,9 @@ namespace impl
 		
 		// Selectively lock the resource cache for the object used by this instance to avoid the destruction of resources that we may want to reuse.
 		// When protecting textures there mustn't be any left from a previous update
-		check(System->ProtectedCachedTextures.Num() == 0);
+		check(System->GetPrivate()->ProtectedCachedTextures.Num() == 0);
 
-		UCustomizableObjectSystemPrivate* SystemPrivateData = System->GetPrivateChecked();
+		UCustomizableObjectSystemPrivate* SystemPrivateData = System->GetPrivate();
 
 		// TODO: If this is the first code that runs after the CO program has finished AND if it's
 		// guaranteed that the next CO program hasn't started yet, we need to call ClearActiveObject
@@ -2442,7 +2425,7 @@ namespace impl
 
 		FMutableResourceCache& Cache = SystemPrivateData->GetObjectCache(CustomizableObject);
 
-		System->ProtectedCachedTextures.Reset(Cache.Images.Num());
+		System->GetPrivate()->ProtectedCachedTextures.Reset(Cache.Images.Num());
 		SystemPrivateData->ProtectedObjectCachedImages.Reset(Cache.Images.Num());
 
 		for (const FInstanceUpdateData::FImage& Image : OperationData->InstanceUpdateData.Images)
@@ -2452,7 +2435,7 @@ namespace impl
 
 			if (TexturePtr && TexturePtr->Get() && SystemPrivateData->TextureHasReferences(Key))
 			{
-				System->ProtectedCachedTextures.Add(TexturePtr->Get());
+				System->GetPrivate()->ProtectedCachedTextures.Add(TexturePtr->Get());
 				SystemPrivateData->ProtectedObjectCachedImages.Add(Image.ImageID);
 			}
 		}
@@ -2483,7 +2466,7 @@ namespace impl
 
 		// Next Task: Load Unreal Assets
 		//-------------------------------------------------------------
-		FGraphEventRef Game_LoadUnrealAssets = ObjectInstancePrivateData->LoadAdditionalAssetsAsync(OperationData, ObjectInstance, UCustomizableObjectSystem::GetInstance()->GetStreamableManager());
+		FGraphEventRef Game_LoadUnrealAssets = ObjectInstancePrivateData->LoadAdditionalAssetsAsync(OperationData, ObjectInstance, UCustomizableObjectSystem::GetInstance()->GetPrivate()->StreamableManager);
 		if (Game_LoadUnrealAssets)
 		{
 			Game_LoadUnrealAssets->SetDebugName(TEXT("LoadAdditionalAssetsAsync"));
@@ -2513,7 +2496,7 @@ namespace impl
 		MUTABLE_CPUPROFILER_SCOPE(Task_Game_ReleaseInstanceID)
 
 		UCustomizableObjectSystem* System = UCustomizableObjectSystem::GetInstanceChecked();
-		UCustomizableObjectSystemPrivate* SystemPrivateData = System->GetPrivateChecked();
+		UCustomizableObjectSystemPrivate* SystemPrivateData = System->GetPrivate();
 
 		const mu::Ptr<mu::System> MutableSystem = SystemPrivateData->MutableSystem;
 
@@ -2536,7 +2519,7 @@ namespace impl
 		MUTABLE_CPUPROFILER_SCOPE(Task_Game_LockMeshCache);
 
 		UCustomizableObjectSystem* System = UCustomizableObjectSystem::GetInstanceChecked();
-		UCustomizableObjectSystemPrivate* SystemPrivate = System->GetPrivateChecked();
+		UCustomizableObjectSystemPrivate* SystemPrivate = System->GetPrivate();
 
 		const UCustomizableObject* CustomizableObject = Operation->Instance->GetCustomizableObject();
 		UCustomizableObjectPrivate* CustomizableObjectPrivate = CustomizableObject->GetPrivate();
@@ -2583,7 +2566,7 @@ namespace impl
 
 
 		UCustomizableObjectSystem* System = UCustomizableObjectSystem::GetInstanceChecked(); // Save since UCustomizableObjectSystem::BeginDestroy always waits for all tasks to finish
-		UCustomizableObjectSystemPrivate* SystemPrivate = System->GetPrivateChecked();
+		UCustomizableObjectSystemPrivate* SystemPrivate = System->GetPrivate();
 		
 		CreateMutableInstance(Operation);
 		FixLODs(Operation);
@@ -2672,10 +2655,10 @@ namespace impl
 		}
 
 		// Only update resources if the instance is in range (it could have got far from the player since the task was queued)
-		check(System->CurrentInstanceLODManagement != nullptr);
-		if (System->CurrentInstanceLODManagement->IsOnlyUpdateCloseCustomizableObjectsEnabled()
+		check(System->GetPrivate()->CurrentInstanceLODManagement != nullptr);
+		if (System->GetPrivate()->CurrentInstanceLODManagement->IsOnlyUpdateCloseCustomizableObjectsEnabled()
 			&& CandidateInstancePrivateData
-			&& CandidateInstancePrivateData->LastMinSquareDistFromComponentToPlayer > FMath::Square(System->CurrentInstanceLODManagement->GetOnlyUpdateCloseCustomizableObjectsDist())
+			&& CandidateInstancePrivateData->LastMinSquareDistFromComponentToPlayer > FMath::Square(System->GetPrivate()->CurrentInstanceLODManagement->GetOnlyUpdateCloseCustomizableObjectsDist())
 			&& CandidateInstancePrivateData->LastMinSquareDistFromComponentToPlayer != FLT_MAX // This means it is the first frame so it has to be updated
 		   )
 		{
@@ -2697,7 +2680,7 @@ namespace impl
 			return;
 		}
 
-		UCustomizableObjectSystemPrivate* SystemPrivateData = System->GetPrivateChecked();
+		UCustomizableObjectSystemPrivate* SystemPrivateData = System->GetPrivate();
 
 		SystemPrivateData->CurrentInstanceBeingUpdated = CandidateInstance;
 
@@ -2796,7 +2779,7 @@ namespace impl
 #endif // WITH_EDITOR
 		
 		if (!System->IsOnlyGenerateRequestedLODsEnabled() ||
-			!System->CurrentInstanceLODManagement->IsOnlyGenerateRequestedLODLevelsEnabled() ||
+			!System->GetPrivate()->CurrentInstanceLODManagement->IsOnlyGenerateRequestedLODLevelsEnabled() ||
 			bIsInEditorViewport)
 		{
 			TArray<uint16> RequestedLODs = Operation->GetRequestedLODs();
@@ -2953,7 +2936,7 @@ bool UCustomizableObjectSystem::Tick(float DeltaTime)
 		// The RequestedUpdates only refer to LOD changes. User Customization and discards are handled separately
 		FMutableInstanceUpdateMap RequestedLODUpdates;
 		
-		CurrentInstanceLODManagement->UpdateInstanceDistsAndLODs(RequestedLODUpdates);
+		GetPrivate()->CurrentInstanceLODManagement->UpdateInstanceDistsAndLODs(RequestedLODUpdates);
 
 		for (TObjectIterator<UCustomizableObjectInstance> CustomizableObjectInstance; CustomizableObjectInstance; ++CustomizableObjectInstance)
 		{
@@ -3164,8 +3147,8 @@ void UCustomizableObjectSystem::DiscardInstances()
 			UCustomizableInstancePrivate* COIPrivateData = COI ? COI->GetPrivate() : nullptr;
 
 			// Only discard resources if the instance is still out range (it could have got closer to the player since the task was queued)
-			if (!CurrentInstanceLODManagement->IsOnlyUpdateCloseCustomizableObjectsEnabled() ||
-				COIPrivateData->LastMinSquareDistFromComponentToPlayer > FMath::Square(CurrentInstanceLODManagement->GetOnlyUpdateCloseCustomizableObjectsDist()))
+			if (!GetPrivate()->CurrentInstanceLODManagement->IsOnlyUpdateCloseCustomizableObjectsEnabled() ||
+				COIPrivateData->LastMinSquareDistFromComponentToPlayer > FMath::Square(GetPrivate()->CurrentInstanceLODManagement->GetOnlyUpdateCloseCustomizableObjectsDist()))
 			{
 				COIPrivateData->DiscardResources();
 				COIPrivateData->SetDefaultSkeletalMesh(!IsReplaceDiscardedWithReferenceMeshEnabled());
@@ -3209,7 +3192,7 @@ bool UCustomizableObjectSystem::IsUpdating(const UCustomizableObjectInstance* In
 		return false;
 	}
 	
-	return GetPrivateChecked()->IsUpdating(*Instance);
+	return GetPrivate()->IsUpdating(*Instance);
 }
 
 
@@ -3217,7 +3200,7 @@ TArray<FCustomizableObjectExternalTexture> UCustomizableObjectSystem::GetTexture
 {
 	TArray<FCustomizableObjectExternalTexture> Result;
 
-	for (const TWeakObjectPtr<UCustomizableSystemImageProvider> Provider : GetPrivateChecked()->GetImageProviderChecked()->ImageProviders)
+	for (const TWeakObjectPtr<UCustomizableSystemImageProvider> Provider : GetPrivate()->GetImageProviderChecked()->ImageProviders)
 	{
 		if (Provider.IsValid())
 		{
@@ -3231,13 +3214,13 @@ TArray<FCustomizableObjectExternalTexture> UCustomizableObjectSystem::GetTexture
 
 void UCustomizableObjectSystem::RegisterImageProvider(UCustomizableSystemImageProvider* Provider)
 {
-	GetPrivateChecked()->GetImageProviderChecked()->ImageProviders.Add(Provider);
+	GetPrivate()->GetImageProviderChecked()->ImageProviders.Add(Provider);
 }
 
 
 void UCustomizableObjectSystem::UnregisterImageProvider(UCustomizableSystemImageProvider* Provider)
 {
-	GetPrivateChecked()->GetImageProviderChecked()->ImageProviders.Remove(Provider);
+	GetPrivate()->GetImageProviderChecked()->ImageProviders.Remove(Provider);
 }
 
 
@@ -3277,14 +3260,14 @@ int32 UCustomizableObjectSystem::GetNumInstances() const
 	int32 NumInstancesLOD1;
 	int32 NumInstancesLOD2;
 	int32 NumAllocatedSkeletalMeshes;
-	GetPrivateChecked()->LogBenchmarkUtil.GetInstancesStats(NumInstances, NumBuiltInstances, NumInstancesLOD0, NumInstancesLOD1, NumInstancesLOD2, NumAllocatedSkeletalMeshes);
+	GetPrivate()->LogBenchmarkUtil.GetInstancesStats(NumInstances, NumBuiltInstances, NumInstancesLOD0, NumInstancesLOD1, NumInstancesLOD2, NumAllocatedSkeletalMeshes);
 
 	return NumBuiltInstances;
 }
 
 int32 UCustomizableObjectSystem::GetNumPendingInstances() const
 {
-	return GetPrivateChecked()->MutablePendingInstanceWork.Num();
+	return GetPrivate()->MutablePendingInstanceWork.Num();
 }
 
 int32 UCustomizableObjectSystem::GetTotalInstances() const
@@ -3306,48 +3289,48 @@ int32 UCustomizableObjectSystem::GetTotalInstances() const
 
 int64 UCustomizableObjectSystem::GetTextureMemoryUsed() const
 {
-	return GetPrivateChecked()->LogBenchmarkUtil.TextureGPUSize.GetValue();
+	return GetPrivate()->LogBenchmarkUtil.TextureGPUSize.GetValue();
 }
 
 int32 UCustomizableObjectSystem::GetAverageBuildTime() const
 {
-	return GetPrivateChecked()->LogBenchmarkUtil.InstanceBuildTimeAvrg.GetValue() * 1000;
+	return GetPrivate()->LogBenchmarkUtil.InstanceBuildTimeAvrg.GetValue() * 1000;
 }
 
 
 int32 UCustomizableObjectSystem::GetSkeletalMeshMinLODQualityLevel() const
 {
-	return GetPrivateChecked()->SkeletalMeshMinLodQualityLevel;
+	return GetPrivate()->SkeletalMeshMinLodQualityLevel;
 }
 
 
 bool UCustomizableObjectSystem::IsSupport16BitBoneIndexEnabled() const
 {
-	return GetPrivateChecked()->bSupport16BitBoneIndex;
+	return GetPrivate()->bSupport16BitBoneIndex;
 }
 
 
 bool UCustomizableObjectSystem::IsProgressiveMipStreamingEnabled() const
 {
-	return GetPrivateChecked()->EnableMutableProgressiveMipStreaming != 0;
+	return GetPrivate()->EnableMutableProgressiveMipStreaming != 0;
 }
 
 
 void UCustomizableObjectSystem::SetProgressiveMipStreamingEnabled(bool bIsEnabled)
 {
-	GetPrivateChecked()->EnableMutableProgressiveMipStreaming = bIsEnabled ? 1 : 0;
+	GetPrivate()->EnableMutableProgressiveMipStreaming = bIsEnabled ? 1 : 0;
 }
 
 
 bool UCustomizableObjectSystem::IsOnlyGenerateRequestedLODsEnabled() const
 {
-	return GetPrivateChecked()->EnableOnlyGenerateRequestedLODs != 0;
+	return GetPrivate()->EnableOnlyGenerateRequestedLODs != 0;
 }
 
 
 void UCustomizableObjectSystem::SetOnlyGenerateRequestedLODsEnabled(bool bIsEnabled)
 {
-	GetPrivateChecked()->EnableOnlyGenerateRequestedLODs = bIsEnabled ? 1 : 0;
+	GetPrivate()->EnableOnlyGenerateRequestedLODs = bIsEnabled ? 1 : 0;
 }
 
 
@@ -3370,18 +3353,18 @@ void UCustomizableObjectSystem::AddUncompiledCOWarning(const UCustomizableObject
 
 #if WITH_EDITOR
 	// Mutable will spam these warnings constantly due to the tick and LOD manager checking for instances to update with every tick. Send only one message per CO in the editor.
-	if (UncompiledCustomizableObjectIds.Find(InObject.GetVersionId()) != INDEX_NONE)
+	if (GetPrivate()->UncompiledCustomizableObjectIds.Find(InObject.GetVersionId()) != INDEX_NONE)
 	{
 		return;
 	}
 	
 	// Add notification
-	UncompiledCustomizableObjectIds.Add(InObject.GetVersionId());
+	GetPrivate()->UncompiledCustomizableObjectIds.Add(InObject.GetVersionId());
 
 	FMessageLog MessageLog("Mutable");
 	MessageLog.Warning(FText::FromString(Msg));
 
-	if (!UncompiledCustomizableObjectsNotificationPtr.IsValid())
+	if (!GetPrivate()->UncompiledCustomizableObjectsNotificationPtr.IsValid())
 	{
 		FNotificationInfo Info(FText::FromString("Uncompiled Customizable Object/s found. Please, check the Message Log - Mutable for more information."));
 		Info.bFireAndForget = true;
@@ -3389,7 +3372,7 @@ void UCustomizableObjectSystem::AddUncompiledCOWarning(const UCustomizableObject
 		Info.FadeOutDuration = 1.0f;
 		Info.ExpireDuration = 5.0f;
 
-		UncompiledCustomizableObjectsNotificationPtr = FSlateNotificationManager::Get().AddNotification(Info);
+		GetPrivate()->UncompiledCustomizableObjectsNotificationPtr = FSlateNotificationManager::Get().AddNotification(Info);
 	}
 
 	const FString ErrorString = FString::Printf(
@@ -3421,7 +3404,7 @@ void UCustomizableObjectSystem::EndBenchmark()
 
 void UCustomizableObjectSystem::SetReleaseMutableTexturesImmediately(bool bReleaseTextures)
 {
-	GetPrivateChecked()->bReleaseTexturesImmediately = bReleaseTextures;
+	GetPrivate()->bReleaseTexturesImmediately = bReleaseTextures;
 }
 
 
@@ -3467,7 +3450,7 @@ void UCustomizableObjectSystem::OnPreBeginPIE(const bool bIsSimulatingInEditor)
 		const FText Msg = FText::FromString(TEXT("Warning: one or more Customizable Objects used in PIE are uncompiled.\n\nDo you want to compile them?"));
 		if (FMessageDialog::Open(EAppMsgType::OkCancel, Msg) == EAppReturnType::Ok)
 		{
-			ObjectsToRecompile.Empty(TempObjectsToRecompile.Num());
+			GetPrivate()->ObjectsToRecompile.Empty(TempObjectsToRecompile.Num());
 			RecompileCustomizableObjects(TempObjectsToRecompile);
 		}
 	}
@@ -3480,18 +3463,18 @@ void UCustomizableObjectSystem::StartNextRecompile()
 		GEngine->ForceGarbageCollection();
 	}
 
-	FAssetData Itr = ObjectsToRecompile.Pop();
+	FAssetData Itr = GetPrivate()->ObjectsToRecompile.Pop();
 
 	if (UCustomizableObject* CustomizableObject = Cast<UCustomizableObject>(Itr.GetAsset()))
 	{
 		const FText UpdateMsg = FText::FromString(FString::Printf(TEXT("Compiling Customizable Objects:\n%s"), *CustomizableObject->GetName()));
-		FSlateNotificationManager::Get().UpdateProgressNotification(RecompileNotificationHandle, NumObjectsCompiled, TotalNumObjectsToRecompile, UpdateMsg);
+		FSlateNotificationManager::Get().UpdateProgressNotification(GetPrivate()->RecompileNotificationHandle, GetPrivate()->NumObjectsCompiled, GetPrivate()->TotalNumObjectsToRecompile, UpdateMsg);
 
 		// Use default options
 		FCompilationOptions Options = CustomizableObject->CompileOptions;
 		Options.bSilentCompilation = true;
-		check(RecompileCustomizableObjectsCompiler != nullptr);
-		RecompileCustomizableObjectsCompiler->Compile(*CustomizableObject, Options, true);
+		check(GetPrivate()->RecompileCustomizableObjectsCompiler != nullptr);
+		GetPrivate()->RecompileCustomizableObjectsCompiler->Compile(*CustomizableObject, Options, true);
 	}
 }
 
@@ -3503,14 +3486,14 @@ void UCustomizableObjectSystem::RecompileCustomizableObjectAsync(const FAssetDat
 		return;
 	}
 	
-	if ((InObject && InObject->GetPrivate()->IsLocked()) || ObjectsToRecompile.Find((InAssetData)) != INDEX_NONE)
+	if ((InObject && InObject->GetPrivate()->IsLocked()) || GetPrivate()->ObjectsToRecompile.Find((InAssetData)) != INDEX_NONE)
 	{
 		return;
 	}
 	
-	if (!ObjectsToRecompile.IsEmpty())
+	if (!GetPrivate()->ObjectsToRecompile.IsEmpty())
 	{
-		ObjectsToRecompile.Add(InAssetData);
+		GetPrivate()->ObjectsToRecompile.Add(InAssetData);
 	}
 	else
 	{
@@ -3527,29 +3510,29 @@ void UCustomizableObjectSystem::RecompileCustomizableObjects(const TArray<FAsset
 
 	if (InObjects.Num())
 	{
-		if (!RecompileCustomizableObjectsCompiler)
+		if (!GetPrivate()->RecompileCustomizableObjectsCompiler)
 		{
-			RecompileCustomizableObjectsCompiler = GetNewCompiler();
+			GetPrivate()->RecompileCustomizableObjectsCompiler = GetNewCompiler();
 
-			if (!RecompileCustomizableObjectsCompiler)
+			if (!GetPrivate()->RecompileCustomizableObjectsCompiler)
 			{
 				return;
 			}
 		}
 
-		ObjectsToRecompile.Append(InObjects);
+		GetPrivate()->ObjectsToRecompile.Append(InObjects);
 
-		TotalNumObjectsToRecompile = ObjectsToRecompile.Num();
-		NumObjectsCompiled = 0;
+		GetPrivate()->TotalNumObjectsToRecompile = GetPrivate()->ObjectsToRecompile.Num();
+		GetPrivate()->NumObjectsCompiled = 0;
 
-		if (RecompileNotificationHandle.IsValid())
+		if (GetPrivate()->RecompileNotificationHandle.IsValid())
 		{
-			++TotalNumObjectsToRecompile;
-			FSlateNotificationManager::Get().UpdateProgressNotification(RecompileNotificationHandle, NumObjectsCompiled, TotalNumObjectsToRecompile);
+			++GetPrivate()->TotalNumObjectsToRecompile;
+			FSlateNotificationManager::Get().UpdateProgressNotification(GetPrivate()->RecompileNotificationHandle, GetPrivate()->NumObjectsCompiled, GetPrivate()->TotalNumObjectsToRecompile);
 		}
 		else
 		{
-			RecompileNotificationHandle = FSlateNotificationManager::Get().StartProgressNotification(FText::FromString(TEXT("Compiling Customizable Objects")), TotalNumObjectsToRecompile);
+			GetPrivate()->RecompileNotificationHandle = FSlateNotificationManager::Get().StartProgressNotification(FText::FromString(TEXT("Compiling Customizable Objects")), GetPrivate()->TotalNumObjectsToRecompile);
 			StartNextRecompile();
 		}
 	}
@@ -3560,29 +3543,29 @@ void UCustomizableObjectSystem::TickRecompileCustomizableObjects()
 {
 	bool bUpdated = false;
 	
-	if (RecompileCustomizableObjectsCompiler)
+	if (GetPrivate()->RecompileCustomizableObjectsCompiler)
 	{
-		bUpdated = RecompileCustomizableObjectsCompiler->Tick() || RecompileCustomizableObjectsCompiler->GetCompilationState() == ECustomizableObjectCompilationState::Failed;
+		bUpdated = GetPrivate()->RecompileCustomizableObjectsCompiler->Tick() || GetPrivate()->RecompileCustomizableObjectsCompiler->GetCompilationState() == ECustomizableObjectCompilationState::Failed;
 	}
 
 	if (bUpdated)
 	{
-		NumObjectsCompiled++;
+		GetPrivate()->NumObjectsCompiled++;
 
-		if (!ObjectsToRecompile.IsEmpty())
+		if (!GetPrivate()->ObjectsToRecompile.IsEmpty())
 		{
 			StartNextRecompile();
 		}
 		else // All objects compiled, clean up
 		{
 			// Delete compiler
-			delete RecompileCustomizableObjectsCompiler;
-			RecompileCustomizableObjectsCompiler = nullptr;
+			delete GetPrivate()->RecompileCustomizableObjectsCompiler;
+			GetPrivate()->RecompileCustomizableObjectsCompiler = nullptr;
 
 			// Remove progress bar
-			FSlateNotificationManager::Get().UpdateProgressNotification(RecompileNotificationHandle, NumObjectsCompiled, TotalNumObjectsToRecompile);
-			FSlateNotificationManager::Get().CancelProgressNotification(RecompileNotificationHandle);
-			RecompileNotificationHandle.Reset();
+			FSlateNotificationManager::Get().UpdateProgressNotification(GetPrivate()->RecompileNotificationHandle, GetPrivate()->NumObjectsCompiled, GetPrivate()->TotalNumObjectsToRecompile);
+			FSlateNotificationManager::Get().CancelProgressNotification(GetPrivate()->RecompileNotificationHandle);
+			GetPrivate()->RecompileNotificationHandle.Reset();
 
 			if (GEngine)
 			{
@@ -3610,7 +3593,7 @@ uint64 UCustomizableObjectSystem::GetMaxChunkSizeForPlatform(const ITargetPlatfo
 {
 	const FString& PlatformName = TargetPlatform ? TargetPlatform->IniPlatformName() : FPlatformProperties::IniPlatformName();
 
-	if (const int64* CachedMaxChunkSize = PlatformMaxChunkSize.Find(PlatformName))
+	if (const int64* CachedMaxChunkSize = GetPrivate()->PlatformMaxChunkSize.Find(PlatformName))
 	{
 		return *CachedMaxChunkSize;
 	}
@@ -3634,7 +3617,7 @@ uint64 UCustomizableObjectSystem::GetMaxChunkSizeForPlatform(const ITargetPlatfo
 		MaxChunkSize = MUTABLE_STREAMED_DATA_MAXCHUNKSIZE;
 	}
 
-	PlatformMaxChunkSize.Add(PlatformName, MaxChunkSize);
+	GetPrivate()->PlatformMaxChunkSize.Add(PlatformName, MaxChunkSize);
 
 	return MaxChunkSize;
 }
@@ -3644,19 +3627,19 @@ uint64 UCustomizableObjectSystem::GetMaxChunkSizeForPlatform(const ITargetPlatfo
 
 void UCustomizableObjectSystem::CacheImage(FName ImageId)
 {
-	GetPrivateChecked()->GetImageProviderChecked()->CacheImage(ImageId, true);
+	GetPrivate()->GetImageProviderChecked()->CacheImage(ImageId, true);
 }
 
 
 void UCustomizableObjectSystem::UnCacheImage(FName ImageId)
 {
-	GetPrivateChecked()->GetImageProviderChecked()->UnCacheImage(ImageId, true);
+	GetPrivate()->GetImageProviderChecked()->UnCacheImage(ImageId, true);
 }
 
 
 void UCustomizableObjectSystem::ClearImageCache()
 {
-	GetPrivateChecked()->GetImageProviderChecked()->ClearCache(true);
+	GetPrivate()->GetImageProviderChecked()->ClearCache(true);
 }
 
 
@@ -3685,7 +3668,7 @@ void UCustomizableObjectSystemPrivate::OnMutableEnabledChanged(IConsoleVariable*
 	}
 
 	UCustomizableObjectSystem* System = UCustomizableObjectSystem::GetInstance();
-	UCustomizableObjectSystemPrivate* SystemPrivate = System->GetPrivateChecked();
+	UCustomizableObjectSystemPrivate* SystemPrivate = System->GetPrivate();
 
 	if (bIsMutableEnabled)
 	{
@@ -3763,7 +3746,7 @@ void UCustomizableObjectSystemPrivate::UpdateStats()
 bool UCustomizableObjectSystem::IsMutableAnimInfoDebuggingEnabled() const
 {
 #if WITH_EDITOR
-	return GetPrivateChecked()->IsMutableAnimInfoDebuggingEnabled();
+	return GetPrivate()->IsMutableAnimInfoDebuggingEnabled();
 #else
 	return false;
 #endif
diff --git a/Engine/Plugins/Experimental/Mutable/Source/CustomizableObject/Private/MuCO/UnrealExtensionDataStreamer.cpp b/Engine/Plugins/Experimental/Mutable/Source/CustomizableObject/Private/MuCO/UnrealExtensionDataStreamer.cpp
--- a/Engine/Plugins/Experimental/Mutable/Source/CustomizableObject/Private/MuCO/UnrealExtensionDataStreamer.cpp
+++ b/Engine/Plugins/Experimental/Mutable/Source/CustomizableObject/Private/MuCO/UnrealExtensionDataStreamer.cpp
@@ -224,7 +224,7 @@ TSharedPtr<FStreamableHandle> FUnrealExtensionDataStreamer::StartLoadOnGameThrea
 	DebugString = FString::Printf(TEXT("UnrealExtensionDataStreamer for %s"), *Object->GetPathName());
 #endif
 
-	FStreamableManager& Manager = UCustomizableObjectSystem::GetInstance()->GetStreamableManager();
+	FStreamableManager& Manager = UCustomizableObjectSystem::GetInstance()->GetPrivate()->StreamableManager;
 
 	return Manager.RequestAsyncLoad(
 		TargetsToStream,
diff --git a/Engine/Plugins/Experimental/Mutable/Source/CustomizableObjectEditor/Private/MuCOE/CustomizableObjectCompiler.cpp b/Engine/Plugins/Experimental/Mutable/Source/CustomizableObjectEditor/Private/MuCOE/CustomizableObjectCompiler.cpp
--- a/Engine/Plugins/Experimental/Mutable/Source/CustomizableObjectEditor/Private/MuCOE/CustomizableObjectCompiler.cpp
+++ b/Engine/Plugins/Experimental/Mutable/Source/CustomizableObjectEditor/Private/MuCOE/CustomizableObjectCompiler.cpp
@@ -25,6 +25,7 @@
 #include "Interfaces/ITargetPlatform.h"
 #include "Misc/App.h"
 #include "MuCO/CustomizableObjectPrivate.h"
+#include "MuCO/CustomizableObjectSystemPrivate.h"
 #include "MuR/Model.h"
 #include "MuR/ModelPrivate.h"
 
@@ -242,7 +243,7 @@ void FCustomizableObjectCompiler::Compile(UCustomizableObject& Object, const FCo
 	FCookLoadScope CookLoadScope(ECookLoadType::EditorOnly);
 	if (ArrayAssetToStream.Num() > 0)
 	{
-		FStreamableManager& Streamable = System->GetStreamableManager();
+		FStreamableManager& Streamable = System->GetPrivate()->StreamableManager;
 
 		if (bAsync) 
 		{
