[Mutable] Fix not generating all LODs.

#rnx

diff --git a/Engine/Plugins/Experimental/Mutable/Source/CustomizableObject/Private/MuCO/CustomizableObjectSystem.cpp b/Engine/Plugins/Experimental/Mutable/Source/CustomizableObject/Private/MuCO/CustomizableObjectSystem.cpp
--- a/Engine/Plugins/Experimental/Mutable/Source/CustomizableObject/Private/MuCO/CustomizableObjectSystem.cpp
+++ b/Engine/Plugins/Experimental/Mutable/Source/CustomizableObject/Private/MuCO/CustomizableObjectSystem.cpp
@@ -186,7 +186,7 @@ void FUpdateContextPrivate::SetMinLOD(int32 MinLOD)
 
 int32 FUpdateContextPrivate::GetMaxLOD() const
 {
-	return CapturedDescriptor.GetMinLod();
+	return CapturedDescriptor.GetMaxLod();
 }
 
 
@@ -1603,7 +1603,7 @@ namespace impl
 		Operation->SetMaxLOD(CurrentMaxLOD);
 
 		// Initialize RequestedLODs to zero if not set
-		TArray<uint16> RequestedLODs;
+		TArray<uint16> RequestedLODs = Operation->GetRequestedLODs();
 		RequestedLODs.SetNumZeroed(Operation->NumComponents);
 
 		for (int32 ComponentIndex = 0; ComponentIndex < Operation->NumComponents; ++ComponentIndex)
