[Mutable] Fix cook reference to an external private AssetUserData object.

#rnx

diff --git a/Engine/Plugins/Experimental/Mutable/Source/CustomizableObject/Internal/MuCO/CustomizableObjectPrivate.h b/Engine/Plugins/Experimental/Mutable/Source/CustomizableObject/Internal/MuCO/CustomizableObjectPrivate.h
--- a/Engine/Plugins/Experimental/Mutable/Source/CustomizableObject/Internal/MuCO/CustomizableObjectPrivate.h
+++ b/Engine/Plugins/Experimental/Mutable/Source/CustomizableObject/Internal/MuCO/CustomizableObjectPrivate.h
@@ -348,7 +348,7 @@ struct FMutableRefAssetUserData
 	GENERATED_BODY()
 
 	UPROPERTY()
-	TObjectPtr<UAssetUserData> AssetUserData;
+	TObjectPtr<UCustomizableObjectResourceDataContainer> AssetUserData;
 
 #if WITH_EDITORONLY_DATA
 	UPROPERTY()
diff --git a/Engine/Plugins/Experimental/Mutable/Source/CustomizableObject/Private/MuCO/CustomizableObject.cpp b/Engine/Plugins/Experimental/Mutable/Source/CustomizableObject/Private/MuCO/CustomizableObject.cpp
--- a/Engine/Plugins/Experimental/Mutable/Source/CustomizableObject/Private/MuCO/CustomizableObject.cpp
+++ b/Engine/Plugins/Experimental/Mutable/Source/CustomizableObject/Private/MuCO/CustomizableObject.cpp
@@ -2279,15 +2279,9 @@ void FMutableRefSkeletalMeshData::InitResources(UCustomizableObject* InOuter, co
 		}
 
 		FCustomizableObjectStreamedResourceData& StreamedResource = InOuter->StreamedResourceData[Data.AssetUserDataIndex];
-		StreamedResource.NotifyLoaded(StreamedResource.GetPath().LoadSynchronous());
-
-		const FCustomizableObjectResourceData& ResourceData = StreamedResource.GetLoadedData();
-		check(ResourceData.Type == ECOResourceDataType::AssetUserData);
-
-		const FCustomizableObjectAssetUserData* AUDResource = ResourceData.Data.GetPtr<FCustomizableObjectAssetUserData>();
-		check(AUDResource);
-
-		Data.AssetUserData = AUDResource->AssetUserData ? AUDResource->AssetUserData : AUDResource->AssetUserDataEditor;
+		Data.AssetUserData = StreamedResource.GetPath().LoadSynchronous();
+		check(Data.AssetUserData);
+		check(Data.AssetUserData->Data.Type == ECOResourceDataType::AssetUserData);
 	}
 }
 
diff --git a/Engine/Plugins/Experimental/Mutable/Source/CustomizableObject/Private/MuCO/CustomizableObjectInstance.cpp b/Engine/Plugins/Experimental/Mutable/Source/CustomizableObject/Private/MuCO/CustomizableObjectInstance.cpp
--- a/Engine/Plugins/Experimental/Mutable/Source/CustomizableObject/Private/MuCO/CustomizableObjectInstance.cpp
+++ b/Engine/Plugins/Experimental/Mutable/Source/CustomizableObject/Private/MuCO/CustomizableObjectInstance.cpp
@@ -3322,9 +3322,15 @@ void UCustomizableInstancePrivate::InitSkeletalMeshData(const TSharedRef<FUpdate
 	{
 		for (const FMutableRefAssetUserData& MutAssetUserData : RefSkeletalMeshData.AssetUserData)
 		{
-			if (MutAssetUserData.AssetUserData)
+			if (MutAssetUserData.AssetUserData && MutAssetUserData.AssetUserData->Data.Type == ECOResourceDataType::AssetUserData)
 			{
-				SkeletalMesh->AddAssetUserData(MutAssetUserData.AssetUserData);
+				const FCustomizableObjectAssetUserData* DataPtr = MutAssetUserData.AssetUserData->Data.Data.GetPtr<FCustomizableObjectAssetUserData>();
+				check(DataPtr);
+#if WITH_EDITORONLY_DATA
+				SkeletalMesh->AddAssetUserData(DataPtr->AssetUserDataEditor);
+#else
+				SkeletalMesh->AddAssetUserData(DataPtr->AssetUserData);
+#endif
 			}
 		}
 
diff --git a/Engine/Plugins/Experimental/Mutable/Source/CustomizableObjectEditor/Private/MuCOE/GenerateMutableSource/GenerateMutableSource.cpp b/Engine/Plugins/Experimental/Mutable/Source/CustomizableObjectEditor/Private/MuCOE/GenerateMutableSource/GenerateMutableSource.cpp
--- a/Engine/Plugins/Experimental/Mutable/Source/CustomizableObjectEditor/Private/MuCOE/GenerateMutableSource/GenerateMutableSource.cpp
+++ b/Engine/Plugins/Experimental/Mutable/Source/CustomizableObjectEditor/Private/MuCOE/GenerateMutableSource/GenerateMutableSource.cpp
@@ -1613,7 +1613,9 @@ void PopulateReferenceSkeletalMeshesData(FMutableGraphGenerationContext& Generat
 				{
 					FMutableRefAssetUserData& MutAssetUserData = Data.AssetUserData.AddDefaulted_GetRef();
 					MutAssetUserData.AssetUserDataIndex = GenerationContext.AddAssetUserDataToStreamedResources(AssetUserData);
-					MutAssetUserData.AssetUserData = AssetUserData;
+					MutAssetUserData.AssetUserData = GenerationContext.StreamedResourceData[MutAssetUserData.AssetUserDataIndex].GetPath().Get();
+					check(MutAssetUserData.AssetUserData);
+					check(MutAssetUserData.AssetUserData->Data.Type == ECOResourceDataType::AssetUserData);
 				}
 			}
 		}
