[Mutable] Warn in PIE if a CO require a recompilation.

A warning will be shown if the asset or its participating objects are dirty or have been re-saved since the last compilation.

- Moved bIsCompiledWithOptimization warning to the tick.

#rnx

diff --git a/Engine/Plugins/Experimental/Mutable/Source/CustomizableObject/Classes/MuCO/CustomizableObjectSystem.h b/Engine/Plugins/Experimental/Mutable/Source/CustomizableObject/Classes/MuCO/CustomizableObjectSystem.h
--- a/Engine/Plugins/Experimental/Mutable/Source/CustomizableObject/Classes/MuCO/CustomizableObjectSystem.h
+++ b/Engine/Plugins/Experimental/Mutable/Source/CustomizableObject/Classes/MuCO/CustomizableObjectSystem.h
@@ -33,6 +33,8 @@ struct FFrame;
 struct FGuid;
 
 
+constexpr uint64 KEY_OFFSET_COMPILATION_OUT_OF_DATE = 1;
+
 extern TAutoConsoleVariable<bool> CVarClearWorkingMemoryOnUpdateEnd;
 
 extern TAutoConsoleVariable<bool> CVarReuseImagesBetweenInstances;
diff --git a/Engine/Plugins/Experimental/Mutable/Source/CustomizableObject/Internal/MuCO/CustomizableObjectPrivate.h b/Engine/Plugins/Experimental/Mutable/Source/CustomizableObject/Internal/MuCO/CustomizableObjectPrivate.h
--- a/Engine/Plugins/Experimental/Mutable/Source/CustomizableObject/Internal/MuCO/CustomizableObjectPrivate.h
+++ b/Engine/Plugins/Experimental/Mutable/Source/CustomizableObject/Internal/MuCO/CustomizableObjectPrivate.h
@@ -573,6 +573,11 @@ public:
 	// See UCustomizableObjectSystem::LockObject()
 	bool IsLocked() const;
 
+#if WITH_EDITORONLY_DATA
+	/** Return true if the CO is not compiled or the ParticipatingObjects system has detected a change (participating objects dirty or re-saved since last compilation). */
+	bool IsCompilationOutOfDate() const;
+#endif
+	
 	/** Cache of generated SkeletalMeshes */
 	FMeshCache MeshCache;
 
diff --git a/Engine/Plugins/Experimental/Mutable/Source/CustomizableObject/Internal/MuCO/CustomizableObjectSystemPrivate.h b/Engine/Plugins/Experimental/Mutable/Source/CustomizableObject/Internal/MuCO/CustomizableObjectSystemPrivate.h
--- a/Engine/Plugins/Experimental/Mutable/Source/CustomizableObject/Internal/MuCO/CustomizableObjectSystemPrivate.h
+++ b/Engine/Plugins/Experimental/Mutable/Source/CustomizableObject/Internal/MuCO/CustomizableObjectSystemPrivate.h
@@ -693,6 +693,13 @@ public:
 	void CacheTextureParameters(const TArray<FCustomizableObjectTextureParameterValue>& TextureParameters) const;
 
 	void UnCacheTextureParameters(const TArray<FCustomizableObjectTextureParameterValue>& TextureParameters) const;
+
+#if WITH_EDITORONLY_DATA
+	/**	PIE CO on-screen warnings. */
+	static void ShowOnScreenCompileWarnings();
+
+	static void HideOnScreenCompileWarnings(const UCustomizableObjectPrivate& ObjectPrivate);
+#endif
 	
 	/** Mutable TaskGraph system (Mutable Thread). */
 	FMutableTaskGraph MutableTaskGraph;
diff --git a/Engine/Plugins/Experimental/Mutable/Source/CustomizableObject/Private/MuCO/CustomizableObject.cpp b/Engine/Plugins/Experimental/Mutable/Source/CustomizableObject/Private/MuCO/CustomizableObject.cpp
--- a/Engine/Plugins/Experimental/Mutable/Source/CustomizableObject/Private/MuCO/CustomizableObject.cpp
+++ b/Engine/Plugins/Experimental/Mutable/Source/CustomizableObject/Private/MuCO/CustomizableObject.cpp
@@ -39,6 +39,8 @@
 #endif
 
 
+#include "MuCO/CustomizableObjectSystemPrivate.h"
+
 #include UE_INLINE_GENERATED_CPP_BY_NAME(CustomizableObject)
 
 #define LOCTEXT_NAMESPACE "CustomizableObject"
@@ -852,13 +854,6 @@ bool UCustomizableObject::ConditionalAutoCompile()
 	// Don't compile compiled objects
 	if (IsCompiled())
 	{
-		// Show a warning if the compilation was not done with optimizations.
-		if (GetPrivate()->bIsCompiledWithOptimization)
-		{
-			FString Msg = FString::Printf(TEXT("Warning: Customizable Object [%s] was compiled without optimization."), *GetName());
-			GEngine->AddOnScreenDebugMessage((uint64)((PTRINT)this), 10.0f, FColor::Red, Msg);
-		}
-
 		return true;
 	}
 
@@ -1186,6 +1181,8 @@ FString UCustomizableObject::GetStateParameterName(int32 StateIndex, int32 Param
 void UCustomizableObject::PostCompile()
 {
 	GetPrivate()->PostCompileDelegate.Broadcast();
+
+	UCustomizableObjectSystemPrivate::HideOnScreenCompileWarnings(*GetPrivate());	
 }
 #endif
 
@@ -1915,6 +1912,60 @@ FModelResources& UCustomizableObjectPrivate::GetModelResources(bool bIsCooking)
 #endif
 
 
+#if WITH_EDITORONLY_DATA
+bool UCustomizableObjectPrivate::IsCompilationOutOfDate() const
+{
+	if (Status.Get() != FCustomizableObjectStatus::EState::ModelLoaded)
+	{
+		return false;
+	}
+	
+	if (GetPackage()->IsDirty())
+	{
+		return true;
+	}
+	
+	const FAssetRegistryModule& AssetRegistryModule = FModuleManager::LoadModuleChecked<FAssetRegistryModule>("AssetRegistry");
+
+	for (const TTuple<FName, FGuid>& ParticipatingObject : ParticipatingObjects)
+	{
+		TSoftObjectPtr Object(ParticipatingObject.Key.ToString());
+		if (Object) // If loaded
+		{
+			PRAGMA_DISABLE_DEPRECATION_WARNINGS
+			const FGuid PackageGuid = Object->GetPackage()->GetGuid();
+			PRAGMA_ENABLE_DEPRECATION_WARNINGS
+			
+			if (Object->GetPackage()->IsDirty() || PackageGuid != ParticipatingObject.Value)
+			{
+				return true;
+			}
+		}
+		else // Not loaded
+		{
+			FAssetPackageData AssetPackageData;
+			const UE::AssetRegistry::EExists Result = AssetRegistryModule.Get().TryGetAssetPackageData(ParticipatingObject.Key, AssetPackageData);
+				
+			if (Result != UE::AssetRegistry::EExists::Exists)
+			{
+				return true;
+			}
+
+			PRAGMA_DISABLE_DEPRECATION_WARNINGS
+			const FGuid PackageGuid = AssetPackageData.PackageGuid;
+			PRAGMA_ENABLE_DEPRECATION_WARNINGS
+			
+			if (PackageGuid != ParticipatingObject.Value)
+			{
+				return true;
+			}
+		}
+	}
+
+	return false;
+}
+#endif
+
 //-------------------------------------------------------------------------------------------------
 //-------------------------------------------------------------------------------------------------
 //-------------------------------------------------------------------------------------------------
diff --git a/Engine/Plugins/Experimental/Mutable/Source/CustomizableObject/Private/MuCO/CustomizableObjectSystem.cpp b/Engine/Plugins/Experimental/Mutable/Source/CustomizableObject/Private/MuCO/CustomizableObjectSystem.cpp
--- a/Engine/Plugins/Experimental/Mutable/Source/CustomizableObject/Private/MuCO/CustomizableObjectSystem.cpp
+++ b/Engine/Plugins/Experimental/Mutable/Source/CustomizableObject/Private/MuCO/CustomizableObjectSystem.cpp
@@ -3102,6 +3102,10 @@ bool UCustomizableObjectSystem::Tick(float DeltaTime)
 	TickRecompileCustomizableObjects();
 #endif
 
+#if WITH_EDITORONLY_DATA
+	UCustomizableObjectSystemPrivate::ShowOnScreenCompileWarnings();
+#endif
+	
 	Private->MutableTaskGraph.Tick();
 
 	Private->LogBenchmarkUtil.UpdateStats(); // Must to be the last thing to perform
@@ -3589,6 +3593,79 @@ int32 UCustomizableObjectSystem::GetWorkingMemory() const
 }
 
 
+#if WITH_EDITORONLY_DATA
+void UCustomizableObjectSystemPrivate::ShowOnScreenCompileWarnings()
+{
+	for (TObjectIterator<UCustomizableObjectInstanceUsage> CustomizableObjectInstanceUsage; CustomizableObjectInstanceUsage; ++CustomizableObjectInstanceUsage)
+	{
+		if (!IsValid(*CustomizableObjectInstanceUsage) || CustomizableObjectInstanceUsage->IsTemplate())
+		{
+			continue;
+		}
+		
+		const UCustomizableObjectInstance* Instance = CustomizableObjectInstanceUsage->GetCustomizableObjectInstance();
+		if (!Instance || !Instance->GetCustomizableObject())
+		{
+			continue;
+		}
+
+		const UCustomizableObject* Object = Cast<UCustomizableObject>(Instance->GetCustomizableObject());
+		if (!Object)
+		{
+			continue;
+		}
+		
+		const USkeletalMeshComponent* Parent = Cast<USkeletalMeshComponent>(CustomizableObjectInstanceUsage->GetAttachParent());
+		if (!Parent)
+		{
+			continue;
+		}
+
+		EWorldType::Type WorldType = EWorldType::Type::None;
+		if (const UWorld* World = Parent ? Parent->GetWorld() : nullptr)
+		{
+			WorldType = World->WorldType;
+		}
+
+		if (WorldType != EWorldType::PIE)
+		{
+			continue;
+		}
+		
+		// Show a warning if the compilation was not done with optimizations.
+		const uint64 KeyCompiledWithOptimization = reinterpret_cast<uint64>(Object);
+		if (!GEngine->OnScreenDebugMessageExists(KeyCompiledWithOptimization) && Object->GetPrivate()->bIsCompiledWithOptimization)
+		{
+			FString Msg = FString::Printf(TEXT("Warning: Customizable Object [%s] was compiled without optimization."), *Object->GetName());
+			GEngine->AddOnScreenDebugMessage(KeyCompiledWithOptimization, 10.0f, FColor::Red, Msg);
+		}
+
+		const uint64 KeyCompiledOutOfData = reinterpret_cast<uint64>(Object) + KEY_OFFSET_COMPILATION_OUT_OF_DATE; // Offset added to avoid collision with bIsCompiledWithOptimization warning
+		if (!GEngine->OnScreenDebugMessageExists(KeyCompiledWithOptimization) && Object->GetPrivate()->IsCompilationOutOfDate())
+		{
+			FString Msg = FString::Printf(TEXT("Warning: Customizable Object [%s] compilation out of date. Save all referenced assets and recompile."), *Object->GetName());
+			GEngine->AddOnScreenDebugMessage(KeyCompiledOutOfData, 10.0f, FColor::Red, Msg);
+		}
+	}
+}
+
+
+void UCustomizableObjectSystemPrivate::HideOnScreenCompileWarnings(const UCustomizableObjectPrivate& ObjectPrivate)
+{
+	if (ObjectPrivate.bIsCompiledWithOptimization) 
+	{
+		GEngine->RemoveOnScreenDebugMessage(reinterpret_cast<uint64>(&ObjectPrivate));
+	}
+
+	if (!ObjectPrivate.IsCompilationOutOfDate())
+	{
+		GEngine->RemoveOnScreenDebugMessage(reinterpret_cast<uint64>(&ObjectPrivate) + KEY_OFFSET_COMPILATION_OUT_OF_DATE);
+	}
+}
+
+#endif
+
+
 uint64 UCustomizableObjectSystem::GetMaxChunkSizeForPlatform(const ITargetPlatform* TargetPlatform)
 {
 	const FString& PlatformName = TargetPlatform ? TargetPlatform->IniPlatformName() : FPlatformProperties::IniPlatformName();
