[mutable] Fixed some bugs related to group projectors:
+ Fixed parameters GUids.
+ Fixed imge parameter selector.
+ Fixed projector type selection. It was always setting the projector's type to planar.

#rnx

diff --git a/Engine/Plugins/Experimental/Mutable/Source/CustomizableObject/Private/MuCO/CustomizableObjectInstance.cpp b/Engine/Plugins/Experimental/Mutable/Source/CustomizableObject/Private/MuCO/CustomizableObjectInstance.cpp
--- a/Engine/Plugins/Experimental/Mutable/Source/CustomizableObject/Private/MuCO/CustomizableObjectInstance.cpp
+++ b/Engine/Plugins/Experimental/Mutable/Source/CustomizableObject/Private/MuCO/CustomizableObjectInstance.cpp
@@ -2213,7 +2213,7 @@ int32 UCustomizableObjectInstance::RemoveValueFromIntRange(const FString& ParamN
 
 int32 UCustomizableObjectInstance::RemoveValueFromIntRange(const FString& ParamName, int32 RangeIndex)
 {
-	return Descriptor.RemoveValueFromIntRange(ParamName);
+	return Descriptor.RemoveValueFromIntRange(ParamName, RangeIndex);
 
 }
 
diff --git a/Engine/Plugins/Experimental/Mutable/Source/CustomizableObject/Private/MuCO/CustomizableObjectInstanceDescriptor.cpp b/Engine/Plugins/Experimental/Mutable/Source/CustomizableObject/Private/MuCO/CustomizableObjectInstanceDescriptor.cpp
--- a/Engine/Plugins/Experimental/Mutable/Source/CustomizableObject/Private/MuCO/CustomizableObjectInstanceDescriptor.cpp
+++ b/Engine/Plugins/Experimental/Mutable/Source/CustomizableObject/Private/MuCO/CustomizableObjectInstanceDescriptor.cpp
@@ -605,6 +605,8 @@ mu::ParametersPtr FCustomizableObjectInstanceDescriptor::GetParameters() const
 						}
 					};
 
+					CopyProjector(ProjectorParameter.Value);
+
 					if (const mu::RangeIndexPtr RangeIdxPtr = MutableParameters->NewRangeIndex(ParamIndex))
 					{
 						for (int32 RangeIndex = 0; RangeIndex < ProjectorParameter.RangeValues.Num(); ++RangeIndex)
@@ -613,10 +615,6 @@ mu::ParametersPtr FCustomizableObjectInstanceDescriptor::GetParameters() const
 							CopyProjector(ProjectorParameter.RangeValues[RangeIndex], RangeIdxPtr);
 						}
 					}
-					else
-					{
-						CopyProjector(ProjectorParameter.Value);
-					}
 				}
 			}
 
@@ -844,28 +842,19 @@ void FCustomizableObjectInstanceDescriptor::ReloadParameters()
 				
 				if (mu::RangeIndexPtr RangeIdxPtr = MutableParameters->NewRangeIndex(ParamIndex)) // Is multidimensional
 				{
-					// Find the max RangeIndex of this parameter 
-					int32 RangeNum = 0;
-					const int32 ValueCount = MutableParameters->GetValueCount(ParamIndex);
-					for (int32 ValueIndex = 0; ValueIndex < ValueCount; ++ValueIndex)
-					{
-						mu::RangeIndexPtr RangeValueIdxPtr = MutableParameters->GetValueIndex(ParamIndex, ValueIndex);
-						RangeNum = FMath::Max(RangeNum, RangeValueIdxPtr->GetPosition(0));
-					}
-					
-					Param.ParameterRangeValueNames.Reserve(RangeNum); 
+					// Get num of ranges (layers) from the instance
+					int32 ValueCount = Result->ParameterRangeValueNames.Num();
+					Param.ParameterRangeValueNames.Reserve(ValueCount);
 
-					for (int32 RangeIndex = 0; RangeIndex < RangeNum; ++RangeIndex)
+					for (int32 RangeIndex = 0; RangeIndex < ValueCount; ++RangeIndex)
 					{
-						if (const FString& OldValue = Result->ParameterRangeValueNames[RangeIndex];
-							ValueExists(OldValue)) // Value still exits
+						// Checking if the selected value still exists as option in the parameter
+						if (const FString& OldValue = Result->ParameterRangeValueNames[RangeIndex]; ValueExists(OldValue))
 						{
-							Param.ParameterRangeValueNames.Add(OldValue);							
+							Param.ParameterRangeValueNames.Add(OldValue);
 						}
 						else
 						{
-							RangeIdxPtr->SetPosition(0, RangeIndex);
-							
 							const int32 Value = MutableParameters->GetIntValue(ParamIndex, RangeIdxPtr);
 							const FString AuxParameterValueName = CustomizableObject->FindIntParameterValueName(ParamIndex, Value);
 							Param.ParameterRangeValueNames.Add(AuxParameterValueName);
@@ -1000,6 +989,9 @@ void FCustomizableObjectInstanceDescriptor::ReloadParameters()
 			Param.ParameterName = Name;
 			Param.Id = Uid;
 
+			// Projector to check if the porojector's type has changed
+			FCustomizableObjectProjector DefaultProjectorValue = CustomizableObject->GetProjectorParameterDefaultValue(Name);
+
 			auto FindByNameAndUid = [&](const FCustomizableObjectProjectorParameterValue& P)
 			{
 				return P.ParameterName == Name || (Uid.IsValid() && P.Id == Uid);
@@ -1010,10 +1002,17 @@ void FCustomizableObjectInstanceDescriptor::ReloadParameters()
 				if (mu::RangeIndexPtr RangeIdxPtr = MutableParameters->NewRangeIndex(ParamIndex))
 				{
 					Param.RangeValues = Result->RangeValues;
+					Param.Value.ProjectionType = DefaultProjectorValue.ProjectionType;
+
+					for (FCustomizableObjectProjector& Projector : Param.RangeValues)
+					{
+						Projector.ProjectionType = DefaultProjectorValue.ProjectionType;
+					}
 				}
 				else
 				{
 					Param.Value = Result->Value;
+					Param.Value.ProjectionType = DefaultProjectorValue.ProjectionType;
 				}
 			} 
 			else // Not found in Instance Parameters. Use Mutable Parameters.
@@ -1040,7 +1039,9 @@ void FCustomizableObjectInstanceDescriptor::ReloadParameters()
 						Value.Scale[2] = -Value.Scale[0];
 						Value.Scale[0] = Value.Scale[1] = Value.Scale[1] * 2.0f;
 					}
-				}; 
+				};
+
+				GetProjector(Param.Value);
 
 				if (mu::RangeIndexPtr RangeIdxPtr = MutableParameters->NewRangeIndex(ParamIndex))
 				{
@@ -1059,10 +1060,6 @@ void FCustomizableObjectInstanceDescriptor::ReloadParameters()
 						GetProjector(Param.RangeValues[RangeIndex], RangeValueIdxPtr);
 					}
 				}
-				else
-				{
-					GetProjector(Param.Value);
-				}
 			}
 
 			ProjectorParameters.Add(Param);
diff --git a/Engine/Plugins/Experimental/Mutable/Source/CustomizableObject/Public/MuCO/CustomizableObjectCustomVersion.h b/Engine/Plugins/Experimental/Mutable/Source/CustomizableObject/Public/MuCO/CustomizableObjectCustomVersion.h
--- a/Engine/Plugins/Experimental/Mutable/Source/CustomizableObject/Public/MuCO/CustomizableObjectCustomVersion.h
+++ b/Engine/Plugins/Experimental/Mutable/Source/CustomizableObject/Public/MuCO/CustomizableObjectCustomVersion.h
@@ -112,6 +112,8 @@ struct CUSTOMIZABLEOBJECT_API FCustomizableObjectCustomVersion
 
 		AddedAnyTextureTypeToPassThroughTextures,
 
+		ProjectorNodesDefaultValueFix,
+
 		// -----<new versions can be added above this line>-------------------------------------------------
 		VersionPlusOne,
 		LatestVersion = VersionPlusOne - 1
diff --git a/Engine/Plugins/Experimental/Mutable/Source/CustomizableObject/Public/MuCO/CustomizableObjectParameterTypeDefinitions.h b/Engine/Plugins/Experimental/Mutable/Source/CustomizableObject/Public/MuCO/CustomizableObjectParameterTypeDefinitions.h
--- a/Engine/Plugins/Experimental/Mutable/Source/CustomizableObject/Public/MuCO/CustomizableObjectParameterTypeDefinitions.h
+++ b/Engine/Plugins/Experimental/Mutable/Source/CustomizableObject/Public/MuCO/CustomizableObjectParameterTypeDefinitions.h
@@ -230,23 +230,23 @@ struct CUSTOMIZABLEOBJECT_API FCustomizableObjectProjector
 {
 	GENERATED_USTRUCT_BODY()
 
-	UPROPERTY(EditAnywhere, Category = CustomizableObject)
+	UPROPERTY()
 	FVector3f Position = FVector3f(0, 0, 0);
 
-	UPROPERTY(EditAnywhere, Category = CustomizableObject)
+	UPROPERTY()
 	FVector3f Direction = FVector3f(1, 0, 0);
 
-	UPROPERTY(EditAnywhere, Category = CustomizableObject)
+	UPROPERTY()
 	FVector3f Up = FVector3f(0, 1, 0);
 
-	UPROPERTY(EditAnywhere, Category = CustomizableObject)
+	UPROPERTY()
 	FVector3f Scale = FVector3f(10, 10, 100);
 
 	UPROPERTY(EditAnywhere, Category = CustomizableObject)
 	ECustomizableObjectProjectorType ProjectionType = ECustomizableObjectProjectorType::Planar;
 
 	// Just for cylindrical projectors, in radians
-	UPROPERTY(EditAnywhere, Category = CustomizableObject)
+	UPROPERTY()
 	float Angle = 2.0f * PI;
 };
 
diff --git a/Engine/Plugins/Experimental/Mutable/Source/CustomizableObjectEditor/Classes/MuCOE/Nodes/CustomizableObjectNodeGroupProjectorParameter.h b/Engine/Plugins/Experimental/Mutable/Source/CustomizableObjectEditor/Classes/MuCOE/Nodes/CustomizableObjectNodeGroupProjectorParameter.h
--- a/Engine/Plugins/Experimental/Mutable/Source/CustomizableObjectEditor/Classes/MuCOE/Nodes/CustomizableObjectNodeGroupProjectorParameter.h
+++ b/Engine/Plugins/Experimental/Mutable/Source/CustomizableObjectEditor/Classes/MuCOE/Nodes/CustomizableObjectNodeGroupProjectorParameter.h
@@ -39,7 +39,7 @@ struct FGroupProjectorParameterPose
 };
 
 
-UCLASS(hideCategories = (CustomizableObjectHide))
+UCLASS()
 class CUSTOMIZABLEOBJECTEDITOR_API UCustomizableObjectNodeGroupProjectorParameter : public UCustomizableObjectNodeProjectorParameter
 {
 public:
diff --git a/Engine/Plugins/Experimental/Mutable/Source/CustomizableObjectEditor/Classes/MuCOE/Nodes/CustomizableObjectNodeProjectorConstant.h b/Engine/Plugins/Experimental/Mutable/Source/CustomizableObjectEditor/Classes/MuCOE/Nodes/CustomizableObjectNodeProjectorConstant.h
--- a/Engine/Plugins/Experimental/Mutable/Source/CustomizableObjectEditor/Classes/MuCOE/Nodes/CustomizableObjectNodeProjectorConstant.h
+++ b/Engine/Plugins/Experimental/Mutable/Source/CustomizableObjectEditor/Classes/MuCOE/Nodes/CustomizableObjectNodeProjectorConstant.h
@@ -14,7 +14,7 @@ class UObject;
 struct FPropertyChangedEvent;
 
 
-UCLASS(hideCategories = (CustomizableObjectHide))
+UCLASS()
 class CUSTOMIZABLEOBJECTEDITOR_API UCustomizableObjectNodeProjectorConstant : public UCustomizableObjectNode
 {
 public:
@@ -23,15 +23,12 @@ public:
 	UCustomizableObjectNodeProjectorConstant();
 
 	/**  */
-	UPROPERTY(EditAnywhere, Category= CustomizableObjectHide)
+	UPROPERTY(EditAnywhere, Category = CustomizableObject, meta = (ShowOnlyInnerProperties))
 	FCustomizableObjectProjector Value;
 
 	UPROPERTY(EditAnywhere, Category = UI, meta = (DisplayName = "Parameter UI Metadata"))
 	FMutableParamUIMetadata ParamUIMetadata;
 
-	UPROPERTY(EditAnywhere, Category = CustomizableObject)
-	ECustomizableObjectProjectorType ProjectionType;
-
 	UPROPERTY(EditAnywhere, Category = CustomizableObject, Meta = (DisplayName = "Projection Angle (degrees)"))
 	float ProjectionAngle;
 
@@ -60,6 +57,7 @@ public:
 
 	// UCustomizableObjectNode interface
 	void AllocateDefaultPins(UCustomizableObjectNodeRemapPins* RemapPins) override;
+	virtual void BackwardsCompatibleFixup() override;
 
 	// Own interface
 	UEdGraphPin* ValuePin();
@@ -85,5 +83,10 @@ public:
 	float GetProjectorAngle();
 
 	void SetProjectorAngle(float Angle);
+
+private:
+
+	UPROPERTY()
+	ECustomizableObjectProjectorType ProjectionType_DEPRECATED;
 };
 
diff --git a/Engine/Plugins/Experimental/Mutable/Source/CustomizableObjectEditor/Classes/MuCOE/Nodes/CustomizableObjectNodeProjectorParameter.h b/Engine/Plugins/Experimental/Mutable/Source/CustomizableObjectEditor/Classes/MuCOE/Nodes/CustomizableObjectNodeProjectorParameter.h
--- a/Engine/Plugins/Experimental/Mutable/Source/CustomizableObjectEditor/Classes/MuCOE/Nodes/CustomizableObjectNodeProjectorParameter.h
+++ b/Engine/Plugins/Experimental/Mutable/Source/CustomizableObjectEditor/Classes/MuCOE/Nodes/CustomizableObjectNodeProjectorParameter.h
@@ -14,13 +14,13 @@ class UObject;
 struct FPropertyChangedEvent;
 
 
-UCLASS(hideCategories = (CustomizableObjectHide))
+UCLASS()
 class CUSTOMIZABLEOBJECTEDITOR_API UCustomizableObjectNodeProjectorParameter : public UCustomizableObjectNode
 {
 public:
 	GENERATED_BODY()
 
-	UPROPERTY(EditAnywhere, Category=CustomizableObjectHide)
+	UPROPERTY(EditAnywhere, Category=CustomizableObject, meta = (ShowOnlyInnerProperties))
 	FCustomizableObjectProjector DefaultValue;
 
 	UPROPERTY(EditAnywhere, Category=CustomizableObject)
@@ -29,9 +29,6 @@ public:
 	UPROPERTY(EditAnywhere, Category = UI, meta = (DisplayName = "Parameter UI Metadata"))
 	FMutableParamUIMetadata ParamUIMetadata;
 
-	UPROPERTY(EditAnywhere, Category = CustomizableObject)
-	ECustomizableObjectProjectorType ProjectionType = ECustomizableObjectProjectorType::Planar;
-
 	UPROPERTY(EditAnywhere, Category = CustomizableObject, Meta = (DisplayName = "Projection Angle (degrees)"))
 	float ProjectionAngle = 360.0f;
 
@@ -61,6 +58,7 @@ public:
 	// UCustomizableObjectNode interface
 	void AllocateDefaultPins(UCustomizableObjectNodeRemapPins* RemapPins) override;
 	virtual bool IsAffectedByLOD() const override { return false; }
+	virtual void BackwardsCompatibleFixup() override;
 
 	// Own interface
 	UEdGraphPin* ProjectorPin() const
@@ -89,5 +87,10 @@ public:
 	float GetProjectorDefaultAngle() const;
 
 	void SetProjectorDefaultAngle(float Angle);
+
+private:
+
+	UPROPERTY()
+	ECustomizableObjectProjectorType ProjectionType_DEPRECATED;
 };
 
diff --git a/Engine/Plugins/Experimental/Mutable/Source/CustomizableObjectEditor/Private/MuCOE/CustomizableInstanceDetails.cpp b/Engine/Plugins/Experimental/Mutable/Source/CustomizableObjectEditor/Private/MuCOE/CustomizableInstanceDetails.cpp
--- a/Engine/Plugins/Experimental/Mutable/Source/CustomizableObjectEditor/Private/MuCOE/CustomizableInstanceDetails.cpp
+++ b/Engine/Plugins/Experimental/Mutable/Source/CustomizableObjectEditor/Private/MuCOE/CustomizableInstanceDetails.cpp
@@ -968,7 +968,20 @@ TSharedRef<SWidget> FCustomizableInstanceDetails::GenerateFloatWidget(const int3
 
 float FCustomizableInstanceDetails::GetFloatParameterValue(FString ParamName, int32 RangeIndex) const
 {
-	return CustomInstance->GetFloatParameterSelectedOption(ParamName, RangeIndex);
+	if (RangeIndex == INDEX_NONE)
+	{
+		return CustomInstance->GetFloatParameterSelectedOption(ParamName, RangeIndex);
+	}
+	else //multidimensional
+	{
+		// We may have deleted a range but the Instance has not been updated yet
+		if (CustomInstance->GetFloatValueRange(ParamName) > RangeIndex)
+		{
+			return CustomInstance->GetFloatParameterSelectedOption(ParamName, RangeIndex);
+		}
+	}
+
+	return 0;
 }
 
 
@@ -1443,7 +1456,7 @@ TSharedRef<SWidget> FCustomizableInstanceDetails::GenerateMultidimensionalProjec
 		FString OpacitySliderParamName = ParamName + FMultilayerProjector::OPACITY_PARAMETER_POSTFIX;
 
 		TArray<FString> OptionNamesAttribute;
-		FString Value = CustomInstance->GetIntParameterSelectedOption(ParamName, RangeIndex);
+		FString Value = CustomInstance->GetIntParameterSelectedOption(TextureSwitchEnumParamName, RangeIndex);
 		int32 ValueIndex = 0;
 
 		const UProjectorParameter* ProjectorParameter = Editor->GetProjectorParameter();
diff --git a/Engine/Plugins/Experimental/Mutable/Source/CustomizableObjectEditor/Private/MuCOE/GenerateMutableSource/GenerateMutableSourceGroupProjector.cpp b/Engine/Plugins/Experimental/Mutable/Source/CustomizableObjectEditor/Private/MuCOE/GenerateMutableSource/GenerateMutableSourceGroupProjector.cpp
--- a/Engine/Plugins/Experimental/Mutable/Source/CustomizableObjectEditor/Private/MuCOE/GenerateMutableSource/GenerateMutableSourceGroupProjector.cpp
+++ b/Engine/Plugins/Experimental/Mutable/Source/CustomizableObjectEditor/Private/MuCOE/GenerateMutableSource/GenerateMutableSourceGroupProjector.cpp
@@ -326,8 +326,15 @@ bool GenerateMutableSourceGroupProjector(const UEdGraphPin* Pin, FMutableGraphGe
 
 		if (GroupProjectorTempData.NodeProjectorParameterPtr)
 		{
-			// Use the projector parameter uid + suffix to identify parameters derived from this node
-			FString ProjectorParamUid = GroupProjectorTempData.NodeProjectorParameterPtr->GetUid();
+			// Use the projector parameter uid + num to identify parameters derived from this node
+			FGuid NumLayersParamUid = Node->NodeGuid;
+			NumLayersParamUid.D += 1;
+			FGuid SelectedPoseParamUid = Node->NodeGuid;
+			SelectedPoseParamUid.D += 2;
+			FGuid OpacityParamUid = Node->NodeGuid;
+			OpacityParamUid.D += 3;
+			FGuid SelectedImageParamUid = Node->NodeGuid;
+			SelectedImageParamUid.D += 4;
 
 			// Add to UCustomizableObjectNodeGroupProjectorParameter::OptionImages those textures that are present in
 			// UCustomizableObjectNodeGroupProjectorParameter::OptionImagesDataTable avoiding any repeated element
@@ -345,7 +352,7 @@ bool GenerateMutableSourceGroupProjector(const UEdGraphPin* Pin, FMutableGraphGe
 			mu::NodeScalarParameterPtr NodeScalarParam = new mu::NodeScalarParameter;
 			FString NodeScalarParamName = ProjParamNode->ParameterName + FMultilayerProjector::NUM_LAYERS_PARAMETER_POSTFIX;
 			NodeScalarParam->SetName(NodeScalarParamName);
-			NodeScalarParam->SetUid(ProjectorParamUid + FString("_NL"));
+			NodeScalarParam->SetUid(NumLayersParamUid.ToString());
 			GenerationContext.AddParameterNameUnique(originalGroup, NodeScalarParamName);
 			NodeScalarParam->SetDefaultValue(0.f);
 
@@ -364,7 +371,7 @@ bool GenerateMutableSourceGroupProjector(const UEdGraphPin* Pin, FMutableGraphGe
 			mu::NodeScalarEnumParameterPtr PoseEnumParameterNode = new mu::NodeScalarEnumParameter;
 			FString PoseNodeEnumParamName = ProjParamNode->ParameterName + FMultilayerProjector::POSE_PARAMETER_POSTFIX;
 			PoseEnumParameterNode->SetName(PoseNodeEnumParamName);
-			PoseEnumParameterNode->SetUid(ProjectorParamUid + FString("_SP"));
+			PoseEnumParameterNode->SetUid(SelectedPoseParamUid.ToString());
 			GenerationContext.AddParameterNameUnique(originalGroup, PoseNodeEnumParamName);
 			PoseEnumParameterNode->SetValueCount(ProjParamNode->OptionPoses.Num() + 1);
 			PoseEnumParameterNode->SetDefaultValueIndex(0);
@@ -378,7 +385,7 @@ bool GenerateMutableSourceGroupProjector(const UEdGraphPin* Pin, FMutableGraphGe
 			mu::NodeScalarParameterPtr OpacityParameterNode = new mu::NodeScalarParameter;
 			FString OpacityParameterNodeName = ProjParamNode->ParameterName + FMultilayerProjector::OPACITY_PARAMETER_POSTFIX;
 			OpacityParameterNode->SetName(OpacityParameterNodeName);
-			OpacityParameterNode->SetUid(ProjectorParamUid + FString("_O"));
+			OpacityParameterNode->SetUid(OpacityParamUid.ToString());
 			GenerationContext.AddParameterNameUnique(originalGroup, OpacityParameterNodeName);
 			OpacityParameterNode->SetDefaultValue(0.75f);
 			OpacityParameterNode->SetRangeCount(1);
@@ -425,7 +432,7 @@ bool GenerateMutableSourceGroupProjector(const UEdGraphPin* Pin, FMutableGraphGe
 			mu::NodeScalarEnumParameterPtr EnumParameterNode = new mu::NodeScalarEnumParameter;
 			FString NodeEnumParamName = ProjParamNode->ParameterName + FMultilayerProjector::IMAGE_PARAMETER_POSTFIX;
 			EnumParameterNode->SetName(NodeEnumParamName);
-			EnumParameterNode->SetUid(ProjectorParamUid + FString("_SI"));
+			EnumParameterNode->SetUid(SelectedImageParamUid.ToString());
 			GenerationContext.AddParameterNameUnique(originalGroup, NodeEnumParamName);
 			EnumParameterNode->SetValueCount(ArrayOptionImage.Num());
 			EnumParameterNode->SetDefaultValueIndex(0);
diff --git a/Engine/Plugins/Experimental/Mutable/Source/CustomizableObjectEditor/Private/MuCOE/Nodes/CustomizableObjectNodeProjectorConstant.cpp b/Engine/Plugins/Experimental/Mutable/Source/CustomizableObjectEditor/Private/MuCOE/Nodes/CustomizableObjectNodeProjectorConstant.cpp
--- a/Engine/Plugins/Experimental/Mutable/Source/CustomizableObjectEditor/Private/MuCOE/Nodes/CustomizableObjectNodeProjectorConstant.cpp
+++ b/Engine/Plugins/Experimental/Mutable/Source/CustomizableObjectEditor/Private/MuCOE/Nodes/CustomizableObjectNodeProjectorConstant.cpp
@@ -12,7 +12,6 @@ class UCustomizableObjectNodeRemapPins;
 
 UCustomizableObjectNodeProjectorConstant::UCustomizableObjectNodeProjectorConstant()
 	: Super()
-	, ProjectionType(ECustomizableObjectProjectorType::Planar)
 	, ProjectionAngle(360.0f)
 	, BoneComboBoxLocation(FVector::ZeroVector)
 	, BoneComboBoxForwardDirection(FVector::ZeroVector)
@@ -51,9 +50,22 @@ void UCustomizableObjectNodeProjectorConstant::AllocateDefaultPins(UCustomizable
 }
 
 
+void UCustomizableObjectNodeProjectorConstant::BackwardsCompatibleFixup()
+{
+	Super::BackwardsCompatibleFixup();
+
+	const int32 CustomizableObjectCustomVersion = GetLinkerCustomVersion(FCustomizableObjectCustomVersion::GUID);
+
+	if (CustomizableObjectCustomVersion < FCustomizableObjectCustomVersion::ProjectorNodesDefaultValueFix)
+	{
+		Value.ProjectionType = ProjectionType_DEPRECATED;
+	}
+}
+
+
 ECustomizableObjectProjectorType UCustomizableObjectNodeProjectorConstant::GetProjectorType() const
 {
-	return ProjectionType;
+	return Value.ProjectionType;
 }
 
 
diff --git a/Engine/Plugins/Experimental/Mutable/Source/CustomizableObjectEditor/Private/MuCOE/Nodes/CustomizableObjectNodeProjectorParameter.cpp b/Engine/Plugins/Experimental/Mutable/Source/CustomizableObjectEditor/Private/MuCOE/Nodes/CustomizableObjectNodeProjectorParameter.cpp
--- a/Engine/Plugins/Experimental/Mutable/Source/CustomizableObjectEditor/Private/MuCOE/Nodes/CustomizableObjectNodeProjectorParameter.cpp
+++ b/Engine/Plugins/Experimental/Mutable/Source/CustomizableObjectEditor/Private/MuCOE/Nodes/CustomizableObjectNodeProjectorParameter.cpp
@@ -13,10 +13,7 @@ void UCustomizableObjectNodeProjectorParameter::PostEditChangeProperty(FProperty
 {
 	const FName PropertyName = PropertyChangedEvent.GetPropertyName();
 
-	if (PropertyName == "ProjectionType")
-	{
-	}
-	else if (PropertyName == "ProjectionAngle")
+	if (PropertyName == "ProjectionAngle")
 	{
 		DefaultValue.Angle = FMath::DegreesToRadians(ProjectionAngle);
 	}
@@ -41,9 +38,22 @@ void UCustomizableObjectNodeProjectorParameter::AllocateDefaultPins(UCustomizabl
 }
 
 
+void UCustomizableObjectNodeProjectorParameter::BackwardsCompatibleFixup()
+{
+	Super::BackwardsCompatibleFixup();
+
+	const int32 CustomizableObjectCustomVersion = GetLinkerCustomVersion(FCustomizableObjectCustomVersion::GUID);
+
+	if (CustomizableObjectCustomVersion < FCustomizableObjectCustomVersion::ProjectorNodesDefaultValueFix)
+	{
+		DefaultValue.ProjectionType = ProjectionType_DEPRECATED;
+	}
+}
+
+
 ECustomizableObjectProjectorType UCustomizableObjectNodeProjectorParameter::GetProjectorType() const
 {
-	return ProjectionType;
+	return DefaultValue.ProjectionType;
 }
 
 
