[Mutable] Optimize ShowOnScreenCompileWarnings.

Solves perofrmance issues when having a large number of instances on PIE.

#rnx

diff --git a/Engine/Plugins/Experimental/Mutable/Source/CustomizableObject/Private/MuCO/CustomizableObjectSystem.cpp b/Engine/Plugins/Experimental/Mutable/Source/CustomizableObject/Private/MuCO/CustomizableObjectSystem.cpp
--- a/Engine/Plugins/Experimental/Mutable/Source/CustomizableObject/Private/MuCO/CustomizableObjectSystem.cpp
+++ b/Engine/Plugins/Experimental/Mutable/Source/CustomizableObject/Private/MuCO/CustomizableObjectSystem.cpp
@@ -3596,6 +3596,8 @@ int32 UCustomizableObjectSystem::GetWorkingMemory() const
 #if WITH_EDITORONLY_DATA
 void UCustomizableObjectSystemPrivate::ShowOnScreenCompileWarnings()
 {
+	TSet<const UCustomizableObject*> Objects;
+	
 	for (TObjectIterator<UCustomizableObjectInstanceUsage> CustomizableObjectInstanceUsage; CustomizableObjectInstanceUsage; ++CustomizableObjectInstanceUsage)
 	{
 		if (!IsValid(*CustomizableObjectInstanceUsage) || CustomizableObjectInstanceUsage->IsTemplate())
@@ -3604,7 +3606,7 @@ void UCustomizableObjectSystemPrivate::ShowOnScreenCompileWarnings()
 		}
 		
 		const UCustomizableObjectInstance* Instance = CustomizableObjectInstanceUsage->GetCustomizableObjectInstance();
-		if (!Instance || !Instance->GetCustomizableObject())
+		if (!Instance)
 		{
 			continue;
 		}
@@ -3621,27 +3623,34 @@ void UCustomizableObjectSystemPrivate::ShowOnScreenCompileWarnings()
 			continue;
 		}
 
-		EWorldType::Type WorldType = EWorldType::Type::None;
-		if (const UWorld* World = Parent ? Parent->GetWorld() : nullptr)
+		const UWorld* World = Parent->GetWorld();
+		if (!World)
 		{
-			WorldType = World->WorldType;
+			continue;
 		}
 
-		if (WorldType != EWorldType::PIE)
+		if (World->WorldType != EWorldType::PIE)
 		{
 			continue;
 		}
-		
+
+		Objects.Add(Object);
+	}
+
+	for (const UCustomizableObject* Object : Objects)
+	{
 		// Show a warning if the compilation was not done with optimizations.
 		const uint64 KeyCompiledWithOptimization = reinterpret_cast<uint64>(Object);
-		if (!GEngine->OnScreenDebugMessageExists(KeyCompiledWithOptimization) && Object->GetPrivate()->bIsCompiledWithOptimization)
+		if (Object->GetPrivate()->bIsCompiledWithOptimization && // Quicker to check
+			!GEngine->OnScreenDebugMessageExists(KeyCompiledWithOptimization))
 		{
 			FString Msg = FString::Printf(TEXT("Warning: Customizable Object [%s] was compiled without optimization."), *Object->GetName());
 			GEngine->AddOnScreenDebugMessage(KeyCompiledWithOptimization, 10.0f, FColor::Red, Msg);
 		}
 
 		const uint64 KeyCompiledOutOfData = reinterpret_cast<uint64>(Object) + KEY_OFFSET_COMPILATION_OUT_OF_DATE; // Offset added to avoid collision with bIsCompiledWithOptimization warning
-		if (!GEngine->OnScreenDebugMessageExists(KeyCompiledWithOptimization) && Object->GetPrivate()->IsCompilationOutOfDate())
+		if (!GEngine->OnScreenDebugMessageExists(KeyCompiledWithOptimization) && // Quicker to check
+			Object->GetPrivate()->IsCompilationOutOfDate())
 		{
 			FString Msg = FString::Printf(TEXT("Warning: Customizable Object [%s] compilation out of date. Save all referenced assets and recompile."), *Object->GetName());
 			GEngine->AddOnScreenDebugMessage(KeyCompiledOutOfData, 10.0f, FColor::Red, Msg);
