[Mutable] Fix LOD updates not working properly.

#rnx

diff --git a/Engine/Plugins/Experimental/Mutable/Source/CustomizableObject/Classes/MuCO/CustomizableObjectInstance.h b/Engine/Plugins/Experimental/Mutable/Source/CustomizableObject/Classes/MuCO/CustomizableObjectInstance.h
--- a/Engine/Plugins/Experimental/Mutable/Source/CustomizableObject/Classes/MuCO/CustomizableObjectInstance.h
+++ b/Engine/Plugins/Experimental/Mutable/Source/CustomizableObject/Classes/MuCO/CustomizableObjectInstance.h
@@ -759,7 +759,7 @@ public:
 
 	void Issue();
 
-	void ApplyLODUpdateParamsToInstance(FUpdateContextPrivate* Context = nullptr);
+	void ApplyLODUpdateParamsToInstance(FUpdateContextPrivate& Context);
 
 private:
 	/** If true it means that EnqueueUpdateSkeletalMesh has decided this update should be performed, if false it should be ignored. Just used for consistency checks */
diff --git a/Engine/Plugins/Experimental/Mutable/Source/CustomizableObject/Private/MuCO/CustomizableObjectInstance.cpp b/Engine/Plugins/Experimental/Mutable/Source/CustomizableObject/Private/MuCO/CustomizableObjectInstance.cpp
--- a/Engine/Plugins/Experimental/Mutable/Source/CustomizableObject/Private/MuCO/CustomizableObjectInstance.cpp
+++ b/Engine/Plugins/Experimental/Mutable/Source/CustomizableObject/Private/MuCO/CustomizableObjectInstance.cpp
@@ -2923,18 +2923,15 @@ void FMutableUpdateCandidate::Issue()
 }
 
 
-void FMutableUpdateCandidate::ApplyLODUpdateParamsToInstance(FUpdateContextPrivate* Context)
+void FMutableUpdateCandidate::ApplyLODUpdateParamsToInstance(FUpdateContextPrivate& Context)
 {
 	CustomizableObjectInstance->Descriptor.MinLOD = MinLOD;
 	CustomizableObjectInstance->Descriptor.MaxLOD = MaxLOD;
 	CustomizableObjectInstance->Descriptor.RequestedLODLevels = RequestedLODLevels;
 
-	if (Context)
-	{
-		Context->SetMinLOD(MinLOD);
-		Context->SetMaxLOD(MaxLOD);
-		Context->SetRequestedLODs(RequestedLODLevels);
-	}
+	Context.SetMinLOD(MinLOD);
+	Context.SetMaxLOD(MaxLOD);
+	Context.SetRequestedLODs(RequestedLODLevels);
 }
 
 
diff --git a/Engine/Plugins/Experimental/Mutable/Source/CustomizableObject/Private/MuCO/CustomizableObjectSystem.cpp b/Engine/Plugins/Experimental/Mutable/Source/CustomizableObject/Private/MuCO/CustomizableObjectSystem.cpp
--- a/Engine/Plugins/Experimental/Mutable/Source/CustomizableObject/Private/MuCO/CustomizableObjectSystem.cpp
+++ b/Engine/Plugins/Experimental/Mutable/Source/CustomizableObject/Private/MuCO/CustomizableObjectSystem.cpp
@@ -3054,7 +3054,7 @@ bool UCustomizableObjectSystem::Tick(float DeltaTime)
 				
 				if (LODUpdateWithSameInstance)
 				{
-					LODUpdateWithSameInstance->ApplyLODUpdateParamsToInstance(&PendingInstanceUpdateFound->Context.Get());
+					LODUpdateWithSameInstance->ApplyLODUpdateParamsToInstance(PendingInstanceUpdateFound->Context.Get());
 				}
 
 				Private->StartUpdateSkeletalMesh(PendingInstanceUpdateFound->Context);				
@@ -3062,13 +3062,14 @@ bool UCustomizableObjectSystem::Tick(float DeltaTime)
 			}
 			else if (LODUpdateCandidateFound)
 			{
-				// Commit the LOD changes
-				LODUpdateCandidateFound->ApplyLODUpdateParamsToInstance();
-
 				UCustomizableObjectInstance* Instance = LODUpdateCandidateFound->CustomizableObjectInstance;
 				FCustomizableObjectInstanceDescriptor& Descriptor = Instance->GetPrivate()->SkeletalMeshStatus == ESkeletalMeshStatus::NotGenerated ? Instance->GetDescriptor() : Instance->GetPrivate()->CommittedDescriptor;
 
 				const TSharedRef<FUpdateContextPrivate> Context = MakeShared<FUpdateContextPrivate>(*Instance, Descriptor);
+
+				// Commit the LOD changes
+				LODUpdateCandidateFound->ApplyLODUpdateParamsToInstance(*Context);
+
 				Private->StartUpdateSkeletalMesh(Context);
 			}
 		}
