[mutable] Added option to skip logging of mutable compilation commandlet data.
- We do not need to know the compilation data in some situations, like when the compilation of the CO is a preparation step for the commandlet (bulk instance update)
- Also removed unnecesary checks for the AssetRegistry sync search invocation.

#changelist validated
#virtualized

diff --git a/Engine/Plugins/Experimental/Mutable/Source/MutableValidation/Private/MuV/COIBulkUpdateTestCommandlet.cpp b/Engine/Plugins/Experimental/Mutable/Source/MutableValidation/Private/MuV/COIBulkUpdateTestCommandlet.cpp
--- a/Engine/Plugins/Experimental/Mutable/Source/MutableValidation/Private/MuV/COIBulkUpdateTestCommandlet.cpp
+++ b/Engine/Plugins/Experimental/Mutable/Source/MutableValidation/Private/MuV/COIBulkUpdateTestCommandlet.cpp
@@ -123,7 +123,7 @@ int32 UCOIBulkUpdateTestCommandlet::Main(const FString& Params)
 		
 		// Compile the current CO object
 		const FString CustomizableObjectName = CustomizableObjectToCompile->GetName();
-		if (!CompileCustomizableObject(CustomizableObjectToCompile))
+		if (!CompileCustomizableObject(CustomizableObjectToCompile, false))	// Do not log mutable data since mongoDB will not be able to handle it correctly 
 		{
 			UE_LOG(LogMutable,Error,TEXT("The CO %s could not be compiled succesfully. Skipping the update of all COIs that use it."), *CustomizableObjectName )
 			continue;
diff --git a/Engine/Plugins/Experimental/Mutable/Source/MutableValidation/Private/MuV/CustomizableObjectValidationCommandlet.cpp b/Engine/Plugins/Experimental/Mutable/Source/MutableValidation/Private/MuV/CustomizableObjectValidationCommandlet.cpp
--- a/Engine/Plugins/Experimental/Mutable/Source/MutableValidation/Private/MuV/CustomizableObjectValidationCommandlet.cpp
+++ b/Engine/Plugins/Experimental/Mutable/Source/MutableValidation/Private/MuV/CustomizableObjectValidationCommandlet.cpp
@@ -62,7 +62,7 @@ int32 UCustomizableObjectValidationCommandlet::Main(const FString& Params)
 		CompilationOptions.OptimizationLevel = 2;			// Set the optimization level to the max
 		CompilationOptions.TextureCompression = ECustomizableObjectTextureCompression::Fast;
 
-		bWasCoCompilationSuccessful = CompileCustomizableObject(ToTestCustomizableObject, &CompilationOptions);
+		bWasCoCompilationSuccessful = CompileCustomizableObject(ToTestCustomizableObject, true, &CompilationOptions);
 	}
 	// -------------------------------------------------------------------------------------------------------------- //
 	
diff --git a/Engine/Plugins/Experimental/Mutable/Source/MutableValidation/Private/MuV/ValidationUtils.cpp b/Engine/Plugins/Experimental/Mutable/Source/MutableValidation/Private/MuV/ValidationUtils.cpp
--- a/Engine/Plugins/Experimental/Mutable/Source/MutableValidation/Private/MuV/ValidationUtils.cpp
+++ b/Engine/Plugins/Experimental/Mutable/Source/MutableValidation/Private/MuV/ValidationUtils.cpp
@@ -13,22 +13,13 @@
 void PrepareAssetRegistry()
 {
 	FAssetRegistryModule& AssetRegistryModule = FModuleManager::LoadModuleChecked<FAssetRegistryModule>(AssetRegistryConstants::ModuleName);
-	if (AssetRegistryModule.Get().IsSearchAllAssets())
-	{
-		UE_LOG(LogMutable,Display,TEXT("Asset search already in progress. Waitting for it to complete."));
-		AssetRegistryModule.Get().WaitForCompletion();
-	}
-	else
-	{
-		UE_LOG(LogMutable,Display,TEXT("Searching all assets (this will take some time)..."));
-		AssetRegistryModule.Get().SearchAllAssets(true /* bSynchronousSearch */);
-	}
-
+	UE_LOG(LogMutable,Display,TEXT("Searching all assets (this will take some time)..."));
+	AssetRegistryModule.Get().SearchAllAssets(true /* bSynchronousSearch */);
 	UE_LOG(LogMutable,Display,TEXT("Asset searching completed!"));
 }
 
 
-bool CompileCustomizableObject(UCustomizableObject* InCustomizableObject, const FCompilationOptions* InCompilationOptionsOverride  /* nullptr */)
+bool CompileCustomizableObject(UCustomizableObject* InCustomizableObject, const bool bLogMutableLogs /* = true */, const FCompilationOptions* InCompilationOptionsOverride  /* nullptr */)
 {
 	if (!InCustomizableObject)
 	{
@@ -49,15 +40,7 @@ bool CompileCustomizableObject(UCustomizableObject* InCustomizableObject, const
 		CompilationOptions = InCustomizableObject->CompileOptions;
 		UE_LOG(LogMutable,Display,TEXT("Compiling CO using it's own compilation options."));
 	}
-
-	UE_LOG(LogMutable, Log, TEXT("(string) model_compile_options_overriden : %s "), bOverrideCompilationOptions ? TEXT("true") : TEXT("false"));
 	
-	// TODO: Add logs for the other relevant configs of the model being compiled
-	// Print MTU parseable logs
-	UE_LOG(LogMutable, Log, TEXT("(int) model_optimization_level : %d "), CompilationOptions.OptimizationLevel);
-	UE_LOG(LogMutable, Log, TEXT("(string) model_texture_compression : %s "), *UEnum::GetValueAsString(CompilationOptions.TextureCompression));
-	UE_LOG(LogMutable, Log, TEXT("(string) model_disk_compilation : %s "), CompilationOptions.bUseDiskCompilation ? TEXT("true") : TEXT("false"));
-
 	// Request a compiler to be able to locate the root and to compile it
 	const TUniquePtr<FCustomizableObjectCompilerBase> Compiler =
 		TUniquePtr<FCustomizableObjectCompilerBase>(UCustomizableObjectSystem::GetInstanceChecked()->GetNewCompiler());
@@ -73,13 +56,10 @@ bool CompileCustomizableObject(UCustomizableObject* InCustomizableObject, const
     // --------------------------------------------------------------------------
 	UE_LOG(LogMutable,Display,TEXT("Compilation of CO completed!"));
 
-	UE_LOG(LogMutable, Log, TEXT("(double) model_compile_time_ms : %f "), CompilationEndSeconds * 1000);
 	UE_LOG(LogMutable, Display, TEXT("The compilation of the %s model took %f seconds."), *InCustomizableObject->GetName(), CompilationEndSeconds);
-
 	
     // Get the compilation result
     const ECustomizableObjectCompilationState CompilationEndResult = Compiler->GetCompilationState();
-	UE_LOG(LogMutable, Log, TEXT("(string) model_compile_end_state : %s "), *UEnum::GetValueAsString(CompilationEndResult));
 	
     const bool bWasCoCompilationSuccessful = CompilationEndResult == ECustomizableObjectCompilationState::Completed;
     if (bWasCoCompilationSuccessful)
@@ -90,7 +70,22 @@ bool CompileCustomizableObject(UCustomizableObject* InCustomizableObject, const
     {
     	UE_LOG(LogMutable, Error, TEXT("The compilation of the %s model failed."), *InCustomizableObject->GetName());
     }
-    
+
+	
+	// Print MTU parseable logs only if asked. The addition of duplicated entries in MongoDB is not available so this way we avoid having to handle them
+	// when we are sure we do not require it.
+	if (bLogMutableLogs)
+	{
+		UE_LOG(LogMutable, Log, TEXT("(string) model_compile_options_overriden : %s "), bOverrideCompilationOptions ? TEXT("true") : TEXT("false"));
+		UE_LOG(LogMutable, Log, TEXT("(int) model_optimization_level : %d "), CompilationOptions.OptimizationLevel);
+		UE_LOG(LogMutable, Log, TEXT("(string) model_texture_compression : %s "), *UEnum::GetValueAsString(CompilationOptions.TextureCompression));
+		UE_LOG(LogMutable, Log, TEXT("(string) model_disk_compilation : %s "), CompilationOptions.bUseDiskCompilation ? TEXT("true") : TEXT("false"));
+		UE_LOG(LogMutable, Log, TEXT("(double) model_compile_time_ms : %f "), CompilationEndSeconds * 1000);
+		UE_LOG(LogMutable, Log, TEXT("(string) model_compile_end_state : %s "), *UEnum::GetValueAsString(CompilationEndResult));
+		// TODO: Add logs for the other relevant configs of the model being compiled
+	}
+
+
 	return bWasCoCompilationSuccessful;
 }
 
diff --git a/Engine/Plugins/Experimental/Mutable/Source/MutableValidation/Private/MuV/ValidationUtils.h b/Engine/Plugins/Experimental/Mutable/Source/MutableValidation/Private/MuV/ValidationUtils.h
--- a/Engine/Plugins/Experimental/Mutable/Source/MutableValidation/Private/MuV/ValidationUtils.h
+++ b/Engine/Plugins/Experimental/Mutable/Source/MutableValidation/Private/MuV/ValidationUtils.h
@@ -22,10 +22,11 @@ void PrepareAssetRegistry();
 /**
  * Compiles a CO synchronously
  * @param InCustomizableObject The Customizable Object we want to synchronously compile
+ * @param bLogMutableLogs Enables or disables the logging of Log category logs relevant to the CO compilation. Required to avoid a MongoDB limitation with the duplication of MongoDB document names.
  * @param InCompilationOptionsOverride The configuration for the compilation of the CO we want to use instead of the one part of the CO.
  * @return True if the compilation was successful and false if it failed.
  */
-bool CompileCustomizableObject(UCustomizableObject* InCustomizableObject, const FCompilationOptions* InCompilationOptionsOverride = nullptr);
+bool CompileCustomizableObject(UCustomizableObject* InCustomizableObject, const bool bLogMutableLogs = true, const FCompilationOptions* InCompilationOptionsOverride = nullptr);
 
 
 /**
