[Mutable] Moved all non-public UCustomizableObject functions to private.

- Removed UCustomizableObject::GetModel (body moved to the only caller).
- Fixed minor typos.
- Standarized doc style.

#rnx

diff --git a/Engine/Plugins/Experimental/Mutable/Source/CustomizableObject/Classes/MuCO/CustomizableObject.h b/Engine/Plugins/Experimental/Mutable/Source/CustomizableObject/Classes/MuCO/CustomizableObject.h
--- a/Engine/Plugins/Experimental/Mutable/Source/CustomizableObject/Classes/MuCO/CustomizableObject.h
+++ b/Engine/Plugins/Experimental/Mutable/Source/CustomizableObject/Classes/MuCO/CustomizableObject.h
@@ -421,7 +421,7 @@ struct FMutableLODSettings
 	UPROPERTY()
 	int32 NumLODsInRoot = 0;
 
-	/** Frist LOD available, some platforms may remove lower LODs when cooking, this MinLOD represents the first LOD we can generate */
+	/** First LOD available, some platforms may remove lower LODs when cooking, this MinLOD represents the first LOD we can generate */
 	UPROPERTY()
 	int32 FirstLODAvailable = 0;
 
@@ -490,20 +490,11 @@ public:
 	UCustomizableObject();
 
 #if WITH_EDITORONLY_DATA
-	/** All the SkeletalMeshes generated for this CustomizableObject instances will use the Reference Skeletal Mesh 
-	* properties for everything that Mutable doesn't create or modify. This includes data like LOD distances, Physics
-	* properties, Bounding Volumes, Skeleton, etc.
-	*
-	* While a CustomizableObject instance is being created for the first time, and in some situation with lots of 
-	* objects this may require some seconds, the Reference Skeletal Mesh is used for the actor. This works as a better
-	* solution than the alternative of not showing anything, although this can be disabled with the function
-	* "SetReplaceDiscardedWithReferenceMeshEnabled" (See the c++ section).
-	* 
-	* For more information on this topic read the Basic Concepts at work.anticto.com.
-	*/
+private:
 	UPROPERTY()
 	TObjectPtr<USkeletalMesh> ReferenceSkeletalMesh_DEPRECATED;
 
+public:
 	/** All the SkeletalMeshes generated for this CustomizableObject instances will use the Reference Skeletal Mesh
 	* properties for everything that Mutable doesn't create or modify. This includes data like LOD distances, Physics
 	* properties, Bounding Volumes, Skeleton, etc.
@@ -511,14 +502,11 @@ public:
 	* While a CustomizableObject instance is being created for the first time, and in some situation with lots of
 	* objects this may require some seconds, the Reference Skeletal Mesh is used for the actor. This works as a better
 	* solution than the alternative of not showing anything, although this can be disabled with the function
-	* "SetReplaceDiscardedWithReferenceMeshEnabled" (See the c++ section).
-	*
-	* For more information on this topic read the Basic Concepts at work.anticto.com.
-	*/
+	* "SetReplaceDiscardedWithReferenceMeshEnabled" (See the c++ section). */
 	UPROPERTY(EditAnywhere, Category = CustomizableObject)
 	TArray<TObjectPtr<USkeletalMesh>> ReferenceSkeletalMeshes;
 #endif
-
+	
 	UPROPERTY(EditAnywhere, Category = CustomizableObject, meta = (DisplayName = "LOD Settings"))
 	FMutableLODSettings LODSettings;
 	
@@ -549,7 +537,6 @@ public:
 	UPROPERTY()
 	TArray<FCustomizableObjectStreamedResourceData> StreamedResourceData;
 
-private:
 	/** Use the SkeletalMesh of reference as a placeholder until the custom mesh is ready to use.
 	  * Note: If disabled, a null mesh will be used to replace the discarded mesh due to 'ReplaceDiscardedWithReferenceMesh' being enabled. */
 	UPROPERTY(EditAnywhere, Category = CustomizableObject)
@@ -564,22 +551,21 @@ private:
 	UPROPERTY(EditAnywhere, Category = CustomizableObject)
 	bool bEnableMeshCache = false;
 	
-public:
 #if WITH_EDITORONLY_DATA
+private:
 	// Hide this property because it is not used yet.
 	//UPROPERTY(EditAnywhere, Category = CustomizableObject)
 	UPROPERTY()
 	ECustomizableObjectRelevancy Relevancy;
 
+public:
 	// Compilation options to use in editor and for packaging for this object.
 	UPROPERTY()
 	FCompilationOptions CompileOptions;
 
-	//
 	UPROPERTY(EditAnywhere, Category = CompileOptions)
 	bool bEnableRealTimeMorphTargets = false;
 
-	//
 	UPROPERTY(EditAnywhere, Category = CompileOptions)
 	bool bEnableClothing = false;
 
@@ -587,11 +573,9 @@ public:
 	UPROPERTY(VisibleAnywhere, Category = CompileOptions)
 	bool bEnable16BitBoneWeights = false;
 
-	//
 	UPROPERTY(EditAnywhere, Category = CompileOptions)
 	bool bEnableAltSkinWeightProfiles = false;
 
-	//
 	UPROPERTY(EditAnywhere, Category = CompileOptions)
 	bool bEnablePhysicsAssetMerge = false;
 
@@ -614,19 +598,19 @@ public:
 	UPROPERTY()
 	TObjectPtr<UEdGraph> Source;
 
+private:
 	// Used to verify the derived data matches this version of the Customizable Object.
 	UPROPERTY()
 	FGuid VersionId;
 
-	// 
+public:
 	UPROPERTY()
-	TArray<FProfileParameterDat>  InstancePropertiesProfiles;
-
+	TArray<FProfileParameterDat> InstancePropertiesProfiles;
 #endif // WITH_EDITORONLY_DATA
 
 	/** Amount of components in this CO. Set at the end of the model compilation process. */
 	UPROPERTY()
-	int32 NumMeshComponentsInRoot = 0;
+	int32 NumMeshComponentsInRoot = 0; // TODO UE-205600, move to FModelResources
 	
 	/** Get the number of components this Customizable Object has. */
 	UFUNCTION(BlueprintCallable, Category = CustomizableObject)
@@ -661,69 +645,49 @@ public:
 	int32 GetIntParameterNumOptions(int32 ParamIndex) const;
 
 	/** Gets the Name of the option at position K in the list of available options for the int parameter.
-	 * Useful to enumerate the int parameter's possible options (Ex: "Hat1", "Hat2", "Cap", "Nothing")
-	 */
+	  * Useful to enumerate the int parameter's possible options (Ex: "Hat1", "Hat2", "Cap", "Nothing") */
 	UFUNCTION(BlueprintCallable, Category = CustomizableObject)
 	const FString& GetIntParameterAvailableOption(int32 ParamIndex, int32 K) const;
 
 	int32 FindIntParameterValue( int32 ParamIndex, const FString& Value ) const;
 	FString FindIntParameterValueName(int32 ParamIndex, int32 ParamValue) const;
 
-	//
 	USkeletalMesh* GetRefSkeletalMesh(int32 ComponentIndex = 0) const;
-	
-private:
 
-	
-	/** Returns true or false if the parameter with name can be located and it has the type the caller is looking for. It will also
-	 * check if the model has been set to ensure access to it can take place at the calculated parameter index.
-	 * @param InParameterName The name of the parameter to look for.
-	 * @param InParameterType The type the parameter we are looking for we know has. If the name does not match this type this check will fail and return false.
-	 * @param OutParameterIndex The index of the parameter.
-	 * @return True if the parameter can be accessed for it's default values, false if it can not be accessed.
-	 */
-	bool CanDefaultParameterBeAccessed(const FString& InParameterName,const EMutableParameterType& InParameterType, int32& OutParameterIndex) const;
-
-public:
 	/** Get the default value of a parameter of type Float.
-	 * @param InParameterName The name of the Float parameter to get the default value of.
-	 * @return The default value of the provided parameter name.
-	 */
+	  * @param InParameterName The name of the Float parameter to get the default value of.
+	  * @return The default value of the provided parameter name. */
 	UFUNCTION(BlueprintCallable, Category = CustomizableObject)
 	float GetFloatParameterDefaultValue(UPARAM(DisplayName = "Parameter Name") const FString& InParameterName) const;
 	
 	/** Get the default value of a parameter of type Int. 
-	 * @param InParameterName The name of the Int parameter to get the default value of.
-	 * @return The default value of the provided parameter name.
-	 */
+	  * @param InParameterName The name of the Int parameter to get the default value of.
+	  * @return The default value of the provided parameter name. */
 	UFUNCTION(BlueprintCallable, Category = CustomizableObject)
 	int32 GetIntParameterDefaultValue(UPARAM(DisplayName = "Parameter Name") const FString& InParameterName) const;
  
 	/** Get the default value of a parameter of type Bool.
-	 * @param InParameterName The name of the Bool parameter to get the default value of.
-	 * @return The default value of the provided parameter name.
-	 */
+	  * @param InParameterName The name of the Bool parameter to get the default value of.
+	  * @return The default value of the provided parameter name. */
 	UFUNCTION(BlueprintCallable, Category = CustomizableObject)
 	bool GetBoolParameterDefaultValue(UPARAM(DisplayName = "Parameter Name") const FString& InParameterName) const;
 
 	/** Get the default value of a parameter of type Color.
-	 * @param InParameterName The name of the Color parameter to get the default value of.
-	 * @return The default value of the provided parameter name.
-	*/
+	  * @param InParameterName The name of the Color parameter to get the default value of.
+	  * @return The default value of the provided parameter name. */
 	UFUNCTION(BlueprintCallable, Category = CustomizableObject)
 	FLinearColor GetColorParameterDefaultValue(UPARAM(DisplayName = "Parameter Name") const FString& InParameterName) const;
 	
 	/** Get the default value of a parameter of type Projector.
-	 * @param InParameterName The name of the Projector parameter to get the default value of.
-	 * @param OutPos The default position of the Projector.
-	 * @param OutDirection The default projection direction of the Projector.
-	 * @param OutUp The default up vector of the Projector.
-	 * @param OutScale The default scale of the Projector.
-	 * @param OutAngle The default angle of the Projector.
-	 * @param OutType The default type of the Projector.
-	 */
+	  * @param InParameterName The name of the Projector parameter to get the default value of.
+	  * @param OutPos The default position of the Projector.
+	  * @param OutDirection The default projection direction of the Projector.
+	  * @param OutUp The default up vector of the Projector.
+	  * @param OutScale The default scale of the Projector.
+	  * @param OutAngle The default angle of the Projector.
+	  * @param OutType The default type of the Projector. */
 	UFUNCTION(BlueprintCallable, Category = CustomizableObject)
-	void GetProjectorParameterDefaultValue (
+	void GetProjectorParameterDefaultValue(
 		UPARAM(DisplayName = "Parameter Name") const FString& InParameterName,
 		UPARAM(DisplayName = "Position") FVector3f& OutPos,
 		UPARAM(DisplayName = "Direction") FVector3f& OutDirection,
@@ -733,38 +697,29 @@ public:
 		UPARAM(DisplayName = "Type") ECustomizableObjectProjectorType& OutType) const;
 
 	/** Get the default value of a projector with the provided name
-	 * @param InParameterName The name of the parameter to get the default value of.
-	 * @return A data structure containing all the default data for the targeted projector parameter.
-	 */
-	FCustomizableObjectProjector GetProjectorParameterDefaultValue ( const FString& InParameterName) const;
+	  * @param InParameterName The name of the parameter to get the default value of.
+	  * @return A data structure containing all the default data for the targeted projector parameter. */
+	FCustomizableObjectProjector GetProjectorParameterDefaultValue(const FString& InParameterName) const;
 	
 	/** Get the default value of a parameter of type Texture.
-	 * @param InParameterName The name of the Projector parameter to get the default value of.
-	 * @return An id representing the default parameter's texture.
-	 */
-	FName GetTextureParameterDefaultValue (const FString& InParameterName) const;
+	  * @param InParameterName The name of the Projector parameter to get the default value of.
+	  * @return An id representing the default parameter's texture. */
+	FName GetTextureParameterDefaultValue(const FString& InParameterName) const;
 
 	/** Return true if the parameter at the index provided is multidimensional.
-	 * @param InParamIndex The index of the parameter to check.
-	 * @return True if the parameter is multidimensional and false if it is not.
-	 */
+	  * @param InParamIndex The index of the parameter to check.
+	  * @return True if the parameter is multidimensional and false if it is not. */
 	bool IsParameterMultidimensional(const int32& InParamIndex) const;
 	
 	/** Return true if the parameter at the index provided is multidimensional.
-	 * @param InParameterName The name of the parameter to check.
-	 * @return True if the parameter is multidimensional and false if it is not.
-	 */
+	  * @param InParameterName The name of the parameter to check.
+	  * @return True if the parameter is multidimensional and false if it is not. */
 	UFUNCTION(BlueprintCallable, Category = CustomizableObject)
-	bool IsParameterMultidimensional ( UPARAM(DisplayName = "Parameter Name") const FString& InParameterName) const;
-	
-private: 
-	// Rebuild ParameterProperties from the current compiled model.
-	void UpdateParameterPropertiesFromModel(const TSharedPtr<mu::Model>& Model);
-
-public:
-
-	void AddUncompiledCOWarning(const FString& AdditionalLoggingInfo);
+	bool IsParameterMultidimensional(UPARAM(DisplayName = "Parameter Name") const FString& InParameterName) const;
 
+	// Begin UObject interface.
+	void PostLoad() override;
+	void Serialize(FArchive& Ar) override;
 #if WITH_EDITOR
 	
 	// Compile the object for a specific platform - Compile for Cook Customizable Object
@@ -774,8 +729,8 @@ public:
 	void ClearCompiledData(bool bIsCooking);
 
 	/** Compile the object if Automatic Compilation is enabled and has not been already compiled.
-	 * Automatic compilation can be enabled/disabled in the Mutable's Plugin Settings.
-	 * @return true if compiled */
+	  * Automatic compilation can be enabled/disabled in the Mutable's Plugin Settings.
+	  * @return true if compiled */
 	bool ConditionalAutoCompile();
 	
 	// Add a profile that stores the values of the parameters used by the CustomInstance.
@@ -791,55 +746,21 @@ public:
 	// Compose file name 
 	FString GetCompiledDataFileName(bool bIsModel, const ITargetPlatform* InTargetPlatform = nullptr, bool bIsDiskStreamer = false);
 
-	/** Used to set the flag IsRoot */
 	virtual void GetAssetRegistryTags(FAssetRegistryTagsContext Context) const override;
 	UE_DEPRECATED(5.4, "Implement the version that takes FAssetRegistryTagsContext instead.")
 	virtual void GetAssetRegistryTags(TArray<FAssetRegistryTag>& OutTags) const override;
 
 	/** Returns a one line description of an object for viewing in the thumbnail view of the generic browser */
 	FString GetDesc() override;
-
 	bool IsEditorOnly() const override;
-
-	// Begin UObject interface.
 	void PreSave(FObjectPreSaveContext ObjectSaveContext) override;
-	void PostRename(UObject* OldOuter, const FName OldName) override;
-	
+	void PostRename(UObject* OldOuter, const FName OldName) override;	
 	void BeginCacheForCookedPlatformData(const ITargetPlatform* TargetPlatform) override;
 	bool IsCachedCookedPlatformDataLoaded(const ITargetPlatform* TargetPlatform) override;
-	// End UObject interface.
-
-	
-	/** Generic Save/Load methods to write/read compiled data */
-	void SaveCompiledData(FArchive& Ar, bool bSkipEditorOnlyData = false);
-	void LoadCompiledData(FArchive& Ar, const ITargetPlatform* InTargetPlatform, bool bSkipEditorOnlyData = false);
-
-	/** Load compiled data for the running platform from disk, this is used to load Editor Compilations. */
-	void LoadCompiledDataFromDisk();
-
-	/** Cache platform data for cook */
-	void CachePlatformData(const ITargetPlatform* InTargetPlatform, const TArray64<uint8>& InObjectBytes, const TArray64<uint8>& InBulkBytes);
-	/**
-	 * Loads data previously compiled in BeginCacheForCookedPlatformData onto the UProperties in *this,
-	 * in preparation for saving the cooked package for *this or for a CustomizableObjectInstance using *this.
-	 * Returns whether the data was successfully loaded.
-	 */
-	bool TryLoadCompiledCookDataForPlatform(const ITargetPlatform* TargetPlatform);
-	
-	void SaveEmbeddedData(FArchive& Ar);
-
 #endif
-	/** Compute bIsChildObject if currently possible to do so. Return whether it was computed. */
-	bool TryUpdateIsChildObject();
-
-	// Data that may be stored in the asset itself, only in packaged builds.
-	void LoadEmbeddedData(FArchive& Ar);
-
-	void PostLoad() override;
-
-	void Serialize(FArchive& Ar) override;
+	// End UObject interface.
 
-	int32 FindState( const FString& Name ) const;
+	int32 FindState(const FString& Name) const;
 
 	/** Return the number of object states that are defined in the CustomizableObject. */
 	UFUNCTION(BlueprintCallable, Category = CustomizableObject)
@@ -877,9 +798,6 @@ public:
 	UFUNCTION(BlueprintCallable, Category = CustomizableObject)
 	FParameterUIData GetParameterUIMetadataFromIndex(int32 ParamIndex) const;
 
-	/** See bPreserveUserLODsOnFirstGeneration. */
-	bool IsPreserveUserLODsOnFirstGeneration() const;
-
 	/** Textures marked as low priority will generate defaulted resident mips (if texture streaming is enabled).
 	  * Generating defaulted resident mips greatly reduce initial generation times. */
 	UPROPERTY(EditAnywhere, Category = CustomizableObject)
@@ -887,14 +805,14 @@ public:
 
 	/** Map of Hash to Streaming blocks, used to stream a block of data representing a resource from the BulkData */
 	UPROPERTY()
-	TMap<uint64, FMutableStreamableBlock> HashToStreamableBlock;
+	TMap<uint64, FMutableStreamableBlock> HashToStreamableBlock; // TODO UE-205600, move to FModelResources
 
 	// Customizable Object Population data start ------------------------------------------------------
 	/** Array to store the selected Population Class tags for this Customizable Object */
 	UPROPERTY()
 	TArray<FString> CustomizableObjectClassTags;
 	
-	/** Array to strore all the Population Class tags */
+	/** Array to store all the Population Class tags */
 	UPROPERTY()
 	TArray<FString> PopulationClassTags;
 
@@ -904,7 +822,6 @@ public:
 	// Customizable Object Population data end --------------------------------------------------------
 
 #if WITH_EDITORONLY_DATA
-
 	/** True if this object references a parent object. This is used basically to exclude this object
 	  * from cooking. This is actually derived from the source graph object node pointing to another
 	  * object or not, but it needs to be cached here because the source graph is not always available.
@@ -912,19 +829,14 @@ public:
 	  * that is the conservative case and shouldn't cause a problem. */
 	UPROPERTY()
 	bool bIsChildObject = false;
-
-	void PostCompile();
-
 #endif
-
+	
+public:
 	FPostCompileDelegate& GetPostCompileDelegate() const;
 	
 	/** Create a new instance of this object. The instance parameters will be initialized with the object default values. */
 	UFUNCTION(BlueprintCallable, Category = CustomizableObject)
-	class UCustomizableObjectInstance* CreateInstance();
-
-	// Return a pointer to the BulkData subobject, only valid in packaged builds
-	const UCustomizableObjectBulk* GetStreamableBulkData() const { return BulkData; }
+	UCustomizableObjectInstance* CreateInstance();
 
 	UCustomizableObjectPrivate* GetPrivate() const;
 
@@ -932,29 +844,11 @@ public:
 	UFUNCTION(BlueprintCallable, Category = CustomizableObject)
 	bool IsCompiled() const;
 
-#if WITH_EDITOR
-	void SetModel(TSharedPtr<mu::Model, ESPMode::ThreadSafe> Model);
-#endif
-
 	int32 GetNumLODs() const;
 
-	/** Modify the provided mutable parameters so that the forced values for the given customizable object state are applied. */
-	void ApplyStateForcedValuesToParameters(int32 State, mu::Parameters* Parameters);
-
-	/** Return the names used by mutable to identify which mu::Image should be considered of LowPriority. */
-	void GetLowPriorityTextureNames(TArray<FString>& OutTextureNames);
-
-	/** Return the MinLOD index to generate based on the active LODSettings (PerPlatformMinLOD or PerQualityLevelMinLOD) */
-	int32 GetMinLODIndex() const;
-
-	/** See bEnableUseRefSkeletalMeshAsPlaceholder. */
-	bool IsEnableUseRefSkeletalMeshAsPlaceholder() const;
-
-	/** See bEnableUseMeshCache. */
-	bool IsMeshCacheEnabled() const;
+	bool IsChildObject() const;
 
 private:
-
 	/** BulkData that stores all in-game resources used by Mutable when generating instances.
 	  * Only valid in packaged builds */
 	UPROPERTY()
diff --git a/Engine/Plugins/Experimental/Mutable/Source/CustomizableObject/Internal/MuCO/CustomizableObjectPrivate.h b/Engine/Plugins/Experimental/Mutable/Source/CustomizableObject/Internal/MuCO/CustomizableObjectPrivate.h
--- a/Engine/Plugins/Experimental/Mutable/Source/CustomizableObject/Internal/MuCO/CustomizableObjectPrivate.h
+++ b/Engine/Plugins/Experimental/Mutable/Source/CustomizableObject/Internal/MuCO/CustomizableObjectPrivate.h
@@ -28,6 +28,9 @@ class UAssetUserData;
 class UCustomizableObject;
 
 
+FGuid CUSTOMIZABLEOBJECT_API GenerateIdentifier(const UCustomizableObject& CustomizableObject);
+
+
 class CUSTOMIZABLEOBJECT_API FCustomizableObjectCompilerBase
 {
 public:
@@ -573,6 +576,90 @@ public:
 	// See UCustomizableObjectSystem::LockObject()
 	bool IsLocked() const;
 
+	/** Modify the provided mutable parameters so that the forced values for the given customizable object state are applied. */
+	void ApplyStateForcedValuesToParameters(int32 State, mu::Parameters* Parameters);
+
+	int32 FindParameter(const FString& Name) const;
+
+	EMutableParameterType GetParameterType(int32 ParamIndex) const;
+
+	int32 FindIntParameterValue(int32 ParamIndex, const FString& Value) const;
+
+	FParameterUIData GetStateUIMetadataFromIndex(int32 StateIndex) const;
+
+	FParameterUIData GetStateUIMetadata(const FString& StateName) const;
+
+	FString GetStateName(int32 StateIndex) const;
+
+#if WITH_EDITORONLY_DATA
+	void PostCompile();
+#endif
+
+	/** Returns a pointer to the BulkData subobject, only valid in packaged builds. */
+	const UCustomizableObjectBulk* GetStreamableBulkData() const;
+
+	UCustomizableObject* GetPublic() const;
+
+#if WITH_EDITOR
+	/** Compose file name. */
+	FString GetCompiledDataFileName(bool bIsModel, const ITargetPlatform* InTargetPlatform = nullptr, bool bIsDiskStreamer = false);
+#endif
+	
+	/** Rebuild ParameterProperties from the current compiled model. */
+	void UpdateParameterPropertiesFromModel(const TSharedPtr<mu::Model>& Model);
+
+	void AddUncompiledCOWarning(const FString& AdditionalLoggingInfo);
+
+#if WITH_EDITOR
+	// Create new GUID for this CO
+	void UpdateVersionId();
+	
+	FGuid GetVersionId() const;
+
+	// Unless we are packaging there is no need for keeping all the data generated during compilation, this information is stored in the derived data.
+	void ClearCompiledData(bool bIsCooking);
+
+	void SaveEmbeddedData(FArchive& Ar);
+
+	// Compile the object for a specific platform - Compile for Cook Customizable Object
+	void CompileForTargetPlatform(const ITargetPlatform* TargetPlatform);
+	
+	// Add a profile that stores the values of the parameters used by the CustomInstance.
+	FReply AddNewParameterProfile(FString Name, class UCustomizableObjectInstance& CustomInstance);
+
+	// Compose folder name where the data is stored
+	FString GetCompiledDataFolderPath() const;
+
+	/** Generic Save/Load methods to write/read compiled data */
+	void SaveCompiledData(FArchive& Ar, bool bSkipEditorOnlyData = false);
+	void LoadCompiledData(FArchive& Ar, const ITargetPlatform* InTargetPlatform, bool bSkipEditorOnlyData = false);
+
+	/** Load compiled data for the running platform from disk, this is used to load Editor Compilations. */
+	void LoadCompiledDataFromDisk();
+
+	/** Cache platform data for cook */
+	void CachePlatformData(const ITargetPlatform* InTargetPlatform, const TArray64<uint8>& InObjectBytes, const TArray64<uint8>& InBulkBytes);
+	
+	/** Loads data previously compiled in BeginCacheForCookedPlatformData onto the UProperties in *this,
+	  * in preparation for saving the cooked package for *this or for a CustomizableObjectInstance using *this.
+      * Returns whether the data was successfully loaded. */
+	bool TryLoadCompiledCookDataForPlatform(const ITargetPlatform* TargetPlatform);
+#endif
+
+	// Data that may be stored in the asset itself, only in packaged builds.
+	void LoadEmbeddedData(FArchive& Ar);
+	
+	/** Compute bIsChildObject if currently possible to do so. Return whether it was computed. */
+	bool TryUpdateIsChildObject();
+
+	void SetIsChildObject(bool bIsChildObject);
+
+	/** Return the names used by mutable to identify which mu::Image should be considered of LowPriority. */
+	void GetLowPriorityTextureNames(TArray<FString>& OutTextureNames);
+
+	/** Return the MinLOD index to generate based on the active LODSettings (PerPlatformMinLOD or PerQualityLevelMinLOD) */
+	int32 GetMinLODIndex() const;
+	
 #if WITH_EDITORONLY_DATA
 	/** Return true if the CO is not compiled or the ParticipatingObjects system has detected a change (participating objects dirty or re-saved since last compilation). */
 	bool IsCompilationOutOfDate() const;
diff --git a/Engine/Plugins/Experimental/Mutable/Source/CustomizableObject/Private/MuCO/CustomizableObject.cpp b/Engine/Plugins/Experimental/Mutable/Source/CustomizableObject/Private/MuCO/CustomizableObject.cpp
--- a/Engine/Plugins/Experimental/Mutable/Source/CustomizableObject/Private/MuCO/CustomizableObject.cpp
+++ b/Engine/Plugins/Experimental/Mutable/Source/CustomizableObject/Private/MuCO/CustomizableObject.cpp
@@ -92,9 +92,15 @@ bool UCustomizableObject::IsEditorOnly() const
 }
 
 
-void UCustomizableObject::UpdateVersionId()
+void UCustomizableObjectPrivate::UpdateVersionId()
 {
-	VersionId = FGuid::NewGuid();
+	GetPublic()->VersionId = FGuid::NewGuid();
+}
+
+
+FGuid UCustomizableObjectPrivate::GetVersionId() const
+{
+	return GetPublic()->VersionId;
 }
 
 
@@ -126,7 +132,7 @@ void UCustomizableObject::PreSave(FObjectPreSaveContext ObjectSaveContext)
 	Super::PreSave(ObjectSaveContext);
 
 	// Update the derived child object flag
-	if (TryUpdateIsChildObject())
+	if (GetPrivate()->TryUpdateIsChildObject())
 	{
 		if (bIsChildObject)
 		{
@@ -143,7 +149,7 @@ void UCustomizableObject::PreSave(FObjectPreSaveContext ObjectSaveContext)
 	{
 		const ITargetPlatform* TargetPlatform = ObjectSaveContext.GetTargetPlatform();
 		// Load cached data before saving
-		if (TryLoadCompiledCookDataForPlatform(TargetPlatform))
+		if (GetPrivate()->TryLoadCompiledCookDataForPlatform(TargetPlatform))
 		{
 			// Create an export object to manage the streamable data
 			if (!BulkData)
@@ -161,18 +167,18 @@ void UCustomizableObject::PreSave(FObjectPreSaveContext ObjectSaveContext)
 			UE_LOG(LogMutable, Warning, TEXT("Cook: Customizable Object [%s] is missing [%s] platform data."), *GetName(),
 				*ObjectSaveContext.GetTargetPlatform()->PlatformName());
 			
-			ClearCompiledData(true);
+			GetPrivate()->ClearCompiledData(true);
 		}
 	}
 #endif
 }
 
-bool UCustomizableObject::TryUpdateIsChildObject()
+bool UCustomizableObjectPrivate::TryUpdateIsChildObject()
 {
 	FCustomizableObjectCompilerBase* Compiler = UCustomizableObjectSystem::GetInstance()->GetNewCompiler();
 	if (Compiler)
 	{
-		bIsChildObject = !Compiler->IsRootObject(this);
+		GetPublic()->bIsChildObject = !Compiler->IsRootObject(GetPublic());
 		delete Compiler;
 		return true;
 	}
@@ -182,10 +188,23 @@ bool UCustomizableObject::TryUpdateIsChildObject()
 	}
 }
 
+
+bool UCustomizableObject::IsChildObject() const
+{
+	return bIsChildObject;
+}
+
+
+void UCustomizableObjectPrivate::SetIsChildObject(bool bIsChildObject)
+{
+	GetPublic()->bIsChildObject = bIsChildObject;
+}
+
+
 #if WITH_EDITOR
-bool UCustomizableObject::TryLoadCompiledCookDataForPlatform(const ITargetPlatform* TargetPlatform)
+bool UCustomizableObjectPrivate::TryLoadCompiledCookDataForPlatform(const ITargetPlatform* TargetPlatform)
 {
-	const FMutableCachedPlatformData* PlatformData = GetPrivate()->CachedPlatformsData.Find(TargetPlatform->PlatformName());
+	const FMutableCachedPlatformData* PlatformData = CachedPlatformsData.Find(TargetPlatform->PlatformName());
 	if (!PlatformData)
 	{
 		return false;
@@ -218,7 +237,7 @@ void UCustomizableObject::PostLoad()
 			{
 				if (GetPackage() == Pkg)
 				{
-					UpdateVersionId();
+					GetPrivate()->UpdateVersionId();
 				}
 			});
 	}
@@ -235,11 +254,11 @@ void UCustomizableObject::PostLoad()
 		const FAssetRegistryModule& AssetRegistryModule = FModuleManager::LoadModuleChecked<FAssetRegistryModule>("AssetRegistry");
 		if (AssetRegistryModule.Get().IsLoadingAssets())
 		{
-			AssetRegistryModule.Get().OnFilesLoaded().AddUObject(this, &UCustomizableObject::LoadCompiledDataFromDisk);
+			AssetRegistryModule.Get().OnFilesLoaded().AddUObject(GetPrivate(), &UCustomizableObjectPrivate::LoadCompiledDataFromDisk);
 		}
 		else
 		{
-			LoadCompiledDataFromDisk();
+			GetPrivate()->LoadCompiledDataFromDisk();
 		}
 	}
 #endif
@@ -264,7 +283,7 @@ void UCustomizableObject::Serialize(FArchive& Ar_Asset)
 		if (Ar_Asset.IsSaving())
 		{
 			UE_LOG(LogMutable, Verbose, TEXT("Serializing cooked data for Customizable Object [%s]."), *GetName());
-			SaveEmbeddedData(Ar_Asset);
+			GetPrivate()->SaveEmbeddedData(Ar_Asset);
 		}
 	}
 	else
@@ -276,7 +295,7 @@ void UCustomizableObject::Serialize(FArchive& Ar_Asset)
 #else
 	if (Ar_Asset.IsLoading())
 	{
-		LoadEmbeddedData(Ar_Asset);
+		GetPrivate()->LoadEmbeddedData(Ar_Asset);
 	}
 #endif
 }
@@ -304,8 +323,8 @@ void UCustomizableObject::BeginCacheForCookedPlatformData(const ITargetPlatform*
 	if (Private->CachedPlatformNames.Find(TargetPlatform->PlatformName()) == INDEX_NONE)
 	{
 		// Compile and save in the CachedPlatformsData map
-		CompileForTargetPlatform(TargetPlatform);
-		Private->CachedPlatformNames.Add(TargetPlatform->PlatformName());
+		GetPrivate()->CompileForTargetPlatform(TargetPlatform);
+		GetPrivate()->CachedPlatformNames.Add(TargetPlatform->PlatformName());
 	}
 }
 
@@ -326,24 +345,24 @@ FGuid GenerateIdentifier(const UCustomizableObject& CustomizableObject)
 }
 
 
-void UCustomizableObject::ClearCompiledData(bool bIsCooking)
+void UCustomizableObjectPrivate::ClearCompiledData(bool bIsCooking)
 {
-	Private->GetModelResources(bIsCooking) = FModelResources();
+	GetModelResources(bIsCooking) = FModelResources();
 
-	ContributingMorphTargetsInfo.Empty();
-	MorphTargetReconstructionData.Empty();
-	ClothMeshToMeshVertData.Empty();
-	ContributingClothingAssetsData.Empty();
-	ClothSharedConfigsData.Empty();
+	GetPublic()->ContributingMorphTargetsInfo.Empty();
+	GetPublic()->MorphTargetReconstructionData.Empty();
+	GetPublic()->ClothMeshToMeshVertData.Empty();
+	GetPublic()->ContributingClothingAssetsData.Empty();
+	GetPublic()->ClothSharedConfigsData.Empty();
 
 #if WITH_EDITORONLY_DATA
-	GetPrivate()->CustomizableObjectPathMap.Empty();
-	GetPrivate()->GroupNodeMap.Empty();
-	Private->ParticipatingObjects.Empty();
+	CustomizableObjectPathMap.Empty();
+	GroupNodeMap.Empty();
+	ParticipatingObjects.Empty();
 #endif
 
-	HashToStreamableBlock.Empty();
-	BulkData = nullptr;
+	GetPublic()->HashToStreamableBlock.Empty();
+	GetPublic()->BulkData = nullptr;
 }
 
 
@@ -458,100 +477,100 @@ void SerializeStreamedResources(FArchive& Ar, UObject* Object, TArray<FCustomiza
 }
 
 
-void UCustomizableObject::SaveCompiledData(FArchive& MemoryWriter, bool bIsCooking)
+void UCustomizableObjectPrivate::SaveCompiledData(FArchive& MemoryWriter, bool bIsCooking)
 {
 	int32 InternalVersion = UCustomizableObjectPrivate::CurrentSupportedVersion;
-	MutableCompiledDataStreamHeader Header(InternalVersion, VersionId);
+	MutableCompiledDataStreamHeader Header(InternalVersion, GetVersionId());
 	MemoryWriter << Header;
 
-	FModelResources& ModelResources = Private->GetModelResources(false);
+	FModelResources& LocalModelResources = GetModelResources(false);
 
-	MemoryWriter << ModelResources.ReferenceSkeletalMeshesData;
+	MemoryWriter << LocalModelResources.ReferenceSkeletalMeshesData;
 
-	SerializeStreamedResources(MemoryWriter, this, StreamedResourceData, bIsCooking);
+	SerializeStreamedResources(MemoryWriter, GetPublic(), GetPublic()->StreamedResourceData, bIsCooking);
 
-	int32 NumReferencedMaterials = ModelResources.Materials.Num();
+	int32 NumReferencedMaterials = LocalModelResources.Materials.Num();
 	MemoryWriter << NumReferencedMaterials;
 
-	for (const TSoftObjectPtr<UMaterialInterface>& Material : ModelResources.Materials)
+	for (const TSoftObjectPtr<UMaterialInterface>& Material : LocalModelResources.Materials)
 	{
 		FString StringRef = Material.ToString();
 		MemoryWriter << StringRef;
 	}
 
-	int32 NumReferencedSkeletons = ModelResources.Skeletons.Num();
+	int32 NumReferencedSkeletons = LocalModelResources.Skeletons.Num();
 	MemoryWriter << NumReferencedSkeletons;
 
-	for (const TSoftObjectPtr<USkeleton>& Skeleton : ModelResources.Skeletons)
+	for (const TSoftObjectPtr<USkeleton>& Skeleton : LocalModelResources.Skeletons)
 	{
 		FString StringRef = Skeleton.ToString();
 		MemoryWriter << StringRef;
 	}
 
-	int32 NumPassthroughTextures = ModelResources.PassThroughTextures.Num();
+	int32 NumPassthroughTextures = LocalModelResources.PassThroughTextures.Num();
 	MemoryWriter << NumPassthroughTextures;
 
-	for (const TSoftObjectPtr<UTexture>& PassthroughTexture : ModelResources.PassThroughTextures)
+	for (const TSoftObjectPtr<UTexture>& PassthroughTexture : LocalModelResources.PassThroughTextures)
 	{
 		FString StringRef = PassthroughTexture.ToString();
 		MemoryWriter << StringRef;
 	}
 
-	int32 NumPhysicsAssets = ModelResources.PhysicsAssets.Num();
+	int32 NumPhysicsAssets = LocalModelResources.PhysicsAssets.Num();
 	MemoryWriter << NumPhysicsAssets;
 
-	for (const TSoftObjectPtr<UPhysicsAsset>& PhysicsAsset : ModelResources.PhysicsAssets)
+	for (const TSoftObjectPtr<UPhysicsAsset>& PhysicsAsset : LocalModelResources.PhysicsAssets)
 	{
 		FString StringRef = PhysicsAsset.ToString();
 		MemoryWriter << StringRef;
 	}
 
-	int32 NumAnimBps = ModelResources.AnimBPs.Num();
+	int32 NumAnimBps = LocalModelResources.AnimBPs.Num();
 	MemoryWriter << NumAnimBps;
 
-	for (const TSoftClassPtr<UAnimInstance>& AnimBp : ModelResources.AnimBPs)
+	for (const TSoftClassPtr<UAnimInstance>& AnimBp : LocalModelResources.AnimBPs)
 	{
 		FString StringRef = AnimBp.ToString();
 		MemoryWriter << StringRef;
 	}
 
-	MemoryWriter << ModelResources.AnimBpOverridePhysiscAssetsInfo;
+	MemoryWriter << LocalModelResources.AnimBpOverridePhysiscAssetsInfo;
 
-	MemoryWriter << ModelResources.MaterialSlotNames;
-	MemoryWriter << ModelResources.BoneNames;
-	MemoryWriter << ModelResources.SocketArray;
+	MemoryWriter << LocalModelResources.MaterialSlotNames;
+	MemoryWriter << LocalModelResources.BoneNames;
+	MemoryWriter << LocalModelResources.SocketArray;
 
-	MemoryWriter << ModelResources.SkinWeightProfilesInfo;
+	MemoryWriter << LocalModelResources.SkinWeightProfilesInfo;
 
-	MemoryWriter << ModelResources.ImageProperties;
-	MemoryWriter << ModelResources.ParameterUIDataMap;
-	MemoryWriter << ModelResources.StateUIDataMap;
+	MemoryWriter << LocalModelResources.ImageProperties;
+	MemoryWriter << LocalModelResources.ParameterUIDataMap;
+	MemoryWriter << LocalModelResources.StateUIDataMap;
 
-	MemoryWriter << ContributingMorphTargetsInfo;
-	MemoryWriter << MorphTargetReconstructionData;
+	MemoryWriter << GetPublic()->ContributingMorphTargetsInfo;
+	MemoryWriter << GetPublic()->MorphTargetReconstructionData;
 	
-	MemoryWriter << ClothMeshToMeshVertData;
-	MemoryWriter << ContributingClothingAssetsData;
-	MemoryWriter << ClothSharedConfigsData; 
+	MemoryWriter << GetPublic()->ClothMeshToMeshVertData;
+	MemoryWriter << GetPublic()->ContributingClothingAssetsData;
+	MemoryWriter << GetPublic()->ClothSharedConfigsData; 
 
-	MemoryWriter << HashToStreamableBlock;
+	MemoryWriter << GetPublic()->HashToStreamableBlock;
 
-	MemoryWriter << LODSettings.NumLODsInRoot;
-	MemoryWriter << NumMeshComponentsInRoot;
+	MemoryWriter << GetPublic()->LODSettings.NumLODsInRoot;
+	MemoryWriter << GetPublic()->NumMeshComponentsInRoot;
 
-	MemoryWriter << LODSettings.FirstLODAvailable;
+	MemoryWriter << GetPublic()->LODSettings.FirstLODAvailable;
 
-	MemoryWriter << LODSettings.NumLODsToStream;
-	MemoryWriter << LODSettings.bLODStreamingEnabled;
+	MemoryWriter << GetPublic()->LODSettings.NumLODsToStream;
+	MemoryWriter << GetPublic()->LODSettings.bLODStreamingEnabled;
 
 	// Editor Only data
-	MemoryWriter << GetPrivate()->bIsCompiledWithOptimization;
-	MemoryWriter << GetPrivate()->CustomizableObjectPathMap;
-	MemoryWriter << GetPrivate()->GroupNodeMap;
-	MemoryWriter << GetPrivate()->ParticipatingObjects;
+	MemoryWriter << bIsCompiledWithOptimization;
+	MemoryWriter << CustomizableObjectPathMap;
+	MemoryWriter << GroupNodeMap;
+	MemoryWriter << ParticipatingObjects;
 }
 
-void UCustomizableObject::LoadCompiledData(FArchive& MemoryReader, const ITargetPlatform* InTargetPlatform, bool bIsCooking)
+void UCustomizableObjectPrivate::LoadCompiledData(FArchive& MemoryReader, const ITargetPlatform* InTargetPlatform, bool bIsCooking)
 {
 	TSharedPtr<mu::Model, ESPMode::ThreadSafe> LoadedModel;
 	ClearCompiledData(bIsCooking);
@@ -564,17 +583,17 @@ void UCustomizableObject::LoadCompiledData(FArchive& MemoryReader, const ITarget
 		// Make sure mutable has been initialised.
 		UCustomizableObjectSystem::GetInstance();
 
-		FModelResources& ModelResources = Private->GetModelResources(bIsCooking);
-		ModelResources = FModelResources();
+		FModelResources& LocalModelResource = GetModelResources(bIsCooking);
+		LocalModelResource = FModelResources();
 
-		MemoryReader << ModelResources.ReferenceSkeletalMeshesData;
+		MemoryReader << LocalModelResource.ReferenceSkeletalMeshesData;
 
-		SerializeStreamedResources(MemoryReader, this, StreamedResourceData, bIsCooking);
+		SerializeStreamedResources(MemoryReader, GetPublic(), GetPublic()->StreamedResourceData, bIsCooking);
 
 		// Initialize resources. 
-		for(FMutableRefSkeletalMeshData& ReferenceSkeletalMeshData : ModelResources.ReferenceSkeletalMeshesData)
+		for(FMutableRefSkeletalMeshData& ReferenceSkeletalMeshData : LocalModelResource.ReferenceSkeletalMeshesData)
 		{
-			ReferenceSkeletalMeshData.InitResources(this, InTargetPlatform);
+			ReferenceSkeletalMeshData.InitResources(GetPublic(), InTargetPlatform);
 		}
 
 		int32 NumReferencedMaterials = 0;
@@ -585,7 +604,7 @@ void UCustomizableObject::LoadCompiledData(FArchive& MemoryReader, const ITarget
 			FString StringRef;
 			MemoryReader << StringRef;
 
-			ModelResources.Materials.Add(TSoftObjectPtr<UMaterialInterface>(FSoftObjectPath(StringRef)));
+			LocalModelResource.Materials.Add(TSoftObjectPtr<UMaterialInterface>(FSoftObjectPath(StringRef)));
 		}
 
 		int32 NumReferencedSkeletons = 0;
@@ -596,7 +615,7 @@ void UCustomizableObject::LoadCompiledData(FArchive& MemoryReader, const ITarget
 			FString StringRef;
 			MemoryReader << StringRef;
 
-			ModelResources.Skeletons.Add(TSoftObjectPtr<USkeleton>(FSoftObjectPath(StringRef)));
+			LocalModelResource.Skeletons.Add(TSoftObjectPtr<USkeleton>(FSoftObjectPath(StringRef)));
 		}
 
 		int32 NumPassthroughTextures = 0;
@@ -607,7 +626,7 @@ void UCustomizableObject::LoadCompiledData(FArchive& MemoryReader, const ITarget
 			FString StringRef;
 			MemoryReader << StringRef;
 
-			ModelResources.PassThroughTextures.Add(TSoftObjectPtr<UTexture>(FSoftObjectPath(StringRef)));
+			LocalModelResource.PassThroughTextures.Add(TSoftObjectPtr<UTexture>(FSoftObjectPath(StringRef)));
 		}
 
 		int32 NumPhysicsAssets = 0;
@@ -618,7 +637,7 @@ void UCustomizableObject::LoadCompiledData(FArchive& MemoryReader, const ITarget
 			FString StringRef;
 			MemoryReader << StringRef;
 
-			ModelResources.PhysicsAssets.Add(TSoftObjectPtr<UPhysicsAsset>(FSoftObjectPath(StringRef)));
+			LocalModelResource.PhysicsAssets.Add(TSoftObjectPtr<UPhysicsAsset>(FSoftObjectPath(StringRef)));
 		}
 
 		int32 NumAnimBps = 0;
@@ -629,48 +648,46 @@ void UCustomizableObject::LoadCompiledData(FArchive& MemoryReader, const ITarget
 			FString StringRef;
 			MemoryReader << StringRef;
 
-			ModelResources.AnimBPs.Add(TSoftClassPtr<UAnimInstance>(StringRef));
+			LocalModelResource.AnimBPs.Add(TSoftClassPtr<UAnimInstance>(StringRef));
 		}
 
-		MemoryReader << ModelResources.AnimBpOverridePhysiscAssetsInfo;
+		MemoryReader << LocalModelResource.AnimBpOverridePhysiscAssetsInfo;
 
-		MemoryReader << ModelResources.MaterialSlotNames;
-		MemoryReader << ModelResources.BoneNames;
-		MemoryReader << ModelResources.SocketArray;
+		MemoryReader << LocalModelResource.MaterialSlotNames;
+		MemoryReader << LocalModelResource.BoneNames;
+		MemoryReader << LocalModelResource.SocketArray;
 
-		MemoryReader << ModelResources.SkinWeightProfilesInfo;
+		MemoryReader << LocalModelResource.SkinWeightProfilesInfo;
 
-		MemoryReader << ModelResources.ImageProperties;
-		MemoryReader << ModelResources.ParameterUIDataMap;
-		MemoryReader << ModelResources.StateUIDataMap;
+		MemoryReader << LocalModelResource.ImageProperties;
+		MemoryReader << LocalModelResource.ParameterUIDataMap;
+		MemoryReader << LocalModelResource.StateUIDataMap;
 
 
-		MemoryReader << ContributingMorphTargetsInfo;
-		MemoryReader << MorphTargetReconstructionData;
+		MemoryReader << GetPublic()->ContributingMorphTargetsInfo;
+		MemoryReader << GetPublic()->MorphTargetReconstructionData;
 
-		MemoryReader << ClothMeshToMeshVertData;
-		MemoryReader << ContributingClothingAssetsData;
-		MemoryReader << ClothSharedConfigsData;
+		MemoryReader << GetPublic()->ClothMeshToMeshVertData;
+		MemoryReader << GetPublic()->ContributingClothingAssetsData;
+		MemoryReader << GetPublic()->ClothSharedConfigsData;
 
-		MemoryReader << HashToStreamableBlock;
+		MemoryReader << GetPublic()->HashToStreamableBlock;
 
-		MemoryReader << LODSettings.NumLODsInRoot;
-		MemoryReader << NumMeshComponentsInRoot;
+		MemoryReader << GetPublic()->LODSettings.NumLODsInRoot;
+		MemoryReader << GetPublic()->NumMeshComponentsInRoot;
 
-		MemoryReader << LODSettings.FirstLODAvailable;
+		MemoryReader << GetPublic()->LODSettings.FirstLODAvailable;
 
-		MemoryReader << LODSettings.NumLODsToStream;
-		MemoryReader << LODSettings.bLODStreamingEnabled;
+		MemoryReader << GetPublic()->LODSettings.NumLODsToStream;
+		MemoryReader << GetPublic()->LODSettings.bLODStreamingEnabled;
 
 		bool bForceRecompilation = false;
 
 		// Editor Only data
 		{
-			MemoryReader << GetPrivate()->bIsCompiledWithOptimization;
-			MemoryReader << GetPrivate()->CustomizableObjectPathMap;
-			MemoryReader << GetPrivate()->GroupNodeMap;
-
-			TMap<FName, FGuid>& ParticipatingObjects = GetPrivate()->ParticipatingObjects;
+			MemoryReader << bIsCompiledWithOptimization;
+			MemoryReader << CustomizableObjectPathMap;
+			MemoryReader << GroupNodeMap;
 			MemoryReader << ParticipatingObjects;
 
 			if (!bIsCooking)
@@ -715,11 +732,11 @@ void UCustomizableObject::LoadCompiledData(FArchive& MemoryReader, const ITarget
 	}
 	
 	UpdateParameterPropertiesFromModel(LoadedModel);
-	Private->SetModel(LoadedModel, GenerateIdentifier(*this));
+	SetModel(LoadedModel, GenerateIdentifier(*GetPublic()));
 }
 
 
-void UCustomizableObject::LoadCompiledDataFromDisk()
+void UCustomizableObjectPrivate::LoadCompiledDataFromDisk()
 {
 	ITargetPlatformManagerModule& TargetPlatformManager = GetTargetPlatformManagerRef();
 	const ITargetPlatform* RunningPlatform = TargetPlatformManager.GetRunningTargetPlatform();
@@ -763,7 +780,7 @@ void UCustomizableObject::LoadCompiledDataFromDisk()
 			&&
 			CompiledDataHeader.VersionId == StreamableDataHeader.VersionId)
 		{
-			if (IsRunningGame() || CompiledDataHeader.VersionId == VersionId)
+			if (IsRunningGame() || CompiledDataHeader.VersionId == GetVersionId())
 			{ 
 				int64 CompiledDataSize = CompiledDataFileHandle->Size() - HeaderSize;
 				TArray64<uint8> CompiledDataBytes;
@@ -781,22 +798,20 @@ void UCustomizableObject::LoadCompiledDataFromDisk()
 		delete StreamableDataFileHandle;
 	}
 
-	UCustomizableObjectPrivate* CustomizableObjectPrivate = GetPrivate();
-
-	if (!CustomizableObjectPrivate->GetModel()) // Not failed to load the model
+	if (!GetModel()) // Not failed to load the model
 	{
-		CustomizableObjectPrivate->Status.NextState(FCustomizableObjectStatusTypes::EState::NoModel);
+		Status.NextState(FCustomizableObjectStatusTypes::EState::NoModel);
 	}
 }
 
 
-void UCustomizableObject::CachePlatformData(const ITargetPlatform* InTargetPlatform, const TArray64<uint8>& InModelBytes, const TArray64< uint8>& InBulkBytes)
+void UCustomizableObjectPrivate::CachePlatformData(const ITargetPlatform* InTargetPlatform, const TArray64<uint8>& InModelBytes, const TArray64< uint8>& InBulkBytes)
 {
 	FString PlatformName = InTargetPlatform ? InTargetPlatform->PlatformName() : FPlatformProperties::PlatformName();
 
-	check(!GetPrivate()->CachedPlatformsData.Find(PlatformName));
+	check(!CachedPlatformsData.Find(PlatformName));
 
-	FMutableCachedPlatformData& Data = GetPrivate()->CachedPlatformsData.Add(PlatformName);
+	FMutableCachedPlatformData& Data = CachedPlatformsData.Add(PlatformName);
 
 	// Cache CO data and mu::Model
 	FMemoryWriter64 MemoryWriter(Data.ModelData);
@@ -809,7 +824,7 @@ void UCustomizableObject::CachePlatformData(const ITargetPlatform* InTargetPlatf
 }
 
 
-void UCustomizableObject::CompileForTargetPlatform(const ITargetPlatform* TargetPlatform)
+void UCustomizableObjectPrivate::CompileForTargetPlatform(const ITargetPlatform* TargetPlatform)
 {
 	if (!TargetPlatform)
 	{
@@ -818,12 +833,12 @@ void UCustomizableObject::CompileForTargetPlatform(const ITargetPlatform* Target
 
 	FCustomizableObjectCompilerBase* Compiler = UCustomizableObjectSystem::GetInstance()->GetNewCompiler();
 
-	bool bIsRootObject = Compiler->IsRootObject(this);
+	bool bIsRootObject = Compiler->IsRootObject(GetPublic());
 	
 	bool bIsRelevantForThisTarget =
-		(Relevancy == ECustomizableObjectRelevancy::All)
+		(GetPublic()->Relevancy == ECustomizableObjectRelevancy::All)
 		||
-		(Relevancy == ECustomizableObjectRelevancy::ClientOnly && !TargetPlatform->IsServerOnly());
+		(GetPublic()->Relevancy == ECustomizableObjectRelevancy::ClientOnly && !TargetPlatform->IsServerOnly());
 
 	if (bIsRootObject && bIsRelevantForThisTarget)
 	{
@@ -834,7 +849,7 @@ void UCustomizableObject::CompileForTargetPlatform(const ITargetPlatform* Target
 		Options.TargetPlatform = TargetPlatform;
 		Options.CustomizableObjectNumBoneInfluences = ICustomizableObjectModule::Get().GetNumBoneInfluences();
 
-		Compiler->Compile(*this, Options, false);
+		Compiler->Compile(*GetPublic(), Options, false);
 	}
 
 	delete Compiler;
@@ -907,9 +922,8 @@ bool UCustomizableObject::ConditionalAutoCompile()
 }
 
 
-FReply UCustomizableObject::AddNewParameterProfile(FString Name, UCustomizableObjectInstance& CustomInstance)
+FReply UCustomizableObjectPrivate::AddNewParameterProfile(FString Name, UCustomizableObjectInstance& CustomInstance)
 {
-
 	if (Name.IsEmpty())
 	{
 		Name = "Unnamed_Profile";
@@ -921,7 +935,7 @@ FReply UCustomizableObject::AddNewParameterProfile(FString Name, UCustomizableOb
 	bool bUniqueNameFound = false;
 	while (!bUniqueNameFound)
 	{
-		FProfileParameterDat* Found = InstancePropertiesProfiles.FindByPredicate(
+		FProfileParameterDat* Found = GetPublic()->InstancePropertiesProfiles.FindByPredicate(
 			[&ProfileName](const FProfileParameterDat& Profile) { return Profile.ProfileName == ProfileName; });
 
 		bUniqueNameFound = static_cast<bool>(!Found);
@@ -932,9 +946,9 @@ FReply UCustomizableObject::AddNewParameterProfile(FString Name, UCustomizableOb
 		}
 	}
 
-	int32 ProfileIndex = InstancePropertiesProfiles.Emplace();
+	int32 ProfileIndex = GetPublic()->InstancePropertiesProfiles.Emplace();
 
-	InstancePropertiesProfiles[ProfileIndex].ProfileName = ProfileName;
+	GetPublic()->InstancePropertiesProfiles[ProfileIndex].ProfileName = ProfileName;
 	CustomInstance.SaveParametersToProfile(ProfileIndex);
 
 	Modify();
@@ -943,16 +957,16 @@ FReply UCustomizableObject::AddNewParameterProfile(FString Name, UCustomizableOb
 }
 
 
-FString UCustomizableObject::GetCompiledDataFolderPath() const
+FString UCustomizableObjectPrivate::GetCompiledDataFolderPath() const
 {	
 	return FPaths::ConvertRelativePathToFull(FPaths::ProjectSavedDir() + TEXT("MutableStreamedDataEditor/"));
 }
 
 
-FString UCustomizableObject::GetCompiledDataFileName(bool bIsModel, const ITargetPlatform* InTargetPlatform, bool bIsDiskStreamer)
+FString UCustomizableObjectPrivate::GetCompiledDataFileName(bool bIsModel, const ITargetPlatform* InTargetPlatform, bool bIsDiskStreamer)
 {
 	const FString PlatformName = InTargetPlatform ? InTargetPlatform->PlatformName() : FPlatformProperties::PlatformName();
-	const FString FileIdentifier = bIsDiskStreamer ? GetPrivate()->Identifier.ToString() : GenerateIdentifier(*this).ToString();
+	const FString FileIdentifier = bIsDiskStreamer ? Identifier.ToString() : GenerateIdentifier(*GetPublic()).ToString();
 	const FString Extension = bIsModel ? TEXT("_M.mut") : TEXT("_S.mut");
 	return PlatformName + FileIdentifier + Extension;
 }
@@ -966,37 +980,37 @@ FString UCustomizableObject::GetDesc()
 }
 
 
-void UCustomizableObject::SaveEmbeddedData(FArchive& Ar)
+void UCustomizableObjectPrivate::SaveEmbeddedData(FArchive& Ar)
 {
 	UE_LOG(LogMutable, Verbose, TEXT("Saving embedded data for Customizable Object [%s] now at position %d."), *GetName(), int(Ar.Tell()));
 
-	int32 InternalVersion = Private->GetModel() ? UCustomizableObjectPrivate::CurrentSupportedVersion : -1;
+	int32 InternalVersion = GetModel() ? CurrentSupportedVersion : -1;
 	Ar << InternalVersion;
 
-	if (Private->GetModel())
+	if (GetModel())
 	{
 		// General derived flags
-		Ar << GetPrivate()->bDisableTextureStreaming;
+		Ar << bDisableTextureStreaming;
 
 		// Serialize morph data
 		{
-			Ar << ContributingMorphTargetsInfo;
-			Ar << MorphTargetReconstructionData;
+			Ar << GetPublic()->ContributingMorphTargetsInfo;
+			Ar << GetPublic()->MorphTargetReconstructionData;
 		}
 		
 		{
-			Ar << ClothMeshToMeshVertData;
-			Ar << ContributingClothingAssetsData;
-			Ar << ClothSharedConfigsData;
+			Ar << GetPublic()->ClothMeshToMeshVertData;
+			Ar << GetPublic()->ContributingClothingAssetsData;
+			Ar << GetPublic()->ClothSharedConfigsData;
 		}
 
 		// Serialise the entire model, but unload the streamable data first.
 		{
-			Private->GetModel()->UnloadExternalData();
+			GetModel()->UnloadExternalData();
 
 			UnrealMutableOutputStream stream(Ar);
 			mu::OutputArchive arch(&stream);
-			mu::Model::Serialise(Private->GetModel().Get(), arch);
+			mu::Model::Serialise(GetModel().Get(), arch);
 		}
 
 		UE_LOG(LogMutable, Verbose, TEXT("Saved embedded data for Customizable Object [%s] now at position %d."), *GetName(), int(Ar.Tell()));
@@ -1005,7 +1019,7 @@ void UCustomizableObject::SaveEmbeddedData(FArchive& Ar)
 
 #endif // End WITH_EDITOR 
 
-void UCustomizableObject::LoadEmbeddedData(FArchive& Ar)
+void UCustomizableObjectPrivate::LoadEmbeddedData(FArchive& Ar)
 {
 	MUTABLE_CPUPROFILER_SCOPE(UCustomizableObject::LoadEmbeddedData)
 
@@ -1014,23 +1028,23 @@ void UCustomizableObject::LoadEmbeddedData(FArchive& Ar)
 
 	// If this fails, something went wrong with the packaging: we have data that belongs
 	// to a different version than the code.
-	check(UCustomizableObjectPrivate::CurrentSupportedVersion==InternalVersion);
+	check(CurrentSupportedVersion == InternalVersion);
 
-	if(UCustomizableObjectPrivate::CurrentSupportedVersion == InternalVersion)
+	if(CurrentSupportedVersion == InternalVersion)
 	{
 		// General derived flags
-		Ar << GetPrivate()->bDisableTextureStreaming;
+		Ar << bDisableTextureStreaming;
 
 		// Load morph data
 		{
-			Ar << ContributingMorphTargetsInfo;
-			Ar << MorphTargetReconstructionData;
+			Ar << GetPublic()->ContributingMorphTargetsInfo;
+			Ar << GetPublic()->MorphTargetReconstructionData;
 		}
 		
 		{
-			Ar << ClothMeshToMeshVertData;
-			Ar << ContributingClothingAssetsData;
-			Ar << ClothSharedConfigsData;
+			Ar << GetPublic()->ClothMeshToMeshVertData;
+			Ar << GetPublic()->ContributingClothingAssetsData;
+			Ar << GetPublic()->ClothSharedConfigsData;
 		}
 		
 		// Load model
@@ -1041,7 +1055,7 @@ void UCustomizableObject::LoadEmbeddedData(FArchive& Ar)
 		// Create parameter properties
 		UpdateParameterPropertiesFromModel(Model);
 
-		Private->SetModel( Model, FGuid());
+		SetModel( Model, FGuid());
 	}
 }
 
@@ -1061,7 +1075,7 @@ bool UCustomizableObject::IsCompiled() const
 }
 
 
-void UCustomizableObject::AddUncompiledCOWarning(const FString& AdditionalLoggingInfo)
+void UCustomizableObjectPrivate::AddUncompiledCOWarning(const FString& AdditionalLoggingInfo)
 {
 	// Send a warning (on-screen notification, log error, and in-editor notification)
 	UCustomizableObjectSystem* System = UCustomizableObjectSystem::GetInstance();
@@ -1070,7 +1084,7 @@ void UCustomizableObject::AddUncompiledCOWarning(const FString& AdditionalLoggin
 		return;
 	}
 
-	System->AddUncompiledCOWarning(*this, &AdditionalLoggingInfo);
+	System->AddUncompiledCOWarning(*GetPublic(), &AdditionalLoggingInfo);
 }
 
 
@@ -1119,12 +1133,18 @@ int32 UCustomizableObject::GetStateCount() const
 
 
 FString UCustomizableObject::GetStateName(int32 StateIndex) const
+{
+	return GetPrivate()->GetStateName(StateIndex);
+}
+
+
+FString UCustomizableObjectPrivate::GetStateName(int32 StateIndex) const
 {
 	FString Result;
 
-	if (Private->GetModel())
+	if (GetModel())
 	{
-		Result = Private->GetModel()->GetStateName(StateIndex);
+		Result = GetModel()->GetStateName(StateIndex);
 	}
 
 	return Result;
@@ -1178,15 +1198,30 @@ FString UCustomizableObject::GetStateParameterName(int32 StateIndex, int32 Param
 
 
 #if WITH_EDITORONLY_DATA
-void UCustomizableObject::PostCompile()
+void UCustomizableObjectPrivate::PostCompile()
 {
-	GetPrivate()->PostCompileDelegate.Broadcast();
+	PostCompileDelegate.Broadcast();
 
-	UCustomizableObjectSystemPrivate::HideOnScreenCompileWarnings(*GetPrivate());	
+	UCustomizableObjectSystemPrivate::HideOnScreenCompileWarnings(*this);	
 }
 #endif
 
 
+const UCustomizableObjectBulk* UCustomizableObjectPrivate::GetStreamableBulkData() const
+{
+	return GetPublic()->BulkData;
+}
+
+
+UCustomizableObject* UCustomizableObjectPrivate::GetPublic() const
+{
+	UCustomizableObject* Public = StaticCast<UCustomizableObject*>(GetOuter());
+	check(Public);
+
+	return Public;
+}
+
+
 FPostCompileDelegate& UCustomizableObject::GetPostCompileDelegate() const
 {
 	return GetPrivate()->PostCompileDelegate;
@@ -1207,35 +1242,6 @@ UCustomizableObjectInstance* UCustomizableObject::CreateInstance()
 }
 
 
-#if WITH_EDITOR
-void UCustomizableObject::SetModel(TSharedPtr<mu::Model, ESPMode::ThreadSafe> Model)
-{
-	// Generate a map that using the resource id tells the offset and size of the resource inside the bulk data
-	if (Model)
-	{
-		uint64 Offset = 0;
-
-		const int32 NumStreamingFiles = Model->GetRomCount();
-		HashToStreamableBlock.Empty(NumStreamingFiles);
-
-		for (size_t FileIndex = 0; FileIndex < NumStreamingFiles; ++FileIndex)
-		{
-			const uint32 ResourceId = Model->GetRomId(FileIndex);
-			const uint32 ResourceSize = Model->GetRomSize(FileIndex);
-
-			HashToStreamableBlock.Add(ResourceId, FMutableStreamableBlock({0, Offset, ResourceSize }));
-			Offset += ResourceSize;
-		}
-	}
-
-	// Generate ParameterProperties and IntParameterLookUpTable
-	UpdateParameterPropertiesFromModel(Model);
-
-	Private->SetModel(Model, GenerateIdentifier(*this));
-}
-#endif // End WITH_EDITOR
-
-
 int32 UCustomizableObject::GetNumLODs() const
 {
 	return LODSettings.NumLODsInRoot;
@@ -1254,9 +1260,15 @@ int32 UCustomizableObject::GetParameterCount() const
 
 EMutableParameterType UCustomizableObject::GetParameterType(int32 ParamIndex) const
 {
-	if (ParamIndex >= 0 && ParamIndex < GetPrivate()->ParameterProperties.Num())
+	return GetPrivate()->GetParameterType(ParamIndex);
+}
+
+
+EMutableParameterType UCustomizableObjectPrivate::GetParameterType(int32 ParamIndex) const
+{
+	if (ParamIndex >= 0 && ParamIndex < ParameterProperties.Num())
 	{
-		return GetPrivate()->ParameterProperties[ParamIndex].Type;
+		return ParameterProperties[ParamIndex].Type;
 	}
 	else
 	{
@@ -1310,15 +1322,15 @@ const FString & UCustomizableObject::GetParameterName(int32 ParamIndex) const
 }
 
 
-void UCustomizableObject::UpdateParameterPropertiesFromModel(const TSharedPtr<mu::Model>& Model)
+void UCustomizableObjectPrivate::UpdateParameterPropertiesFromModel(const TSharedPtr<mu::Model>& Model)
 {
 	if (Model)
 	{
 		mu::ParametersPtr MutableParameters = mu::Model::NewParameters(Model);
 		int paramCount = MutableParameters->GetCount();
 
-		GetPrivate()->ParameterProperties.Reset(paramCount);
-		GetPrivate()->ParameterPropertiesLookupTable.Empty(paramCount);
+		ParameterProperties.Reset(paramCount);
+		ParameterPropertiesLookupTable.Empty(paramCount);
 		for (int paramIndex = 0; paramIndex<paramCount; ++paramIndex)
 		{
 			FMutableModelParameterProperties Data;
@@ -1380,14 +1392,14 @@ void UCustomizableObject::UpdateParameterPropertiesFromModel(const TSharedPtr<mu
 				break;
 			}
 
-			GetPrivate()->ParameterProperties.Add(Data);
-			GetPrivate()->ParameterPropertiesLookupTable.Add(Data.Name, paramIndex);
+			ParameterProperties.Add(Data);
+			ParameterPropertiesLookupTable.Add(Data.Name, paramIndex);
 		}
 	}
 	else
 	{
-		GetPrivate()->ParameterProperties.Empty();
-		GetPrivate()->ParameterPropertiesLookupTable.Empty();
+		ParameterProperties.Empty();
+		ParameterPropertiesLookupTable.Empty();
 	}
 }
 
@@ -1437,7 +1449,13 @@ const FString& UCustomizableObject::GetIntParameterAvailableOption(int32 ParamIn
 
 int32 UCustomizableObject::FindParameter(const FString& Name) const
 {
-	const int32 * Found = GetPrivate()->ParameterPropertiesLookupTable.Find(Name);
+	return GetPrivate()->FindParameter(Name);
+}
+
+
+int32 UCustomizableObjectPrivate::FindParameter(const FString& Name) const
+{
+	const int32 * Found = ParameterPropertiesLookupTable.Find(Name);
 	if (Found == nullptr)
 	{
 		return INDEX_NONE;
@@ -1447,11 +1465,17 @@ int32 UCustomizableObject::FindParameter(const FString& Name) const
 
 
 int32 UCustomizableObject::FindIntParameterValue(int32 ParamIndex, const FString& Value) const
+{
+	return GetPrivate()->FindIntParameterValue(ParamIndex, Value);
+}
+
+
+int32 UCustomizableObjectPrivate::FindIntParameterValue(int32 ParamIndex, const FString& Value) const
 {
 	int32 MinValueIndex = 0;
-	if (ParamIndex >= 0 && ParamIndex < GetPrivate()->ParameterProperties.Num())
+	if (ParamIndex >= 0 && ParamIndex < ParameterProperties.Num())
 	{
-		const TArray<FMutableModelParameterValue>& PossibleValues = GetPrivate()->ParameterProperties[ParamIndex].PossibleValues;
+		const TArray<FMutableModelParameterValue>& PossibleValues = ParameterProperties[ParamIndex].PossibleValues;
 		if (PossibleValues.Num())
 		{
 			MinValueIndex = PossibleValues[0].Value;
@@ -1520,13 +1544,25 @@ FParameterUIData UCustomizableObject::GetParameterUIMetadataFromIndex(int32 Para
 
 FParameterUIData UCustomizableObject::GetStateUIMetadata(const FString& StateName) const
 {
-	const FParameterUIData* StateUIData = Private->GetModelResources().StateUIDataMap.Find(StateName);
+	return GetPrivate()->GetStateUIMetadata(StateName);
+}
+
+
+FParameterUIData UCustomizableObjectPrivate::GetStateUIMetadata(const FString& StateName) const
+{
+	const FParameterUIData* StateUIData = GetModelResources().StateUIDataMap.Find(StateName);
 
 	return StateUIData ? *StateUIData : FParameterUIData();
 }
 
 
 FParameterUIData UCustomizableObject::GetStateUIMetadataFromIndex(int32 StateIndex) const
+{
+	return GetPrivate()->GetStateUIMetadataFromIndex(StateIndex);
+}
+
+
+FParameterUIData UCustomizableObjectPrivate::GetStateUIMetadataFromIndex(int32 StateIndex) const
 {
 	return GetStateUIMetadata(GetStateName(StateIndex));
 }
@@ -1700,7 +1736,7 @@ bool UCustomizableObject::IsParameterMultidimensional(const int32& InParamIndex)
 }
 
 
-void UCustomizableObject::ApplyStateForcedValuesToParameters( int32 State, mu::Parameters* Parameters)
+void UCustomizableObjectPrivate::ApplyStateForcedValuesToParameters( int32 State, mu::Parameters* Parameters)
 {
 	FParameterUIData StateMetaData = GetStateUIMetadataFromIndex(State);
 	for (const TPair<FString, FString>& ForcedParameter : StateMetaData.ForcedParameterValues)
@@ -1747,17 +1783,17 @@ void UCustomizableObject::ApplyStateForcedValuesToParameters( int32 State, mu::P
 }
 
 
-void UCustomizableObject::GetLowPriorityTextureNames(TArray<FString>& OutTextureNames)
+void UCustomizableObjectPrivate::GetLowPriorityTextureNames(TArray<FString>& OutTextureNames)
 {
-	OutTextureNames.Reset(LowPriorityTextures.Num());
+	OutTextureNames.Reset(GetPublic()->LowPriorityTextures.Num());
 
-	if (!LowPriorityTextures.IsEmpty())
+	if (!GetPublic()->LowPriorityTextures.IsEmpty())
 	{
-		const FModelResources& ModelResources = Private->GetModelResources();
-		const int32 ImageCount = ModelResources.ImageProperties.Num();
+		const FModelResources& LocalModelResources = GetModelResources();
+		const int32 ImageCount = LocalModelResources.ImageProperties.Num();
 		for (int32 ImageIndex = 0; ImageIndex < ImageCount; ++ImageIndex)
 		{
-			if (LowPriorityTextures.Find(FName(ModelResources.ImageProperties[ImageIndex].TextureParameterName)) != INDEX_NONE)
+			if (GetPublic()->LowPriorityTextures.Find(FName(LocalModelResources.ImageProperties[ImageIndex].TextureParameterName)) != INDEX_NONE)
 			{
 				OutTextureNames.Add(FString::FromInt(ImageIndex));
 			}
@@ -1766,7 +1802,7 @@ void UCustomizableObject::GetLowPriorityTextureNames(TArray<FString>& OutTexture
 }
 
 
-int32 UCustomizableObject::GetMinLODIndex() const
+int32 UCustomizableObjectPrivate::GetMinLODIndex() const
 {
 	int32 MinLODIdx = 0;
 
@@ -1774,27 +1810,15 @@ int32 UCustomizableObject::GetMinLODIndex() const
 	{
 		if (UCustomizableObjectSystem::GetInstance() != nullptr)
 		{
-			MinLODIdx = LODSettings.MinQualityLevelLOD.GetValue(UCustomizableObjectSystem::GetInstance()->GetSkeletalMeshMinLODQualityLevel());
+			MinLODIdx = GetPublic()->LODSettings.MinQualityLevelLOD.GetValue(UCustomizableObjectSystem::GetInstance()->GetSkeletalMeshMinLODQualityLevel());
 		}
 	}
 	else
 	{
-		MinLODIdx = LODSettings.MinLOD.GetValue();
+		MinLODIdx = GetPublic()->LODSettings.MinLOD.GetValue();
 	}
 
-	return FMath::Max(MinLODIdx, LODSettings.FirstLODAvailable);
-}
-
-
-bool UCustomizableObject::IsEnableUseRefSkeletalMeshAsPlaceholder() const
-{
-	return bEnableUseRefSkeletalMeshAsPlaceholder;
-}
-
-
-bool UCustomizableObject::IsMeshCacheEnabled() const
-{
-	return bEnableMeshCache;
+	return FMath::Max(MinLODIdx, GetPublic()->LODSettings.FirstLODAvailable);
 }
 
 
@@ -2018,12 +2042,6 @@ TArray<TSharedPtr<IAsyncReadFileHandle>> UCustomizableObjectBulk::GetAsyncReadFi
 }
 
 
-bool UCustomizableObject::IsPreserveUserLODsOnFirstGeneration() const
-{
-	return bPreserveUserLODsOnFirstGeneration;
-}
-
-
 #if WITH_EDITOR
 
 void UCustomizableObjectBulk::CookAdditionalFilesOverride(const TCHAR* PackageFilename,
@@ -2315,7 +2333,7 @@ void FMutableRefSkeletalMeshData::InitResources(UCustomizableObject* InOuter, co
 	check(InOuter);
 
 	const bool bHasServer = InTargetPlatform ? !InTargetPlatform->IsClientOnly() : false;
-	if (InOuter->IsEnableUseRefSkeletalMeshAsPlaceholder() || bHasServer)
+	if (InOuter->bEnableUseRefSkeletalMeshAsPlaceholder || bHasServer)
 	{
 		SkeletalMesh = TSoftObjectPtr<USkeletalMesh>(SoftSkeletalMesh).LoadSynchronous();
 	}
diff --git a/Engine/Plugins/Experimental/Mutable/Source/CustomizableObject/Private/MuCO/CustomizableObjectInstance.cpp b/Engine/Plugins/Experimental/Mutable/Source/CustomizableObject/Private/MuCO/CustomizableObjectInstance.cpp
--- a/Engine/Plugins/Experimental/Mutable/Source/CustomizableObject/Private/MuCO/CustomizableObjectInstance.cpp
+++ b/Engine/Plugins/Experimental/Mutable/Source/CustomizableObject/Private/MuCO/CustomizableObjectInstance.cpp
@@ -623,14 +623,14 @@ void UCustomizableObjectInstance::PreSave(FObjectPreSaveContext ObjectSaveContex
 	{
 		return;
 	}
-	if (CustomizableObject->TryUpdateIsChildObject() && CustomizableObject->bIsChildObject)
+	if (CustomizableObject->GetPrivate()->TryUpdateIsChildObject() && CustomizableObject->IsChildObject())
 	{
 		UE_LOG(LogMutable, Error,
 			TEXT("CO Instance [%s] has an invalid dependency on CO that is a child object: [%s]. The instance will be unusable at runtime and may crash."),
 			*GetPathName(), *CustomizableObject->GetPathName());
 	}
 
-	CustomizableObject->TryLoadCompiledCookDataForPlatform(TargetPlatform);
+	CustomizableObject->GetPrivate()->TryLoadCompiledCookDataForPlatform(TargetPlatform);
 	Descriptor.ReloadParameters();
 }
 #endif
@@ -6203,7 +6203,7 @@ void UCustomizableObjectInstance::SetRequestedLODs(int32 InMinLOD, int32 InMaxLO
 	}
 
 	if (CVarPreserveUserLODsOnFirstGeneration.GetValueOnGameThread() &&
-		GetCustomizableObject()->IsPreserveUserLODsOnFirstGeneration() &&
+		GetCustomizableObject()->bPreserveUserLODsOnFirstGeneration &&
 		GetPrivate()->SkeletalMeshStatus != ESkeletalMeshStatus::Success)
 	{
 		return;
@@ -6221,7 +6221,7 @@ void UCustomizableObjectInstance::SetRequestedLODs(int32 InMinLOD, int32 InMaxLO
 	FMutableUpdateCandidate MutableUpdateCandidate(this);
 
 	// Clamp Min LOD
-	const int32 MinLODIdx = GetCustomizableObject()->GetMinLODIndex();
+	const int32 MinLODIdx = GetCustomizableObject()->GetPrivate()->GetMinLODIndex();
 	InMinLOD = FMath::Min(FMath::Max(InMinLOD, MinLODIdx), MinLODIdx + PrivateData->NumMaxLODsToStream);
 
 	// Clamp Max LOD
diff --git a/Engine/Plugins/Experimental/Mutable/Source/CustomizableObject/Private/MuCO/CustomizableObjectInstanceDescriptor.cpp b/Engine/Plugins/Experimental/Mutable/Source/CustomizableObject/Private/MuCO/CustomizableObjectInstanceDescriptor.cpp
--- a/Engine/Plugins/Experimental/Mutable/Source/CustomizableObject/Private/MuCO/CustomizableObjectInstanceDescriptor.cpp
+++ b/Engine/Plugins/Experimental/Mutable/Source/CustomizableObject/Private/MuCO/CustomizableObjectInstanceDescriptor.cpp
@@ -1257,7 +1257,7 @@ bool FCustomizableObjectInstanceDescriptor::HasAnyParameters() const
 	if (!CustomizableObject->IsCompiled()) \
 	{ \
 		FString AdditionalLoggingInfo = FString::Printf(TEXT("Calling function: %hs.  %s"), __FUNCTION__, ErrorMessage); \
-		CustomizableObject->AddUncompiledCOWarning(AdditionalLoggingInfo);\
+		CustomizableObject->GetPrivate()->AddUncompiledCOWarning(AdditionalLoggingInfo);\
 		return; \
 	} \
 
@@ -1267,7 +1267,7 @@ bool FCustomizableObjectInstanceDescriptor::HasAnyParameters() const
 	if (!ensureMsgf(CustomizableObject->IsCompiled(), TEXT("Customizable Object (%s) was not compiled."), *GetNameSafe(CustomizableObject))) \
 	{ \
 		FString AdditionalLoggingInfo = FString::Printf(TEXT("Calling function: %hs.  %s"), __FUNCTION__, ErrorMessage); \
-		CustomizableObject->AddUncompiledCOWarning(AdditionalLoggingInfo);\
+		CustomizableObject->GetPrivate()->AddUncompiledCOWarning(AdditionalLoggingInfo);\
 		return; \
 	} \
  
diff --git a/Engine/Plugins/Experimental/Mutable/Source/CustomizableObject/Private/MuCO/CustomizableObjectInstanceUsage.cpp b/Engine/Plugins/Experimental/Mutable/Source/CustomizableObject/Private/MuCO/CustomizableObjectInstanceUsage.cpp
--- a/Engine/Plugins/Experimental/Mutable/Source/CustomizableObject/Private/MuCO/CustomizableObjectInstanceUsage.cpp
+++ b/Engine/Plugins/Experimental/Mutable/Source/CustomizableObject/Private/MuCO/CustomizableObjectInstanceUsage.cpp
@@ -213,7 +213,7 @@ void UCustomizableObjectInstanceUsage::SetSkeletalMesh(USkeletalMesh* SkeletalMe
 		}
 		
 		if (CustomizableObject &&
-			CustomizableObject->IsMeshCacheEnabled() &&
+			CustomizableObject->bEnableMeshCache &&
 			CVarEnableMeshCache.GetValueOnAnyThread())
 		{
 			if (FCustomizableInstanceComponentData* ComponentData = Instance->GetPrivate()->GetComponentData(GetComponentIndex()))
@@ -490,7 +490,7 @@ void UCustomizableObjectInstanceUsage::Tick(float DeltaTime)
 		{
 			// If not generated yet, conditionally set the SkeletalMesh of reference
 			if (!bInstanceGenerated && !GetSkipSetReferenceSkeletalMesh() && 
-				CustomizableObject->IsEnableUseRefSkeletalMeshAsPlaceholder())
+				CustomizableObject->bEnableUseRefSkeletalMeshAsPlaceholder)
 			{
 				// Can be nullptr
 				SkeletalMesh = CustomizableObject->GetRefSkeletalMesh(ComponentIndex);
@@ -516,7 +516,7 @@ void UCustomizableObjectInstanceUsage::Tick(float DeltaTime)
 
 			if (CustomizableObject &&
 				bInstanceGenerated &&
-				CustomizableObject->IsMeshCacheEnabled() &&
+				CustomizableObject->bEnableMeshCache &&
 				CVarEnableMeshCache.GetValueOnAnyThread())
 			{
 				if (FCustomizableInstanceComponentData* ComponentData = CustomizableObjectInstance->GetPrivate()->GetComponentData(GetComponentIndex()))
diff --git a/Engine/Plugins/Experimental/Mutable/Source/CustomizableObject/Private/MuCO/CustomizableObjectSystem.cpp b/Engine/Plugins/Experimental/Mutable/Source/CustomizableObject/Private/MuCO/CustomizableObjectSystem.cpp
--- a/Engine/Plugins/Experimental/Mutable/Source/CustomizableObject/Private/MuCO/CustomizableObjectSystem.cpp
+++ b/Engine/Plugins/Experimental/Mutable/Source/CustomizableObject/Private/MuCO/CustomizableObjectSystem.cpp
@@ -136,7 +136,7 @@ FUpdateContextPrivate::FUpdateContextPrivate(UCustomizableObjectInstance& InInst
 	Parameters = Descriptor.GetParameters();
 	NumComponents = InInstance.GetCustomizableObject()->GetComponentCount();
 	
-	InInstance.GetCustomizableObject()->ApplyStateForcedValuesToParameters(CapturedDescriptor.GetState(), Parameters.get());
+	InInstance.GetCustomizableObject()->GetPrivate()->ApplyStateForcedValuesToParameters(CapturedDescriptor.GetState(), Parameters.get());
 
 	UCustomizableObjectSystem* System = UCustomizableObjectSystem::GetInstance();
 	System->GetPrivate()->CacheTextureParameters(CapturedDescriptor.GetTextureParameters());
@@ -2734,14 +2734,14 @@ namespace impl
 		// Task: Mutable Update and GetMesh
 		//-------------------------------------------------------------
 		Operation->InstanceID = Operation->bLiveUpdateMode ? CandidateInstancePrivateData->LiveUpdateModeInstanceID : 0;
-		Operation->bUseMeshCache = CustomizableObject->IsMeshCacheEnabled() && !Operation->bLiveUpdateMode && CVarEnableMeshCache.GetValueOnGameThread();
+		Operation->bUseMeshCache = CustomizableObject->bEnableMeshCache && !Operation->bLiveUpdateMode && CVarEnableMeshCache.GetValueOnGameThread();
 #if WITH_EDITOR
 		Operation->PixelFormatOverride = SystemPrivateData->ImageFormatOverrideFunc;
 #endif
 
 		if (!CandidateInstancePrivateData->HasCOInstanceFlags(ForceGenerateMipTail))
 		{
-			CustomizableObject->GetLowPriorityTextureNames(Operation->LowPriorityTextures);
+			CustomizableObject->GetPrivate()->GetLowPriorityTextureNames(Operation->LowPriorityTextures);
 		}
 
 		bool bIsInEditorViewport = false;
@@ -3358,13 +3358,13 @@ void UCustomizableObjectSystem::AddUncompiledCOWarning(const UCustomizableObject
 
 #if WITH_EDITOR
 	// Mutable will spam these warnings constantly due to the tick and LOD manager checking for instances to update with every tick. Send only one message per CO in the editor.
-	if (GetPrivate()->UncompiledCustomizableObjectIds.Find(InObject.GetVersionId()) != INDEX_NONE)
+	if (GetPrivate()->UncompiledCustomizableObjectIds.Find(InObject.GetPrivate()->GetVersionId()) != INDEX_NONE)
 	{
 		return;
 	}
 	
 	// Add notification
-	GetPrivate()->UncompiledCustomizableObjectIds.Add(InObject.GetVersionId());
+	GetPrivate()->UncompiledCustomizableObjectIds.Add(InObject.GetPrivate()->GetVersionId());
 
 	FMessageLog MessageLog("Mutable");
 	MessageLog.Warning(FText::FromString(Msg));
@@ -3441,7 +3441,7 @@ void UCustomizableObjectSystem::OnPreBeginPIE(const bool bIsSimulatingInEditor)
 		}
 		
 		const UCustomizableObject* Object = Cast<UCustomizableObject>(Asset.GetAsset());
-		if (!Object || Object->IsCompiled() || Object->GetPrivate()->IsLocked() || Object->bIsChildObject)
+		if (!Object || Object->IsCompiled() || Object->GetPrivate()->IsLocked() || Object->IsChildObject())
 		{
 			continue;
 		}
diff --git a/Engine/Plugins/Experimental/Mutable/Source/CustomizableObject/Private/MuCO/UnrealMutableModelDiskStreamer.cpp b/Engine/Plugins/Experimental/Mutable/Source/CustomizableObject/Private/MuCO/UnrealMutableModelDiskStreamer.cpp
--- a/Engine/Plugins/Experimental/Mutable/Source/CustomizableObject/Private/MuCO/UnrealMutableModelDiskStreamer.cpp
+++ b/Engine/Plugins/Experimental/Mutable/Source/CustomizableObject/Private/MuCO/UnrealMutableModelDiskStreamer.cpp
@@ -87,8 +87,8 @@ bool FUnrealMutableModelBulkReader::PrepareStreamingForObject(UCustomizableObjec
 		NewData.Model = TWeakPtr<const mu::Model>(CustomizableObject->GetPrivate()->GetModel());
 
 #if WITH_EDITOR
-		FString FolderPath = CustomizableObject->GetCompiledDataFolderPath();
-		FString FullFileName = FolderPath + CustomizableObject->GetCompiledDataFileName(false, nullptr, true);
+		FString FolderPath = CustomizableObject->GetPrivate()->GetCompiledDataFolderPath();
+		FString FullFileName = FolderPath + CustomizableObject->GetPrivate()->GetCompiledDataFileName(false, nullptr, true);
 
 		const TSharedPtr<IAsyncReadFileHandle> ReadFileHandle = MakeShareable(FPlatformFileManager::Get().GetPlatformFile().OpenAsyncRead(*FullFileName));
 		if (!ReadFileHandle)
@@ -99,7 +99,7 @@ bool FUnrealMutableModelBulkReader::PrepareStreamingForObject(UCustomizableObjec
 
 		NewData.ReadFileHandles.Add(ReadFileHandle);
 #else
-		const UCustomizableObjectBulk* BulkData = CustomizableObject->GetStreamableBulkData();
+		const UCustomizableObjectBulk* BulkData = CustomizableObject->GetPrivate()->GetStreamableBulkData();
 		if (!BulkData)
 		{
 			UE_LOG(LogMutable, Warning, TEXT("Streaming: Customizable Object %s is missing the BulkData export."), *CustomizableObject->GetName());
diff --git a/Engine/Plugins/Experimental/Mutable/Source/CustomizableObjectEditor/Private/MuCOE/CustomizableInstanceDetails.cpp b/Engine/Plugins/Experimental/Mutable/Source/CustomizableObjectEditor/Private/MuCOE/CustomizableInstanceDetails.cpp
--- a/Engine/Plugins/Experimental/Mutable/Source/CustomizableObjectEditor/Private/MuCOE/CustomizableInstanceDetails.cpp
+++ b/Engine/Plugins/Experimental/Mutable/Source/CustomizableObjectEditor/Private/MuCOE/CustomizableInstanceDetails.cpp
@@ -2099,7 +2099,7 @@ void SProfileParametersWindow::OnNameChange(const FText& NewName, ETextCommit::T
 		RequestDestroyWindow();
 
 		UCustomizableObject* CustomizableObject = CustomInstance->GetCustomizableObject();
-		CustomizableObject->AddNewParameterProfile(GetFileName(), *CustomInstance.Get());
+		CustomizableObject->GetPrivate()->AddNewParameterProfile(GetFileName(), *CustomInstance.Get());
 
 		if (CustomInstance->bSelectedProfileDirty && CustomInstance->SelectedProfileIndex != INDEX_NONE)
 		{
@@ -2124,7 +2124,7 @@ FReply SProfileParametersWindow::OnButtonClick(EAppReturnType::Type ButtonID)
 		RequestDestroyWindow();
 
 		UCustomizableObject* CustomizableObject = CustomInstance->GetCustomizableObject();
-		CustomizableObject->AddNewParameterProfile(GetFileName(), *CustomInstance.Get());
+		CustomizableObject->GetPrivate()->AddNewParameterProfile(GetFileName(), *CustomInstance.Get());
 
 		if (CustomInstance->bSelectedProfileDirty && CustomInstance->SelectedProfileIndex != INDEX_NONE)
 		{
diff --git a/Engine/Plugins/Experimental/Mutable/Source/CustomizableObjectEditor/Private/MuCOE/CustomizableObjectCompileRunnable.cpp b/Engine/Plugins/Experimental/Mutable/Source/CustomizableObjectEditor/Private/MuCOE/CustomizableObjectCompileRunnable.cpp
--- a/Engine/Plugins/Experimental/Mutable/Source/CustomizableObjectEditor/Private/MuCOE/CustomizableObjectCompileRunnable.cpp
+++ b/Engine/Plugins/Experimental/Mutable/Source/CustomizableObjectEditor/Private/MuCOE/CustomizableObjectCompileRunnable.cpp
@@ -221,18 +221,18 @@ FCustomizableObjectSaveDDRunnable::FCustomizableObjectSaveDDRunnable(UCustomizab
 	Options = InOptions;
 	
 	CustomizableObjectHeader.InternalVersion = CustomizableObject->GetPrivate()->CurrentSupportedVersion;
-	CustomizableObjectHeader.VersionId = Options.bIsCooking? FGuid::NewGuid() : CustomizableObject->GetVersionId();
+	CustomizableObjectHeader.VersionId = Options.bIsCooking? FGuid::NewGuid() : CustomizableObject->GetPrivate()->GetVersionId();
 
 	if (!Options.bIsCooking)
 	{
 		// We will be saving all compilation data in two separate files, write CO Data
-		FolderPath = CustomizableObject->GetCompiledDataFolderPath();
-		CompileDataFullFileName = FolderPath + CustomizableObject->GetCompiledDataFileName(true, InOptions.TargetPlatform);
-		StreamableDataFullFileName = FolderPath + CustomizableObject->GetCompiledDataFileName(false, InOptions.TargetPlatform);
+		FolderPath = CustomizableObject->GetPrivate()->GetCompiledDataFolderPath();
+		CompileDataFullFileName = FolderPath + CustomizableObject->GetPrivate()->GetCompiledDataFileName(true, InOptions.TargetPlatform);
+		StreamableDataFullFileName = FolderPath + CustomizableObject->GetPrivate()->GetCompiledDataFileName(false, InOptions.TargetPlatform);
 
 		// Serialize Customizable Object's data
 		FMemoryWriter64 MemoryWriter(Bytes);
-		CustomizableObject->SaveCompiledData(MemoryWriter, Options.bIsCooking);
+		CustomizableObject->GetPrivate()->SaveCompiledData(MemoryWriter, Options.bIsCooking);
 	}
 }
 
diff --git a/Engine/Plugins/Experimental/Mutable/Source/CustomizableObjectEditor/Private/MuCOE/CustomizableObjectCompiler.cpp b/Engine/Plugins/Experimental/Mutable/Source/CustomizableObjectEditor/Private/MuCOE/CustomizableObjectCompiler.cpp
--- a/Engine/Plugins/Experimental/Mutable/Source/CustomizableObjectEditor/Private/MuCOE/CustomizableObjectCompiler.cpp
+++ b/Engine/Plugins/Experimental/Mutable/Source/CustomizableObjectEditor/Private/MuCOE/CustomizableObjectCompiler.cpp
@@ -97,7 +97,7 @@ bool FCustomizableObjectCompiler::Tick()
 
 		FinishSavingDerivedData();
 	
-		CurrentObject->PostCompile();
+		CurrentObject->GetPrivate()->PostCompile();
 
 		UE_LOG(LogMutable, Verbose, TEXT("PROFILE: [ %16.8f ] Finished Saving Derived Data task."), FPlatformTime::Seconds());
 		UE_LOG(LogMutable, Verbose, TEXT("PROFILE: -----------------------------------------------------------"));
@@ -1113,7 +1113,7 @@ void FCustomizableObjectCompiler::CompileInternal(UCustomizableObject* Object, c
 			CleanCachedReferencers();
 			UpdateArrayGCProtect();
 
-			CurrentObject->PostCompile(); 
+			CurrentObject->GetPrivate()->PostCompile(); 
 
 			SetCompilationState(ECustomizableObjectCompilationState::Completed);
 		}
@@ -1215,8 +1215,29 @@ void FCustomizableObjectCompiler::FinishCompilation()
 	UpdateCompilerLogData();
 	TSharedPtr<mu::Model, ESPMode::ThreadSafe> Model = CompileTask->Model;
 
-	CurrentObject->SetModel(Model);
+	// Generate a map that using the resource id tells the offset and size of the resource inside the bulk data
+	if (Model)
+	{
+		uint64 Offset = 0;
+
+		const int32 NumStreamingFiles = Model->GetRomCount();
+		CurrentObject->HashToStreamableBlock.Empty(NumStreamingFiles);
 
+		for (size_t FileIndex = 0; FileIndex < NumStreamingFiles; ++FileIndex)
+		{
+			const uint32 ResourceId = Model->GetRomId(FileIndex);
+			const uint32 ResourceSize = Model->GetRomSize(FileIndex);
+
+			CurrentObject->HashToStreamableBlock.Add(ResourceId, FMutableStreamableBlock({0, Offset, ResourceSize }));
+			Offset += ResourceSize;
+		}
+	}
+
+	// Generate ParameterProperties and IntParameterLookUpTable
+	CurrentObject->GetPrivate()->UpdateParameterPropertiesFromModel(Model);
+
+	CurrentObject->GetPrivate()->SetModel(Model, GenerateIdentifier(*CurrentObject));
+	
 	// Reset all instances, as the parameters may need to be rebuilt.
 	for (TObjectIterator<UCustomizableObjectInstance> It; It; ++It)
 	{
@@ -1255,7 +1276,7 @@ void FCustomizableObjectCompiler::FinishSavingDerivedData()
 
 	if (Options.bIsCooking)
 	{
-		CurrentObject->CachePlatformData(SaveDDTask->GetTargetPlatform(), SaveDDTask->GetModelBytes(), SaveDDTask->GetBulkBytes());
+		CurrentObject->GetPrivate()->CachePlatformData(SaveDDTask->GetTargetPlatform(), SaveDDTask->GetModelBytes(), SaveDDTask->GetBulkBytes());
 	}
 
 	// Order matters
diff --git a/Engine/Plugins/Experimental/Mutable/Source/CustomizableObjectEditor/Private/MuCOE/CustomizableObjectDetails.cpp b/Engine/Plugins/Experimental/Mutable/Source/CustomizableObjectEditor/Private/MuCOE/CustomizableObjectDetails.cpp
--- a/Engine/Plugins/Experimental/Mutable/Source/CustomizableObjectEditor/Private/MuCOE/CustomizableObjectDetails.cpp
+++ b/Engine/Plugins/Experimental/Mutable/Source/CustomizableObjectEditor/Private/MuCOE/CustomizableObjectDetails.cpp
@@ -137,7 +137,7 @@ void FCustomizableObjectDetails::CustomizeDetails( IDetailLayoutBuilder& DetailB
 
 	if (Property->IsValidHandle() && CustomizableObject)
 	{
-		if (CustomizableObject->bIsChildObject)
+		if (CustomizableObject->IsChildObject())
 		{
 			Property->MarkHiddenByCustomization();
 		}
diff --git a/Engine/Plugins/Experimental/Mutable/Source/CustomizableObjectEditor/Private/MuCOE/CustomizableObjectEditor.cpp b/Engine/Plugins/Experimental/Mutable/Source/CustomizableObjectEditor/Private/MuCOE/CustomizableObjectEditor.cpp
--- a/Engine/Plugins/Experimental/Mutable/Source/CustomizableObjectEditor/Private/MuCOE/CustomizableObjectEditor.cpp
+++ b/Engine/Plugins/Experimental/Mutable/Source/CustomizableObjectEditor/Private/MuCOE/CustomizableObjectEditor.cpp
@@ -1412,7 +1412,7 @@ void FCustomizableObjectEditor::OnObjectModified(UObject* Object)
 		// Sometimes when another CO is open in another editor window/tab, it triggers this callback, so prevent the modification of this object by a callback triggered by another one
 		if (UCustomizableObject* AuxCustomizableObject = Cast<UCustomizableObject>(Object))
 		{
-			AuxCustomizableObject->UpdateVersionId();
+			AuxCustomizableObject->GetPrivate()->UpdateVersionId();
 		}
 		else if (UCustomizableObjectNode* Node = Cast<UCustomizableObjectNode>(Object))
 		{
@@ -1420,7 +1420,7 @@ void FCustomizableObjectEditor::OnObjectModified(UObject* Object)
 			{
 				if (UCustomizableObject* AuxOuterCustomizableObject = Cast<UCustomizableObject>(Graph->GetOuter()))
 				{
-					AuxOuterCustomizableObject->UpdateVersionId();
+					AuxOuterCustomizableObject->GetPrivate()->UpdateVersionId();
 				}
 			}
 		}
@@ -1428,7 +1428,7 @@ void FCustomizableObjectEditor::OnObjectModified(UObject* Object)
 		{
 			if (UCustomizableObject* AuxOuterCustomizableObject = Cast<UCustomizableObject>(Graph->GetOuter()))
 			{
-				AuxOuterCustomizableObject->UpdateVersionId();
+				AuxOuterCustomizableObject->GetPrivate()->UpdateVersionId();
 			}
 		}
 	}
diff --git a/Engine/Plugins/Experimental/Mutable/Source/CustomizableObjectEditor/Private/MuCOE/Nodes/CustomizableObjectNodeMeshClipMorphDetails.cpp b/Engine/Plugins/Experimental/Mutable/Source/CustomizableObjectEditor/Private/MuCOE/Nodes/CustomizableObjectNodeMeshClipMorphDetails.cpp
--- a/Engine/Plugins/Experimental/Mutable/Source/CustomizableObjectEditor/Private/MuCOE/Nodes/CustomizableObjectNodeMeshClipMorphDetails.cpp
+++ b/Engine/Plugins/Experimental/Mutable/Source/CustomizableObjectEditor/Private/MuCOE/Nodes/CustomizableObjectNodeMeshClipMorphDetails.cpp
@@ -11,6 +11,7 @@
 #include "MuCOE/Nodes/CustomizableObjectNodeMeshClipMorph.h"
 #include "MuCOE/UnrealEditorPortabilityHelpers.h"
 #include "PropertyCustomizationHelpers.h"
+#include "MuCO/CustomizableObjectPrivate.h"
 #include "Widgets/Input/SCheckBox.h"
 #include "Widgets/Input/STextComboBox.h"
 
@@ -57,7 +58,7 @@ void FCustomizableObjectNodeMeshClipMorphDetails::CustomizeDetails(IDetailLayout
 
 		if (CustomizableObject)
 		{
-			if (CustomizableObject->bIsChildObject)
+			if (CustomizableObject->IsChildObject())
 			{
 				bool bMultipleBaseObjects;
 				UCustomizableObjectNodeObject* RootNode = GetRootNode(CustomizableObject, bMultipleBaseObjects);
diff --git a/Engine/Plugins/Experimental/Mutable/Source/CustomizableObjectEditor/Private/MuCOE/Nodes/CustomizableObjectNodeObject.cpp b/Engine/Plugins/Experimental/Mutable/Source/CustomizableObjectEditor/Private/MuCOE/Nodes/CustomizableObjectNodeObject.cpp
--- a/Engine/Plugins/Experimental/Mutable/Source/CustomizableObjectEditor/Private/MuCOE/Nodes/CustomizableObjectNodeObject.cpp
+++ b/Engine/Plugins/Experimental/Mutable/Source/CustomizableObjectEditor/Private/MuCOE/Nodes/CustomizableObjectNodeObject.cpp
@@ -6,6 +6,7 @@
 #include "Logging/MessageLog.h"
 #include "MuCO/ICustomizableObjectModule.h"
 #include "MuCO/CustomizableObjectExtension.h"
+#include "MuCO/CustomizableObjectPrivate.h"
 #include "MuCOE/CustomizableObjectEditorLogger.h"
 #include "MuCOE/CustomizableObjectGraph.h"
 #include "MuCOE/EdGraphSchema_CustomizableObject.h"
@@ -88,7 +89,7 @@ void UCustomizableObjectNodeObject::PostEditChangeProperty(FPropertyChangedEvent
 	UCustomizableObject* CustomizableObject = Cast<UCustomizableObject>( GetCustomizableObjectGraph()->GetOuter() );
 	if (CustomizableObject)
 	{
-		CustomizableObject->bIsChildObject = ParentObject != nullptr;
+		CustomizableObject->GetPrivate()->SetIsChildObject(ParentObject != nullptr);
 	}
 
 	FProperty* PropertyThatChanged = PropertyChangedEvent.Property;
@@ -321,7 +322,7 @@ void UCustomizableObjectNodeObject::SetParentObject(UCustomizableObject* Customi
 		UCustomizableObject* CustomizableObject = Cast<UCustomizableObject>(GetCustomizableObjectGraph()->GetOuter());
 		if (CustomizableObject)
 		{
-			CustomizableObject->bIsChildObject = ParentObject != nullptr;
+			CustomizableObject->GetPrivate()->SetIsChildObject(ParentObject != nullptr);
 
 			TSharedPtr<ICustomizableObjectEditor> Editor = GetGraphEditor();
 
diff --git a/Engine/Plugins/Experimental/Mutable/Source/CustomizableObjectEditor/Private/MuCOE/SCustomizableInstanceProperties.cpp b/Engine/Plugins/Experimental/Mutable/Source/CustomizableObjectEditor/Private/MuCOE/SCustomizableInstanceProperties.cpp
--- a/Engine/Plugins/Experimental/Mutable/Source/CustomizableObjectEditor/Private/MuCOE/SCustomizableInstanceProperties.cpp
+++ b/Engine/Plugins/Experimental/Mutable/Source/CustomizableObjectEditor/Private/MuCOE/SCustomizableInstanceProperties.cpp
@@ -2325,7 +2325,7 @@ FReply SCreateProfileParameters::OnButtonClick(EAppReturnType::Type ButtonID)
 		RequestDestroyWindow();
 
 		UCustomizableObject* CustomizableObject = CustomInstance->GetCustomizableObject(); 
-		CustomizableObject->AddNewParameterProfile(GetFileName(), *CustomInstance.Get());
+		CustomizableObject->GetPrivate()->AddNewParameterProfile(GetFileName(), *CustomInstance.Get());
 
 		if (CustomInstance->bSelectedProfileDirty && CustomInstance->SelectedProfileIndex != INDEX_NONE)
 		{
diff --git a/Engine/Plugins/Experimental/Mutable/Source/MutableValidation/Private/MuV/CustomizableObjectValidationCommandlet.cpp b/Engine/Plugins/Experimental/Mutable/Source/MutableValidation/Private/MuV/CustomizableObjectValidationCommandlet.cpp
--- a/Engine/Plugins/Experimental/Mutable/Source/MutableValidation/Private/MuV/CustomizableObjectValidationCommandlet.cpp
+++ b/Engine/Plugins/Experimental/Mutable/Source/MutableValidation/Private/MuV/CustomizableObjectValidationCommandlet.cpp
@@ -94,7 +94,7 @@ int32 UCustomizableObjectValidationCommandlet::Main(const FString& Params)
 			TArray<uint8> EmbeddedDataBytes{};
 			FMemoryWriter SerializationTarget{EmbeddedDataBytes, false};
 			
-			ToTestCustomizableObject->SaveEmbeddedData(SerializationTarget);
+			ToTestCustomizableObject->GetPrivate()->SaveEmbeddedData(SerializationTarget);
 			const int64 COEmbeddedDataSizeBytes = EmbeddedDataBytes.Num();
 			
 			UE_LOG(LogMutable, Log,TEXT("(int) co_embedded_data_bytes : %lld "), COEmbeddedDataSizeBytes);
