[Mutable] Moved all non-public UCustomizableObject functions to private. Part 2.

- Removed UCustomizableObject::GetModel (body moved to the only caller).
- Fixed minor typos.
- Standarized doc style.

#rnx

diff --git a/Engine/Plugins/Experimental/Mutable/Source/CustomizableObject/Classes/MuCO/CustomizableObject.h b/Engine/Plugins/Experimental/Mutable/Source/CustomizableObject/Classes/MuCO/CustomizableObject.h
--- a/Engine/Plugins/Experimental/Mutable/Source/CustomizableObject/Classes/MuCO/CustomizableObject.h
+++ b/Engine/Plugins/Experimental/Mutable/Source/CustomizableObject/Classes/MuCO/CustomizableObject.h
@@ -822,6 +822,7 @@ public:
 	// Customizable Object Population data end --------------------------------------------------------
 
 #if WITH_EDITORONLY_DATA
+private:
 	/** True if this object references a parent object. This is used basically to exclude this object
 	  * from cooking. This is actually derived from the source graph object node pointing to another
 	  * object or not, but it needs to be cached here because the source graph is not always available.
