[mutable] Restored "Compile Only Selected" option and added support for tables.
#rnx

#tests Editor

#changelist validated
#virtualized

diff --git a/Engine/Plugins/Experimental/Mutable/Source/CustomizableObjectEditor/Private/MuCOE/CustomizableObjectCompiler.cpp b/Engine/Plugins/Experimental/Mutable/Source/CustomizableObjectEditor/Private/MuCOE/CustomizableObjectCompiler.cpp
--- a/Engine/Plugins/Experimental/Mutable/Source/CustomizableObjectEditor/Private/MuCOE/CustomizableObjectCompiler.cpp
+++ b/Engine/Plugins/Experimental/Mutable/Source/CustomizableObjectEditor/Private/MuCOE/CustomizableObjectCompiler.cpp
@@ -879,11 +879,6 @@ void FCustomizableObjectCompiler::CompileInternal(UCustomizableObject* Object, c
 
 	FMutableGraphGenerationContext GenerationContext(Object, this, Options);
 	GenerationContext.ParamNamesToSelectedOptions = ParamNamesToSelectedOptions;
-	if (ParamNamesToSelectedOptions.Num() > 0)
-	{
-		// A partial compilation shouldn't update the derived data cache, because it would force a recompilation at the next restart of the editor
-		Options.bDontUpdateStreamedDataAndCache = true;
-	}
 
 	// If we don't have the target platform yet (in editor) we need to get it
 	if (!GenerationContext.Options.TargetPlatform)
diff --git a/Engine/Plugins/Experimental/Mutable/Source/CustomizableObjectEditor/Private/MuCOE/CustomizableObjectEditor.cpp b/Engine/Plugins/Experimental/Mutable/Source/CustomizableObjectEditor/Private/MuCOE/CustomizableObjectEditor.cpp
--- a/Engine/Plugins/Experimental/Mutable/Source/CustomizableObjectEditor/Private/MuCOE/CustomizableObjectEditor.cpp
+++ b/Engine/Plugins/Experimental/Mutable/Source/CustomizableObjectEditor/Private/MuCOE/CustomizableObjectEditor.cpp
@@ -1214,7 +1214,7 @@ void FCustomizableObjectEditor::ExtendToolbar()
 		{
 			ToolbarBuilder.BeginSection("Compilation");
 			ToolbarBuilder.AddToolBarButton(FCustomizableObjectEditorCommands::Get().Compile);
-			//ToolbarBuilder.AddToolBarButton(FCustomizableObjectEditorCommands::Get().CompileOnlySelected); // Disabled due to a crash with this mode (MTBL-1716). 
+			ToolbarBuilder.AddToolBarButton(FCustomizableObjectEditorCommands::Get().CompileOnlySelected);
 			ToolbarBuilder.AddComboButton(
 				FUIAction(),
 				FOnGetContent::CreateSP(Editor, &FCustomizableObjectEditor::GenerateCompileOptionsMenuContent, CommandList.ToSharedRef()),
diff --git a/Engine/Plugins/Experimental/Mutable/Source/CustomizableObjectEditor/Private/MuCOE/GenerateMutableSource/GenerateMutableSourceTable.cpp b/Engine/Plugins/Experimental/Mutable/Source/CustomizableObjectEditor/Private/MuCOE/GenerateMutableSource/GenerateMutableSourceTable.cpp
--- a/Engine/Plugins/Experimental/Mutable/Source/CustomizableObjectEditor/Private/MuCOE/GenerateMutableSource/GenerateMutableSourceTable.cpp
+++ b/Engine/Plugins/Experimental/Mutable/Source/CustomizableObjectEditor/Private/MuCOE/GenerateMutableSource/GenerateMutableSourceTable.cpp
@@ -627,6 +627,47 @@ bool FillTableColumn(const UCustomizableObjectNodeTable* TableNode,	mu::TablePtr
 }
 
 
+FName GetAnotherOption(FName SelectedOptionName, const TArray<FName>& RowNames)
+{
+	for (const FName& CandidateOption : RowNames)
+	{
+		if (CandidateOption != SelectedOptionName)
+		{
+			return CandidateOption;
+		}
+	}
+
+	return FName("None");
+}
+
+
+void RestrictRowNamesToSelectedOption(TArray<FName>& InOutRowNames, const UCustomizableObjectNodeTable& TableNode, FMutableGraphGenerationContext& GenerationContext)
+{
+	// If the param is in the map restrict to only the selected option
+	FString* SelectedOptionString = GenerationContext.ParamNamesToSelectedOptions.Find(TableNode.ParameterName);
+
+	if (SelectedOptionString)
+	{
+		FName SelectedOptionName = FName(*SelectedOptionString);
+
+		if (InOutRowNames.Contains(SelectedOptionName))
+		{
+			FName AnotherOption = GetAnotherOption(SelectedOptionName, InOutRowNames);
+			InOutRowNames.Empty(2);
+			InOutRowNames.Add(SelectedOptionName);
+
+			// To prevent the optimization of the parameter for having just one option, which would prevent the restriction 
+			// of that parameter in the next compile only selected
+			InOutRowNames.Add(AnotherOption);
+		}
+		else
+		{
+			InOutRowNames.Empty(0);
+		}
+	}
+}
+
+
 bool GenerateTableColumn(const UCustomizableObjectNodeTable* TableNode, const UEdGraphPin* Pin, mu::TablePtr MutableTable, const FString& DataTableColumnName, const FProperty* ColumnProperty,
 	const int32 LODIndexConnected, const int32 SectionIndexConnected, const int32 LODIndex, const int32 SectionIndex, const bool bOnlyConnectedLOD, FMutableGraphGenerationContext& GenerationContext)
 {
@@ -648,6 +689,7 @@ bool GenerateTableColumn(const UCustomizableObjectNodeTable* TableNode, const UE
 
 	// Getting names of the rows to access the information
 	TArray<FName> RowNames = TableNode->GetRowNames(DataTable);
+	RestrictRowNamesToSelectedOption(RowNames, *TableNode, GenerationContext);
 
 	for (int32 RowIndex = 0; RowIndex < RowNames.Num(); ++RowIndex)
 	{
@@ -685,7 +727,9 @@ void GenerateTableParameterUIData(const UDataTable* DataTable, const UCustomizab
 	if (!GenerationContext.ParameterUIDataMap.Contains(TableNode->ParameterName))
 	{
 		// Getting Table and row names to access the information
-		const TArray<FName>& RowNames = TableNode->GetRowNames(DataTable);
+		TArray<FName> RowNames = TableNode->GetRowNames(DataTable);
+		RestrictRowNamesToSelectedOption(RowNames, *TableNode, GenerationContext);
+
 		FParameterUIData ParameterUIData(TableNode->ParameterName, TableNode->ParamUIMetadata, EMutableParameterType::Int);
 
 		for (int32 NameIndex = 0; NameIndex < RowNames.Num(); ++NameIndex)
@@ -719,7 +763,8 @@ mu::TablePtr GenerateMutableSourceTable(const UDataTable* DataTable, const UCust
 	if (const UScriptStruct* TableStruct = DataTable->GetRowStruct())
 	{
 		// Getting Table and row names to access the information
-		const TArray<FName>& RowNames = TableNode->GetRowNames(DataTable);
+		TArray<FName> RowNames = TableNode->GetRowNames(DataTable);
+		RestrictRowNamesToSelectedOption(RowNames, *TableNode, GenerationContext);
 
 		// Adding and filling Name Column
 		MutableTable->AddColumn("Name", mu::ETableColumnType::String);
diff --git a/Engine/Plugins/Experimental/Mutable/Source/MutableTools/Private/MuT/CodeGenerator_Image.cpp b/Engine/Plugins/Experimental/Mutable/Source/MutableTools/Private/MuT/CodeGenerator_Image.cpp
--- a/Engine/Plugins/Experimental/Mutable/Source/MutableTools/Private/MuT/CodeGenerator_Image.cpp
+++ b/Engine/Plugins/Experimental/Mutable/Source/MutableTools/Private/MuT/CodeGenerator_Image.cpp
@@ -1859,7 +1859,7 @@ namespace mu
         Result.op = currentOp;
     }
 
-
+	
 	//---------------------------------------------------------------------------------------------
 	void CodeGenerator::GenerateImage_Table(const FImageGenerationOptions& Options, FImageGenerationResult& Result, const NodeImageTable* InNode)
 	{
@@ -1882,6 +1882,8 @@ namespace mu
 				{
 					FString Msg = FString::Printf(TEXT("Table has a missing image in column %d, row %d."), colIndex, row);
 					pErrorLog->GetPrivate()->Add(Msg, ELMT_ERROR, node.m_errorContext);
+
+					return ImageOp;
 				}
 				else
 				{
