[Mutable] Fixed column names for null meshes in table nodes.
+Removed unnecessary function.

#rnx

diff --git a/Engine/Plugins/Experimental/Mutable/Source/CustomizableObjectEditor/Classes/MuCOE/Nodes/CustomizableObjectNodeTable.h b/Engine/Plugins/Experimental/Mutable/Source/CustomizableObjectEditor/Classes/MuCOE/Nodes/CustomizableObjectNodeTable.h
--- a/Engine/Plugins/Experimental/Mutable/Source/CustomizableObjectEditor/Classes/MuCOE/Nodes/CustomizableObjectNodeTable.h
+++ b/Engine/Plugins/Experimental/Mutable/Source/CustomizableObjectEditor/Classes/MuCOE/Nodes/CustomizableObjectNodeTable.h
@@ -264,9 +264,6 @@ public:
 	// Returns the name of the table column related to a pin
 	FString GetColumnNameByPin(const UEdGraphPin* Pin) const;
 
-	// Returns the name of a Mesh Mutable Column from its LOD and pin (Only for AutomaticFromMesh LOD Strategy)
-	FString GetMutableColumnName(const UEdGraphPin* Pin, const int32& LOD) const;
-
 	// Returns the LOD of the mesh associated to the input pin
 	void GetPinLODAndSection(const UEdGraphPin* Pin, int32& LODIndex, int32& SectionIndex) const;
 
diff --git a/Engine/Plugins/Experimental/Mutable/Source/CustomizableObjectEditor/Private/MuCOE/GenerateMutableSource/GenerateMutableSourceMesh.cpp b/Engine/Plugins/Experimental/Mutable/Source/CustomizableObjectEditor/Private/MuCOE/GenerateMutableSource/GenerateMutableSourceMesh.cpp
--- a/Engine/Plugins/Experimental/Mutable/Source/CustomizableObjectEditor/Private/MuCOE/GenerateMutableSource/GenerateMutableSourceMesh.cpp
+++ b/Engine/Plugins/Experimental/Mutable/Source/CustomizableObjectEditor/Private/MuCOE/GenerateMutableSource/GenerateMutableSourceMesh.cpp
@@ -3669,7 +3669,7 @@ mu::NodeMeshPtr GenerateMutableSourceMesh(const UEdGraphPin* Pin,
 					}
 					
 					// Getting the mutable table mesh column name
-					FString MutableColumnName = TypedNodeTable->GetMutableColumnName(Pin, LODIndex);
+					FString MutableColumnName = TypedNodeTable->GenerateSkeletalMeshMutableColumName(DataTableColumnName, LODIndex, SectionIndex);
 
 					// Generating a new Mesh column if not exists
 					if (Table->FindColumn(MutableColumnName) == INDEX_NONE)
diff --git a/Engine/Plugins/Experimental/Mutable/Source/CustomizableObjectEditor/Private/MuCOE/Nodes/CustomizableObjectNodeTable.cpp b/Engine/Plugins/Experimental/Mutable/Source/CustomizableObjectEditor/Private/MuCOE/Nodes/CustomizableObjectNodeTable.cpp
--- a/Engine/Plugins/Experimental/Mutable/Source/CustomizableObjectEditor/Private/MuCOE/Nodes/CustomizableObjectNodeTable.cpp
+++ b/Engine/Plugins/Experimental/Mutable/Source/CustomizableObjectEditor/Private/MuCOE/Nodes/CustomizableObjectNodeTable.cpp
@@ -985,22 +985,6 @@ FString UCustomizableObjectNodeTable::GetColumnNameByPin(const UEdGraphPin* Pin)
 }
 
 
-FString UCustomizableObjectNodeTable::GetMutableColumnName(const UEdGraphPin* Pin, const int32& LOD) const
-{
-	UCustomizableObjectNodeTableMeshPinData* MeshData = Cast<UCustomizableObjectNodeTableMeshPinData >(GetPinData(*Pin));
-
-	if (MeshData)
-	{
-		FString ColumnName = MeshData->ColumnName;
-		int32 MaterialIndex = MeshData->Material;
-
-		return GenerateSkeletalMeshMutableColumName(MeshData->ColumnName, LOD, MeshData->Material);
-	}
-
-	return FString();
-}
-
-
 void UCustomizableObjectNodeTable::GetPinLODAndSection(const UEdGraphPin* Pin, int32& LODIndex, int32& SectionIndex) const
 {
 	if (const UCustomizableObjectNodeTableMeshPinData* PinData = Cast<UCustomizableObjectNodeTableMeshPinData >(GetPinData(*Pin)))
