[Mutable] Fix crashes and visual glitches related to the reuse of mesh components and MeshCache hits.
Cache hits did not update the MeshIds of the component, potentially reusing meshes in cases requiring an update.

#rnx

#tests Frontend-Locker

diff --git a/Engine/Plugins/Experimental/Mutable/Source/CustomizableObject/Private/MuCO/CustomizableObjectInstance.cpp b/Engine/Plugins/Experimental/Mutable/Source/CustomizableObject/Private/MuCO/CustomizableObjectInstance.cpp
--- a/Engine/Plugins/Experimental/Mutable/Source/CustomizableObject/Private/MuCO/CustomizableObjectInstance.cpp
+++ b/Engine/Plugins/Experimental/Mutable/Source/CustomizableObject/Private/MuCO/CustomizableObjectInstance.cpp
@@ -1874,22 +1874,24 @@ bool UCustomizableInstancePrivate::UpdateSkeletalMesh_PostBeginUpdate0(UCustomiz
 			continue;
 		}
 
+		if (!ComponentsData.IsValidIndex(Component.Id))
+		{
+			bSuccess = false;
+			ensure(false);
+			continue;
+		}
+
 		if (OperationData->bUseMeshCache)
 		{
 			if (USkeletalMesh* CachedMesh = CustomizableObject->GetPrivate()->MeshCache.Get(OperationData->MeshDescriptors[Component.Id]))
 			{
+				check(OperationData->MeshDescriptors[Component.Id].Num() == MAX_MESH_LOD_COUNT);
+				ComponentsData[Component.Id].LastMeshIdPerLOD = OperationData->MeshDescriptors[Component.Id];
 				SkeletalMeshes[Component.Id] = CachedMesh;
 				continue;
 			}
 		}
 
-		if (!ComponentsData.IsValidIndex(Component.Id))
-		{
-			bSuccess = false;
-			ensure(false);
-			continue;
-		}
-
 		// Reset last mesh IDs.
 		ComponentsData[Component.Id].LastMeshIdPerLOD.Init(MAX_uint64, MAX_MESH_LOD_COUNT);
 
