[mutable] Fixed the mutable.EnableMutableProgressiveMipStreaming cvar not disabling streaming when set to 0.
#rnx
#tests PIE#preflight 65c4d464e80bb2fbab52b748

diff --git a/Engine/Plugins/Experimental/Mutable/Source/CustomizableObject/Private/MuCO/CustomizableObjectSystem.cpp b/Engine/Plugins/Experimental/Mutable/Source/CustomizableObject/Private/MuCO/CustomizableObjectSystem.cpp
--- a/Engine/Plugins/Experimental/Mutable/Source/CustomizableObject/Private/MuCO/CustomizableObjectSystem.cpp
+++ b/Engine/Plugins/Experimental/Mutable/Source/CustomizableObject/Private/MuCO/CustomizableObjectSystem.cpp
@@ -974,21 +974,19 @@ void UCustomizableObjectSystemPrivate::GetMipStreamingConfig(const UCustomizable
 		}
 	}
 
-	bool bUseMipmapStreaming = !bOutNeverStream;
 	OutMipsToSkip = 0; // 0 means generate all mips
 
-	// Streaming disabled from platform settings?
+	// Streaming disabled from platform settings or from platform CustomizableObjectSystem properties?
 #if PLATFORM_SUPPORTS_TEXTURE_STREAMING
-	if (!IStreamingManager::Get().IsTextureStreamingEnabled())
+	if (!IStreamingManager::Get().IsTextureStreamingEnabled() || !EnableMutableProgressiveMipStreaming)
 	{
-		bUseMipmapStreaming = false;
+		bOutNeverStream = true;
 	}
 #else
-	bUseMipmapStreaming = false;
+	bOutNeverStream = true;
 #endif
-
-	// Streaming disabled from platform CustomizableObjectSystem properties?
-	if (bUseMipmapStreaming && EnableMutableProgressiveMipStreaming)
+	
+	if (!bOutNeverStream)
 	{
 		OutMipsToSkip = 255; // This means skip all possible mips until only UTexture::GetStaticMinTextureResidentMipCount() are left
 	}
