[Mutable] Removed PrivateIncludePaths in COE module.

Moved required headers from private to internal.

#rnx

diff --git a/Engine/Plugins/Experimental/Mutable/Source/CustomizableObjectEditor/CustomizableObjectEditor.Build.cs b/Engine/Plugins/Experimental/Mutable/Source/CustomizableObjectEditor/CustomizableObjectEditor.Build.cs
--- a/Engine/Plugins/Experimental/Mutable/Source/CustomizableObjectEditor/CustomizableObjectEditor.Build.cs
+++ b/Engine/Plugins/Experimental/Mutable/Source/CustomizableObjectEditor/CustomizableObjectEditor.Build.cs
@@ -68,7 +68,6 @@ public class CustomizableObjectEditor : ModuleRules
 				"TextureCompressor",
 				"TextureBuildUtilities",
 				"ImageCore",
-
 			}
 		);
 
@@ -79,8 +78,8 @@ public class CustomizableObjectEditor : ModuleRules
 				"SceneOutliner",
 				"ClassViewer",
 				"WorkspaceMenuStructure",
-				}
-			);
+			}
+		);
 
 
 		PublicDependencyModuleNames.AddRange(
@@ -90,10 +89,6 @@ public class CustomizableObjectEditor : ModuleRules
 				"StructUtils", 
 				"MutableTools",
 			}
-			);
-
-		PrivateIncludePaths.AddRange(new string[] {
-				"MutableTools/Private",
-			});
+		);
 	}
 }
diff --git a/Engine/Plugins/Experimental/Mutable/Source/MutableRuntime/Internal/MuR/SerialisationPrivate.h b/Engine/Plugins/Experimental/Mutable/Source/MutableRuntime/Internal/MuR/SerialisationPrivate.h
--- a/Engine/Plugins/Experimental/Mutable/Source/MutableRuntime/Internal/MuR/SerialisationPrivate.h
+++ b/Engine/Plugins/Experimental/Mutable/Source/MutableRuntime/Internal/MuR/SerialisationPrivate.h
@@ -233,25 +233,6 @@ namespace mu
 
 		TVariantLoadFromInputArchiveLookup<Ts...>::Load(static_cast<SIZE_T>(Index), Ar, Variant);
 	}
-	
-
-	//---------------------------------------------------------------------------------------------
-	// Bool size is not a standard
-	//---------------------------------------------------------------------------------------------
-	template<>														
-	inline void operator<< <bool>( OutputArchive& arch, const bool& t )
-	{
-        uint8 s = t ? 1 : 0;
-        arch.GetPrivate()->m_pStream->Write( &s, sizeof(uint8) );
-	}
-
-	template<>
-	inline void operator>> <bool>( InputArchive& arch, bool& t )
-	{
-        uint8 s;
-        arch.GetPrivate()->m_pStream->Read( &s, sizeof(uint8) );
-		t = s!=0;
-	}
 
 
 	//---------------------------------------------------------------------------------------------
@@ -372,30 +353,7 @@ namespace mu
 	MUTABLE_IMPLEMENT_POD_VECTOR_SERIALISABLE(UE::Math::TIntVector2<int16>);
 	MUTABLE_IMPLEMENT_POD_VECTOR_SERIALISABLE(FVector4f);
 
-	//---------------------------------------------------------------------------------------------
-	//template<>
-	//inline void operator<< <std::string>( OutputArchive& arch, const std::string& v )
-	//{
-	//	arch << (uint32)v.size();
-	//	if ( v.size() )
-	//	{
-	//		arch.GetPrivate()->m_pStream->Write( &v[0], (unsigned)v.size()*sizeof(char) );
-	//	}
-	//}
-
-	template<>
-	inline void operator>> <std::string>( InputArchive& arch, std::string& v )
-	{
-        uint32 size;
-		arch >> size;
-		v.resize( size );
-		if (size)
-		{
-			arch.GetPrivate()->m_pStream->Read( &v[0], (unsigned)size*sizeof(char) );
-		}
-	}
-
-
+	
 	//---------------------------------------------------------------------------------------------
 	template< typename T >
 	inline void operator<< ( OutputArchive& arch, const Ptr<const T>& p )
diff --git a/Engine/Plugins/Experimental/Mutable/Source/MutableRuntime/Private/MuR/Serialisation.cpp b/Engine/Plugins/Experimental/Mutable/Source/MutableRuntime/Private/MuR/Serialisation.cpp
--- a/Engine/Plugins/Experimental/Mutable/Source/MutableRuntime/Private/MuR/Serialisation.cpp
+++ b/Engine/Plugins/Experimental/Mutable/Source/MutableRuntime/Private/MuR/Serialisation.cpp
@@ -69,8 +69,35 @@ namespace mu
 	    arch >> Temp;
 	    v = FName(Temp);
     }
+	
+
+	void operator>> ( InputArchive& arch, std::string& v )
+    {
+    	uint32 size;
+    	arch >> size;
+    	v.resize( size );
+    	if (size)
+    	{
+    		arch.GetPrivate()->m_pStream->Read( &v[0], (unsigned)size*sizeof(char) );
+    	}
+    }
+
+	
+	void operator<< ( OutputArchive& arch, const bool& t )
+    {
+    	uint8 s = t ? 1 : 0;
+    	arch.GetPrivate()->m_pStream->Write( &s, sizeof(uint8) );
+    }
+
 
+	void operator>> ( InputArchive& arch, bool& t )
+    {
+    	uint8 s;
+    	arch.GetPrivate()->m_pStream->Read( &s, sizeof(uint8) );
+    	t = s!=0;
+    }
 
+	
     //---------------------------------------------------------------------------------------------
     //---------------------------------------------------------------------------------------------
     //---------------------------------------------------------------------------------------------
diff --git a/Engine/Plugins/Experimental/Mutable/Source/MutableRuntime/Public/MuR/Serialisation.h b/Engine/Plugins/Experimental/Mutable/Source/MutableRuntime/Public/MuR/Serialisation.h
--- a/Engine/Plugins/Experimental/Mutable/Source/MutableRuntime/Public/MuR/Serialisation.h
+++ b/Engine/Plugins/Experimental/Mutable/Source/MutableRuntime/Public/MuR/Serialisation.h
@@ -2,6 +2,8 @@
 
 #pragma once
 
+#include <string>
+
 #include "HAL/Platform.h"
 #include "MuR/MutableMemory.h"
 #include "MuR/Ptr.h"
@@ -12,6 +14,7 @@
 #include "Math/Vector4.h"
 #include "MuR/MutableMath.h"
 
+
 namespace mu
 {    
     class Image;
@@ -411,6 +414,13 @@ namespace mu
 			Arch >> V[Index];
 		}
 	}
+
 	
+	MUTABLERUNTIME_API void operator>> ( InputArchive& arch, std::string& v );
+
+	// Bool size is not a standard
+	MUTABLERUNTIME_API void operator<< ( OutputArchive& arch, const bool& t );
+
+	MUTABLERUNTIME_API void operator>> ( InputArchive& arch, bool& t );
 }
 
diff --git a/Engine/Plugins/Experimental/Mutable/Source/MutableTools/Internal/MuT/NodeObjectGroupPrivate.h b/Engine/Plugins/Experimental/Mutable/Source/MutableTools/Internal/MuT/NodeObjectGroupPrivate.h
new file mode 100644
--- /dev/null
+++ b/Engine/Plugins/Experimental/Mutable/Source/MutableTools/Internal/MuT/NodeObjectGroupPrivate.h
@@ -0,0 +1,82 @@
+// Copyright Epic Games, Inc. All Rights Reserved.
+
+#pragma once
+
+#include "MuT/NodeObjectPrivate.h"
+
+#include "MuT/NodeObjectGroup.h"
+#include "MuT/NodeLayout.h"
+
+
+namespace mu
+{
+	MUTABLE_DEFINE_ENUM_SERIALISABLE(NodeObjectGroup::CHILD_SELECTION)
+
+
+	class NodeObjectGroup::Private : public NodeObject::Private
+	{
+	public:
+
+		static FNodeType s_type;
+
+		FString Name;
+		FString Uid;
+
+		CHILD_SELECTION m_type;
+
+		//! Set the child selection type
+		void SetSelectionType( CHILD_SELECTION );
+
+		TArray<NodeObjectPtr> m_children;
+		int32 DefaultValue;
+
+		//!
+		void Serialise( OutputArchive& arch ) const
+		{
+            uint32_t ver = 3;
+			arch << ver;
+
+			arch << m_type;
+			arch << Name;
+			arch << Uid;
+			arch << m_children;
+			arch << DefaultValue;
+		}
+
+		//!
+		void Unserialise( InputArchive& arch )
+		{
+            uint32_t ver;
+			arch >> ver;
+            check(ver>=1 && ver<=3);
+
+			arch >> m_type;
+			if (ver <= 1)
+			{
+				std::string Temp;
+				arch >> Temp;
+				Name = Temp.c_str();
+				arch >> Temp;
+				Uid = Temp.c_str();
+			}
+			else
+			{
+				arch >> Name;
+				arch >> Uid;
+			}
+
+			arch >> m_children;
+
+			if (ver >= 3)
+			{
+				arch >> DefaultValue;
+			}
+		}
+
+
+        // NodeObject::Private interface
+        NodeLayoutPtr GetLayout( int lod, int component, int surface, int texture ) const override;
+
+	};
+
+}
diff --git a/Engine/Plugins/Experimental/Mutable/Source/MutableTools/Internal/MuT/NodeObjectNewPrivate.h b/Engine/Plugins/Experimental/Mutable/Source/MutableTools/Internal/MuT/NodeObjectNewPrivate.h
new file mode 100644
--- /dev/null
+++ b/Engine/Plugins/Experimental/Mutable/Source/MutableTools/Internal/MuT/NodeObjectNewPrivate.h
@@ -0,0 +1,122 @@
+// Copyright Epic Games, Inc. All Rights Reserved.
+
+#pragma once
+
+#include "MuT/NodeObjectPrivate.h"
+
+#include "MuT/NodeExtensionData.h"
+#include "MuT/NodeObjectNew.h"
+#include "MuT/NodeLayout.h"
+#include "MuT/NodeLOD.h"
+#include "MuT/Compiler.h"
+
+
+namespace mu
+{
+
+
+	class NodeObjectNew::Private : public NodeObject::Private
+	{
+	public:
+
+		static FNodeType s_type;
+
+		FString m_name;
+		FString m_uid;
+
+		TArray<NodeLODPtr> m_lods;
+
+		TArray<NodeObjectPtr> m_children;
+
+		struct NamedExtensionDataNode
+		{
+			Ptr<NodeExtensionData> Node;
+			FString Name;
+
+			void Serialise(OutputArchive& arch) const
+			{
+				arch << Node;
+				arch << Name;
+			}
+
+			void Unserialise(InputArchive& arch)
+			{
+				arch >> Node;
+				arch >> Name;
+			}
+		};
+
+		TArray<NamedExtensionDataNode> m_extensionDataNodes;
+
+
+		//! List of states
+        TArray<FObjectState> m_states;
+
+
+		//!
+		void Serialise( OutputArchive& arch ) const
+		{
+            uint32_t ver = 4;
+			arch << ver;
+
+			arch << m_name;
+			arch << m_uid;
+			arch << m_lods;
+			arch << m_children;
+			arch << m_states;
+			arch << m_extensionDataNodes;
+		}
+
+		//!
+		void Unserialise( InputArchive& arch )
+		{
+            uint32_t ver;
+			arch >> ver;
+            check(ver>=2 && ver<=4);
+
+			if (ver <= 3)
+			{
+				std::string Temp;
+				arch >> Temp;
+				m_name = Temp.c_str();
+				arch >> Temp;
+				m_uid = Temp.c_str();
+			}
+			else
+			{
+				arch >> m_name;
+				arch >> m_uid;
+			}
+
+			arch >> m_lods;
+			arch >> m_children;
+			arch >> m_states;
+
+			if (ver >= 4)
+			{
+				arch >> m_extensionDataNodes;
+			}
+			else if (ver >= 3)
+			{
+				int32 Num = 0;
+				arch >> Num;
+				m_extensionDataNodes.SetNum(Num);
+				for (int i=0;i<Num; ++i)
+				{
+					std::string Temp;
+					arch >> Temp;
+					m_extensionDataNodes[i].Name = Temp.c_str();
+					arch >> m_extensionDataNodes[i].Node;
+				}
+			}
+		}
+
+        //! Return true if the given component is set in any lod of this object.
+        bool HasComponent( const NodeComponent* pComponent ) const;
+
+        // NodeObject::Private interface
+        NodeLayoutPtr GetLayout( int lod, int component, int surface, int texture ) const override;
+
+	};
+
+}
diff --git a/Engine/Plugins/Experimental/Mutable/Source/MutableTools/Internal/MuT/NodeObjectPrivate.h b/Engine/Plugins/Experimental/Mutable/Source/MutableTools/Internal/MuT/NodeObjectPrivate.h
new file mode 100644
--- /dev/null
+++ b/Engine/Plugins/Experimental/Mutable/Source/MutableTools/Internal/MuT/NodeObjectPrivate.h
@@ -0,0 +1,22 @@
+// Copyright Epic Games, Inc. All Rights Reserved.
+
+#pragma once
+
+#include "MuT/NodePrivate.h"
+
+#include "MuT/NodeObject.h"
+#include "MuT/NodeLayout.h"
+
+namespace mu
+{
+
+	class NodeObject::Private : public Node::Private
+	{
+	public:
+
+        virtual NodeLayoutPtr GetLayout( int lod, int component, int surface, int texture ) const = 0;
+
+	};
+
+}
+
diff --git a/Engine/Plugins/Experimental/Mutable/Source/MutableTools/Internal/MuT/NodePatchImage.h b/Engine/Plugins/Experimental/Mutable/Source/MutableTools/Internal/MuT/NodePatchImage.h
new file mode 100644
--- /dev/null
+++ b/Engine/Plugins/Experimental/Mutable/Source/MutableTools/Internal/MuT/NodePatchImage.h
@@ -0,0 +1,111 @@
+// Copyright Epic Games, Inc. All Rights Reserved.
+
+#pragma once
+
+#include "MuR/Image.h"
+#include "MuR/Ptr.h"
+#include "MuR/RefCounted.h"
+#include "MuT/Node.h"
+#include "MuT/NodeImage.h"
+
+
+namespace mu
+{
+
+    // Forward definitions
+    class NodeImage;
+    typedef Ptr<NodeImage> NodeImagePtr;
+    typedef Ptr<const NodeImage> NodeImagePtrConst;
+
+    class NodePatchImage;
+    typedef Ptr<NodePatchImage> NodePatchImagePtr;
+    typedef Ptr<const NodePatchImage> NodePatchImagePtrConst;
+
+
+    //! Node that allows to modify an image from an object by blending other images on specific
+    //! layout blocks.
+    //! \ingroup model
+    class MUTABLETOOLS_API NodePatchImage : public Node
+    {
+    public:
+
+        //-----------------------------------------------------------------------------------------
+        // Life cycle
+        //-----------------------------------------------------------------------------------------
+
+        NodePatchImage();
+
+		void SerialiseWrapper(OutputArchive& arch) const override;
+		static void Serialise( const NodePatchImage* pNode, OutputArchive& arch );
+        static NodePatchImagePtr StaticUnserialise( InputArchive& arch );
+
+
+        //-----------------------------------------------------------------------------------------
+        // Node Interface
+        //-----------------------------------------------------------------------------------------
+
+        const FNodeType* GetType() const override;
+        static const FNodeType* GetStaticType();
+
+        //-----------------------------------------------------------------------------------------
+        // Own Interface
+        //-----------------------------------------------------------------------------------------
+
+        //! Set the image that will be blended on the destination.
+        void SetImage( NodeImagePtr );
+
+        //! Get the image that will be blended on the destination.
+        NodeImagePtr GetImage() const;
+
+        //! Set the blending mask image.
+        void SetMask( NodeImagePtr );
+
+        //! Get the blending mask image.
+        NodeImagePtr GetMask() const;
+
+        //! Set the number of blocks in the layout that will be patched
+        void SetBlockCount( int );
+
+        //! Get the number of blocks in the layout that will be patched
+        int GetBlockCount() const;
+
+        //! Set a block of the layout that will be patched.
+        //! \param index is the index in the patching list, from 0 to GetBlockCount()-1
+        //! \param block is the block index in the layout of the image were this patch operation
+        //! will be applied.
+        void SetBlock( int index, int block );
+
+        //! Get a block of the layout that will be patched.
+        //! \param index is the index in the patching list, from 0 to GetBlockCount()-1
+        int GetBlock( int index ) const;
+
+        //! Set the blending operation to use to combine the pixels
+        EBlendType GetBlendType() const;
+        void SetBlendType(EBlendType);
+
+        //! Enable patching the alpha channel if present
+        bool GetApplyToAlphaChannel() const;
+        void SetApplyToAlphaChannel( bool );
+
+
+        //-----------------------------------------------------------------------------------------
+        // Interface pattern
+        //-----------------------------------------------------------------------------------------
+        class Private;
+        Private* GetPrivate() const;
+        Node::Private* GetBasePrivate() const override;
+
+    protected:
+
+        //! Forbidden. Manage with the Ptr<> template.
+        ~NodePatchImage();
+
+    private:
+
+        Private* m_pD;
+
+    };
+
+
+
+}
diff --git a/Engine/Plugins/Experimental/Mutable/Source/MutableTools/Internal/MuT/NodePrivate.h b/Engine/Plugins/Experimental/Mutable/Source/MutableTools/Internal/MuT/NodePrivate.h
new file mode 100644
--- /dev/null
+++ b/Engine/Plugins/Experimental/Mutable/Source/MutableTools/Internal/MuT/NodePrivate.h
@@ -0,0 +1,81 @@
+// Copyright Epic Games, Inc. All Rights Reserved.
+
+#pragma once
+
+#include "MuT/Node.h"
+#include "MuR/Operations.h"
+
+
+namespace mu
+{
+
+    class Node::Private
+    {
+    public:
+
+		/** Force a virtual destructor. */
+		virtual ~Private() = default; 
+
+        /** This is an opaque context used to attach to reported error messages. */
+		const void* m_errorContext = nullptr;
+
+		/** Generic pointer to the node owning this private. */
+		const Node* m_pNode = nullptr;
+    };
+
+
+
+#define MUTABLE_IMPLEMENT_NODE( N, T, PN, PT )				\
+    N::N()													\
+    {														\
+		Type = T;											\
+        m_pD = new Private();								\
+		m_pD->m_pNode = this;								\
+		PN::Type = PT;										\
+	}														\
+                                                            \
+    N::~N()													\
+    {														\
+        check( m_pD );										\
+        delete m_pD;										\
+        m_pD = nullptr;										\
+    }														\
+                                                            \
+    void N::Serialise( const N* p, OutputArchive& arch )	\
+    {														\
+        arch << *p->m_pD;									\
+    }														\
+                                                            \
+	void N::SerialiseWrapper( OutputArchive& arch ) const	\
+	{														\
+		N::Serialise(this, arch);							\
+	}														\
+                                                            \
+    Ptr<N> N::StaticUnserialise( InputArchive& arch )		\
+    {														\
+        Ptr<N> pResult = new N();							\
+        arch >> *pResult->m_pD;								\
+        return pResult;										\
+    }														\
+                                                            \
+    N::Private* N::GetPrivate() const						\
+    {														\
+        return m_pD;										\
+    }														\
+                                                            \
+    Node::Private* N::GetBasePrivate() const				\
+    {														\
+        return m_pD;										\
+    }														\
+                                                            \
+    const FNodeType* N::GetType() const						\
+    {														\
+        return GetStaticType();								\
+    }														\
+                                                            \
+    const FNodeType* N::GetStaticType()						\
+    {														\
+        return &Private::s_type;							\
+    }
+
+}
diff --git a/Engine/Plugins/Experimental/Mutable/Source/MutableTools/Internal/MuT/NodeSurfaceEditPrivate.h b/Engine/Plugins/Experimental/Mutable/Source/MutableTools/Internal/MuT/NodeSurfaceEditPrivate.h
new file mode 100644
--- /dev/null
+++ b/Engine/Plugins/Experimental/Mutable/Source/MutableTools/Internal/MuT/NodeSurfaceEditPrivate.h
@@ -0,0 +1,108 @@
+// Copyright Epic Games, Inc. All Rights Reserved.
+
+#pragma once
+
+#include "MuT/NodeSurfaceEdit.h"
+
+#include "MuT/NodePrivate.h"
+#include "MuT/NodeMesh.h"
+#include "MuT/NodePatchMesh.h"
+#include "MuT/NodeImage.h"
+#include "MuT/NodeScalar.h"
+#include "MuT/NodePatchImage.h"
+
+namespace mu
+{
+
+
+	//---------------------------------------------------------------------------------------------
+	//!
+	//---------------------------------------------------------------------------------------------
+    class NodeSurfaceEdit::Private : public Node::Private
+	{
+	public:
+
+		static FNodeType s_type;
+
+        NodeSurfacePtr m_pParent;
+		NodePatchMeshPtr m_pMesh;
+		NodeMeshPtr m_pMorph;
+
+        //! This flag indicates that the mesh in the m_pMorph connection is a target mesh, and not
+        //! morph information. This means that the morph information needs to be generated.
+        bool m_morphIsTarget = true;
+
+		struct FTexture
+		{
+			NodeImagePtr m_pExtend;
+			NodePatchImagePtr m_pPatch;
+
+			//!
+			void Serialise( OutputArchive& arch ) const
+			{
+				arch << m_pExtend;
+				arch << m_pPatch;
+			}
+
+			void Unserialise( InputArchive& arch )
+			{
+				arch >> m_pExtend;
+				arch >> m_pPatch;
+			}
+		};
+
+		TArray<FTexture> m_textures;
+
+        //! Tags in this surface edit
+		TArray<FString> m_tags;
+
+		//! Factor of the morph
+		NodeScalarPtr m_pFactor;
+
+		//!
+		void Serialise( OutputArchive& arch ) const
+		{
+            uint32_t ver = 5;
+			arch << ver;
+
+			arch << m_pParent;
+			arch << m_pMesh;
+            arch << m_pMorph;
+            arch << m_morphIsTarget;
+            arch << m_textures;
+            arch << m_tags;
+			arch << m_pFactor;
+        }
+
+		//!
+		void Unserialise( InputArchive& arch )
+		{
+            uint32_t ver;
+			arch >> ver;
+            check( ver>=4 && ver<=5 );
+
+			arch >> m_pParent;
+			arch >> m_pMesh;
+			arch >> m_pMorph;
+            arch >> m_morphIsTarget;
+			arch >> m_textures;
+			if (ver <= 4)
+			{
+				TArray <std::string> Temp;
+				arch >> Temp;
+				m_tags.SetNum(Temp.Num());
+				for (int32 i=0; i<Temp.Num(); ++i)
+				{
+					m_tags[i] = Temp[i].c_str();
+				}
+			}
+			else
+			{
+				arch >> m_tags;
+			}
+			arch >> m_pFactor;
+        }
+	};
+
+}
+
diff --git a/Engine/Plugins/Experimental/Mutable/Source/MutableTools/Internal/MuT/NodeSurfaceNewPrivate.h b/Engine/Plugins/Experimental/Mutable/Source/MutableTools/Internal/MuT/NodeSurfaceNewPrivate.h
new file mode 100644
--- /dev/null
+++ b/Engine/Plugins/Experimental/Mutable/Source/MutableTools/Internal/MuT/NodeSurfaceNewPrivate.h
@@ -0,0 +1,289 @@
+// Copyright Epic Games, Inc. All Rights Reserved.
+
+#pragma once
+
+#include "MuT/NodePrivate.h"
+#include "MuT/NodeSurfaceNew.h"
+
+#include "MuT/NodeMesh.h"
+#include "MuT/NodePatchMesh.h"
+#include "MuT/NodeImage.h"
+#include "MuT/NodeScalar.h"
+#include "MuT/NodeString.h"
+#include "MuT/NodeColour.h"
+
+
+namespace mu
+{
+
+	//---------------------------------------------------------------------------------------------
+	//!
+	//---------------------------------------------------------------------------------------------
+    class NodeSurfaceNew::Private : public Node::Private
+	{
+	public:
+
+		static FNodeType s_type;
+
+		FString m_name;
+        uint32 ExternalId =0;
+        int32 SharedSurfaceId =INDEX_NONE;
+
+		struct MESH
+		{
+			FString m_name;
+			NodeMeshPtr m_pMesh;
+
+			//!
+			void Serialise( OutputArchive& arch ) const
+			{
+				arch << m_name;
+				arch << m_pMesh;
+			}
+
+			void Unserialise( InputArchive& arch )
+			{
+				arch >> m_name;
+				arch >> m_pMesh;
+			}
+		};
+
+		TArray<MESH> m_meshes;
+
+		struct IMAGE
+		{
+			FString m_name;
+			FString m_materialName;
+			FString m_materialParameterName;
+			NodeImagePtr m_pImage;
+
+			// It could be negative, to indicate no layout.
+            int8 m_layoutIndex = 0;
+
+            //!
+			void Serialise(OutputArchive& arch) const
+			{
+				arch << m_name;
+				arch << m_materialName;
+				arch << m_materialParameterName;
+				arch << m_pImage;
+                arch << m_layoutIndex;
+            }
+
+			void Unserialise(InputArchive& arch)
+			{
+				arch >> m_name;
+				arch >> m_materialName;
+				arch >> m_materialParameterName;
+				arch >> m_pImage;
+                arch >> m_layoutIndex;
+            }
+        };
+
+		TArray<IMAGE> m_images;
+
+		struct VECTOR
+		{
+			FString m_name;
+			NodeColourPtr m_pVector;
+
+			//!
+			void Serialise(OutputArchive& arch) const
+			{
+				arch << m_name;
+				arch << m_pVector;
+			}
+
+			void Unserialise(InputArchive& arch)
+			{
+				arch >> m_name;
+				arch >> m_pVector;
+			}
+		};
+
+		TArray<VECTOR> m_vectors;
+
+        struct SCALAR
+        {
+			FString m_name;
+            NodeScalarPtr m_pScalar;
+
+            //!
+            void Serialise(OutputArchive& arch) const
+            {
+                arch << m_name;
+                arch << m_pScalar;
+            }
+
+            void Unserialise(InputArchive& arch)
+            {
+                arch >> m_name;
+                arch >> m_pScalar;
+            }
+        };
+
+        TArray<SCALAR> m_scalars;
+
+        struct STRING
+        {
+			FString m_name;
+            NodeStringPtr m_pString;
+
+            //!
+            void Serialise( OutputArchive& arch ) const
+            {
+                arch << m_name;
+                arch << m_pString;
+            }
+
+            void Unserialise( InputArchive& arch )
+            {
+                arch >> m_name;
+                arch >> m_pString;
+            }
+        };
+
+		TArray<STRING> m_strings;
+
+        //! Tags in this surface
+		TArray<FString> m_tags;
+
+		//! Find an image node index by name or return -1
+		int FindImage( const FString& strName ) const;
+
+		//! Find a mesh node index by name or return -1
+		int FindMesh(const FString& strName) const;
+
+        //! Find a vector node index by name or return -1
+        int FindVector(const FString& strName) const;
+
+        //! Find a scalar node index by name or return -1
+        int FindScalar( const FString& strName ) const;
+
+        //! Find a string node index by name or return -1
+        int FindString( const FString& strName ) const;
+
+        //!
+		void Serialise( OutputArchive& arch ) const
+		{
+            uint32 ver = 9;
+			arch << ver;
+
+			arch << m_name;
+            arch << ExternalId;
+			arch << SharedSurfaceId;
+
+			arch << m_meshes;
+			arch << m_images;
+			arch << m_tags;
+            arch << m_vectors;
+            arch << m_scalars;
+            arch << m_strings;
+        }
+
+        //!
+		void Unserialise( InputArchive& arch )
+		{
+            uint32 ver;
+			arch >> ver;
+            check(ver<=9);
+
+			std::string Temp;
+			if (ver <= 8)
+			{
+				arch >> Temp;
+				m_name = Temp.c_str();
+			}
+			else
+			{
+				arch >> m_name;
+			}
+            arch >> ExternalId;
+
+			if (ver >= 8)
+			{
+				arch >> SharedSurfaceId;
+			}
+
+			if (ver <= 8)
+			{
+				int32 Num = 0;
+
+				arch >> Num;
+				m_meshes.SetNum(Num);
+				for (int32 i=0; i<Num; ++i)
+				{
+					arch >> Temp;
+					m_meshes[i].m_name = Temp.c_str();
+					arch >> m_meshes[i].m_pMesh;
+				}
+
+				arch >> Num;
+				m_images.SetNum(Num);
+				for (int32 i = 0; i < Num; ++i)
+				{
+					arch >> Temp;
+					m_images[i].m_name = Temp.c_str();
+
+					arch >> Temp;
+					m_images[i].m_materialName = Temp.c_str();
+
+					arch >> Temp;
+					m_images[i].m_materialParameterName = Temp.c_str();
+
+					arch >> m_images[i].m_pImage;
+					arch >> m_images[i].m_layoutIndex;
+				}
+
+				arch >> Num;
+				m_tags.SetNum(Num);
+				for (int32 i = 0; i < Num; ++i)
+				{
+					arch >> Temp;
+					m_tags[i] = Temp.c_str();
+				}
+
+				arch >> Num;
+				m_vectors.SetNum(Num);
+				for (int32 i = 0; i < Num; ++i)
+				{
+					arch >> Temp;
+					m_vectors[i].m_name = Temp.c_str();
+					arch >> m_vectors[i].m_pVector;
+				}
+
+				arch >> Num;
+				m_scalars.SetNum(Num);
+				for (int32 i = 0; i < Num; ++i)
+				{
+					arch >> Temp;
+					m_scalars[i].m_name = Temp.c_str();
+					arch >> m_scalars[i].m_pScalar;
+				}
+
+				arch >> Num;
+				m_strings.SetNum(Num);
+				for (int32 i = 0; i < Num; ++i)
+				{
+					arch >> Temp;
+					m_strings[i].m_name = Temp.c_str();
+					arch >> m_strings[i].m_pString;
+				}
+			}
+			else
+			{
+				arch >> m_meshes;
+				arch >> m_images;
+				arch >> m_tags;
+				arch >> m_vectors;
+				arch >> m_scalars;
+				arch >> m_strings;
+			}
+        }
+    };
+
+
+
+
+}
+
diff --git a/Engine/Plugins/Experimental/Mutable/Source/MutableTools/Internal/MuT/NodeSurfaceSwitchPrivate.h b/Engine/Plugins/Experimental/Mutable/Source/MutableTools/Internal/MuT/NodeSurfaceSwitchPrivate.h
new file mode 100644
--- /dev/null
+++ b/Engine/Plugins/Experimental/Mutable/Source/MutableTools/Internal/MuT/NodeSurfaceSwitchPrivate.h
@@ -0,0 +1,48 @@
+// Copyright Epic Games, Inc. All Rights Reserved.
+
+#pragma once
+
+#include "MuT/NodeSurfaceSwitch.h"
+
+#include "MuT/NodePrivate.h"
+#include "MuT/NodeScalar.h"
+
+
+namespace mu
+{
+
+
+	class NodeSurfaceSwitch::Private : public NodeSurface::Private
+	{
+	public:
+
+		static FNodeType s_type;
+
+		Ptr<NodeScalar> Parameter;
+		TArray<Ptr<NodeSurface>> Options;
+
+		//!
+		void Serialise( OutputArchive& arch ) const
+		{
+            uint32_t ver = 0;
+			arch << ver;
+
+			arch << Parameter;
+			arch << Options;
+		}
+
+		//!
+		void Unserialise( InputArchive& arch )
+		{
+            uint32_t ver;
+			arch >> ver;
+			check(ver==0);
+
+			arch >> Parameter;
+			arch >> Options;
+		}
+
+	};
+
+
+}
diff --git a/Engine/Plugins/Experimental/Mutable/Source/MutableTools/Internal/MuT/NodeSurfaceVariationPrivate.h b/Engine/Plugins/Experimental/Mutable/Source/MutableTools/Internal/MuT/NodeSurfaceVariationPrivate.h
new file mode 100644
--- /dev/null
+++ b/Engine/Plugins/Experimental/Mutable/Source/MutableTools/Internal/MuT/NodeSurfaceVariationPrivate.h
@@ -0,0 +1,97 @@
+// Copyright Epic Games, Inc. All Rights Reserved.
+
+#pragma once
+
+#include "MuT/NodePrivate.h"
+#include "MuT/NodeSurfaceVariation.h"
+#include "MuT/NodeSurfaceNewPrivate.h"
+#include "MuT/NodeModifier.h"
+
+
+namespace mu
+{
+
+	//---------------------------------------------------------------------------------------------
+	//!
+	//---------------------------------------------------------------------------------------------
+    class NodeSurfaceVariation::Private : public NodeSurface::Private
+	{
+	public:
+
+		static FNodeType s_type;
+
+        TArray<NodeSurfacePtr> m_defaultSurfaces;
+		TArray<NodeModifierPtr> m_defaultModifiers;
+
+		struct FVariation
+		{
+			TArray<NodeSurfacePtr> m_surfaces;
+			TArray<NodeModifierPtr> m_modifiers;
+            FString m_tag;
+
+			//!
+			void Serialise( OutputArchive& arch ) const
+			{
+                uint32 ver = 2;
+				arch << ver;
+
+				arch << m_tag;
+                arch << m_surfaces;
+                arch << m_modifiers;
+            }
+
+			void Unserialise( InputArchive& arch )
+			{
+                uint32 ver = 0;
+				arch >> ver;
+                check(ver>=1&&ver<=2);
+
+				if (ver <= 1)
+				{
+					std::string Temp;
+					arch >> Temp;
+					m_tag = Temp.c_str();
+				}
+				else
+				{
+					arch >> m_tag;
+				}
+				arch >> m_surfaces;
+                arch >> m_modifiers;
+			}
+		};
+
+        NodeSurfaceVariation::VariationType m_type = NodeSurfaceVariation::VariationType::Tag;
+
+		TArray<FVariation> m_variations;
+
+		//!
+		void Serialise( OutputArchive& arch ) const
+		{
+            uint32 ver = 3;
+			arch << ver;
+
+            arch << uint32_t(m_type);
+            arch << m_defaultSurfaces;
+            arch << m_defaultModifiers;
+            arch << m_variations;
+		}
+
+		//!
+		void Unserialise( InputArchive& arch )
+		{
+            uint32 ver;
+			arch >> ver;
+            check( ver==3 );
+
+            uint32 temp;
+            arch >> temp;
+            m_type = NodeSurfaceVariation::VariationType(temp);
+
+            arch >> m_defaultSurfaces;
+            arch >> m_defaultModifiers;
+            arch >> m_variations;
+		}
+	};
+
+}
diff --git a/Engine/Plugins/Experimental/Mutable/Source/MutableTools/Private/MuT/Compiler.cpp b/Engine/Plugins/Experimental/Mutable/Source/MutableTools/Private/MuT/Compiler.cpp
--- a/Engine/Plugins/Experimental/Mutable/Source/MutableTools/Private/MuT/Compiler.cpp
+++ b/Engine/Plugins/Experimental/Mutable/Source/MutableTools/Private/MuT/Compiler.cpp
@@ -34,6 +34,8 @@
 #include "MuT/Table.h"
 #include "Trace/Detail/Channel.h"
 
+#include <string>
+
 
 namespace mu
 {
@@ -164,7 +166,53 @@ namespace mu
 	}
 
 
-    //---------------------------------------------------------------------------------------------
+	void FObjectState::Serialise(OutputArchive& arch) const
+    {
+    	const int32 ver = 6;
+    	arch << ver;
+
+    	arch << m_name;
+    	arch << m_optimisation;
+    	arch << m_runtimeParams;
+    }
+
+
+	void FObjectState::Unserialise(InputArchive& arch)
+    {
+    	int32 ver = 0;
+    	arch >> ver;
+    	check( ver>=5 && ver<=6 );
+
+    	if (ver <= 5)
+    	{
+    		std::string Temp;
+    		arch >> Temp;
+    		m_name = Temp.c_str();
+    	}
+    	else
+    	{
+    		arch >> m_name;
+    	}
+    	arch >> m_optimisation;
+
+    	if (ver <= 5)
+    	{
+    		TArray<std::string> Temp;
+    		arch >> Temp;
+    		m_runtimeParams.SetNum(Temp.Num());
+    		for ( int32 i=0; i<Temp.Num(); ++i)
+    		{
+    			m_runtimeParams[i] = Temp[i].c_str();
+    		}
+    	}
+    	else
+    	{
+    		arch >> m_runtimeParams;
+    	}
+    }
+	
+	
+	//---------------------------------------------------------------------------------------------
     //---------------------------------------------------------------------------------------------
     //---------------------------------------------------------------------------------------------
     Compiler::Compiler( Ptr<CompilerOptions> options )
diff --git a/Engine/Plugins/Experimental/Mutable/Source/MutableTools/Private/MuT/CompilerPrivate.h b/Engine/Plugins/Experimental/Mutable/Source/MutableTools/Private/MuT/CompilerPrivate.h
--- a/Engine/Plugins/Experimental/Mutable/Source/MutableTools/Private/MuT/CompilerPrivate.h
+++ b/Engine/Plugins/Experimental/Mutable/Source/MutableTools/Private/MuT/CompilerPrivate.h
@@ -13,74 +13,7 @@
 
 namespace mu
 {
-	
-    //!
-    struct FStateOptimizationOptions
-    {
-		uint8 FirstLOD = 0;
-		uint8 NumExtraLODsToBuildAfterFirstLOD = 0;
-		bool bOnlyFirstLOD = false;
-		ETextureCompressionStrategy TextureCompressionStrategy = ETextureCompressionStrategy::None;
-
-        void Serialise( OutputArchive& arch ) const
-        {
-            const int32_t ver = 4;
-            arch << ver;
-
-			arch << FirstLOD;
-			arch << bOnlyFirstLOD;
-            arch << TextureCompressionStrategy;
-			arch << NumExtraLODsToBuildAfterFirstLOD;
-        }
-
-
-        void Unserialise( InputArchive& arch )
-        {
-            int32_t ver = 0;
-            arch >> ver;
-			check(ver <= 4);
-
-			if (ver >= 2)
-			{
-				arch >> FirstLOD;
-			}
-			else
-			{
-				FirstLOD = 0;
-			}
-
-            arch >> bOnlyFirstLOD;
-
-			if (ver >= 4)
-			{
-				arch >> TextureCompressionStrategy;
-			}
-			else
-			{
-				bool bAvoidRuntimeCompression;
-				arch >> bAvoidRuntimeCompression;
-				TextureCompressionStrategy = bAvoidRuntimeCompression ? ETextureCompressionStrategy::DontCompressRuntime : ETextureCompressionStrategy::None;
-			}
-
-			if (ver == 3)
-			{
-				int32 OldNumExtraLODsToBuildAfterFirstLOD;
-				arch >> OldNumExtraLODsToBuildAfterFirstLOD;
-				NumExtraLODsToBuildAfterFirstLOD = OldNumExtraLODsToBuildAfterFirstLOD;
-			}
-			else if (ver >= 4)
-			{
-				arch >> NumExtraLODsToBuildAfterFirstLOD;
-			}
-			else
-			{
-				NumExtraLODsToBuildAfterFirstLOD = 0;
-			}
-
-		}
-    };
-
-
+		
     //!
     class CompilerOptions::Private
     {
@@ -102,65 +35,6 @@ namespace mu
     };
 
 
-    //! Information about an object state in the source data
-    struct FObjectState
-    {
-        //! Name used to identify the state from the code and user interface.
-        FString m_name;
-
-        //! GPU Optimisation options
-		FStateOptimizationOptions m_optimisation;
-
-        //! List of names of the runtime parameters in this state
-        TArray<FString> m_runtimeParams;
-
-        void Serialise( OutputArchive& arch ) const
-        {
-            const int32 ver = 6;
-            arch << ver;
-
-            arch << m_name;
-            arch << m_optimisation;
-            arch << m_runtimeParams;
-        }
-
-
-        void Unserialise( InputArchive& arch )
-        {
-            int32 ver = 0;
-            arch >> ver;
-            check( ver>=5 && ver<=6 );
-
-			if (ver <= 5)
-			{
-				std::string Temp;
-				arch >> Temp;
-				m_name = Temp.c_str();
-			}
-			else
-			{
-				arch >> m_name;
-			}
-            arch >> m_optimisation;
-
-			if (ver <= 5)
-			{
-				TArray<std::string> Temp;
-				arch >> Temp;
-				m_runtimeParams.SetNum(Temp.Num());
-				for ( int32 i=0; i<Temp.Num(); ++i)
-				{
-					m_runtimeParams[i] = Temp[i].c_str();
-				}
-			}
-			else
-			{
-				arch >> m_runtimeParams;
-			}
-        }
-    };
-
-
     //!
     struct FStateCompilationData
     {
diff --git a/Engine/Plugins/Experimental/Mutable/Source/MutableTools/Private/MuT/NodeObjectGroupPrivate.h b/Engine/Plugins/Experimental/Mutable/Source/MutableTools/Private/MuT/NodeObjectGroupPrivate.h
deleted file mode 100644
diff --git a/Engine/Plugins/Experimental/Mutable/Source/MutableTools/Private/MuT/NodeObjectNewPrivate.h b/Engine/Plugins/Experimental/Mutable/Source/MutableTools/Private/MuT/NodeObjectNewPrivate.h
deleted file mode 100644
diff --git a/Engine/Plugins/Experimental/Mutable/Source/MutableTools/Private/MuT/NodeObjectPrivate.h b/Engine/Plugins/Experimental/Mutable/Source/MutableTools/Private/MuT/NodeObjectPrivate.h
deleted file mode 100644
diff --git a/Engine/Plugins/Experimental/Mutable/Source/MutableTools/Private/MuT/NodePrivate.h b/Engine/Plugins/Experimental/Mutable/Source/MutableTools/Private/MuT/NodePrivate.h
deleted file mode 100644
diff --git a/Engine/Plugins/Experimental/Mutable/Source/MutableTools/Private/MuT/NodeSurfaceEditPrivate.h b/Engine/Plugins/Experimental/Mutable/Source/MutableTools/Private/MuT/NodeSurfaceEditPrivate.h
deleted file mode 100644
diff --git a/Engine/Plugins/Experimental/Mutable/Source/MutableTools/Private/MuT/NodeSurfaceNewPrivate.h b/Engine/Plugins/Experimental/Mutable/Source/MutableTools/Private/MuT/NodeSurfaceNewPrivate.h
deleted file mode 100644
diff --git a/Engine/Plugins/Experimental/Mutable/Source/MutableTools/Private/MuT/NodeSurfaceSwitchPrivate.h b/Engine/Plugins/Experimental/Mutable/Source/MutableTools/Private/MuT/NodeSurfaceSwitchPrivate.h
deleted file mode 100644
diff --git a/Engine/Plugins/Experimental/Mutable/Source/MutableTools/Private/MuT/NodeSurfaceVariationPrivate.h b/Engine/Plugins/Experimental/Mutable/Source/MutableTools/Private/MuT/NodeSurfaceVariationPrivate.h
deleted file mode 100644
diff --git a/Engine/Plugins/Experimental/Mutable/Source/MutableTools/Public/MuT/Compiler.h b/Engine/Plugins/Experimental/Mutable/Source/MutableTools/Public/MuT/Compiler.h
--- a/Engine/Plugins/Experimental/Mutable/Source/MutableTools/Public/MuT/Compiler.h
+++ b/Engine/Plugins/Experimental/Mutable/Source/MutableTools/Public/MuT/Compiler.h
@@ -7,6 +7,8 @@
 #include "MuR/Image.h"
 #include "Templates/SharedPointer.h"
 #include "HAL/PlatformMath.h"
+#include "MuR/System.h"
+
 
 namespace mu
 {
@@ -119,6 +121,92 @@ namespace mu
     };
 
 
+    //!
+    struct FStateOptimizationOptions
+    {
+		uint8 FirstLOD = 0;
+		uint8 NumExtraLODsToBuildAfterFirstLOD = 0;
+		bool bOnlyFirstLOD = false;
+		ETextureCompressionStrategy TextureCompressionStrategy = ETextureCompressionStrategy::None;
+
+        void Serialise( OutputArchive& arch ) const
+        {
+            const int32_t ver = 4;
+            arch << ver;
+
+			arch << FirstLOD;
+			arch << bOnlyFirstLOD;
+            arch << TextureCompressionStrategy;
+			arch << NumExtraLODsToBuildAfterFirstLOD;
+        }
+
+
+        void Unserialise( InputArchive& arch )
+        {
+            int32_t ver = 0;
+            arch >> ver;
+			check(ver <= 4);
+
+			if (ver >= 2)
+			{
+				arch >> FirstLOD;
+			}
+			else
+			{
+				FirstLOD = 0;
+			}
+
+            arch >> bOnlyFirstLOD;
+
+			if (ver >= 4)
+			{
+				arch >> TextureCompressionStrategy;
+			}
+			else
+			{
+				bool bAvoidRuntimeCompression;
+				arch >> bAvoidRuntimeCompression;
+				TextureCompressionStrategy = bAvoidRuntimeCompression ? ETextureCompressionStrategy::DontCompressRuntime : ETextureCompressionStrategy::None;
+			}
+
+			if (ver == 3)
+			{
+				int32 OldNumExtraLODsToBuildAfterFirstLOD;
+				arch >> OldNumExtraLODsToBuildAfterFirstLOD;
+				NumExtraLODsToBuildAfterFirstLOD = OldNumExtraLODsToBuildAfterFirstLOD;
+			}
+			else if (ver >= 4)
+			{
+				arch >> NumExtraLODsToBuildAfterFirstLOD;
+			}
+			else
+			{
+				NumExtraLODsToBuildAfterFirstLOD = 0;
+			}
+
+		}
+    };
+
+
+	//! Information about an object state in the source data
+	struct FObjectState
+	{
+		//! Name used to identify the state from the code and user interface.
+		FString m_name;
+
+		//! GPU Optimisation options
+		FStateOptimizationOptions m_optimisation;
+
+		//! List of names of the runtime parameters in this state
+		TArray<FString> m_runtimeParams;
+
+		void Serialise( OutputArchive& arch ) const;
+
+
+		void Unserialise( InputArchive& arch );
+	};
+
+	
     //! This utility class compiles two types of expressions: models and transforms.
     //! Model expressions are compiled into run-time usable objects.
     //! Transform expressions are compiled into Transform objects than can be applied to Model
diff --git a/Engine/Plugins/Experimental/Mutable/Source/MutableTools/Public/MuT/NodePatchImage.h b/Engine/Plugins/Experimental/Mutable/Source/MutableTools/Public/MuT/NodePatchImage.h
deleted file mode 100644
