[Mutable] Fix PROJECTOR_TYPE compilation errors due to missing header.

#rnx

diff --git a/Engine/Plugins/Experimental/Mutable/Source/MutableTools/Private/MuT/NodeProjectorPrivate.h b/Engine/Plugins/Experimental/Mutable/Source/MutableTools/Private/MuT/NodeProjectorPrivate.h
--- a/Engine/Plugins/Experimental/Mutable/Source/MutableTools/Private/MuT/NodeProjectorPrivate.h
+++ b/Engine/Plugins/Experimental/Mutable/Source/MutableTools/Private/MuT/NodeProjectorPrivate.h
@@ -7,6 +7,7 @@
 
 #include "MuT/NodeRange.h"
 #include "MuR/MutableMath.h"
+#include "MuR/ParametersPrivate.h"
 
 
 namespace mu
