[Mutable] Fix clip Morph Mesh UI controls do not updateing the preview gizmo.

#rnx

diff --git a/Engine/Plugins/Experimental/Mutable/Source/CustomizableObjectEditor/Private/MuCOE/CustomizableInstanceDetails.cpp b/Engine/Plugins/Experimental/Mutable/Source/CustomizableObjectEditor/Private/MuCOE/CustomizableInstanceDetails.cpp
--- a/Engine/Plugins/Experimental/Mutable/Source/CustomizableObjectEditor/Private/MuCOE/CustomizableInstanceDetails.cpp
+++ b/Engine/Plugins/Experimental/Mutable/Source/CustomizableObjectEditor/Private/MuCOE/CustomizableInstanceDetails.cpp
@@ -1262,7 +1262,7 @@ FReply FCustomizableInstanceDetails::OnProjectorSelectChanged(const FString Para
 	const UProjectorParameter* ProjectorParameter = Editor->GetProjectorParameter();
 	if (ProjectorParameter->IsProjectorSelected(ParamName, RangeIndex))
 	{
-		Editor->HideGizmoProjectorParameter();
+		Editor->HideGizmo();
 	}
 	else
 	{
diff --git a/Engine/Plugins/Experimental/Mutable/Source/CustomizableObjectEditor/Private/MuCOE/CustomizableObjectEditor.cpp b/Engine/Plugins/Experimental/Mutable/Source/CustomizableObjectEditor/Private/MuCOE/CustomizableObjectEditor.cpp
--- a/Engine/Plugins/Experimental/Mutable/Source/CustomizableObjectEditor/Private/MuCOE/CustomizableObjectEditor.cpp
+++ b/Engine/Plugins/Experimental/Mutable/Source/CustomizableObjectEditor/Private/MuCOE/CustomizableObjectEditor.cpp
@@ -886,26 +886,38 @@ UCustomSettings* FCustomizableObjectEditor::GetCustomSettings()
 }
 
 
+void FCustomizableObjectEditor::SelectSingleNode(UCustomizableObjectNode& Node)
+{
+	FGraphPanelSelectionSet SelectedNodes = GraphEditor->GetSelectedNodes();
+	if (SelectedNodes.Num() != 1 || Cast<UCustomizableObjectNode>(*SelectedNodes.CreateIterator()) != &Node)
+	{
+		GraphEditor->ClearSelectionSet();
+		GraphEditor->SetNodeSelection(&Node, true);
+	}
+}
+
+
 void FCustomizableObjectEditor::HideGizmo()
 {
 	HideGizmoProjectorNodeProjectorConstant();
 	HideGizmoProjectorNodeProjectorParameter();
+	HideGizmoProjectorParameter();
 	HideGizmoClipMorph();
 	HideGizmoClipMesh();
 	HideGizmoLight();
-
-	FCustomizableObjectInstanceEditor::HideGizmo(SharedThis(this), Viewport, CustomizableInstanceDetailsView);
 }
 
 
 void FCustomizableObjectEditor::ShowGizmoProjectorNodeProjectorConstant(UCustomizableObjectNodeProjectorConstant& Node)
 {
-	HideGizmo();
-	
-	ProjectorGizmo = EProjectorGizmo::NodeProjectorConstant;
-	
-	GraphEditor->ClearSelectionSet();
-	GraphEditor->SetNodeSelection(&Node, true);
+	if (GizmoType != EGizmoType::NodeProjectorConstant)
+	{
+		HideGizmo();
+	}
+
+	GizmoType = EGizmoType::NodeProjectorConstant;
+
+	SelectSingleNode(Node);
 	
 	FProjectorTypeDelegate ProjectorTypeDelegate;
 	ProjectorTypeDelegate.BindUObject(&Node, &UCustomizableObjectNodeProjectorConstant::GetProjectorType);		
@@ -962,10 +974,12 @@ void FCustomizableObjectEditor::ShowGizmoProjectorNodeProjectorConstant(UCustomi
 
 void FCustomizableObjectEditor::HideGizmoProjectorNodeProjectorConstant()
 {
-	if (ProjectorGizmo != EProjectorGizmo::NodeProjectorConstant)
+	if (GizmoType != EGizmoType::NodeProjectorConstant)
 	{
 		return;
 	}
+
+	GizmoType = EGizmoType::Hidden;
 	
 	Viewport->HideGizmoProjector();
 
@@ -984,12 +998,13 @@ void FCustomizableObjectEditor::HideGizmoProjectorNodeProjectorConstant()
 
 void FCustomizableObjectEditor::ShowGizmoProjectorNodeProjectorParameter(UCustomizableObjectNodeProjectorParameter& Node)
 {
-	HideGizmo();
+	if (GizmoType != EGizmoType::NodeProjectorParameter)
+	{
+		HideGizmo();
+		GizmoType = EGizmoType::NodeProjectorParameter;
+	}
 
-	ProjectorGizmo = EProjectorGizmo::NodeProjectorParameter;
-	
-	GraphEditor->ClearSelectionSet();
-	GraphEditor->SetNodeSelection(&Node, true);
+	SelectSingleNode(Node);
 	
 	FProjectorTypeDelegate ProjectorTypeDelegate;
 	ProjectorTypeDelegate.BindUObject(&Node, &UCustomizableObjectNodeProjectorParameter::GetProjectorType);		
@@ -1046,11 +1061,13 @@ void FCustomizableObjectEditor::ShowGizmoProjectorNodeProjectorParameter(UCustom
 
 void FCustomizableObjectEditor::HideGizmoProjectorNodeProjectorParameter()
 {
-	if (ProjectorGizmo != EProjectorGizmo::NodeProjectorParameter)
+	if (GizmoType != EGizmoType::NodeProjectorParameter)
 	{
 		return;
 	}
-	
+
+	GizmoType = EGizmoType::Hidden;
+
 	Viewport->HideGizmoProjector();
 
 	const FGraphPanelSelectionSet SelectedNodes = GraphEditor->GetSelectedNodes();
@@ -1068,7 +1085,11 @@ void FCustomizableObjectEditor::HideGizmoProjectorNodeProjectorParameter()
 
 void FCustomizableObjectEditor::ShowGizmoProjectorParameter(const FString& ParamName, int32 RangeIndex)
 {
-	ProjectorGizmo = EProjectorGizmo::Parameter;
+	if (GizmoType != EGizmoType::ProjectorParameter)
+	{
+		HideGizmo();		
+		GizmoType = EGizmoType::ProjectorParameter;
+	}
 	
 	FCustomizableObjectInstanceEditor::ShowGizmoProjectorParameter(ParamName, RangeIndex, SharedThis(this), Viewport, CustomizableInstanceDetailsView, ProjectorParameter, PreviewInstance);
 }
@@ -1076,31 +1097,45 @@ void FCustomizableObjectEditor::ShowGizmoProjectorParameter(const FString& Param
 
 void FCustomizableObjectEditor::HideGizmoProjectorParameter()
 {
-	if (ProjectorGizmo != EProjectorGizmo::Parameter)
+	if (GizmoType != EGizmoType::ProjectorParameter)
 	{
 		return;	
 	}
-	
+
+	GizmoType = EGizmoType::Hidden;
+
 	FCustomizableObjectInstanceEditor::HideGizmoProjectorParameter(SharedThis(this), Viewport, CustomizableInstanceDetailsView);
 }
 
 
 void FCustomizableObjectEditor::ShowGizmoClipMorph(UCustomizableObjectNodeMeshClipMorph& Node)
 {
-	if (Node.BoneName != FName())
-	{
-		HideGizmo();
-
-		GraphEditor->ClearSelectionSet();
-		GraphEditor->SetNodeSelection(&Node, true);
+	if (Node.BoneName == FName())
+	{	
+		return;
+	}
 
-		Viewport->ShowGizmoClipMorph(Node);
+	if (GizmoType != EGizmoType::ClipMorph)
+	{
+		HideGizmo();		
+		GizmoType = EGizmoType::ClipMorph;
 	}
+	
+	SelectSingleNode(Node);
+
+	Viewport->ShowGizmoClipMorph(Node);
 }
 
 
 void FCustomizableObjectEditor::HideGizmoClipMorph()
 {
+	if (GizmoType != EGizmoType::ClipMorph)
+	{
+		return;	
+	}
+	
+	GizmoType = EGizmoType::Hidden;
+
 	Viewport->HideGizmoClipMorph();
 
 	const FGraphPanelSelectionSet SelectedNodes = GraphEditor->GetSelectedNodes();
@@ -1137,10 +1172,13 @@ void FCustomizableObjectEditor::ShowGizmoClipMesh(UCustomizableObjectNodeMeshCli
 
 	if (ClipMesh)
 	{
-		HideGizmo();
+		if (GizmoType != EGizmoType::ClipMesh)
+		{
+			HideGizmo();
+			GizmoType = EGizmoType::ClipMesh;
+		}
 
-		GraphEditor->ClearSelectionSet();
-		GraphEditor->SetNodeSelection(&Node, true);
+		SelectSingleNode(Node);
 
 		Viewport->ShowGizmoClipMesh(Node, *ClipMesh);
 	}
@@ -1149,6 +1187,13 @@ void FCustomizableObjectEditor::ShowGizmoClipMesh(UCustomizableObjectNodeMeshCli
 
 void FCustomizableObjectEditor::HideGizmoClipMesh()
 {
+	if (GizmoType != EGizmoType::ClipMesh)
+	{
+		return;	
+	}
+
+	GizmoType = EGizmoType::Hidden;
+
 	Viewport->HideGizmoClipMesh();
 
 	const FGraphPanelSelectionSet SelectedNodes = GraphEditor->GetSelectedNodes();
@@ -1166,8 +1211,12 @@ void FCustomizableObjectEditor::HideGizmoClipMesh()
 
 void FCustomizableObjectEditor::ShowGizmoLight(ULightComponent& InSelectedLight)
 {
-	HideGizmo();
-
+	if (GizmoType != EGizmoType::Light)
+	{
+		HideGizmo();
+		GizmoType = EGizmoType::Light;
+	}
+	
 	CustomSettings->SetSelectedLight(&InSelectedLight);
 
 	Viewport->ShowGizmoLight(InSelectedLight);
@@ -1178,6 +1227,13 @@ void FCustomizableObjectEditor::ShowGizmoLight(ULightComponent& InSelectedLight)
 
 void FCustomizableObjectEditor::HideGizmoLight()
 {
+	if (GizmoType != EGizmoType::Light)
+	{
+		return;	
+	}
+	
+	GizmoType = EGizmoType::Hidden;
+
 	CustomSettings->SetSelectedLight(nullptr);
 
 	Viewport->HideGizmoLight();
diff --git a/Engine/Plugins/Experimental/Mutable/Source/CustomizableObjectEditor/Private/MuCOE/CustomizableObjectEditor.h b/Engine/Plugins/Experimental/Mutable/Source/CustomizableObjectEditor/Private/MuCOE/CustomizableObjectEditor.h
--- a/Engine/Plugins/Experimental/Mutable/Source/CustomizableObjectEditor/Private/MuCOE/CustomizableObjectEditor.h
+++ b/Engine/Plugins/Experimental/Mutable/Source/CustomizableObjectEditor/Private/MuCOE/CustomizableObjectEditor.h
@@ -50,6 +50,19 @@ DECLARE_DELEGATE(FCreatePreviewInstanceFlagDelegate);
 
 extern void RemoveRestrictedChars(FString& String);
 
+
+enum class EGizmoType : uint8
+{
+	Hidden,
+	ProjectorParameter,
+	NodeProjectorConstant,
+	NodeProjectorParameter,
+	ClipMorph,
+	ClipMesh,
+	Light,
+};
+
+
 /**
 * Wrapper UObject class for the UCustomizableObjectInstance::FObjectInstanceUpdatedDelegate dynamic multicast delegate
 */
@@ -209,6 +222,9 @@ public:
 	virtual void ShowGizmoLight(ULightComponent& SelectedLight) override;
 	virtual void HideGizmoLight() override;
 
+	/** Select only this node only. Do nothing if already was only selected. */
+	void SelectSingleNode(UCustomizableObjectNode& Node);
+	
 	/** Called to undo the last action */
 	void UndoGraphAction();
 
@@ -447,7 +463,7 @@ private:
 
 	bool bRecursionGuard = false;
 	
-	EProjectorGizmo ProjectorGizmo = EProjectorGizmo::None;
+	EGizmoType GizmoType = EGizmoType::Hidden;
 
 protected:
 	/** @return the documentation location for this editor */
diff --git a/Engine/Plugins/Experimental/Mutable/Source/CustomizableObjectEditor/Private/MuCOE/CustomizableObjectEditorViewportClient.cpp b/Engine/Plugins/Experimental/Mutable/Source/CustomizableObjectEditor/Private/MuCOE/CustomizableObjectEditorViewportClient.cpp
--- a/Engine/Plugins/Experimental/Mutable/Source/CustomizableObjectEditor/Private/MuCOE/CustomizableObjectEditorViewportClient.cpp
+++ b/Engine/Plugins/Experimental/Mutable/Source/CustomizableObjectEditor/Private/MuCOE/CustomizableObjectEditorViewportClient.cpp
@@ -689,33 +689,28 @@ void FCustomizableObjectEditorViewportClient::ShowGizmoClipMorph(UCustomizableOb
 {
 	SetWidgetType(EWidgetType::ClipMorph);
 
-	if (ClipMorphNode != &NodeMeshClipMorph || NodeMeshClipMorph.bUpdateViewportWidget)
-	{	
-		NodeMeshClipMorph.bUpdateViewportWidget = false;
+	ClipMorphNode = &NodeMeshClipMorph;
 
-		bClipMorphLocalStartOffset = NodeMeshClipMorph.bLocalStartOffset;
-		MorphLength = NodeMeshClipMorph.B;
-		Radius1 = NodeMeshClipMorph.Radius;
-		Radius2 = NodeMeshClipMorph.Radius2;
-		RotationAngle = NodeMeshClipMorph.RotationAngle;
-		ClipMorphOrigin = NodeMeshClipMorph.Origin;
-		ClipMorphLocalOffset = NodeMeshClipMorph.StartOffset;
+	bClipMorphLocalStartOffset = NodeMeshClipMorph.bLocalStartOffset;
+	MorphLength = NodeMeshClipMorph.B;
+	Radius1 = NodeMeshClipMorph.Radius;
+	Radius2 = NodeMeshClipMorph.Radius2;
+	RotationAngle = NodeMeshClipMorph.RotationAngle;
+	ClipMorphOrigin = NodeMeshClipMorph.Origin;
+	ClipMorphLocalOffset = NodeMeshClipMorph.StartOffset;
 
-		NodeMeshClipMorph.FindLocalAxes(ClipMorphXAxis, ClipMorphYAxis, ClipMorphNormal);
+	NodeMeshClipMorph.FindLocalAxes(ClipMorphXAxis, ClipMorphYAxis, ClipMorphNormal);
 
-		if (bClipMorphLocalStartOffset)
-		{
-			ClipMorphOffset = ClipMorphLocalOffset.X * ClipMorphXAxis
-				+ ClipMorphLocalOffset.Y * ClipMorphYAxis
-				+ ClipMorphLocalOffset.Z * ClipMorphNormal;
-		}
-		else
-		{
-			ClipMorphOffset = ClipMorphLocalOffset;
-		}
+	if (bClipMorphLocalStartOffset)
+	{
+		ClipMorphOffset = ClipMorphLocalOffset.X * ClipMorphXAxis
+			+ ClipMorphLocalOffset.Y * ClipMorphYAxis
+			+ ClipMorphLocalOffset.Z * ClipMorphNormal;
+	}
+	else
+	{
+		ClipMorphOffset = ClipMorphLocalOffset;
 	}
-
-	ClipMorphNode = &NodeMeshClipMorph;
 }
 
 
diff --git a/Engine/Plugins/Experimental/Mutable/Source/CustomizableObjectEditor/Private/MuCOE/CustomizableObjectInstanceEditor.h b/Engine/Plugins/Experimental/Mutable/Source/CustomizableObjectEditor/Private/MuCOE/CustomizableObjectInstanceEditor.h
--- a/Engine/Plugins/Experimental/Mutable/Source/CustomizableObjectEditor/Private/MuCOE/CustomizableObjectInstanceEditor.h
+++ b/Engine/Plugins/Experimental/Mutable/Source/CustomizableObjectEditor/Private/MuCOE/CustomizableObjectInstanceEditor.h
@@ -27,13 +27,6 @@ struct FFrame;
 
 DECLARE_DELEGATE(FCreatePreviewInstanceFlagDelegate);
 
-enum class EProjectorGizmo : uint8
-{
-	None,
-	Parameter,
-	NodeProjectorConstant,
-	NodeProjectorParameter
-};
 
 /**
 * Wrapper UObject class for the UCustomizableObjectInstance::FObjectInstanceUpdatedDelegate dynamic multicast delegate
@@ -307,8 +300,6 @@ private:
 
 	TObjectPtr<UCustomSettings> CustomSettings = nullptr;
 
-	EProjectorGizmo ProjectorGizmo = EProjectorGizmo::None;
-	
 	/** Adds the customizable Object Instance Editor commands to the default toolbar */
 	void ExtendToolbar();
 };
diff --git a/Engine/Plugins/Experimental/Mutable/Source/CustomizableObjectEditor/Private/MuCOE/Nodes/CustomizableObjectNodeMeshClipMorph.cpp b/Engine/Plugins/Experimental/Mutable/Source/CustomizableObjectEditor/Private/MuCOE/Nodes/CustomizableObjectNodeMeshClipMorph.cpp
--- a/Engine/Plugins/Experimental/Mutable/Source/CustomizableObjectEditor/Private/MuCOE/Nodes/CustomizableObjectNodeMeshClipMorph.cpp
+++ b/Engine/Plugins/Experimental/Mutable/Source/CustomizableObjectEditor/Private/MuCOE/Nodes/CustomizableObjectNodeMeshClipMorph.cpp
@@ -26,8 +26,6 @@ UCustomizableObjectNodeMeshClipMorph::UCustomizableObjectNodeMeshClipMorph()
 	Origin = FVector::ZeroVector;
 	Normal = -FVector::UpVector;
 	MaxEffectRadius = -1.f;
-
-	bUpdateViewportWidget = true;
 }
 
 
@@ -93,10 +91,8 @@ UEdGraphPin* UCustomizableObjectNodeMeshClipMorph::OutputPin() const
 }
 
 
-void UCustomizableObjectNodeMeshClipMorph::PostEditChangeProperty(FPropertyChangedEvent & PropertyChangedEvent)
+void UCustomizableObjectNodeMeshClipMorph::PostEditChangeProperty(FPropertyChangedEvent& PropertyChangedEvent)
 {
-	FProperty* PropertyThatChanged = PropertyChangedEvent.Property;
-
 	const FName PropertyName = PropertyChangedEvent.GetPropertyName();
 
 	if (PropertyName == "bLocalStartOffset")
@@ -104,7 +100,7 @@ void UCustomizableObjectNodeMeshClipMorph::PostEditChangeProperty(FPropertyChang
 		ChangeStartOffsetTransform();
 	}
 
-	bUpdateViewportWidget = true;
+	GetGraphEditor()->ShowGizmoClipMorph(*this);
 
 	Super::PostEditChangeProperty(PropertyChangedEvent);
 }
diff --git a/Engine/Plugins/Experimental/Mutable/Source/CustomizableObjectEditor/Private/MuCOE/SCustomizableInstanceProperties.cpp b/Engine/Plugins/Experimental/Mutable/Source/CustomizableObjectEditor/Private/MuCOE/SCustomizableInstanceProperties.cpp
--- a/Engine/Plugins/Experimental/Mutable/Source/CustomizableObjectEditor/Private/MuCOE/SCustomizableInstanceProperties.cpp
+++ b/Engine/Plugins/Experimental/Mutable/Source/CustomizableObjectEditor/Private/MuCOE/SCustomizableInstanceProperties.cpp
@@ -1763,7 +1763,7 @@ FReply SCustomizableInstanceProperties::OnProjectorSelectChanged(const FString P
 	const UProjectorParameter* ProjectorParameter = Editor->GetProjectorParameter();
 	if (ProjectorParameter->IsProjectorSelected(ParamName, RangeIndex))
 	{
-		Editor->HideGizmoProjectorParameter();
+		Editor->HideGizmo();
 	}
 	else
 	{
