[mutable] Ensure all parameters connected in the source graph appear in the compiled object regardless of their relevance.
#rnx

diff --git a/Engine/Plugins/Experimental/Mutable/Source/MutableTools/Private/MuT/ASTOpParameter.cpp b/Engine/Plugins/Experimental/Mutable/Source/MutableTools/Private/MuT/ASTOpParameter.cpp
--- a/Engine/Plugins/Experimental/Mutable/Source/MutableTools/Private/MuT/ASTOpParameter.cpp
+++ b/Engine/Plugins/Experimental/Mutable/Source/MutableTools/Private/MuT/ASTOpParameter.cpp
@@ -72,7 +72,7 @@ namespace mu
 	}
 
 
-	void ASTOpParameter::Link(FProgram& program, FLinkerOptions*)
+	void ASTOpParameter::Link(FProgram& Program, FLinkerOptions*)
 	{
 		// Already linked?
 		if (!linkedAddress)
@@ -80,24 +80,27 @@ namespace mu
 			OP::ParameterArgs args;
 			FMemory::Memzero(&args, sizeof(args));
 
-			LinkedParameterIndex = program.m_parameters.Num();
+			LinkedParameterIndex = Program.m_parameters.Find(parameter);
+
+			// If this fails, it means an ASTOpParameter was created at code generation time, but not registered into the
+			// parameters map in the CodeGenerator_FirstPass.
+			check(LinkedParameterIndex != INDEX_NONE);
+
 			args.variable = (OP::ADDRESS)LinkedParameterIndex;
-			program.m_parameters.Add(parameter);
 
 			for (const auto& d : ranges)
 			{
 				OP::ADDRESS sizeAt = 0;
 				uint16 rangeId = 0;
-				LinkRange(program, d, sizeAt, rangeId);
-				program.m_parameters.Last().m_ranges.Add(rangeId);
+				LinkRange(Program, d, sizeAt, rangeId);
+				Program.m_parameters.Last().m_ranges.Add(rangeId);
 			}
 
-			linkedAddress = (OP::ADDRESS)program.m_opAddress.Num();
-			//program.m_code.push_back(op);
+			linkedAddress = (OP::ADDRESS)Program.m_opAddress.Num();
 
-			program.m_opAddress.Add((uint32_t)program.m_byteCode.Num());
-			AppendCode(program.m_byteCode, type);
-			AppendCode(program.m_byteCode, args);
+			Program.m_opAddress.Add((uint32)Program.m_byteCode.Num());
+			AppendCode(Program.m_byteCode, type);
+			AppendCode(Program.m_byteCode, args);
 		}
 	}
 
diff --git a/Engine/Plugins/Experimental/Mutable/Source/MutableTools/Private/MuT/CodeGenerator.cpp b/Engine/Plugins/Experimental/Mutable/Source/MutableTools/Private/MuT/CodeGenerator.cpp
--- a/Engine/Plugins/Experimental/Mutable/Source/MutableTools/Private/MuT/CodeGenerator.cpp
+++ b/Engine/Plugins/Experimental/Mutable/Source/MutableTools/Private/MuT/CodeGenerator.cpp
@@ -142,21 +142,6 @@ namespace mu
 				++CurrentStateIndex;
 			}
 		}
-
-		// Free caches
-		GeneratedGenericNodes.Reset();
-		m_constantMeshes.SetNum(0,EAllowShrinking::No);
-		GeneratedLayouts.Reset();
-		NodeVariables.Reset();
-		m_generatedMeshes.Reset();
-		GeneratedProjectors.Reset();
-		GeneratedRanges.Reset();
-		GeneratedStrings.Reset();
-		GeneratedTables.Reset();
-		m_firstPass = FirstPassGenerator();
-		m_currentParents.Reset();
-		m_currentObject.Reset();
-		AdditionalComponents.Reset();
 	}
 
 
@@ -324,15 +309,15 @@ namespace mu
 	}
 
 
-	Ptr<ASTOp> CodeGenerator::GenerateTableVariable(TablePtr pTable, const FString& strName, bool bAddNoneOption)
+	Ptr<ASTOp> CodeGenerator::GenerateTableVariable(Ptr<const Node> InNode, const FTableCacheKey& CacheKey, bool bAddNoneOption)
 	{
 		Ptr<ASTOp> result;
 
         FParameterDesc param;
-        param.m_name = strName;
+        param.m_name = CacheKey.ParameterName;
         if ( param.m_name.Len()==0 )
         {
-            param.m_name = TCHAR_TO_ANSI(*pTable->GetName());
+            param.m_name = TCHAR_TO_ANSI(*CacheKey.Table->GetName());
         }
         param.m_type = PARAMETER_TYPE::T_INT;
         param.m_defaultValue.Set<ParamIntType>(0);
@@ -342,10 +327,10 @@ namespace mu
 			// See if there is a string column. If there is one, we will use it as names for the
 			// options. Only the first string column will be used.
 			int nameCol = -1;
-			int32 cols = pTable->GetPrivate()->Columns.Num();
+			int32 cols = CacheKey.Table->GetPrivate()->Columns.Num();
 			for (int32 c = 0; c < cols && nameCol < 0; ++c)
 			{
-				if (pTable->GetPrivate()->Columns[c].Type == ETableColumnType::String)
+				if (CacheKey.Table->GetPrivate()->Columns[c].Type == ETableColumnType::String)
 				{
 					nameCol = c;
 				}
@@ -361,15 +346,15 @@ namespace mu
 			}
 
 			// Add every row
-			int32 rows = pTable->GetPrivate()->Rows.Num();
+			int32 rows = CacheKey.Table->GetPrivate()->Rows.Num();
 			for (size_t i = 0; i < rows; ++i)
 			{
 				FParameterDesc::FIntValueDesc value;
-				value.m_value = (int16_t)pTable->GetPrivate()->Rows[i].Id;
+				value.m_value = (int16)CacheKey.Table->GetPrivate()->Rows[i].Id;
 
 				if (nameCol > -1)
 				{
-					value.m_name = TCHAR_TO_ANSI(*pTable->GetPrivate()->Rows[i].Values[nameCol].String);
+					value.m_name = TCHAR_TO_ANSI(*CacheKey.Table->GetPrivate()->Rows[i].Values[nameCol].String);
 				}
 
 				param.m_possibleValues.Add(value);
@@ -386,6 +371,8 @@ namespace mu
 		op->type = OP_TYPE::NU_PARAMETER;
 		op->parameter = param;
 
+		m_firstPass.ParameterNodes.Add(InNode, op);
+
 		return op;
 	}
 
diff --git a/Engine/Plugins/Experimental/Mutable/Source/MutableTools/Private/MuT/CodeGenerator.h b/Engine/Plugins/Experimental/Mutable/Source/MutableTools/Private/MuT/CodeGenerator.h
--- a/Engine/Plugins/Experimental/Mutable/Source/MutableTools/Private/MuT/CodeGenerator.h
+++ b/Engine/Plugins/Experimental/Mutable/Source/MutableTools/Private/MuT/CodeGenerator.h
@@ -272,10 +272,27 @@ namespace mu
         };
 		TArray<FObjectGenerationData> m_currentObject;
 
-		TMap< TPair<Ptr<Table>, FString>, TPair<Ptr<Table>,Ptr<ASTOp>> > GeneratedTables;
+		/** The key for generated tables is made of the source table and a parameter name. */
+		struct FTableCacheKey
+		{
+			Ptr<const Table> Table;
+			FString ParameterName;
+
+			friend FORCEINLINE uint32 GetTypeHash(const FTableCacheKey& InKey)
+			{
+				uint32 KeyHash = ::GetTypeHash(InKey.Table.get());
+				KeyHash = HashCombineFast(KeyHash, GetTypeHash(InKey.ParameterName));
+				return KeyHash;
+			}
 
-        //! Variables added for every node
-		TMap< Ptr<const Node>, Ptr<ASTOpParameter> > NodeVariables;
+			FORCEINLINE bool operator==(const FTableCacheKey& InKey) const
+			{
+				if (Table != InKey.Table) return false;
+				if (ParameterName != InKey.ParameterName) return false;
+				return true;
+			}
+		};
+		TMap< FTableCacheKey, Ptr<ASTOp> > GeneratedTables;
 
 		struct FConditionalExtensionDataOp
 		{
@@ -298,7 +315,7 @@ namespace mu
 
         //-----------------------------------------------------------------------------------------
         //!
-        Ptr<ASTOp> GenerateTableVariable( TablePtr pTable, const FString& strName, bool bAddNoneOption );
+        Ptr<ASTOp> GenerateTableVariable(Ptr<const Node>, const FTableCacheKey&, bool bAddNoneOption);
 
         //!
         Ptr<ASTOp> GenerateMissingBoolCode(const TCHAR* strWhere, bool value, const void* errorContext );
@@ -713,43 +730,41 @@ namespace mu
     template<class NODE_TABLE_PRIVATE, ETableColumnType TYPE, OP_TYPE OPTYPE, typename F>
     Ptr<ASTOp> CodeGenerator::GenerateTableSwitch( const NODE_TABLE_PRIVATE& node, F&& GenerateOption )
     {
-        Ptr<Table> Table;
+        Ptr<const Table> NodeTable = node.Table;
         Ptr<ASTOp> Variable;
 
-        TPair<TablePtr,Ptr<ASTOp>>* it = GeneratedTables.Find( TPair<TablePtr,FString>(node.Table,node.ParameterName) );
+		FTableCacheKey CacheKey(node.Table, node.ParameterName);
+        Ptr<ASTOp>* it = GeneratedTables.Find( CacheKey );
         if ( it )
         {
-            Table = it->Key;
-            Variable = it->Value;
+            Variable = *it;
         }
 
-        if ( !Table )
+        if ( !Variable)
         {
             // Create the table variable expression
-            Table = node.Table;
-            Variable = GenerateTableVariable( Table, node.ParameterName, node.bNoneOption);
+            Variable = GenerateTableVariable( node.m_pNode, CacheKey, node.bNoneOption);
 
-            GeneratedTables.Add( TPair<TablePtr, FString>(node.Table,node.ParameterName), TPair<TablePtr,Ptr<ASTOp>>( Table, Variable ) );
+            GeneratedTables.Add(CacheKey, Variable );
         }
 
         // Verify that the table column is the right type
-        int32 ColIndex = Table->FindColumn( node.ColumnName );
+        int32 ColIndex = NodeTable->FindColumn( node.ColumnName );
         if ( ColIndex<0 )
         {
             m_pErrorLog->GetPrivate()->Add("Table column not found.", ELMT_ERROR, node.m_errorContext);
             return nullptr;
         }
 
-        if ( Table->GetPrivate()->Columns[ ColIndex ].Type != TYPE )
+        if (NodeTable->GetPrivate()->Columns[ ColIndex ].Type != TYPE )
         {
-            m_pErrorLog->GetPrivate()->Add("Table column type is not the right type.",
-                                           ELMT_ERROR, node.m_errorContext);
+            m_pErrorLog->GetPrivate()->Add("Table column type is not the right type.", ELMT_ERROR, node.m_errorContext);
             return nullptr;
         }
 
         // Create the switch to cover all the options
         Ptr<ASTOp> lastSwitch;
-        int32 Rows = Table->GetPrivate()->Rows.Num();
+        int32 Rows = NodeTable->GetPrivate()->Rows.Num();
 
         Ptr<ASTOpSwitch> SwitchOp = new ASTOpSwitch();
 		SwitchOp->type = OPTYPE;
@@ -758,8 +773,8 @@ namespace mu
 
 		for (int32 i = 0; i < Rows; ++i)
         {
-            check( Table->GetPrivate()->Rows[i].Id <= 0xFFFF);
-            auto Condition = (uint16)Table->GetPrivate()->Rows[i].Id;
+            check(NodeTable->GetPrivate()->Rows[i].Id <= 0xFFFF);
+            auto Condition = (uint16)NodeTable->GetPrivate()->Rows[i].Id;
             Ptr<ASTOp> Branch = GenerateOption( node, ColIndex, (int)i, m_pErrorLog.get() );
 
 			if (Branch || TYPE != ETableColumnType::Mesh)
diff --git a/Engine/Plugins/Experimental/Mutable/Source/MutableTools/Private/MuT/CodeGenerator_Bool.cpp b/Engine/Plugins/Experimental/Mutable/Source/MutableTools/Private/MuT/CodeGenerator_Bool.cpp
--- a/Engine/Plugins/Experimental/Mutable/Source/MutableTools/Private/MuT/CodeGenerator_Bool.cpp
+++ b/Engine/Plugins/Experimental/Mutable/Source/MutableTools/Private/MuT/CodeGenerator_Bool.cpp
@@ -87,7 +87,7 @@ namespace mu
 
 		Ptr<ASTOpParameter> op;
 
-		Ptr<ASTOpParameter>* it = NodeVariables.Find(node.m_pNode);
+		Ptr<ASTOpParameter>* it = m_firstPass.ParameterNodes.Find(node.m_pNode);
 		if (!it)
 		{
 			FParameterDesc param;
@@ -109,7 +109,7 @@ namespace mu
 				op->ranges.Emplace(op.get(), rangeResult.sizeOp, rangeResult.rangeName, rangeResult.rangeUID);
 			}
 
-			NodeVariables.Add(node.m_pNode,op);
+			m_firstPass.ParameterNodes.Add(node.m_pNode,op);
 		}
 		else
 		{
diff --git a/Engine/Plugins/Experimental/Mutable/Source/MutableTools/Private/MuT/CodeGenerator_Color.cpp b/Engine/Plugins/Experimental/Mutable/Source/MutableTools/Private/MuT/CodeGenerator_Color.cpp
--- a/Engine/Plugins/Experimental/Mutable/Source/MutableTools/Private/MuT/CodeGenerator_Color.cpp
+++ b/Engine/Plugins/Experimental/Mutable/Source/MutableTools/Private/MuT/CodeGenerator_Color.cpp
@@ -147,7 +147,7 @@ namespace mu
 
 		Ptr<ASTOpParameter> op;
 
-		Ptr<ASTOpParameter>* it = NodeVariables.Find(node.m_pNode);
+		Ptr<ASTOpParameter>* it = m_firstPass.ParameterNodes.Find(node.m_pNode);
 
 		if (!it)
 		{
@@ -177,7 +177,7 @@ namespace mu
 				op->ranges.Emplace(op.get(), rangeResult.sizeOp, rangeResult.rangeName, rangeResult.rangeUID);
 			}
 
-			NodeVariables.Add(node.m_pNode, op);
+			m_firstPass.ParameterNodes.Add(node.m_pNode, op);
 		}
 		else
 		{
diff --git a/Engine/Plugins/Experimental/Mutable/Source/MutableTools/Private/MuT/CodeGenerator_FirstPass.cpp b/Engine/Plugins/Experimental/Mutable/Source/MutableTools/Private/MuT/CodeGenerator_FirstPass.cpp
--- a/Engine/Plugins/Experimental/Mutable/Source/MutableTools/Private/MuT/CodeGenerator_FirstPass.cpp
+++ b/Engine/Plugins/Experimental/Mutable/Source/MutableTools/Private/MuT/CodeGenerator_FirstPass.cpp
@@ -594,6 +594,8 @@ namespace mu
                 op->parameter.m_possibleValues.Add( nullValue );
             }
 
+			ParameterNodes.Add(InNode, op);
+
             enumOp = op;
         }
 
@@ -601,7 +603,7 @@ namespace mu
         // Parse the child objects
 		for ( int32 t=0; t< Private->m_children.Num(); ++t )
         {
-            if ( const NodeObject* pChildNode = Private->m_children[t].get() )
+            if ( mu::Ptr<const NodeObject> pChildNode = Private->m_children[t] )
             {
                 // Overwrite the implicit condition
                 Ptr<ASTOp> paramOp = 0;
@@ -619,6 +621,8 @@ namespace mu
                         op->parameter.m_type = PARAMETER_TYPE::T_BOOL;
                         op->parameter.m_defaultValue.Set<ParamBoolType>(false);
 
+						ParameterNodes.Add(pChildNode, op);
+
                         paramOp = op;
                         break;
                     }
@@ -672,7 +676,7 @@ namespace mu
                 data.objectCondition = paramOp;
                 m_currentCondition.Add( data );
 
-				Generate_Generic(pChildNode);
+				Generate_Generic(pChildNode.get());
 
                 m_currentCondition.Pop();
             }
diff --git a/Engine/Plugins/Experimental/Mutable/Source/MutableTools/Private/MuT/CodeGenerator_FirstPass.h b/Engine/Plugins/Experimental/Mutable/Source/MutableTools/Private/MuT/CodeGenerator_FirstPass.h
--- a/Engine/Plugins/Experimental/Mutable/Source/MutableTools/Private/MuT/CodeGenerator_FirstPass.h
+++ b/Engine/Plugins/Experimental/Mutable/Source/MutableTools/Private/MuT/CodeGenerator_FirstPass.h
@@ -23,6 +23,7 @@
 #include "MuT/NodeSurfaceNew.h"
 #include "MuT/NodeSurfaceVariation.h"
 #include "MuT/NodeSurfaceSwitch.h"
+#include "MuT/ASTOpParameter.h"
 
 namespace mu
 {
@@ -173,6 +174,9 @@ namespace mu
         typedef TArray< TPair<FObjectState, const Node::Private*> > StateList;
         StateList m_states;
 
+		/** Parameters added for every node. */
+		TMap< Ptr<const Node>, Ptr<ASTOpParameter> > ParameterNodes;
+
 	private:
 
         struct FConditionContext
diff --git a/Engine/Plugins/Experimental/Mutable/Source/MutableTools/Private/MuT/CodeGenerator_Image.cpp b/Engine/Plugins/Experimental/Mutable/Source/MutableTools/Private/MuT/CodeGenerator_Image.cpp
--- a/Engine/Plugins/Experimental/Mutable/Source/MutableTools/Private/MuT/CodeGenerator_Image.cpp
+++ b/Engine/Plugins/Experimental/Mutable/Source/MutableTools/Private/MuT/CodeGenerator_Image.cpp
@@ -266,7 +266,7 @@ namespace mu
 
         Ptr<ASTOpParameter> op;
 
-		Ptr<ASTOpParameter>* it = NodeVariables.Find( node.m_pNode );
+		Ptr<ASTOpParameter>* it = m_firstPass.ParameterNodes.Find( node.m_pNode );
         if ( !it )
         {
             op = new ASTOpParameter();
@@ -286,7 +286,7 @@ namespace mu
 				op->ranges.Emplace(op.get(), rangeResult.sizeOp, rangeResult.rangeName, rangeResult.rangeUID);
 			}
 
-			NodeVariables.Add(node.m_pNode, op);
+			m_firstPass.ParameterNodes.Add(node.m_pNode, op);
 		}
         else
         {
diff --git a/Engine/Plugins/Experimental/Mutable/Source/MutableTools/Private/MuT/CodeGenerator_Projector.cpp b/Engine/Plugins/Experimental/Mutable/Source/MutableTools/Private/MuT/CodeGenerator_Projector.cpp
--- a/Engine/Plugins/Experimental/Mutable/Source/MutableTools/Private/MuT/CodeGenerator_Projector.cpp
+++ b/Engine/Plugins/Experimental/Mutable/Source/MutableTools/Private/MuT/CodeGenerator_Projector.cpp
@@ -109,7 +109,7 @@ namespace mu
 
 		Ptr<ASTOpParameter> op;
 
-		Ptr<ASTOpParameter>* it = NodeVariables.Find(node.m_pNode);
+		Ptr<ASTOpParameter>* it = m_firstPass.ParameterNodes.Find(node.m_pNode);
 		if (!it)
 		{
 			FParameterDesc param;
@@ -133,7 +133,7 @@ namespace mu
 				op->ranges.Emplace(op.get(), rangeResult.sizeOp, rangeResult.rangeName, rangeResult.rangeUID);
 			}
 
-			NodeVariables.Add(node.m_pNode, op);
+			m_firstPass.ParameterNodes.Add(node.m_pNode, op);
 		}
 		else
 		{
diff --git a/Engine/Plugins/Experimental/Mutable/Source/MutableTools/Private/MuT/CodeGenerator_Scalar.cpp b/Engine/Plugins/Experimental/Mutable/Source/MutableTools/Private/MuT/CodeGenerator_Scalar.cpp
--- a/Engine/Plugins/Experimental/Mutable/Source/MutableTools/Private/MuT/CodeGenerator_Scalar.cpp
+++ b/Engine/Plugins/Experimental/Mutable/Source/MutableTools/Private/MuT/CodeGenerator_Scalar.cpp
@@ -139,7 +139,7 @@ namespace mu
 
 		Ptr<ASTOpParameter> op;
 
-		Ptr<ASTOpParameter>* it = NodeVariables.Find(node.m_pNode);
+		Ptr<ASTOpParameter>* it = m_firstPass.ParameterNodes.Find(node.m_pNode);
 		if (!it)
 		{
 			FParameterDesc param;
@@ -161,7 +161,7 @@ namespace mu
 				op->ranges.Emplace(op.get(), rangeResult.sizeOp, rangeResult.rangeName, rangeResult.rangeUID);
 			}
 
-			NodeVariables.Add(node.m_pNode, op);
+			m_firstPass.ParameterNodes.Add(node.m_pNode, op);
 		}
 		else
 		{
@@ -179,7 +179,7 @@ namespace mu
 
 		Ptr<ASTOpParameter> op;
 
-		Ptr<ASTOpParameter>* it = NodeVariables.Find(node.m_pNode);
+		Ptr<ASTOpParameter>* it = m_firstPass.ParameterNodes.Find(node.m_pNode);
 		if (!it)
 		{
 			FParameterDesc param;
@@ -208,7 +208,7 @@ namespace mu
 				op->ranges.Emplace(op.get(), rangeResult.sizeOp, rangeResult.rangeName, rangeResult.rangeUID);
 			}
 
-			NodeVariables.Add(node.m_pNode, op);
+			m_firstPass.ParameterNodes.Add(node.m_pNode, op);
 		}
 		else
 		{
diff --git a/Engine/Plugins/Experimental/Mutable/Source/MutableTools/Private/MuT/CodeGenerator_String.cpp b/Engine/Plugins/Experimental/Mutable/Source/MutableTools/Private/MuT/CodeGenerator_String.cpp
--- a/Engine/Plugins/Experimental/Mutable/Source/MutableTools/Private/MuT/CodeGenerator_String.cpp
+++ b/Engine/Plugins/Experimental/Mutable/Source/MutableTools/Private/MuT/CodeGenerator_String.cpp
@@ -82,7 +82,7 @@ namespace mu
 
 		Ptr<ASTOpParameter> op;
 
-		Ptr<ASTOpParameter>* it = NodeVariables.Find(node.m_pNode);
+		Ptr<ASTOpParameter>* it = m_firstPass.ParameterNodes.Find(node.m_pNode);
 		if (!it)
 		{
 			FParameterDesc param;
@@ -104,7 +104,7 @@ namespace mu
 				op->ranges.Emplace(op.get(), rangeResult.sizeOp, rangeResult.rangeName, rangeResult.rangeUID);
 			}
 
-			NodeVariables.Add(node.m_pNode, op);
+			m_firstPass.ParameterNodes.Add(node.m_pNode, op);
 		}
 		else
 		{
diff --git a/Engine/Plugins/Experimental/Mutable/Source/MutableTools/Private/MuT/Compiler.cpp b/Engine/Plugins/Experimental/Mutable/Source/MutableTools/Private/MuT/Compiler.cpp
--- a/Engine/Plugins/Experimental/Mutable/Source/MutableTools/Private/MuT/Compiler.cpp
+++ b/Engine/Plugins/Experimental/Mutable/Source/MutableTools/Private/MuT/Compiler.cpp
@@ -249,6 +249,7 @@ namespace mu
 
         TArray< FStateCompilationData > states;
         Ptr<ErrorLog> genErrorLog;
+		TArray<FParameterDesc> Parameters;
         {
             CodeGenerator gen( m_pD->m_options->GetPrivate() );
 
@@ -266,7 +267,31 @@ namespace mu
             }
 
             genErrorLog = gen.m_pErrorLog;
-        }
+
+			// Set the parameter list from the non-optimized data, so that we have them all even if they are optimized out
+			int32 ParameterCount = gen.m_firstPass.ParameterNodes.Num();
+			Parameters.SetNum(ParameterCount);
+			int32 ParameterIndex = 0;
+			for ( const TPair<Ptr<const Node>, Ptr<ASTOpParameter>>& Entry : gen.m_firstPass.ParameterNodes )
+			{
+				Parameters[ParameterIndex] = Entry.Value->parameter;
+				++ParameterIndex;
+			}
+
+			// Sort the parameters as deterministically as possible.
+			struct FParameterSortPredicate
+			{
+				bool operator()(const FParameterDesc& A, const FParameterDesc& B) const
+				{
+					if (A.m_name < B.m_name) return true;
+					if (A.m_name > B.m_name) return false;
+					return A.m_uid < B.m_uid;
+				}
+			};
+			
+			FParameterSortPredicate SortPredicate;
+			Parameters.Sort(SortPredicate);
+		}
 
 
         // Slow AST code verification for debugging.
@@ -283,11 +308,13 @@ namespace mu
             optimiser.OptimiseAST( );
         }
 
-
         // Link the program and generate state data.
 		TSharedPtr<Model> pResult = MakeShared<Model>();
         FProgram& program = pResult->GetPrivate()->m_program;
 
+		check(program.m_parameters.IsEmpty());
+		program.m_parameters = Parameters;
+
 		// Preallocate ample memory
 		program.m_byteCode.Reserve(16 * 1024 * 1024);
 		program.m_opAddress.Reserve(1024 * 1024);
