[mutable] Fix compilation for some platforms.
#rnx

diff --git a/Engine/Plugins/Experimental/Mutable/Source/MutableTools/Private/MuT/CodeGenerator.h b/Engine/Plugins/Experimental/Mutable/Source/MutableTools/Private/MuT/CodeGenerator.h
--- a/Engine/Plugins/Experimental/Mutable/Source/MutableTools/Private/MuT/CodeGenerator.h
+++ b/Engine/Plugins/Experimental/Mutable/Source/MutableTools/Private/MuT/CodeGenerator.h
@@ -733,7 +733,7 @@ namespace mu
         Ptr<const Table> NodeTable = node.Table;
         Ptr<ASTOp> Variable;
 
-		FTableCacheKey CacheKey(node.Table, node.ParameterName);
+		FTableCacheKey CacheKey = FTableCacheKey{ node.Table, node.ParameterName };
         Ptr<ASTOp>* it = GeneratedTables.Find( CacheKey );
         if ( it )
         {
