[mutable] Fixed the detection of whether Mutable textures should be streamed or not depending on the compilation optimization level. This could lead to deadlocks when opening the texture editor or in the MRQ.
#rnx
#tests Editor

diff --git a/Engine/Plugins/Experimental/Mutable/Source/CustomizableObject/Private/MuCO/CustomizableObjectSystem.cpp b/Engine/Plugins/Experimental/Mutable/Source/CustomizableObject/Private/MuCO/CustomizableObjectSystem.cpp
--- a/Engine/Plugins/Experimental/Mutable/Source/CustomizableObject/Private/MuCO/CustomizableObjectSystem.cpp
+++ b/Engine/Plugins/Experimental/Mutable/Source/CustomizableObject/Private/MuCO/CustomizableObjectSystem.cpp
@@ -966,12 +966,12 @@ void UCustomizableObjectSystemPrivate::GetMipStreamingConfig(const UCustomizable
 	if (State)
 	{
 		bOutNeverStream = State->bDisableTextureStreaming;
+	}
 
-		// Was streaming disabled at object-compilation time? 
-		if (Instance.GetCustomizableObject()->GetPrivate()->bDisableTextureStreaming)
-		{
-			bOutNeverStream = true;
-		}
+	// Was streaming disabled at object-compilation time? 
+	if (Instance.GetCustomizableObject()->GetPrivate()->bDisableTextureStreaming)
+	{
+		bOutNeverStream = true;
 	}
 
 	OutMipsToSkip = 0; // 0 means generate all mips
diff --git a/Engine/Plugins/Experimental/Mutable/Source/CustomizableObjectEditor/Private/MuCOE/CustomizableObjectCompiler.cpp b/Engine/Plugins/Experimental/Mutable/Source/CustomizableObjectEditor/Private/MuCOE/CustomizableObjectCompiler.cpp
--- a/Engine/Plugins/Experimental/Mutable/Source/CustomizableObjectEditor/Private/MuCOE/CustomizableObjectCompiler.cpp
+++ b/Engine/Plugins/Experimental/Mutable/Source/CustomizableObjectEditor/Private/MuCOE/CustomizableObjectCompiler.cpp
@@ -1031,6 +1031,10 @@ void FCustomizableObjectCompiler::CompileInternal(UCustomizableObject* Object, c
 			// unreal assets and progressive generation is not supported.
 			Object->GetPrivate()->bDisableTextureStreaming = true;
 		}
+		else
+		{
+			Object->GetPrivate()->bDisableTextureStreaming = false;
+		}
 		
 		Object->GetPrivate()->bIsCompiledWithOptimization = GenerationContext.Options.OptimizationLevel < UE_MUTABLE_MAX_OPTIMIZATION;
 
