[Mutable] Add bDisableTextureStreaming to the derived data.

#rnx

diff --git a/Engine/Plugins/Experimental/Mutable/Source/CustomizableObject/Internal/MuCO/CustomizableObjectPrivate.h b/Engine/Plugins/Experimental/Mutable/Source/CustomizableObject/Internal/MuCO/CustomizableObjectPrivate.h
--- a/Engine/Plugins/Experimental/Mutable/Source/CustomizableObject/Internal/MuCO/CustomizableObjectPrivate.h
+++ b/Engine/Plugins/Experimental/Mutable/Source/CustomizableObject/Internal/MuCO/CustomizableObjectPrivate.h
@@ -691,6 +691,9 @@ public:
 	/** If the object is compiled, this flag is true if it was compiled with maximum optimizations. If the object is not compiled, its value is meaningless. */
 	bool bIsCompiledWithOptimization = true;
 
+	/** This is a non-user-controlled flag to disable streaming (set at object compilation time, depending on optimization). */
+	bool bDisableTextureStreaming = false;
+
 	ECustomizableObjectCompilationState CompilationState = ECustomizableObjectCompilationState::None;
 	
 #if WITH_EDITOR
@@ -722,9 +725,6 @@ public:
 	// This is a manual version number for the binary blobs in this asset.
 	// Increasing it invalidates all the previously compiled models.
 	// Warning: If while merging code both versions have changed, take the highest+1.
-	static constexpr int32 CurrentSupportedVersion = 424;
-
-	/** This is a non-user-controlled flag to disable streaming (set at object compilation time, depending on optimization). */
-	bool bDisableTextureStreaming = false;
+	static constexpr int32 CurrentSupportedVersion = 425;
 };
 
diff --git a/Engine/Plugins/Experimental/Mutable/Source/CustomizableObject/Private/MuCO/CustomizableObject.cpp b/Engine/Plugins/Experimental/Mutable/Source/CustomizableObject/Private/MuCO/CustomizableObject.cpp
--- a/Engine/Plugins/Experimental/Mutable/Source/CustomizableObject/Private/MuCO/CustomizableObject.cpp
+++ b/Engine/Plugins/Experimental/Mutable/Source/CustomizableObject/Private/MuCO/CustomizableObject.cpp
@@ -564,6 +564,7 @@ void UCustomizableObjectPrivate::SaveCompiledData(FArchive& MemoryWriter, bool b
 	MemoryWriter << GetPublic()->LODSettings.bLODStreamingEnabled;
 
 	// Editor Only data
+	MemoryWriter << bDisableTextureStreaming;
 	MemoryWriter << bIsCompiledWithOptimization;
 	MemoryWriter << CustomizableObjectPathMap;
 	MemoryWriter << GroupNodeMap;
@@ -685,6 +686,7 @@ void UCustomizableObjectPrivate::LoadCompiledData(FArchive& MemoryReader, const
 
 		// Editor Only data
 		{
+			MemoryReader << bDisableTextureStreaming;
 			MemoryReader << bIsCompiledWithOptimization;
 			MemoryReader << CustomizableObjectPathMap;
 			MemoryReader << GroupNodeMap;
@@ -989,9 +991,6 @@ void UCustomizableObjectPrivate::SaveEmbeddedData(FArchive& Ar)
 
 	if (GetModel())
 	{
-		// General derived flags
-		Ar << bDisableTextureStreaming;
-
 		// Serialize morph data
 		{
 			Ar << GetPublic()->ContributingMorphTargetsInfo;
@@ -1032,9 +1031,6 @@ void UCustomizableObjectPrivate::LoadEmbeddedData(FArchive& Ar)
 
 	if(CurrentSupportedVersion == InternalVersion)
 	{
-		// General derived flags
-		Ar << bDisableTextureStreaming;
-
 		// Load morph data
 		{
 			Ar << GetPublic()->ContributingMorphTargetsInfo;
diff --git a/Engine/Plugins/Experimental/Mutable/Source/CustomizableObject/Private/MuCO/CustomizableObjectSystem.cpp b/Engine/Plugins/Experimental/Mutable/Source/CustomizableObject/Private/MuCO/CustomizableObjectSystem.cpp
--- a/Engine/Plugins/Experimental/Mutable/Source/CustomizableObject/Private/MuCO/CustomizableObjectSystem.cpp
+++ b/Engine/Plugins/Experimental/Mutable/Source/CustomizableObject/Private/MuCO/CustomizableObjectSystem.cpp
@@ -968,11 +968,13 @@ void UCustomizableObjectSystemPrivate::GetMipStreamingConfig(const UCustomizable
 		bOutNeverStream = State->bDisableTextureStreaming;
 	}
 
+#if WITH_EDITORONLY_DATA
 	// Was streaming disabled at object-compilation time? 
 	if (Instance.GetCustomizableObject()->GetPrivate()->bDisableTextureStreaming)
 	{
 		bOutNeverStream = true;
 	}
+#endif
 
 	OutMipsToSkip = 0; // 0 means generate all mips
 
