[Mutable] Fix invalid access to mu::Image in CodeGenerator::GenerateImage_Table if a table node has an empty column.

#rnx

diff --git a/Engine/Plugins/Experimental/Mutable/Source/MutableTools/Private/MuT/CodeGenerator_Image.cpp b/Engine/Plugins/Experimental/Mutable/Source/MutableTools/Private/MuT/CodeGenerator_Image.cpp
--- a/Engine/Plugins/Experimental/Mutable/Source/MutableTools/Private/MuT/CodeGenerator_Image.cpp
+++ b/Engine/Plugins/Experimental/Mutable/Source/MutableTools/Private/MuT/CodeGenerator_Image.cpp
@@ -1893,19 +1893,19 @@ namespace mu
 					FImageGenerationResult Result;
 					GenerateImage(Options, Result, ImageConst);
 					ImageOp = Result.op;
-				}
-
-				int32 MaxTextureSize = FMath::Max(node.ReferenceImageDesc.m_size[0], node.ReferenceImageDesc.m_size[1]);
 
-				if (MaxTextureSize > 0 && (MaxTextureSize < pImage->GetSizeX() || MaxTextureSize < pImage->GetSizeY()))
-				{
-					float Factor = FMath::Min(MaxTextureSize / (float)(pImage->GetSizeX()), MaxTextureSize / (float)(pImage->GetSizeY()));
-					Ptr<ASTOpFixed> op = new ASTOpFixed();
-					op->op.type = OP_TYPE::IM_RESIZE;
-					op->op.args.ImageResize.size[0] = (uint16)pImage->GetSizeX() * Factor;
-					op->op.args.ImageResize.size[1] = (uint16)pImage->GetSizeY() * Factor;
-					op->SetChild(op->op.args.ImageResize.source, ImageOp);
-					ImageOp = op;
+					int32 MaxTextureSize = FMath::Max(node.ReferenceImageDesc.m_size[0], node.ReferenceImageDesc.m_size[1]);
+
+					if (MaxTextureSize > 0 && (MaxTextureSize < pImage->GetSizeX() || MaxTextureSize < pImage->GetSizeY()))
+					{
+						float Factor = FMath::Min(MaxTextureSize / (float)(pImage->GetSizeX()), MaxTextureSize / (float)(pImage->GetSizeY()));
+						Ptr<ASTOpFixed> op = new ASTOpFixed();
+						op->op.type = OP_TYPE::IM_RESIZE;
+						op->op.args.ImageResize.size[0] = (uint16)pImage->GetSizeX() * Factor;
+						op->op.args.ImageResize.size[1] = (uint16)pImage->GetSizeY() * Factor;
+						op->SetChild(op->op.args.ImageResize.source, ImageOp);
+						ImageOp = op;
+					}
 				}
 
 				return ImageOp;
