[Mutable] Invalidate mu::model and pending instance updates before compiling CustomizableObjects. Check if the model is valid in UCustomizableObject::IsCompiled.
- Remove the early AdvanceTo clean-up step in the FMutableTextureMipDataProvider::Init method that would trigger a check if done.

#rnx

diff --git a/Engine/Plugins/Experimental/Mutable/Source/CustomizableObject/Internal/MuCO/CustomizableObjectSystemPrivate.h b/Engine/Plugins/Experimental/Mutable/Source/CustomizableObject/Internal/MuCO/CustomizableObjectSystemPrivate.h
--- a/Engine/Plugins/Experimental/Mutable/Source/CustomizableObject/Internal/MuCO/CustomizableObjectSystemPrivate.h
+++ b/Engine/Plugins/Experimental/Mutable/Source/CustomizableObject/Internal/MuCO/CustomizableObjectSystemPrivate.h
@@ -142,6 +142,10 @@ public:
 	// Removes an instance update
 	void RemoveUpdate(const TWeakObjectPtr<UCustomizableObjectInstance>& Instance);
 
+#if WITH_EDITOR
+	void RemoveUpdatesForObject(const UCustomizableObject* InObject);
+#endif
+
 	const FMutablePendingInstanceUpdate* GetUpdate(const TWeakObjectPtr<const UCustomizableObjectInstance>& Instance) const;
 
 	TSet<FMutablePendingInstanceUpdate, FPendingInstanceUpdateKeyFuncs>::TIterator GetUpdateIterator()
diff --git a/Engine/Plugins/Experimental/Mutable/Source/CustomizableObject/Private/MuCO/CustomizableObject.cpp b/Engine/Plugins/Experimental/Mutable/Source/CustomizableObject/Private/MuCO/CustomizableObject.cpp
--- a/Engine/Plugins/Experimental/Mutable/Source/CustomizableObject/Private/MuCO/CustomizableObject.cpp
+++ b/Engine/Plugins/Experimental/Mutable/Source/CustomizableObject/Private/MuCO/CustomizableObject.cpp
@@ -1065,9 +1065,13 @@ UCustomizableObjectPrivate* UCustomizableObject::GetPrivate() const
 
 bool UCustomizableObject::IsCompiled() const
 {
-	bool IsCompiled = Private->GetModel() != nullptr;
+#if WITH_EDITOR
+	const bool bIsCompiled = Private->GetModel() != nullptr && Private->GetModel()->IsValid();
+#else
+	const bool bIsCompiled = Private->GetModel() != nullptr;
+#endif
 
-	return IsCompiled;
+	return bIsCompiled;
 }
 
 
diff --git a/Engine/Plugins/Experimental/Mutable/Source/CustomizableObject/Private/MuCO/CustomizableObjectSystem.cpp b/Engine/Plugins/Experimental/Mutable/Source/CustomizableObject/Private/MuCO/CustomizableObjectSystem.cpp
--- a/Engine/Plugins/Experimental/Mutable/Source/CustomizableObject/Private/MuCO/CustomizableObjectSystem.cpp
+++ b/Engine/Plugins/Experimental/Mutable/Source/CustomizableObject/Private/MuCO/CustomizableObjectSystem.cpp
@@ -333,6 +333,19 @@ void FMutablePendingInstanceWork::RemoveUpdate(const TWeakObjectPtr<UCustomizabl
 	}	
 }
 
+#if WITH_EDITOR
+void FMutablePendingInstanceWork::RemoveUpdatesForObject(const UCustomizableObject* InObject)
+{
+	check(InObject);
+	for (auto Iterator = PendingInstanceUpdates.CreateIterator(); Iterator; ++Iterator)
+	{
+		if (Iterator->Context->Instance.IsValid() && Iterator->Context->Instance->GetCustomizableObject() == InObject)
+		{
+			Iterator.RemoveCurrent();
+		}
+	}
+}
+#endif
 
 const FMutablePendingInstanceUpdate* FMutablePendingInstanceWork::GetUpdate(const TWeakObjectPtr<const UCustomizableObjectInstance>& Instance) const
 {
@@ -1349,6 +1362,12 @@ bool UCustomizableObjectSystem::LockObject(const class UCustomizableObject* InOb
 		// Lock the object, no new file or mip streaming operations should start from this point
 		InObject->GetPrivate()->bLocked = true;
 
+		// Invalidate the current model to avoid further disk or mip updates.
+		if (InObject->GetPrivate()->GetModel())
+		{
+			InObject->GetPrivate()->GetModel()->Invalidate();
+		}
+
 		// But some could have started between the first CheckIfDiskOrMipUpdateOperationsPending and the lock a few lines back, so check again
 		if (CheckIfDiskOrMipUpdateOperationsPending(*InObject))
 		{
@@ -1372,6 +1391,8 @@ bool UCustomizableObjectSystem::LockObject(const class UCustomizableObject* InOb
 			Task.Wait();
 		}
 
+		Private->MutablePendingInstanceWork.RemoveUpdatesForObject(InObject);
+
 		// Clear the cache for the instance, since we will remake it
 		FMutableResourceCache& Cache = GetPrivate()->GetObjectCache(InObject);
 		Cache.Clear();
