[Mutable] Moved UCustomizableObjectInstance private functions and members to private.

- Removed deprecated functions.
- Removed commented code.
- Removed unused BeginDestroyDelegate.
- Removed non-constant GetDescriptor.
- Changed SetProjectorPosition to FVector.
- Minor typos.
- Standardized doc formatting.

#rnx

diff --git a/Engine/Plugins/Experimental/Mutable/Source/CustomizableObject/Classes/MuCO/CustomizableObjectInstance.h b/Engine/Plugins/Experimental/Mutable/Source/CustomizableObject/Classes/MuCO/CustomizableObjectInstance.h
--- a/Engine/Plugins/Experimental/Mutable/Source/CustomizableObject/Classes/MuCO/CustomizableObjectInstance.h
+++ b/Engine/Plugins/Experimental/Mutable/Source/CustomizableObject/Classes/MuCO/CustomizableObjectInstance.h
@@ -6,7 +6,6 @@
 #include "MuCO/CustomizableObjectInstanceDescriptor.h"
 #include "Templates/SubclassOf.h"
 #include "Math/RandomStream.h"
-#include "MuR/Ptr.h"
 
 #include "CustomizableObjectInstance.generated.h"
 
@@ -99,9 +98,6 @@ DECLARE_MULTICAST_DELEGATE_OneParam(FInstanceUpdateNativeDelegate, const FUpdate
 DECLARE_DYNAMIC_MULTICAST_DELEGATE_OneParam(FObjectInstanceUpdatedDelegate, UCustomizableObjectInstance*, Instance);
 DECLARE_MULTICAST_DELEGATE_OneParam(FObjectInstanceUpdatedNativeDelegate, UCustomizableObjectInstance*);
 
-DECLARE_DYNAMIC_MULTICAST_DELEGATE_OneParam(FBeginDestroyDelegate, UCustomizableObjectInstance*, Instance);
-DECLARE_MULTICAST_DELEGATE_OneParam(FBeginDestroyNativeDelegate, UCustomizableObjectInstance*);
-
 DECLARE_DELEGATE_OneParam(FProjectorStateChangedDelegate, FString);
 
 DECLARE_DYNAMIC_DELEGATE_TwoParams(FEachComponentAnimInstanceClassDelegate, FName, SlotIndex, TSubclassOf<UAnimInstance>, AnimInstClass);
@@ -114,14 +110,13 @@ class CUSTOMIZABLEOBJECT_API UCustomizableObjectInstance : public UObject
 {
 	GENERATED_BODY()
 
+	// Friends
 	friend UCustomizableInstancePrivate;
 	friend FMutableUpdateCandidate;
 
 public:
 	UCustomizableObjectInstance();
-
-	FCustomizableObjectInstanceDescriptor& GetDescriptor();
-
+	
 	const FCustomizableObjectInstanceDescriptor& GetDescriptor() const;
 	
 	void SetDescriptor(const FCustomizableObjectInstanceDescriptor& InDescriptor);
@@ -133,22 +128,6 @@ public:
 	/** Broadcast when the Customizable Object Instance is updated. */
 	FObjectInstanceUpdatedNativeDelegate UpdatedNativeDelegate;
 
-	/** Broadcast when UObject::BeginDestroy is being called. */	
-	UPROPERTY(BlueprintAssignable, Category = CustomizableObjectInstance)
-	FBeginDestroyDelegate BeginDestroyDelegate;
-
-	/** Broadcast when UObject::BeginDestroy is being called. */
-	FBeginDestroyNativeDelegate BeginDestroyNativeDelegate;
-
-	TMap<FString, bool> ParamNameToExpandedMap; // Used to check whether a mutable param is expanded in the editor to show its child params
-
-#if WITH_EDITOR
-	/** During editor, always remember the duration of the last update in the mutable runtime, for profiling. */
-	int32 LastUpdateMutableRuntimeCycles = 0;
-#endif
-
-public:
-
 	// UObject interface.
 #if WITH_EDITOR
 	virtual void PostEditChangeProperty(FPropertyChangedEvent& PropertyChangedEvent) override;
@@ -157,7 +136,6 @@ public:
 	virtual bool IsCachedCookedPlatformDataLoaded(const ITargetPlatform* TargetPlatform) override;
 	virtual void PreSave(FObjectPreSaveContext ObjectSaveContext) override;
 #endif //WITH_EDITOR
-
 	virtual void Serialize(FArchive& Ar) override;
 	virtual void PostLoad() override;
 	virtual void BeginDestroy() override;
@@ -166,8 +144,7 @@ public:
 	virtual bool IsEditorOnly() const override;
 
 	/** Set the CustomizableObject this instance will be generated from. 
-	* It is usually not necessary to call this since instances are already generated from a CustomizableObject.
-	*/
+	  * It is usually not necessary to call this since instances are already generated from a CustomizableObject. */
 	UFUNCTION(BlueprintCallable, Category = CustomizableObjectInstance)
 	void SetObject(UCustomizableObject* InObject);
 
@@ -175,14 +152,6 @@ public:
 	UFUNCTION(BlueprintCallable, Category = CustomizableObjectInstance)
 	UCustomizableObject* GetCustomizableObject() const;
 
-	/** Deprecated. */
-	UFUNCTION(BlueprintCallable, Category = CustomizableObjectInstance, meta = (DeprecatedFunction, DeprecationMessage = "Parameter decorations have been removed. This method will be removed in future versions."))
-	bool GetBuildParameterDecorations() const;
-
-	/** Deprecated. */
-	UFUNCTION(BlueprintCallable, Category = CustomizableObjectInstance, meta = (DeprecatedFunction, DeprecationMessage = "Parameter decorations have been removed. This method will be removed in future versions."))
-	void SetBuildParameterDecorations(bool Value);
-
 	/** Return true if the parameter relevancy will be updated when this instance is generated. */
 	UFUNCTION(BlueprintCallable, Category = CustomizableObjectInstance)
 	bool GetBuildParameterRelevancy() const;
@@ -191,9 +160,6 @@ public:
 	UFUNCTION(BlueprintCallable, Category = CustomizableObjectInstance)
 	void SetBuildParameterRelevancy(bool Value);
 
-	int32 GetState() const;
-	void SetState(int32 InState);
-
 	/** Return the name of the current CustomizableObject state this is instance is set to. */
 	UFUNCTION(BlueprintCallable, Category = CustomizableObjectInstance)
 	FString GetCurrentState() const;
@@ -203,8 +169,7 @@ public:
 	void SetCurrentState(const FString& StateName);
 
 	/** Get the skeletal mesh generated for this instance. 
-	 * If the object has multiple components, an index of the component can be specified. 
-	 */
+	  * If the object has multiple components, an index of the component can be specified. */
 	UFUNCTION(BlueprintCallable, Category = CustomizableObjectInstance)
 	USkeletalMesh* GetSkeletalMesh(int32 ComponentIndex = 0) const;
 	
@@ -256,24 +221,16 @@ public:
 	UFUNCTION(BlueprintCallable, Category = CustomizableObjectInstance)
 	void SetRandomValues();
 
-	/**  Set random values to the parameters using a stream. */
+	/** Set random values to the parameters using a stream. */
 	UFUNCTION(BlueprintCallable, Category = CustomizableObjectInstance)
 	void SetRandomValuesFromStream(const FRandomStream& InStream);
 
 	/** Returns the AssetUserData that was gathered from all the constituent mesh parts during the last update. 
-	 *  It requires that the CustomizableObject had the bEnableAssetUserDataMerge set to true during compilation. 
-	 */
+	  * It requires that the CustomizableObject had the bEnableAssetUserDataMerge set to true during compilation. */
 	UFUNCTION(BlueprintCallable, Category = CustomizableObjectInstance)
 	TSet<UAssetUserData*> GetMergedAssetUserData(int32 ComponentIndex) const;
-
-	// Utilities to manage saving and loading parameters from profiles.
-	bool LoadParametersFromProfile(int32 ProfileIndex);
-	bool SaveParametersToProfile(int32 ProfileIndex);
-	bool MigrateProfileParametersToCurrentInstance(int32 ProfileIndex);
-	void SetSelectedParameterProfileDirty();
-	bool IsSelectedParameterProfileDirty() const;
-
-	/** Return true if the instance is not locked and if it's compiled */
+	
+	/** Return true if the instance is not locked and if it's compiled. */
 	bool CanUpdateInstance() const;
 
 	/** Generate the instance with the current parameters and update all the components Skeletal Meshes asynchronously. */
@@ -281,44 +238,33 @@ public:
 	void UpdateSkeletalMeshAsync(bool bIgnoreCloseDist = false, bool bForceHighPriority = false);
 		
 	/** Generate the instance with the current parameters and update all the components Skeletal Meshes asynchronously.
-	  * Callback will be called once the update finishes, even if it fails. 
-	  */
+	  * Callback will be called once the update finishes, even if it fails. */
 	UFUNCTION(BlueprintCallable, Category = CustomizableObjectInstance)
 	void UpdateSkeletalMeshAsyncResult(FInstanceUpdateDelegate Callback, bool bIgnoreCloseDist = false, bool bForceHighPriority = false);
 
 	void UpdateSkeletalMeshAsyncResult(FInstanceUpdateNativeDelegate Callback, bool bIgnoreCloseDist = false, bool bForceHighPriority = false);
 
-	
-	// Clones the instance creating a new identical transient instance.
+	/** Clones the instance creating a new identical transient instance. */
 	UFUNCTION(BlueprintCallable, Category = CustomizableObjectInstance)
 	UCustomizableObjectInstance* Clone();
 
-	// Clones the instance creating a new identical static instance with the given Outer.
+	/** Clones the instance creating a new identical static instance with the given Outer. */
 	UFUNCTION(BlueprintCallable, Category = CustomizableObjectInstance)
 	UCustomizableObjectInstance* CloneStatic(UObject* Outer);
 
-	// Copy parameters from input instance
+	/** Copy parameters from the given Instance. */
 	void CopyParametersFromInstance(UCustomizableObjectInstance* Instance);
 
 	/** Immediately destroy the Mutable Core Live Update Instance attached to this (if exists). */
 	void DestroyLiveUpdateInstance();
 
-	/** Deprecated. */
-	UFUNCTION(BlueprintCallable, Category = CustomizableObjectInstance, meta = (DeprecatedFunction, DeprecationMessage = "Parameter decorations have been removed. This method will be removed in future versions."))
-	UTexture2D* GetParameterDescription(const FString& ParamName, int32 DescIndex);
-
-	//! 
-	bool IsParameterRelevant( int32 ParameterIndex ) const;
+	/** Return true if changing the parameter would affect the Instance given its current generation. */
+	bool IsParameterRelevant(int32 ParameterIndex) const;
 
 	/** Return true if the given parameter has any effect in the current object state, and considering the current values of the other parameters. */
 	UFUNCTION(BlueprintCallable, Category = CustomizableObjectInstance)
 	bool IsParameterRelevant(const FString& ParamName) const;
 
-	/** Return true if the given parameter supports multiple dimensions, and so, multiple values. */
-	UFUNCTION(BlueprintCallable, Category = CustomizableObjectInstance, 
-		meta = (DeprecatedFunction, DeprecationMessage = "Use the method with the same name located at the Customizable Object instead."))
-	bool IsParamMultidimensional(const FString& ParamName) const;
-
 	/** For multidimensional parameters, return the number of dimensions that the given projector parameter supports. */
 	UFUNCTION(BlueprintCallable, Category = CustomizableObjectInstance)
 	int32 GetProjectorValueRange(const FString& ParamName) const;
@@ -334,18 +280,6 @@ public:
 	/** For multidimensional parameters, return the number of dimensions that the given texture parameter supports. */
 	UFUNCTION(BlueprintCallable, Category = CustomizableObjectInstance)
 	int32 GetTextureValueRange(const FString& ParamName) const;
-	
-	bool bShowOnlyRuntimeParameters = true;
-	bool bShowOnlyRelevantParameters = true;
-	bool bShowUISections = false;
-
-	// DEPRECATED: Use the method in the CustomizableObject instead which takes an index among all parameters
-	// Returns how many possible options an int parameter has
-	//int32 GetIntParameterNumOptions(int32 IntParamIndex);
-
-	// DEPRECATED: Use the method in the CustomizableObject instead which takes an index among all parameters
-	// Gets the Name of the option at position K in the list of available options for the int parameter. Useful to enumerate the int parameter's possible options (Ex: "Hat1", "Hat2", "Cap", "Nothing")
-	//const FString& GetIntParameterAvailableOption(int32 IntParamIndex, int32 K);
 
 	/** Return the name of the option currently set in the given parameter. */
 	UFUNCTION(BlueprintCallable, Category = CustomizableObjectInstance)
@@ -403,8 +337,8 @@ public:
 
 	/** Set only the projector position keeping the rest of values. */
 	UFUNCTION(BlueprintCallable, Category = CustomizableObjectInstance)
-	void SetProjectorPosition(const FString& ProjectorParamName, const FVector3f& Pos, int32 RangeIndex = -1); // TODO GMTFuture Make API consistent (FVector3f -> FVector)
-
+	void SetProjectorPosition(const FString& ProjectorParamName, const FVector& Pos, int32 RangeIndex = -1);
+	
 	/** Set only the projector direction vector keeping the rest of values. */
 	UFUNCTION(BlueprintCallable, Category = CustomizableObjectInstance)
 	void SetProjectorDirection(const FString& ProjectorParamName, const FVector& Direction, int32 RangeIndex = -1);
@@ -481,20 +415,17 @@ public:
 	int32 FindProjectorParameterNameIndex(const FString& ParamName) const;
 
 	/** Increases the range of values of the integer with ParamName and returns the index of the new integer value, -1 otherwise.
-	* The added value is initialized with the first integer option and is the last one of the range.
-	*/
+	  * The added value is initialized with the first integer option and is the last one of the range. */
 	UFUNCTION(BlueprintCallable, Category = CustomizableObjectInstance)
 	int32 AddValueToIntRange(const FString& ParamName);
 
 	/** Increases the range of values of the float with ParamName, returns the index of the new float value, -1 otherwise.
-	* The added value is initialized with 0.5f and is the last one of the range.
-	*/
+	  * The added value is initialized with 0.5f and is the last one of the range. */
 	UFUNCTION(BlueprintCallable, Category = CustomizableObjectInstance)
 	int32 AddValueToFloatRange(const FString& ParamName);
 
 	/** Increases the range of values of the projector with ParamName, returns the index of the new projector value, -1 otherwise.
-	* The added value is initialized with the default projector as set up in the editor and is the last one of the range.
-	*/
+	  * The added value is initialized with the default projector as set up in the editor and is the last one of the range. */
 	UFUNCTION(BlueprintCallable, Category = CustomizableObjectInstance)
 	int32 AddValueToProjectorRange(const FString& ParamName);
 
@@ -524,8 +455,7 @@ public:
 	// ------------------------------------------------------------
 	
 	/** Given Multilayer Projector name, create a new Multilayer Projector Helper (if non-existent). See FMultilayerProjector.
-	 * @return true if successfully created (or was already created).
-	 */
+	  * @return true if successfully created (or was already created). */
 	UFUNCTION(BlueprintCallable, Category = CustomizableObjectInstance)
 	bool CreateMultiLayerProjector(const FName& ProjectorParamName);
 	
@@ -591,17 +521,20 @@ public:
 	UFUNCTION(BlueprintCallable, Category = CustomizableObjectInstance)
 	const FGameplayTagContainer& GetAnimationGameplayTags() const;
 	
-	/** Execute a delegate for each animation instance invovled in this customizable object instance. */
+	/** Execute a delegate for each animation instance involved in this customizable object instance. */
 	UFUNCTION(BlueprintCallable, Category = CustomizableObjectInstance)
 	void ForEachAnimInstance(int32 ComponentIndex, FEachComponentAnimInstanceClassDelegate Delegate) const;
 
 	void ForEachAnimInstance(int32 ComponentIndex, FEachComponentAnimInstanceClassNativeDelegate Delegate) const;
 
+	/** Check if the given UAnimInstance class requires to be fixed up. */
 	bool AnimInstanceNeedsFixup(TSubclassOf<UAnimInstance> AnimInstance) const;
+
+	/** Fix the given UAnimInstance instance. */
 	void AnimInstanceFixup(UAnimInstance* AnimInstance) const;
 
 	/** See FCustomizableObjectInstanceDescriptor::SaveDescriptor. */
-	void SaveDescriptor(FArchive &CustomizableObjectDescriptor,bool bUseCompactDescriptor);
+	void SaveDescriptor(FArchive &CustomizableObjectDescriptor, bool bUseCompactDescriptor);
 
 	/** See FCustomizableObjectInstanceDescriptor::LoadDescriptor. */
 	void LoadDescriptor(FArchive &CustomizableObjectDescriptor);
@@ -610,20 +543,18 @@ public:
 	UFUNCTION(BlueprintCallable, Category = CustomizableObjectInstance)
 	void SetReplacePhysicsAssets(bool bReplaceEnabled);
 
-	// Enables the reuse of all possible textures when the instance is updated without any changes in geometry or state (the first update after creation doesn't reuse any)
-	// It will only work if the textures aren't compressed, so set the instance to a Mutable state with texture compression disabled
-	// WARNING! If texture reuse is enabled, do NOT keep external references to the textures of the instance. The instance owns the textures.
+	/** Enables the reuse of all possible textures when the instance is updated without any changes in geometry or state (the first update after creation doesn't reuse any)
+	  * It will only work if the textures aren't compressed, so set the instance to a Mutable state with texture compression disabled
+	  * WARNING! If texture reuse is enabled, do NOT keep external references to the textures of the instance. The instance owns the textures. */
 	void SetReuseInstanceTextures(bool bTextureReuseEnabled);
 	
-	// If enabled, low-priority textures will generate resident mipmaps too.
+	/** If enabled, low-priority textures will generate resident mipmaps too. */
 	void SetForceGenerateResidentMips(bool bForceGenerateResidentMips);
 
-	void AdditionalAssetsAsyncLoaded( FGraphEventRef CompletionEvent );
-
 	const TArray<TObjectPtr<UMaterialInterface>>* GetOverrideMaterials(int32 ComponentIndex) const;
 	
 	// The following methods should only be used in an LOD management class
-	void SetIsBeingUsedByComponentInPlay(bool bIsUsedByComponent );
+	void SetIsBeingUsedByComponentInPlay(bool bIsUsedByComponent);
 	bool GetIsBeingUsedByComponentInPlay() const;
 	void SetIsDiscardedBecauseOfTooManyInstances(bool bIsDiscarded);
 	bool GetIsDiscardedBecauseOfTooManyInstances() const;
@@ -635,50 +566,32 @@ public:
 	UFUNCTION(BlueprintCallable, Category = CustomizableObjectInstance)
 	int32 GetNumComponents() const;
 
+	/** Return the min LOD that will be used in the next Instance update. */
 	int32 GetMinLODToLoad() const;
+
+	/** Return the max LOD that will be used in the next Instance update. */
 	int32 GetMaxLODToLoad() const;
+	
 	int32 GetNumLODsAvailable() const;
 
-	UE_DEPRECATED(5.2, "Use SetRequestedLODs instead.")
-	void SetMinMaxLODToLoad(FMutableInstanceUpdateMap& InOutRequestedUpdates, int32 NewMinLOD = 0, int32 NewMaxLOD = INT32_MAX, bool bLimitLODUpgrades = true);
-
-	/** Return the Min LOD this Instance is using (from the beginning of an update. If an update fails this value will be incorrect). */
+	/** Return the min LOD this Instance is using. */
 	int32 GetCurrentMinLOD() const;
 
-	/** Return the Max LOD this Instance is using (from the beginning of an update. If an update fails this value will be incorrect). */
+	/** Return the max LOD this Instance is using. */
 	int32 GetCurrentMaxLOD() const;
 
 	/** Sets an array of LODs to generate per component. Mutable will generate those plus the currently generated LODs (if any).
-	 * Requires mutable.EnableOnlyGenerateRequestedLODs and CurrentInstanceLODManagement->IsOnlyGenerateRequestedLODLevelsEnabled() to be true.
-	 * @param InMinLOD - MinLOD to generate.
-	 * @param InMaxLOD - MaxLOD to generate.
-	 * @param InRequestedLODsPerComponent - Array with bitmasks of requested LODs per component with range from [0 .. CO->GetComponentCount()].
-	 * @param InOutRequestedUpdates - Map from Instance to Update data that stores a request for the Instance to be updated, which will be either processed or discarded by priority (to be rerequested the next tick) */
+	  * Requires mutable.EnableOnlyGenerateRequestedLODs and CurrentInstanceLODManagement->IsOnlyGenerateRequestedLODLevelsEnabled() to be true.
+	  * @param InMinLOD - MinLOD to generate.
+	  * @param InMaxLOD - MaxLOD to generate.
+	  * @param InRequestedLODsPerComponent - Array with bitmasks of requested LODs per component with range from [0 .. CO->GetComponentCount()].
+	  * @param InOutRequestedUpdates - Map from Instance to Update data that stores a request for the Instance to be updated, which will be either processed or discarded by priority (to be rerequested the next tick) */
 	void SetRequestedLODs(int32 InMinLOD, int32 InMaxLOD, const TArray<uint16>& InRequestedLODsPerComponent, FMutableInstanceUpdateMap& InOutRequestedUpdates);
 
 	const TArray<uint16>& GetRequestedLODsPerComponent() const;
 
-#if WITH_EDITOR
-	const TArray<TObjectPtr<UTexture2D>>& GetTextureParameterDeclarations() const;
-#endif
-	
-	// --------------------------------------------------------------------
-
-	/** Flag to know if a property of this instance changed in the editor */
-	bool bEditorPropertyChanged = false;
-
 	UCustomizableInstancePrivate* GetPrivate() const;
 
-	// TEMP VARIABLE to check the Min desired LODs for this instance
-	TWeakObjectPtr<UCustomizableObjectInstanceUsage> NearestToActor;
-	TWeakObjectPtr<const AActor> NearestToViewCenter;
-
-#if WITH_EDITOR
-	/** Profile index the instance parameters are in and if the profile needs to be refreshed */
-	int32 SelectedProfileIndex = INDEX_NONE;
-	bool bSelectedProfileDirty = false;
-#endif 
-
 private:
 	UPROPERTY()
 	FCustomizableObjectInstanceDescriptor Descriptor;
@@ -686,19 +599,15 @@ private:
 	UPROPERTY()
 	TObjectPtr<UCustomizableInstancePrivate> PrivateData;
 
+public:
 #if WITH_EDITORONLY_DATA
 	/** Textures which can used as values in Texture Parameters. */
 	UPROPERTY(EditAnywhere, Category = TextureParameter)
 	TArray<TObjectPtr<UTexture2D>> TextureParameterDeclarations;
-
-public:
-	/** Preview Instance Properties search box filter. Saved here to avoid losing the text during UI refreshes. */
-	FText ParametersSearchFilter;
 #endif
 
-private:	
-	// Deprecated properties
-	
+private:
+	// Deprecated properties	
 	UPROPERTY()
 	TObjectPtr<UCustomizableObject> CustomizableObject_DEPRECATED;
 	
@@ -723,56 +632,10 @@ private:
    	UPROPERTY()
    	TMap<FName, FMultilayerProjector> MultilayerProjectors_DEPRECATED;
 
-	/** If this is set to true, when updating the instance an additional step will be performed to calculate the list of instance parameters that are relevant for the current parameter vaules. */
 	bool bBuildParameterRelevancy_DEPRECATED = false;
 };
 
 
-class FMutableUpdateCandidate
-{
-public:
-	// The Instance to possibly update
-	UCustomizableObjectInstance* CustomizableObjectInstance;
-	EQueuePriorityType Priority = EQueuePriorityType::Med;
-
-	// These are the LODs that would be applied if this candidate is chosen
-	int32 MinLOD = 0;
-	int32 MaxLOD = INT32_MAX;
-
-	/** Array of RequestedLODs per component to generate if this candidate is chosen */
-	TArray<uint16> RequestedLODLevels;
-
-	FMutableUpdateCandidate(UCustomizableObjectInstance* InCustomizableObjectInstance) : CustomizableObjectInstance(InCustomizableObjectInstance)
-	{
-		const FCustomizableObjectInstanceDescriptor& Descriptor = InCustomizableObjectInstance->GetDescriptor();
-		MinLOD = Descriptor.GetMinLod();
-		MaxLOD = Descriptor.GetMaxLod();
-		RequestedLODLevels = Descriptor.GetRequestedLODLevels();
-	}
-
-	FMutableUpdateCandidate(const UCustomizableObjectInstance* InCustomizableObjectInstance, const int32 InMinLOD, const int32 InMaxLOD,
-		const TArray<uint16>& InRequestedLODLevels) :
-		CustomizableObjectInstance(const_cast<UCustomizableObjectInstance*>(InCustomizableObjectInstance)), MinLOD(InMinLOD), MaxLOD(InMaxLOD),
-		RequestedLODLevels(InRequestedLODLevels) {}
-
-	bool HasBeenIssued() const;
-
-	void Issue();
-
-	void ApplyLODUpdateParamsToInstance(FUpdateContextPrivate& Context);
-
-private:
-	/** If true it means that EnqueueUpdateSkeletalMesh has decided this update should be performed, if false it should be ignored. Just used for consistency checks */
-	bool bHasBeenIssued = false;
-};
-
-
-/** 
- * \param OnlyLOD: If not 0, extract and convert only one single LOD from the source image.
- * \param ExtractChannel: If different than -1, extract a single-channel image with the specified source channel data. */
-CUSTOMIZABLEOBJECT_API void ConvertImage(UTexture2D* Texture, mu::Ptr<const mu::Image> MutableImage, const FMutableModelImageProperties& Props, int32 OnlyLOD = -1, int32 ExtractChannel = -1);
-
-
 #if UE_ENABLE_INCLUDE_ORDER_DEPRECATED_IN_5_2
 #include "Async/TaskGraphInterfaces.h"
 #include "MuCO/CustomizableObjectParameterTypeDefinitions.h"
diff --git a/Engine/Plugins/Experimental/Mutable/Source/CustomizableObject/Internal/MuCO/CustomizableObjectInstancePrivate.h b/Engine/Plugins/Experimental/Mutable/Source/CustomizableObject/Internal/MuCO/CustomizableObjectInstancePrivate.h
--- a/Engine/Plugins/Experimental/Mutable/Source/CustomizableObject/Internal/MuCO/CustomizableObjectInstancePrivate.h
+++ b/Engine/Plugins/Experimental/Mutable/Source/CustomizableObject/Internal/MuCO/CustomizableObjectInstancePrivate.h
@@ -26,6 +26,7 @@ struct FModelResources;
 struct FMutableModelImageProperties;
 struct FMutableRefSkeletalMeshData;
 struct FMutableImageCacheKey;
+struct FStreamableManager;
 struct FStreamableHandle;
 struct FGeneratedMaterial;
 struct FGeneratedTexture;
@@ -33,6 +34,11 @@ class UPhysicsAsset;
 class USkeleton;
 
 
+/** \param OnlyLOD: If not 0, extract and convert only one single LOD from the source image.
+  * \param ExtractChannel: If different than -1, extract a single-channel image with the specified source channel data. */
+CUSTOMIZABLEOBJECT_API void ConvertImage(UTexture2D* Texture, mu::Ptr<const mu::Image> MutableImage, const FMutableModelImageProperties& Props, int32 OnlyLOD = -1, int32 ExtractChannel = -1);
+
+
 /** CustomizableObject Instance flags for internal use  */
 enum ECOInstanceFlags
 {
@@ -157,7 +163,7 @@ struct FAnimBpGeneratedPhysicsAssets
 
 
 UCLASS()
-class UCustomizableInstancePrivate : public UObject
+class CUSTOMIZABLEOBJECT_API UCustomizableInstancePrivate : public UObject
 {
 public:
 	GENERATED_BODY()
@@ -216,7 +222,7 @@ public:
 	void ReuseTexture(UTexture2D* Texture, TSharedRef<FTexturePlatformData, ESPMode::ThreadSafe>& PlatformData);
 
 	// Return an event that will be fired when the assets  have been loaded. It returns null if no asset needs loading.
-	FGraphEventRef LoadAdditionalAssetsAsync(const TSharedRef<FUpdateContextPrivate>& OperationData, UCustomizableObjectInstance* Public, struct FStreamableManager &StreamableManager);
+	FGraphEventRef LoadAdditionalAssetsAsync(const TSharedRef<FUpdateContextPrivate>& OperationData, FStreamableManager &StreamableManager);
 	void AdditionalAssetsAsyncLoaded(UCustomizableObjectInstance* Public);
 
 	void TickUpdateCloseCustomizableObjects(UCustomizableObjectInstance& Public, FMutableInstanceUpdateMap& InOutRequestedUpdates);
@@ -280,6 +286,26 @@ private:
 	void SetLastMeshId(int32 ComponentIndex, int32 LODIndex, mu::FResourceID MeshId);
 
 public:
+	bool LoadParametersFromProfile(int32 ProfileIndex);
+	
+	bool SaveParametersToProfile(int32 ProfileIndex);
+	
+	bool MigrateProfileParametersToCurrentInstance(int32 ProfileIndex);
+
+	void SetSelectedParameterProfileDirty();
+
+	bool IsSelectedParameterProfileDirty() const;
+
+	int32 GetState() const;
+
+	void SetState(int32 InState);
+
+	void AdditionalAssetsAsyncLoaded(FGraphEventRef CompletionEvent);
+
+	FCustomizableObjectInstanceDescriptor& GetDescriptor() const;
+	
+	UCustomizableObjectInstance* GetPublic() const;
+
 	// If any components are using this instance, they will store the min of their distances to the player here every frame for LOD purposes
 	float MinSquareDistFromComponentToPlayer;
 	float LastMinSquareDistFromComponentToPlayer; // The same as the previous dist for last frame
@@ -338,5 +364,34 @@ public:
 	
 	/** Status of the generated Skeletal Mesh. Not to be confused with the Update Result. */
 	ESkeletalMeshStatus SkeletalMeshStatus = ESkeletalMeshStatus::NotGenerated;
+
+	TMap<FString, bool> ParamNameToExpandedMap; // Used to check whether a mutable param is expanded in the editor to show its child params
+
+#if WITH_EDITOR
+	/** During editor, always remember the duration of the last update in the mutable runtime, for profiling. */
+	int32 LastUpdateMutableRuntimeCycles = 0;
+#endif
+
+	bool bShowOnlyRuntimeParameters = true;
+	bool bShowOnlyRelevantParameters = true;
+	bool bShowUISections = false;
+
+	/** Flag to know if a property of this instance changed in the editor */
+	bool bEditorPropertyChanged = false;
+	
+	// TEMP VARIABLE to check the Min desired LODs for this instance
+	TWeakObjectPtr<UCustomizableObjectInstanceUsage> NearestToActor;
+	TWeakObjectPtr<const AActor> NearestToViewCenter;
+
+#if WITH_EDITOR
+	/** Profile index the instance parameters are in and if the profile needs to be refreshed */
+	int32 SelectedProfileIndex = INDEX_NONE;
+	bool bSelectedProfileDirty = false;
+#endif 
+	
+#if WITH_EDITORONLY_DATA
+	/** Preview Instance Properties search box filter. Saved here to avoid losing the text during UI refreshes. */
+	FText ParametersSearchFilter;
+#endif
 };
 
diff --git a/Engine/Plugins/Experimental/Mutable/Source/CustomizableObject/Internal/MuCO/CustomizableObjectSystemPrivate.h b/Engine/Plugins/Experimental/Mutable/Source/CustomizableObject/Internal/MuCO/CustomizableObjectSystemPrivate.h
--- a/Engine/Plugins/Experimental/Mutable/Source/CustomizableObject/Internal/MuCO/CustomizableObjectSystemPrivate.h
+++ b/Engine/Plugins/Experimental/Mutable/Source/CustomizableObject/Internal/MuCO/CustomizableObjectSystemPrivate.h
@@ -50,6 +50,45 @@ struct FTexturePlatformData;
 void FinishUpdateGlobal(const TSharedRef<FUpdateContextPrivate>& Context);
 
 
+class FMutableUpdateCandidate
+{
+public:
+	// The Instance to possibly update
+	UCustomizableObjectInstance* CustomizableObjectInstance;
+	EQueuePriorityType Priority = EQueuePriorityType::Med;
+
+	// These are the LODs that would be applied if this candidate is chosen
+	int32 MinLOD = 0;
+	int32 MaxLOD = INT32_MAX;
+
+	/** Array of RequestedLODs per component to generate if this candidate is chosen */
+	TArray<uint16> RequestedLODLevels;
+
+	FMutableUpdateCandidate(UCustomizableObjectInstance* InCustomizableObjectInstance) : CustomizableObjectInstance(InCustomizableObjectInstance)
+	{
+		const FCustomizableObjectInstanceDescriptor& Descriptor = InCustomizableObjectInstance->GetDescriptor();
+		MinLOD = Descriptor.GetMinLod();
+		MaxLOD = Descriptor.GetMaxLod();
+		RequestedLODLevels = Descriptor.GetRequestedLODLevels();
+	}
+
+	FMutableUpdateCandidate(const UCustomizableObjectInstance* InCustomizableObjectInstance, const int32 InMinLOD, const int32 InMaxLOD,
+		const TArray<uint16>& InRequestedLODLevels) :
+		CustomizableObjectInstance(const_cast<UCustomizableObjectInstance*>(InCustomizableObjectInstance)), MinLOD(InMinLOD), MaxLOD(InMaxLOD),
+		RequestedLODLevels(InRequestedLODLevels) {}
+
+	bool HasBeenIssued() const;
+
+	void Issue();
+
+	void ApplyLODUpdateParamsToInstance(FUpdateContextPrivate& Context);
+
+private:
+	/** If true it means that EnqueueUpdateSkeletalMesh has decided this update should be performed, if false it should be ignored. Just used for consistency checks */
+	bool bHasBeenIssued = false;
+};
+
+
 struct FMutablePendingInstanceUpdate
 {
 	TSharedRef<FUpdateContextPrivate> Context;
diff --git a/Engine/Plugins/Experimental/Mutable/Source/CustomizableObject/Private/MuCO/CustomizableInstanceLODManagement.cpp b/Engine/Plugins/Experimental/Mutable/Source/CustomizableObject/Private/MuCO/CustomizableInstanceLODManagement.cpp
--- a/Engine/Plugins/Experimental/Mutable/Source/CustomizableObject/Private/MuCO/CustomizableInstanceLODManagement.cpp
+++ b/Engine/Plugins/Experimental/Mutable/Source/CustomizableObject/Private/MuCO/CustomizableInstanceLODManagement.cpp
@@ -15,6 +15,7 @@
 #include "MuCO/CustomizableObjectInstanceUsage.h"
 #include "UObject/UObjectIterator.h"
 #include "Components/SkeletalMeshComponent.h"
+#include "MuCO/CustomizableObjectInstancePrivate.h"
 #include "MuCO/CustomizableObjectPrivate.h"
 
 #include UE_INLINE_GENERATED_CPP_BY_NAME(CustomizableInstanceLODManagement)
@@ -409,7 +410,7 @@ void UCustomizableInstanceLODManagement::UpdateInstanceDistsAndLODs(FMutableInst
 					continue;
 				}
 
-				if (IsOnlyUpdateCloseCustomizableObjectsEnabled() && (It.Key->GetMinSquareDistToPlayer() > FMath::Square(CloseCustomizableObjectsDist) || !It.Key->NearestToActor.IsValid()))
+				if (IsOnlyUpdateCloseCustomizableObjectsEnabled() && (It.Key->GetMinSquareDistToPlayer() > FMath::Square(CloseCustomizableObjectsDist) || !It.Key->GetPrivate()->NearestToActor.IsValid()))
 				{
 					continue;
 				}
diff --git a/Engine/Plugins/Experimental/Mutable/Source/CustomizableObject/Private/MuCO/CustomizableObject.cpp b/Engine/Plugins/Experimental/Mutable/Source/CustomizableObject/Private/MuCO/CustomizableObject.cpp
--- a/Engine/Plugins/Experimental/Mutable/Source/CustomizableObject/Private/MuCO/CustomizableObject.cpp
+++ b/Engine/Plugins/Experimental/Mutable/Source/CustomizableObject/Private/MuCO/CustomizableObject.cpp
@@ -20,8 +20,10 @@
 #include "Misc/PackageName.h"
 #include "Misc/Paths.h"
 #include "MuCO/CustomizableObjectInstance.h"
+#include "MuCO/CustomizableObjectInstancePrivate.h"
 #include "MuCO/CustomizableObjectPrivate.h"
 #include "MuCO/CustomizableObjectSystem.h"
+#include "MuCO/CustomizableObjectSystemPrivate.h"
 #include "MuCO/CustomizableObjectUIData.h"
 #include "MuCO/ICustomizableObjectModule.h"
 #include "MuCO/MutableProjectorTypeUtils.h"
@@ -39,8 +41,6 @@
 #endif
 
 
-#include "MuCO/CustomizableObjectSystemPrivate.h"
-
 #include UE_INLINE_GENERATED_CPP_BY_NAME(CustomizableObject)
 
 #define LOCTEXT_NAMESPACE "CustomizableObject"
@@ -951,7 +951,7 @@ FReply UCustomizableObjectPrivate::AddNewParameterProfile(FString Name, UCustomi
 	int32 ProfileIndex = GetPublic()->InstancePropertiesProfiles.Emplace();
 
 	GetPublic()->InstancePropertiesProfiles[ProfileIndex].ProfileName = ProfileName;
-	CustomInstance.SaveParametersToProfile(ProfileIndex);
+	CustomInstance.GetPrivate()->SaveParametersToProfile(ProfileIndex);
 
 	Modify();
 
@@ -1234,7 +1234,7 @@ UCustomizableObjectInstance* UCustomizableObject::CreateInstance()
 
 	UCustomizableObjectInstance* PreviewInstance = NewObject<UCustomizableObjectInstance>(GetTransientPackage(), NAME_None, RF_Transient);
 	PreviewInstance->SetObject(this);
-	PreviewInstance->bShowOnlyRuntimeParameters = false;
+	PreviewInstance->GetPrivate()->bShowOnlyRuntimeParameters = false;
 
 	UE_LOG(LogMutable, Verbose, TEXT("Created Customizable Object Instance."));
 
diff --git a/Engine/Plugins/Experimental/Mutable/Source/CustomizableObject/Private/MuCO/CustomizableObjectInstance.cpp b/Engine/Plugins/Experimental/Mutable/Source/CustomizableObject/Private/MuCO/CustomizableObjectInstance.cpp
--- a/Engine/Plugins/Experimental/Mutable/Source/CustomizableObject/Private/MuCO/CustomizableObjectInstance.cpp
+++ b/Engine/Plugins/Experimental/Mutable/Source/CustomizableObject/Private/MuCO/CustomizableObjectInstance.cpp
@@ -251,15 +251,9 @@ UCustomizableObjectInstance::UCustomizableObjectInstance()
 }
 
 
-FCustomizableObjectInstanceDescriptor& UCustomizableObjectInstance::GetDescriptor()
-{
-	return Descriptor;	
-}
-
-
 const FCustomizableObjectInstanceDescriptor& UCustomizableObjectInstance::GetDescriptor() const
 {
-	return const_cast<UCustomizableObjectInstance*>(this)->GetDescriptor();
+	return Descriptor;
 }
 
 
@@ -321,14 +315,8 @@ void UCustomizableInstancePrivate::PostDuplicate(bool bDuplicateForPIE)
 
 void UCustomizableInstancePrivate::OnPostCompile()
 {
-	UCustomizableObjectInstance* Instance = Cast<UCustomizableObjectInstance>(GetOuter());
-	if (!Instance)
-	{
-		return;
-	}
-
-	Instance->GetDescriptor().ReloadParameters();
-	InitCustomizableObjectData(Instance->GetCustomizableObject());
+	GetDescriptor().ReloadParameters();
+	InitCustomizableObjectData(GetPublic()->GetCustomizableObject());
 }
 
 
@@ -368,7 +356,7 @@ void UCustomizableObjectInstance::PostEditChangeProperty(FPropertyChangedEvent&
 {
 	Super::PostEditChangeProperty(PropertyChangedEvent);
 
-	bEditorPropertyChanged = true;
+	GetPrivate()->bEditorPropertyChanged = true;
 
 	const FName PropertyName = PropertyChangedEvent.MemberProperty ? PropertyChangedEvent.MemberProperty->GetFName() : NAME_None;
 
@@ -413,9 +401,6 @@ bool UCustomizableObjectInstance::IsEditorOnly() const
 
 void UCustomizableObjectInstance::BeginDestroy()
 {
-	BeginDestroyDelegate.Broadcast(this);
-	BeginDestroyNativeDelegate.Broadcast(this);
-
 	// Release the Live Instance ID if there it hadn't been released before
 	DestroyLiveUpdateInstance();
 
@@ -671,12 +656,6 @@ int32 UCustomizableObjectInstance::GetTextureValueRange(const FString& ParamName
 }
 
 
-bool UCustomizableObjectInstance::IsParamMultidimensional(const FString& ParamName) const
-{
-	return GetCustomizableObject()->IsParameterMultidimensional(ParamName);
-}
-
-
 // Only safe to call if the Mutable texture ref count system returns 0 and absolutely sure nobody holds a reference to the texture
 void UCustomizableInstancePrivate::ReleaseMutableTexture(const FMutableImageCacheKey& MutableTextureKey, UTexture2D* Texture, FMutableResourceCache& Cache)
 {
@@ -713,19 +692,6 @@ UCustomizableObject* UCustomizableObjectInstance::GetCustomizableObject() const
 }
 
 
-// Marked as deprecated. Will be removed in future versions.
-bool UCustomizableObjectInstance::GetBuildParameterDecorations() const
-{
-	return false;
-}
-
-
-// Marked as deprecated. Will be removed in future versions.
-void UCustomizableObjectInstance::SetBuildParameterDecorations(const bool Value)
-{
-}
-
-
 bool UCustomizableObjectInstance::GetBuildParameterRelevancy() const
 {
 	return Descriptor.GetBuildParameterRelevancy();
@@ -738,22 +704,22 @@ void UCustomizableObjectInstance::SetBuildParameterRelevancy(bool Value)
 }
 
 
-int32 UCustomizableObjectInstance::GetState() const
+int32 UCustomizableInstancePrivate::GetState() const
 {
-	return Descriptor.GetState();
+	return GetPublic()->Descriptor.GetState();
 }
 
 
-void UCustomizableObjectInstance::SetState(const int32 InState)
+void UCustomizableInstancePrivate::SetState(const int32 InState)
 {
 	const int32 OldState = GetState();
 	
-	Descriptor.SetState(InState);
+	GetPublic()->Descriptor.SetState(InState);
 
 	if (OldState != InState)
 	{
 		// State may change texture properties, so invalidate the texture reuse cache
-		PrivateData->TextureReuseCache.Empty();
+		TextureReuseCache.Empty();
 	}
 }
 
@@ -770,23 +736,16 @@ void UCustomizableObjectInstance::SetCurrentState(const FString& StateName)
 }
 
 
-// Marked as deprecated. Will be removed in future versions.
-UTexture2D* UCustomizableObjectInstance::GetParameterDescription(const FString& ParamName, int32 DescIndex)
-{
-	return nullptr;
-}
-
-
 bool UCustomizableObjectInstance::IsParameterRelevant(int32 ParameterIndex) const
 {
-	// This should have been precalculated in the last update if the appropiate flag in the instance was set.
+	// This should have been precalculated in the last update if the appropriate flag in the instance was set.
 	return GetPrivate()->RelevantParameters.Contains(ParameterIndex);
 }
 
 
 bool UCustomizableObjectInstance::IsParameterRelevant(const FString& ParamName) const
 {
-	// This should have been precalculated in the last update if the appropiate flag in the instance was set.
+	// This should have been precalculated in the last update if the appropriate flag in the instance was set.
 	int32 ParameterIndexInObject = GetCustomizableObject()->FindParameter(ParamName);
 	return GetPrivate()->RelevantParameters.Contains(ParameterIndexInObject);
 }
@@ -2177,7 +2136,7 @@ UCustomizableObjectInstance* UCustomizableObjectInstance::CloneStatic(UObject* O
 {
 	UCustomizableObjectInstance* NewInstance = NewObject<UCustomizableObjectInstance>(Outer, GetClass());
 	NewInstance->CopyParametersFromInstance(this);
-	NewInstance->bShowOnlyRuntimeParameters = false;
+	NewInstance->GetPrivate()->bShowOnlyRuntimeParameters = false;
 
 	return NewInstance;
 }
@@ -2414,7 +2373,7 @@ void UCustomizableObjectInstance::SetProjectorValue(const FString& ProjectorPara
 }
 
 
-void UCustomizableObjectInstance::SetProjectorPosition(const FString& ProjectorParamName, const FVector3f& Pos, const int32 RangeIndex)
+void UCustomizableObjectInstance::SetProjectorPosition(const FString& ProjectorParamName, const FVector& Pos, const int32 RangeIndex)
 {
 	Descriptor.SetProjectorPosition(ProjectorParamName, Pos, RangeIndex);
 }
@@ -2544,9 +2503,9 @@ void UCustomizableObjectInstance::SetRandomValuesFromStream(const FRandomStream&
 }
 
 
-bool UCustomizableObjectInstance::LoadParametersFromProfile(int32 ProfileIndex)
+bool UCustomizableInstancePrivate::LoadParametersFromProfile(int32 ProfileIndex)
 {
-	UCustomizableObject* CustomizableObject = GetCustomizableObject();
+	UCustomizableObject* CustomizableObject = GetPublic()->GetCustomizableObject();
 	if (!CustomizableObject)
 	{
 		return false;
@@ -2563,20 +2522,20 @@ bool UCustomizableObjectInstance::LoadParametersFromProfile(int32 ProfileIndex)
 
 	const FProfileParameterDat& Profile = CustomizableObject->InstancePropertiesProfiles[ProfileIndex];
 
-	Descriptor.BoolParameters = Profile.BoolParameters;
-	Descriptor.IntParameters = Profile.IntParameters;
-	Descriptor.FloatParameters = Profile.FloatParameters;
-	Descriptor.TextureParameters = Profile.TextureParameters;
-	Descriptor.ProjectorParameters = Profile.ProjectorParameters;
-	Descriptor.VectorParameters = Profile.VectorParameters;
+	GetPublic()->Descriptor.BoolParameters = Profile.BoolParameters;
+	GetPublic()->Descriptor.IntParameters = Profile.IntParameters;
+	GetPublic()->Descriptor.FloatParameters = Profile.FloatParameters;
+	GetPublic()->Descriptor.TextureParameters = Profile.TextureParameters;
+	GetPublic()->Descriptor.ProjectorParameters = Profile.ProjectorParameters;
+	GetPublic()->Descriptor.VectorParameters = Profile.VectorParameters;
 #endif
 	return true;
 
 }
 
-bool UCustomizableObjectInstance::SaveParametersToProfile(int32 ProfileIndex)
+bool UCustomizableInstancePrivate::SaveParametersToProfile(int32 ProfileIndex)
 {
-	UCustomizableObject* CustomizableObject = GetCustomizableObject();
+	UCustomizableObject* CustomizableObject = GetPublic()->GetCustomizableObject();
 	if (!CustomizableObject)
 	{
 		return false;
@@ -2592,19 +2551,19 @@ bool UCustomizableObjectInstance::SaveParametersToProfile(int32 ProfileIndex)
 
 	FProfileParameterDat& Profile = CustomizableObject->InstancePropertiesProfiles[ProfileIndex];
 
-	Profile.BoolParameters = Descriptor.BoolParameters;
-	Profile.IntParameters = Descriptor.IntParameters;
-	Profile.FloatParameters = Descriptor.FloatParameters;
-	Profile.TextureParameters = Descriptor.TextureParameters;
-	Profile.ProjectorParameters = Descriptor.ProjectorParameters;
-	Profile.VectorParameters = Descriptor.VectorParameters;
+	Profile.BoolParameters = GetPublic()->Descriptor.BoolParameters;
+	Profile.IntParameters = GetPublic()->Descriptor.IntParameters;
+	Profile.FloatParameters = GetPublic()->Descriptor.FloatParameters;
+	Profile.TextureParameters = GetPublic()->Descriptor.TextureParameters;
+	Profile.ProjectorParameters = GetPublic()->Descriptor.ProjectorParameters;
+	Profile.VectorParameters = GetPublic()->Descriptor.VectorParameters;
 #endif
 	return true;
 }
 
-bool UCustomizableObjectInstance::MigrateProfileParametersToCurrentInstance(int32 ProfileIndex)
+bool UCustomizableInstancePrivate::MigrateProfileParametersToCurrentInstance(int32 ProfileIndex)
 {
-	UCustomizableObject* CustomizableObject = GetCustomizableObject();
+	UCustomizableObject* CustomizableObject = GetPublic()->GetCustomizableObject();
 	if (!CustomizableObject)
 	{
 		return false;
@@ -2620,12 +2579,12 @@ bool UCustomizableObjectInstance::MigrateProfileParametersToCurrentInstance(int3
 	FProfileParameterDat TempProfile;
 
 	TempProfile.ProfileName = Profile.ProfileName;
-	TempProfile.BoolParameters = Descriptor.BoolParameters;
-	TempProfile.FloatParameters = Descriptor.FloatParameters;
-	TempProfile.IntParameters = Descriptor.IntParameters;
-	TempProfile.ProjectorParameters = Descriptor.ProjectorParameters;
-	TempProfile.TextureParameters = Descriptor.TextureParameters;
-	TempProfile.VectorParameters = Descriptor.VectorParameters;
+	TempProfile.BoolParameters = GetPublic()->Descriptor.BoolParameters;
+	TempProfile.FloatParameters = GetPublic()->Descriptor.FloatParameters;
+	TempProfile.IntParameters = GetPublic()->Descriptor.IntParameters;
+	TempProfile.ProjectorParameters = GetPublic()->Descriptor.ProjectorParameters;
+	TempProfile.TextureParameters = GetPublic()->Descriptor.TextureParameters;
+	TempProfile.VectorParameters = GetPublic()->Descriptor.VectorParameters;
 	
 
 	// Populate TempProfile with the parameters found in the profile.
@@ -2713,9 +2672,18 @@ bool UCustomizableObjectInstance::MigrateProfileParametersToCurrentInstance(int3
 }
 
 
-void UCustomizableObjectInstance::SetSelectedParameterProfileDirty()
+UCustomizableObjectInstance* UCustomizableInstancePrivate::GetPublic() const
 {
-	UCustomizableObject* CustomizableObject = GetCustomizableObject();
+	UCustomizableObjectInstance* Public = StaticCast<UCustomizableObjectInstance*>(GetOuter());
+	check(Public);
+
+	return Public;
+}
+
+
+void UCustomizableInstancePrivate::SetSelectedParameterProfileDirty()
+{
+	UCustomizableObject* CustomizableObject = GetPublic()->GetCustomizableObject();
 	if (!CustomizableObject)
 	{
 		return;
@@ -2731,7 +2699,7 @@ void UCustomizableObjectInstance::SetSelectedParameterProfileDirty()
 #endif
 }
 
-bool UCustomizableObjectInstance::IsSelectedParameterProfileDirty() const
+bool UCustomizableInstancePrivate::IsSelectedParameterProfileDirty() const
 {
 	
 #if WITH_EDITOR
@@ -2741,43 +2709,6 @@ bool UCustomizableObjectInstance::IsSelectedParameterProfileDirty() const
 #endif
 }
 
-//int32 UCustomizableObjectInstance::GetIntParameterNumOptions(int32 IntParamIndex)
-//{
-//	if (CustomizableObject && IntParamIndex>=0 && IntParamIndex<IntParameters.Num())
-//	{
-//		int32 ObjectParamIndex = CustomizableObject->FindParameter(IntParameters[IntParamIndex].ParameterName);
-//		return CustomizableObject->GetIntParameterNumOptions(ObjectParamIndex);
-//	}
-//
-//	return 0;
-//}
-
-
-//void UCustomizableObjectInstance::SetIntParameter(int32 IntParamIndex, int32 i)
-//{
-//	if (CustomizableObject && IntParamIndex >= 0 && IntParamIndex < IntParameters.Num())
-//	{
-//		int32 ObjectParamIndex = CustomizableObject->FindParameter(IntParameters[IntParamIndex].ParameterName);
-//		int32 PossibleValue = CustomizableObject->GetIntParameterAvailableOptionValue(ObjectParamIndex, i);
-//		IntParameters[IntParamIndex].ParameterValue = PossibleValue;
-//	}
-//}
-
-
-//const FString & UCustomizableObjectInstance::GetIntParameterAvailableOption(int32 IntParamIndex, int32 K)
-//{
-//	if (CustomizableObject && IntParamIndex >= 0 && IntParamIndex<IntParameters.Num())
-//	{
-//		int32 ObjectParamIndex = CustomizableObject->FindParameter(IntParameters[IntParamIndex].ParameterName);
-//		if (ObjectParamIndex != INDEX_NONE)
-//		{
-//			return CustomizableObject->GetIntParameterAvailableOption(ObjectParamIndex,K);
-//		}
-//	}
-//
-//	return s_EmptyStringCI;
-//}
-
 
 void UCustomizableInstancePrivate::DiscardResources()
 {
@@ -2898,14 +2829,6 @@ void SetTexturePropertiesFromMutableImageProps(UTexture2D* Texture, const FMutab
 }
 
 
-#if WITH_EDITOR
-const TArray<TObjectPtr<UTexture2D>>& UCustomizableObjectInstance::GetTextureParameterDeclarations() const
-{
-	return TextureParameterDeclarations;
-}
-#endif
-
-
 UCustomizableInstancePrivate* UCustomizableObjectInstance::GetPrivate() const
 { 
 	check(PrivateData); // Currently this is initialized in the constructor so we expect it always to exist.
@@ -3289,7 +3212,6 @@ void ConvertImage(UTexture2D* Texture, mu::Ptr<const mu::Image> MutableImage, co
 	}
 
 	Texture->SetPlatformData(MutableCreateImagePlatformData(MutableImage,OnlyLOD,0,0) );
-
 }
 
 
@@ -4833,11 +4755,11 @@ FAutoConsoleVariableRef CVarMutableHighPriorityLoading(
 	TEXT("If enabled, the request to load additional assets will have high priority."));
 
 
-FGraphEventRef UCustomizableInstancePrivate::LoadAdditionalAssetsAsync(const TSharedRef<FUpdateContextPrivate>& OperationData, UCustomizableObjectInstance* Public, FStreamableManager& StreamableManager)
+FGraphEventRef UCustomizableInstancePrivate::LoadAdditionalAssetsAsync(const TSharedRef<FUpdateContextPrivate>& OperationData, FStreamableManager& StreamableManager)
 {
 	MUTABLE_CPUPROFILER_SCOPE(UCustomizableInstancePrivate::LoadAdditionalAssetsAsync);
 
-	UCustomizableObject* CustomizableObject = Public->GetCustomizableObject();
+	UCustomizableObject* CustomizableObject = GetPublic()->GetCustomizableObject();
 
 	const FModelResources& ModelResources = CustomizableObject->GetPrivate()->GetModelResources();
 
@@ -5095,7 +5017,7 @@ FGraphEventRef UCustomizableInstancePrivate::LoadAdditionalAssetsAsync(const TSh
 		check(!StreamingHandle);
 
 		Result = FGraphEvent::CreateGraphEvent();
-		StreamingHandle = StreamableManager.RequestAsyncLoad(AssetsToStream, FStreamableDelegate::CreateUObject(Public, &UCustomizableObjectInstance::AdditionalAssetsAsyncLoaded, Result),
+		StreamingHandle = StreamableManager.RequestAsyncLoad(AssetsToStream, FStreamableDelegate::CreateUObject(this, &UCustomizableInstancePrivate::AdditionalAssetsAsyncLoaded, Result),
 			bEnableHighPriorityLoading ? FStreamableManager::AsyncLoadHighPriority : FStreamableManager::DefaultAsyncLoadPriority);
 	}
 
@@ -5103,13 +5025,19 @@ FGraphEventRef UCustomizableInstancePrivate::LoadAdditionalAssetsAsync(const TSh
 }
 
 
-void UCustomizableObjectInstance::AdditionalAssetsAsyncLoaded( FGraphEventRef CompletionEvent )
+void UCustomizableInstancePrivate::AdditionalAssetsAsyncLoaded( FGraphEventRef CompletionEvent )
 {
 	// TODO: Do we need this separated?
-	PrivateData->AdditionalAssetsAsyncLoaded(this);
+	AdditionalAssetsAsyncLoaded(GetPublic());
 	CompletionEvent->DispatchSubsequents(); // TODO: we know it is game thread?
 
-	PrivateData->StreamingHandle = nullptr;
+	StreamingHandle = nullptr;
+}
+
+
+FCustomizableObjectInstanceDescriptor& UCustomizableInstancePrivate::GetDescriptor() const
+{
+	return GetPublic()->Descriptor;
 }
 
 
@@ -6134,12 +6062,6 @@ void UCustomizableObjectInstance::SetMinSquareDistToPlayer(float NewValue)
 }
 
 
-void UCustomizableObjectInstance::SetMinMaxLODToLoad(FMutableInstanceUpdateMap& InOutRequestedUpdates, int32 NewMinLOD, int32 NewMaxLOD, bool bLimitLODUpgrades)
-{
-	SetRequestedLODs(NewMinLOD, NewMaxLOD, Descriptor.RequestedLODLevels, InOutRequestedUpdates);
-}
-
-
 int32 UCustomizableObjectInstance::GetNumComponents() const
 {
 	return GetCustomizableObject() ? GetCustomizableObject()->GetComponentCount() : 0;
@@ -6490,11 +6412,13 @@ void UCustomizableObjectInstance::ForEachAnimInstance(int32 ComponentIndex, FEac
 	}
 }
 
+
 bool UCustomizableObjectInstance::AnimInstanceNeedsFixup(TSubclassOf<UAnimInstance> AnimInstanceClass) const
 {
 	return PrivateData->AnimBpPhysicsAssets.Contains(AnimInstanceClass);
 }
 
+
 void UCustomizableObjectInstance::AnimInstanceFixup(UAnimInstance* InAnimInstance) const
 {
 	if (!InAnimInstance)
diff --git a/Engine/Plugins/Experimental/Mutable/Source/CustomizableObject/Private/MuCO/CustomizableObjectInstanceDescriptor.cpp b/Engine/Plugins/Experimental/Mutable/Source/CustomizableObject/Private/MuCO/CustomizableObjectInstanceDescriptor.cpp
--- a/Engine/Plugins/Experimental/Mutable/Source/CustomizableObject/Private/MuCO/CustomizableObjectInstanceDescriptor.cpp
+++ b/Engine/Plugins/Experimental/Mutable/Source/CustomizableObject/Private/MuCO/CustomizableObjectInstanceDescriptor.cpp
@@ -1665,7 +1665,7 @@ void FCustomizableObjectInstanceDescriptor::SetProjectorValue(const FString& Pro
 }
 
 
-void FCustomizableObjectInstanceDescriptor::SetProjectorPosition(const FString& ProjectorParamName, const FVector3f& Pos, const int32 RangeIndex)
+void FCustomizableObjectInstanceDescriptor::SetProjectorPosition(const FString& ProjectorParamName, const FVector& Pos, const int32 RangeIndex)
 {
 	FVector DummyPos, Direction, Up, Scale;
 	float Angle;
diff --git a/Engine/Plugins/Experimental/Mutable/Source/CustomizableObject/Private/MuCO/CustomizableObjectInstanceUsage.cpp b/Engine/Plugins/Experimental/Mutable/Source/CustomizableObject/Private/MuCO/CustomizableObjectInstanceUsage.cpp
--- a/Engine/Plugins/Experimental/Mutable/Source/CustomizableObject/Private/MuCO/CustomizableObjectInstanceUsage.cpp
+++ b/Engine/Plugins/Experimental/Mutable/Source/CustomizableObject/Private/MuCO/CustomizableObjectInstanceUsage.cpp
@@ -438,8 +438,8 @@ void UCustomizableObjectInstanceUsage::UpdateDistFromComponentToPlayer(const AAc
 
 				if (CustomizableObjectInstance->GetPrivate()->MinSquareDistFromComponentToPlayer == SquareDist)
 				{
-					CustomizableObjectInstance->NearestToActor = this;
-					CustomizableObjectInstance->NearestToViewCenter = ViewCenter;
+					CustomizableObjectInstance->GetPrivate()->NearestToActor = this;
+					CustomizableObjectInstance->GetPrivate()->NearestToViewCenter = ViewCenter;
 				}
 			}
 		}
diff --git a/Engine/Plugins/Experimental/Mutable/Source/CustomizableObject/Private/MuCO/CustomizableObjectSystem.cpp b/Engine/Plugins/Experimental/Mutable/Source/CustomizableObject/Private/MuCO/CustomizableObjectSystem.cpp
--- a/Engine/Plugins/Experimental/Mutable/Source/CustomizableObject/Private/MuCO/CustomizableObjectSystem.cpp
+++ b/Engine/Plugins/Experimental/Mutable/Source/CustomizableObject/Private/MuCO/CustomizableObjectSystem.cpp
@@ -965,7 +965,7 @@ void UpdateSkeletalMesh(const TSharedRef<FUpdateContextPrivate>& Context)
 		}
 	}
 
-	CustomizableObjectInstance->bEditorPropertyChanged = false;
+	CustomizableObjectInstance->GetPrivate()->bEditorPropertyChanged = false;
 }
 
 
@@ -2283,7 +2283,7 @@ namespace impl
 			UCustomizableInstancePrivate* CustomizableInstancePrivateData = CustomizableObjectInstance->GetPrivate();
 
 #if WITH_EDITOR
-			CustomizableObjectInstance->LastUpdateMutableRuntimeCycles = OperationData->MutableRuntimeCycles;
+			CustomizableObjectInstance->GetPrivate()->LastUpdateMutableRuntimeCycles = OperationData->MutableRuntimeCycles;
 #endif
 
 			// Convert Step
@@ -2487,7 +2487,7 @@ namespace impl
 
 		// Next Task: Load Unreal Assets
 		//-------------------------------------------------------------
-		FGraphEventRef Game_LoadUnrealAssets = ObjectInstancePrivateData->LoadAdditionalAssetsAsync(OperationData, ObjectInstance, UCustomizableObjectSystem::GetInstance()->GetPrivate()->StreamableManager);
+		FGraphEventRef Game_LoadUnrealAssets = ObjectInstancePrivateData->LoadAdditionalAssetsAsync(OperationData, UCustomizableObjectSystem::GetInstance()->GetPrivate()->StreamableManager);
 		if (Game_LoadUnrealAssets)
 		{
 			Game_LoadUnrealAssets->SetDebugName(TEXT("LoadAdditionalAssetsAsync"));
@@ -2708,7 +2708,7 @@ namespace impl
 		check(SystemPrivateData->ExtensionDataStreamer != nullptr);
 		SystemPrivateData->ExtensionDataStreamer->SetActiveObject(CustomizableObject);
 
-		FString StateName = CandidateInstance->GetCustomizableObject()->GetStateName(CandidateInstance->GetState());
+		FString StateName = CandidateInstance->GetCustomizableObject()->GetStateName(CandidateInstance->GetPrivate()->GetState());
 		const FParameterUIData* StateData = CandidateInstance->GetCustomizableObject()->GetPrivate()->GetModelResources().StateUIDataMap.Find(StateName);
 
 		Operation->bLiveUpdateMode = false;
@@ -3084,7 +3084,7 @@ bool UCustomizableObjectSystem::Tick(float DeltaTime)
 			else if (LODUpdateCandidateFound)
 			{
 				UCustomizableObjectInstance* Instance = LODUpdateCandidateFound->CustomizableObjectInstance;
-				FCustomizableObjectInstanceDescriptor& Descriptor = Instance->GetPrivate()->SkeletalMeshStatus == ESkeletalMeshStatus::NotGenerated ? Instance->GetDescriptor() : Instance->GetPrivate()->CommittedDescriptor;
+				const FCustomizableObjectInstanceDescriptor& Descriptor = Instance->GetPrivate()->SkeletalMeshStatus == ESkeletalMeshStatus::NotGenerated ? Instance->GetDescriptor() : Instance->GetPrivate()->CommittedDescriptor;
 
 				const TSharedRef<FUpdateContextPrivate> Context = MakeShared<FUpdateContextPrivate>(*Instance, Descriptor);
 
diff --git a/Engine/Plugins/Experimental/Mutable/Source/CustomizableObject/Private/MuCO/LogInformationUtil.cpp b/Engine/Plugins/Experimental/Mutable/Source/CustomizableObject/Private/MuCO/LogInformationUtil.cpp
--- a/Engine/Plugins/Experimental/Mutable/Source/CustomizableObject/Private/MuCO/LogInformationUtil.cpp
+++ b/Engine/Plugins/Experimental/Mutable/Source/CustomizableObject/Private/MuCO/LogInformationUtil.cpp
@@ -146,8 +146,8 @@ void LogInformationUtil::LogShowInstanceDataFull(const UCustomizableObjectInstan
 	FString MessageChunk;
 
 	MessageChunk += "\n\t";
-	MessageChunk += FString::Printf(TEXT("        bShowOnlyRuntimeParameters = %d\n"), CustomizableObjectInstance->bShowOnlyRuntimeParameters);
-	MessageChunk += FString::Printf(TEXT("        bShowOnlyRelevantParameters = %d\n"), CustomizableObjectInstance->bShowOnlyRelevantParameters);
+	MessageChunk += FString::Printf(TEXT("        bShowOnlyRuntimeParameters = %d\n"), CustomizableObjectInstance->GetPrivate()->bShowOnlyRuntimeParameters);
+	MessageChunk += FString::Printf(TEXT("        bShowOnlyRelevantParameters = %d\n"), CustomizableObjectInstance->GetPrivate()->bShowOnlyRelevantParameters);
 	MessageChunk += FString::Printf(TEXT("        MinSquareDistFromComponentToPlayer = %.2f\n"), CustomizableObjectInstance->GetPrivate()->MinSquareDistFromComponentToPlayer);
 	LogData += MessageChunk;
 
diff --git a/Engine/Plugins/Experimental/Mutable/Source/CustomizableObject/Public/MuCO/CustomizableObjectInstanceDescriptor.h b/Engine/Plugins/Experimental/Mutable/Source/CustomizableObject/Public/MuCO/CustomizableObjectInstanceDescriptor.h
--- a/Engine/Plugins/Experimental/Mutable/Source/CustomizableObject/Public/MuCO/CustomizableObjectInstanceDescriptor.h
+++ b/Engine/Plugins/Experimental/Mutable/Source/CustomizableObject/Public/MuCO/CustomizableObjectInstanceDescriptor.h
@@ -155,7 +155,7 @@ struct CUSTOMIZABLEOBJECT_API FCustomizableObjectInstanceDescriptor
 		int32 RangeIndex = -1);
 
 	/** Set only the projector position. */
-	void SetProjectorPosition(const FString& ProjectorParamName, const FVector3f& Pos, int32 RangeIndex = -1);
+	void SetProjectorPosition(const FString& ProjectorParamName, const FVector& Pos, int32 RangeIndex = -1);
 
 	/** Set only the projector direction. */
 	void SetProjectorDirection(const FString& ProjectorParamName, const FVector& Direction, int32 RangeIndex = -1);
@@ -377,7 +377,7 @@ private:
 	/** Mutable parameters optimization state. */
 	int32 State = 0;
 	
-	/** If this is set to true, when updating the instance an additional step will be performed to calculate the list of instance parameters that are relevant for the current parameter vaules. */
+	/** If this is set to true, when updating the instance an additional step will be performed to calculate the list of instance parameters that are relevant for the current parameter values. */
 	bool bBuildParameterRelevancy = false;
 
 	/** These are the LODs Mutable can generate, they MUST NOT be used in an update (Mutable thread). */
diff --git a/Engine/Plugins/Experimental/Mutable/Source/CustomizableObjectEditor/Private/MuCOE/CustomizableInstanceDetails.cpp b/Engine/Plugins/Experimental/Mutable/Source/CustomizableObjectEditor/Private/MuCOE/CustomizableInstanceDetails.cpp
--- a/Engine/Plugins/Experimental/Mutable/Source/CustomizableObjectEditor/Private/MuCOE/CustomizableInstanceDetails.cpp
+++ b/Engine/Plugins/Experimental/Mutable/Source/CustomizableObjectEditor/Private/MuCOE/CustomizableInstanceDetails.cpp
@@ -66,9 +66,9 @@ void FCustomizableInstanceDetails::CustomizeDetails(const TSharedPtr<IDetailLayo
 
 	WeakEditor = StaticCastSharedPtr<ICustomizableObjectInstanceEditor>(FoundAssetEditor).ToWeakPtr();
 
-	if (CustomInstance->IsSelectedParameterProfileDirty())
+	if (CustomInstance->GetPrivate()->IsSelectedParameterProfileDirty())
 	{
-		CustomInstance->SaveParametersToProfile(CustomInstance->SelectedProfileIndex);
+		CustomInstance->GetPrivate()->SaveParametersToProfile(CustomInstance->GetPrivate()->SelectedProfileIndex);
 	}
 
 	// Delegate to refresh the detils when the instance has finished the Update
@@ -157,7 +157,7 @@ void FCustomizableInstanceDetails::CustomizeDetails(const TSharedPtr<IDetailLayo
 	.HAlign(EHorizontalAlignment::HAlign_Fill)
 	[
 		SNew(SCheckBox)
-		.IsChecked(CustomInstance->bShowOnlyRuntimeParameters ? ECheckBoxState::Checked : ECheckBoxState::Unchecked)
+		.IsChecked(CustomInstance->GetPrivate()->bShowOnlyRuntimeParameters ? ECheckBoxState::Checked : ECheckBoxState::Unchecked)
 		.OnCheckStateChanged(this, &FCustomizableInstanceDetails::OnShowOnlyRuntimeSelectionChanged)
 	];
 	
@@ -172,7 +172,7 @@ void FCustomizableInstanceDetails::CustomizeDetails(const TSharedPtr<IDetailLayo
 	.HAlign(EHorizontalAlignment::HAlign_Fill)
 	[
 		SNew(SCheckBox)
-		.IsChecked(CustomInstance->bShowOnlyRelevantParameters ? ECheckBoxState::Checked : ECheckBoxState::Unchecked)
+		.IsChecked(CustomInstance->GetPrivate()->bShowOnlyRelevantParameters ? ECheckBoxState::Checked : ECheckBoxState::Unchecked)
 		.OnCheckStateChanged(this, &FCustomizableInstanceDetails::OnShowOnlyRelevantSelectionChanged)
 	];
 	
@@ -187,7 +187,7 @@ void FCustomizableInstanceDetails::CustomizeDetails(const TSharedPtr<IDetailLayo
 	.HAlign(EHorizontalAlignment::HAlign_Fill)
 	[
 		SNew(SCheckBox)
-		.IsChecked(CustomInstance->bShowUISections ? ECheckBoxState::Checked : ECheckBoxState::Unchecked)
+		.IsChecked(CustomInstance->GetPrivate()->bShowUISections ? ECheckBoxState::Checked : ECheckBoxState::Unchecked)
 		.OnCheckStateChanged(this, &FCustomizableInstanceDetails::OnUseUISectionsSelectionChanged)
 	];
 
@@ -257,7 +257,7 @@ TSharedRef<SWidget> FCustomizableInstanceDetails::GenerateStateSelector()
 
 	// States selector options
 	const int NumStates = CustomizableObject->GetStateCount();
-	const int CurrentState = CustomInstance->GetState();
+	const int CurrentState = CustomInstance->GetPrivate()->GetState();
 	TSharedPtr<FString> CurrentStateName = nullptr;
 	
 	//I think that this is not necessary. There is always a "Default" state
@@ -287,7 +287,7 @@ TSharedRef<SWidget> FCustomizableInstanceDetails::GenerateStateSelector()
 
 	return SNew(STextComboBox)
 		.OptionsSource(&StateNames)
-		.InitiallySelectedItem((CustomInstance->GetState() != -1) ? CurrentStateName : StateNames[0])
+		.InitiallySelectedItem((CustomInstance->GetPrivate()->GetState() != -1) ? CurrentStateName : StateNames[0])
 		.OnSelectionChanged(this, &FCustomizableInstanceDetails::OnStateComboBoxSelectionChanged);
 }
 
@@ -309,7 +309,7 @@ void FCustomizableInstanceDetails::OnStateComboBoxSelectionChanged(TSharedPtr<FS
 TSharedRef<SWidget> FCustomizableInstanceDetails::GenerateInstanceProfileSelector()
 {
 	UCustomizableObject* CustomizableObject = CustomInstance->GetCustomizableObject();
-	const int32 ProfileIdx = CustomInstance->SelectedProfileIndex;
+	const int32 ProfileIdx = CustomInstance->GetPrivate()->SelectedProfileIndex;
 
 	ParameterProfileNames.Emplace(MakeShared<FString>("None"));
 	TSharedPtr<FString> CurrentProfileName = ParameterProfileNames.Last();
@@ -363,7 +363,7 @@ TSharedRef<SWidget> FCustomizableInstanceDetails::GenerateInstanceProfileSelecto
 			SNew(SButton)
 			.Text(LOCTEXT("RemoveButtonLabel", " - "))
 			.ToolTipText(FText::FromString(CustomInstance->HasAnyParameters() ? FString("Delete selected profile") : FString("Delete selected profile functionality is not available, no profile is selected.")))
-			.IsEnabled(CustomInstance->SelectedProfileIndex != INDEX_NONE)
+			.IsEnabled(CustomInstance->GetPrivate()->SelectedProfileIndex != INDEX_NONE)
 			.IsFocusable(false)
 			.OnClicked(this, &FCustomizableInstanceDetails::RemoveParameterProfile)
 		];
@@ -425,7 +425,7 @@ FReply FCustomizableInstanceDetails::RemoveParameterProfile()
 {
 	UCustomizableObject* CustomizableObject = CustomInstance->GetCustomizableObject();
 
-	const int32 ProfileIdx = CustomInstance->SelectedProfileIndex;
+	const int32 ProfileIdx = CustomInstance->GetPrivate()->SelectedProfileIndex;
 
 	if (ProfileIdx == INDEX_NONE)
 	{
@@ -435,7 +435,7 @@ FReply FCustomizableInstanceDetails::RemoveParameterProfile()
 	TArray<FProfileParameterDat>& Profiles = CustomizableObject->InstancePropertiesProfiles;
 
 	Profiles.RemoveAt(ProfileIdx);
-	CustomInstance->SelectedProfileIndex = INDEX_NONE;
+	CustomInstance->GetPrivate()->SelectedProfileIndex = INDEX_NONE;
 	CustomizableObject->Modify();
 
 	CustomInstance->PreEditChange(nullptr);
@@ -451,15 +451,15 @@ FReply FCustomizableInstanceDetails::RemoveParameterProfile()
 
 void FCustomizableInstanceDetails::OnProfileSelectedChanged(TSharedPtr<FString> Selection, ESelectInfo::Type SelectInfo)
 {
-	const int32 ProfileIdx = CustomInstance->SelectedProfileIndex;
-	if (CustomInstance->IsSelectedParameterProfileDirty())
+	const int32 ProfileIdx = CustomInstance->GetPrivate()->SelectedProfileIndex;
+	if (CustomInstance->GetPrivate()->IsSelectedParameterProfileDirty())
 	{
-		CustomInstance->SaveParametersToProfile(ProfileIdx);
+		CustomInstance->GetPrivate()->SaveParametersToProfile(ProfileIdx);
 	}
 
 	if (*Selection == "None")
 	{
-		CustomInstance->SelectedProfileIndex = INDEX_NONE;
+		CustomInstance->GetPrivate()->SelectedProfileIndex = INDEX_NONE;
 	}
 	else
 	{
@@ -469,13 +469,13 @@ void FCustomizableInstanceDetails::OnProfileSelectedChanged(TSharedPtr<FString>
 		{
 			if (Profiles[Idx].ProfileName == *Selection)
 			{
-				CustomInstance->SelectedProfileIndex = Idx;
+				CustomInstance->GetPrivate()->SelectedProfileIndex = Idx;
 				break;
 			}
 		}
 	}	
 
-	CustomInstance->LoadParametersFromProfile(CustomInstance->SelectedProfileIndex);
+	CustomInstance->GetPrivate()->LoadParametersFromProfile(CustomInstance->GetPrivate()->SelectedProfileIndex);
 	CustomInstance->PreEditChange(nullptr);
 	CustomInstance->UpdateSkeletalMeshAsync(true, true);
 	CustomInstance->PostEditChange();
@@ -487,21 +487,21 @@ void FCustomizableInstanceDetails::OnProfileSelectedChanged(TSharedPtr<FString>
 
 void FCustomizableInstanceDetails::OnShowOnlyRuntimeSelectionChanged(ECheckBoxState InCheckboxState)
 {
-	CustomInstance->bShowOnlyRuntimeParameters = InCheckboxState == ECheckBoxState::Checked;
+	CustomInstance->GetPrivate()->bShowOnlyRuntimeParameters = InCheckboxState == ECheckBoxState::Checked;
 	Refresh();
 }
 
 
 void FCustomizableInstanceDetails::OnShowOnlyRelevantSelectionChanged(ECheckBoxState InCheckboxState)
 {
-	CustomInstance->bShowOnlyRelevantParameters = InCheckboxState == ECheckBoxState::Checked;
+	CustomInstance->GetPrivate()->bShowOnlyRelevantParameters = InCheckboxState == ECheckBoxState::Checked;
 	Refresh();
 }
 
 
 void FCustomizableInstanceDetails::OnUseUISectionsSelectionChanged(ECheckBoxState InCheckboxState)
 {
-	CustomInstance->bShowUISections = InCheckboxState == ECheckBoxState::Checked;
+	CustomInstance->GetPrivate()->bShowUISections = InCheckboxState == ECheckBoxState::Checked;
 	Refresh();
 }
 
@@ -550,14 +550,14 @@ void FCustomizableInstanceDetails::GenerateParametersView(IDetailCategoryBuilder
 		return;
 	}
 	
-	if (CustomInstance->bShowOnlyRuntimeParameters)
+	if (CustomInstance->GetPrivate()->bShowOnlyRuntimeParameters)
 	{
-		const int32 NumStateParameters = CustomizableObject->GetStateParameterCount(CustomInstance->GetState());
+		const int32 NumStateParameters = CustomizableObject->GetStateParameterCount(CustomInstance->GetPrivate()->GetState());
 
 		for (int32 ParamIndexInState = 0; ParamIndexInState < NumStateParameters; ++ParamIndexInState)
 		{
 			FParameterInfo ParameterSortInfo;
-			ParameterSortInfo.ParameterIndex = CustomizableObject->GetStateParameterIndex(CustomInstance->GetState(), ParamIndexInState);
+			ParameterSortInfo.ParameterIndex = CustomizableObject->GetStateParameterIndex(CustomInstance->GetPrivate()->GetState(), ParamIndexInState);
 
 			if (CustomInstance->IsParameterRelevant(ParameterSortInfo.ParameterIndex) && IsVisible(ParameterSortInfo.ParameterIndex))
 			{
@@ -573,7 +573,7 @@ void FCustomizableInstanceDetails::GenerateParametersView(IDetailCategoryBuilder
 		{
 			const FParameterInfo& ParamInfo = ParametersTree[ParamIndexInObject];
 			
-			if (CustomInstance->bShowUISections)
+			if (CustomInstance->GetPrivate()->bShowUISections)
 			{
 				IDetailGroup* CurrentSection = GenerateParameterSection(ParamIndexInObject, DetailsCategory);
 				check(CurrentSection);
@@ -617,7 +617,7 @@ void FCustomizableInstanceDetails::GenerateParametersView(IDetailCategoryBuilder
 		{
 			FParameterInfo ParameterSortInfo;
 			ParameterSortInfo.ParameterIndex = ParamIndexInObject;
-			if ((!CustomInstance->bShowOnlyRelevantParameters || CustomInstance->IsParameterRelevant(ParameterSortInfo.ParameterIndex)) && IsVisible(ParameterSortInfo.ParameterIndex))
+			if ((!CustomInstance->GetPrivate()->bShowOnlyRelevantParameters || CustomInstance->IsParameterRelevant(ParameterSortInfo.ParameterIndex)) && IsVisible(ParameterSortInfo.ParameterIndex))
 			{
 				ParameterSortInfo.ParameterUIOrder = CustomizableObject->GetParameterUIMetadataFromIndex(ParameterSortInfo.ParameterIndex).ParamUIMetadata.UIOrder;
 				ParameterSortInfo.ParameterName = CustomizableObject->GetParameterName(ParameterSortInfo.ParameterIndex);
@@ -653,7 +653,7 @@ void FCustomizableInstanceDetails::RecursivelyAddParamAndChildren(const int32 Pa
 
 	if (ParentName.IsEmpty())
 	{
-		if (CustomInstance->bShowUISections)
+		if (CustomInstance->GetPrivate()->bShowUISections)
 		{
 			IDetailGroup* CurrentSection = GenerateParameterSection(ParamIndexInObject, DetailsCategory);
 			check(CurrentSection);
@@ -740,7 +740,7 @@ bool FCustomizableInstanceDetails::IsVisible(int32 ParamIndexInObject)
 		|| (ParamName.EndsWith(FMultilayerProjector::OPACITY_PARAMETER_POSTFIX) && CustomizableObject->IsParameterMultidimensional(ParamIndexInObject))
 		|| (ParamName.EndsWith(FMultilayerProjector::POSE_PARAMETER_POSTFIX));
 
-	if (!IsAProjectorParam && ParentName && CustomInstance->bShowOnlyRelevantParameters)
+	if (!IsAProjectorParam && ParentName && CustomInstance->GetPrivate()->bShowOnlyRelevantParameters)
 	{
 		FString* Value = UIData.ParamUIMetadata.ExtraInformation.Find(FString("__DisplayWhenParentValueEquals"));
 
@@ -940,7 +940,7 @@ void FCustomizableInstanceDetails::OnIntParameterComboBoxChanged(TSharedPtr<FStr
 		}
 	}
 
-	CustomInstance->SetSelectedParameterProfileDirty();
+	CustomInstance->GetPrivate()->SetSelectedParameterProfileDirty();
 	CustomInstance->PreEditChange(nullptr);
 	CustomInstance->UpdateSkeletalMeshAsync(true, true);
 	CustomInstance->PostEditChange();
@@ -995,7 +995,7 @@ void FCustomizableInstanceDetails::OnFloatParameterSliderEnd(float Value, FStrin
 {
 	CustomInstance->SetFloatParameterSelectedOption(ParamName, Value, RangeIndex);
 
-	CustomInstance->SetSelectedParameterProfileDirty();
+	CustomInstance->GetPrivate()->SetSelectedParameterProfileDirty();
 	CustomInstance->PreEditChange(nullptr);
 	CustomInstance->UpdateSkeletalMeshAsync(true, true);
 	CustomInstance->PostEditChange();
@@ -1070,7 +1070,7 @@ void FCustomizableInstanceDetails::GenerateTextureParameterOptions()
 	}
 
 	// Get values from TextureParameterDeclarations.
-	for (TObjectPtr<UTexture2D> Declaration : CustomInstance->GetTextureParameterDeclarations())
+	for (TObjectPtr<UTexture2D> Declaration : CustomInstance->TextureParameterDeclarations)
 	{
 		if (!Declaration)
 		{
@@ -1108,7 +1108,7 @@ void FCustomizableInstanceDetails::OnTextureParameterComboBoxSelectionChanged(TS
 		}
 	}
 
-	CustomInstance->SetSelectedParameterProfileDirty();
+	CustomInstance->GetPrivate()->SetSelectedParameterProfileDirty();
 	CustomInstance->UpdateSkeletalMeshAsync(true, true);
 	CustomInstance->PostEditChange();
 
@@ -1165,7 +1165,7 @@ void FCustomizableInstanceDetails::OnSetColorFromColorPicker(FLinearColor NewCol
 {
 	CustomInstance->PreEditChange(nullptr);
 	CustomInstance->SetColorParameterSelectedOption(PickerParamName, NewColor);
-	CustomInstance->SetSelectedParameterProfileDirty();
+	CustomInstance->GetPrivate()->SetSelectedParameterProfileDirty();
 
 	PickerParamName = FString();
 
@@ -1675,7 +1675,7 @@ FReply FCustomizableInstanceDetails::OnProjectorLayerAdded(FString ParamName) co
 	check(CustomInstance->FindFloatParameterNameIndex(OpacitySliderParamName) != INDEX_NONE);
 	check(NumLayers == FloatParameters[CustomInstance->FindFloatParameterNameIndex(OpacitySliderParamName)].ParameterRangeValues.Num());
 
-	CustomInstance->SetSelectedParameterProfileDirty();
+	CustomInstance->GetPrivate()->SetSelectedParameterProfileDirty();
 	CustomInstance->PreEditChange(nullptr);
 	CustomInstance->UpdateSkeletalMeshAsync(true, true);
 	CustomInstance->PostEditChange();
@@ -1713,7 +1713,7 @@ FReply FCustomizableInstanceDetails::OnProjectorLayerRemoved(const FString Param
 	check(CustomInstance->FindFloatParameterNameIndex(OpacitySliderParamName) != INDEX_NONE);
 	check(NumLayers == FloatParameters[CustomInstance->FindFloatParameterNameIndex(OpacitySliderParamName)].ParameterRangeValues.Num());
 
-	CustomInstance->SetSelectedParameterProfileDirty();
+	CustomInstance->GetPrivate()->SetSelectedParameterProfileDirty();
 	CustomInstance->PreEditChange(nullptr);
 	CustomInstance->UpdateSkeletalMeshAsync(true, true);
 	CustomInstance->PostEditChange();
@@ -1761,7 +1761,7 @@ void FCustomizableInstanceDetails::OnProjectorTextureParameterComboBoxChanged(TS
 		}
 	}
 
-	CustomInstance->SetSelectedParameterProfileDirty();
+	CustomInstance->GetPrivate()->SetSelectedParameterProfileDirty();
 	CustomInstance->UpdateSkeletalMeshAsync(true, true);
 	CustomInstance->PostEditChange();
 
@@ -1809,7 +1809,7 @@ void FCustomizableInstanceDetails::OnBoolParameterChanged(ECheckBoxState InCheck
 {
 	CustomInstance->PreEditChange(nullptr);
 	CustomInstance->SetBoolParameterSelectedOption(ParamName, InCheckboxState == ECheckBoxState::Checked);
-	CustomInstance->SetSelectedParameterProfileDirty();
+	CustomInstance->GetPrivate()->SetSelectedParameterProfileDirty();
 	CustomInstance->UpdateSkeletalMeshAsync(true, true);
 	CustomInstance->PostEditChange();
 
@@ -1847,7 +1847,7 @@ FReply FCustomizableInstanceDetails::OnPasteAllParameters()
 		FromBinary.Seek(0);
 
 		CustomInstance->LoadDescriptor(FromBinary);
-		CustomInstance->SetSelectedParameterProfileDirty();
+		CustomInstance->GetPrivate()->SetSelectedParameterProfileDirty();
 		CustomInstance->UpdateSkeletalMeshAsync(true, true);
 		CustomInstance->PostEditChange();
 
@@ -1875,7 +1875,7 @@ FReply FCustomizableInstanceDetails::OnResetAllParameters()
 
 	CustomInstance->UpdateSkeletalMeshAsync(true, true);
 	CustomInstance->PostEditChange();
-	CustomInstance->SelectedProfileIndex = INDEX_NONE;
+	CustomInstance->GetPrivate()->SelectedProfileIndex = INDEX_NONE;
 
 	// Non-continuous change: collect garbage.
 	GEngine->ForceGarbageCollection();
@@ -1888,7 +1888,7 @@ void FCustomizableInstanceDetails::OnResetParameterButtonClicked(int32 Parameter
 {	
 	CustomInstance->PreEditChange(nullptr);
 	SetParameterValueToDefault(ParameterIndex);
-	CustomInstance->SetSelectedParameterProfileDirty();
+	CustomInstance->GetPrivate()->SetSelectedParameterProfileDirty();
 	CustomInstance->UpdateSkeletalMeshAsync(true, true);
 	CustomInstance->PostEditChange();
 }
@@ -2101,11 +2101,11 @@ void SProfileParametersWindow::OnNameChange(const FText& NewName, ETextCommit::T
 		UCustomizableObject* CustomizableObject = CustomInstance->GetCustomizableObject();
 		CustomizableObject->GetPrivate()->AddNewParameterProfile(GetFileName(), *CustomInstance.Get());
 
-		if (CustomInstance->bSelectedProfileDirty && CustomInstance->SelectedProfileIndex != INDEX_NONE)
+		if (CustomInstance->GetPrivate()->bSelectedProfileDirty && CustomInstance->GetPrivate()->SelectedProfileIndex != INDEX_NONE)
 		{
-			CustomInstance->SaveParametersToProfile(CustomInstance->SelectedProfileIndex);
+			CustomInstance->GetPrivate()->SaveParametersToProfile(CustomInstance->GetPrivate()->SelectedProfileIndex);
 		}
-		CustomInstance->SelectedProfileIndex = CustomizableObject->InstancePropertiesProfiles.Num() - 1;
+		CustomInstance->GetPrivate()->SelectedProfileIndex = CustomizableObject->InstancePropertiesProfiles.Num() - 1;
 
 		if (InstanceDetails)
 		{
@@ -2126,11 +2126,11 @@ FReply SProfileParametersWindow::OnButtonClick(EAppReturnType::Type ButtonID)
 		UCustomizableObject* CustomizableObject = CustomInstance->GetCustomizableObject();
 		CustomizableObject->GetPrivate()->AddNewParameterProfile(GetFileName(), *CustomInstance.Get());
 
-		if (CustomInstance->bSelectedProfileDirty && CustomInstance->SelectedProfileIndex != INDEX_NONE)
+		if (CustomInstance->GetPrivate()->bSelectedProfileDirty && CustomInstance->GetPrivate()->SelectedProfileIndex != INDEX_NONE)
 		{
-			CustomInstance->SaveParametersToProfile(CustomInstance->SelectedProfileIndex);
+			CustomInstance->GetPrivate()->SaveParametersToProfile(CustomInstance->GetPrivate()->SelectedProfileIndex);
 		}
-		CustomInstance->SelectedProfileIndex = CustomizableObject->InstancePropertiesProfiles.Num() - 1;
+		CustomInstance->GetPrivate()->SelectedProfileIndex = CustomizableObject->InstancePropertiesProfiles.Num() - 1;
 
 		if (InstanceDetails)
 		{
diff --git a/Engine/Plugins/Experimental/Mutable/Source/CustomizableObjectEditor/Private/MuCOE/CustomizableObjectEditor.cpp b/Engine/Plugins/Experimental/Mutable/Source/CustomizableObjectEditor/Private/MuCOE/CustomizableObjectEditor.cpp
--- a/Engine/Plugins/Experimental/Mutable/Source/CustomizableObjectEditor/Private/MuCOE/CustomizableObjectEditor.cpp
+++ b/Engine/Plugins/Experimental/Mutable/Source/CustomizableObjectEditor/Private/MuCOE/CustomizableObjectEditor.cpp
@@ -166,9 +166,9 @@ FCustomizableObjectEditor::~FCustomizableObjectEditor()
 
 	if (PreviewInstance)
 	{
-		if (PreviewInstance->bSelectedProfileDirty && PreviewInstance->SelectedProfileIndex != INDEX_NONE)
+		if (PreviewInstance->GetPrivate()->bSelectedProfileDirty && PreviewInstance->GetPrivate()->SelectedProfileIndex != INDEX_NONE)
 		{
-			PreviewInstance->SaveParametersToProfile(PreviewInstance->SelectedProfileIndex);
+			PreviewInstance->GetPrivate()->SaveParametersToProfile(PreviewInstance->GetPrivate()->SelectedProfileIndex);
 		}
 	}
 	
@@ -1577,9 +1577,9 @@ void FCustomizableObjectEditor::SaveAsset_Execute()
 {
 	if (PreviewInstance)
 	{
-		if (PreviewInstance->IsSelectedParameterProfileDirty())
+		if (PreviewInstance->GetPrivate()->IsSelectedParameterProfileDirty())
 		{
-			PreviewInstance->SaveParametersToProfile(PreviewInstance->SelectedProfileIndex);
+			PreviewInstance->GetPrivate()->SaveParametersToProfile(PreviewInstance->GetPrivate()->SelectedProfileIndex);
 		}
 	}
 
diff --git a/Engine/Plugins/Experimental/Mutable/Source/CustomizableObjectEditor/Private/MuCOE/CustomizableObjectInstanceEditor.cpp b/Engine/Plugins/Experimental/Mutable/Source/CustomizableObjectEditor/Private/MuCOE/CustomizableObjectInstanceEditor.cpp
--- a/Engine/Plugins/Experimental/Mutable/Source/CustomizableObjectEditor/Private/MuCOE/CustomizableObjectInstanceEditor.cpp
+++ b/Engine/Plugins/Experimental/Mutable/Source/CustomizableObjectEditor/Private/MuCOE/CustomizableObjectInstanceEditor.cpp
@@ -270,8 +270,8 @@ void FCustomizableObjectInstanceEditor::InitCustomizableObjectInstanceEditor( co
 	// Set the instance
 	check(InCustomizableObjectInstance);
 	CustomizableObjectInstance = InCustomizableObjectInstance;
-	bOnlyRelevantParameters = InCustomizableObjectInstance->bShowOnlyRelevantParameters;
-	bOnlyRuntimeParameters = InCustomizableObjectInstance->bShowOnlyRuntimeParameters;
+	bOnlyRelevantParameters = InCustomizableObjectInstance->GetPrivate()->bShowOnlyRelevantParameters;
+	bOnlyRuntimeParameters = InCustomizableObjectInstance->GetPrivate()->bShowOnlyRuntimeParameters;
 
 	FAdvancedPreviewSceneModule& AdvancedPreviewSceneModule = FModuleManager::LoadModuleChecked<FAdvancedPreviewSceneModule>("AdvancedPreviewScene");
 
@@ -713,15 +713,15 @@ void FCustomizableObjectInstanceEditor::Tick(float InDeltaTime)
 	}
 
 	// If we want to show the Relevant/Runtime parameters, we need to refresh the details view to make sure that the scroll bar appears		
-	if (bOnlyRelevantParameters != CustomizableObjectInstance->bShowOnlyRelevantParameters)
+	if (bOnlyRelevantParameters != CustomizableObjectInstance->GetPrivate()->bShowOnlyRelevantParameters)
 	{
-		bOnlyRelevantParameters = CustomizableObjectInstance->bShowOnlyRelevantParameters;
+		bOnlyRelevantParameters = CustomizableObjectInstance->GetPrivate()->bShowOnlyRelevantParameters;
 		CustomizableInstanceDetailsView->ForceRefresh();
 	}
 
-	if (bOnlyRuntimeParameters != CustomizableObjectInstance->bShowOnlyRuntimeParameters)
+	if (bOnlyRuntimeParameters != CustomizableObjectInstance->GetPrivate()->bShowOnlyRuntimeParameters)
 	{
-		bOnlyRuntimeParameters = CustomizableObjectInstance->bShowOnlyRuntimeParameters;
+		bOnlyRuntimeParameters = CustomizableObjectInstance->GetPrivate()->bShowOnlyRuntimeParameters;
 		CustomizableInstanceDetailsView->ForceRefresh();
 	}
 }
@@ -841,7 +841,7 @@ void FCustomizableObjectInstanceEditor::ShowGizmoProjectorParameter(const FStrin
 		
 		Editor->GetProjectorParameter()->SetPosition(Location);
 		
-		Instance->SetProjectorPosition(ParamName, static_cast<FVector3f>(Location), RangeIndex);
+		Instance->SetProjectorPosition(ParamName, Location, RangeIndex);
 		Instance->UpdateSkeletalMeshAsync(true, true);
 	});
 
diff --git a/Engine/Plugins/Experimental/Mutable/Source/CustomizableObjectEditor/Private/MuCOE/SCustomizableInstanceProperties.cpp b/Engine/Plugins/Experimental/Mutable/Source/CustomizableObjectEditor/Private/MuCOE/SCustomizableInstanceProperties.cpp
--- a/Engine/Plugins/Experimental/Mutable/Source/CustomizableObjectEditor/Private/MuCOE/SCustomizableInstanceProperties.cpp
+++ b/Engine/Plugins/Experimental/Mutable/Source/CustomizableObjectEditor/Private/MuCOE/SCustomizableInstanceProperties.cpp
@@ -10,6 +10,7 @@
 #include "Misc/Paths.h"
 #include "MuCO/CustomizableObject.h"
 #include "MuCO/CustomizableObjectInstance.h"
+#include "MuCO/CustomizableObjectInstancePrivate.h"
 #include "MuCO/CustomizableObjectSystem.h"
 #include "MuCOE/CustomizableInstanceDetails.h"
 #include "MuCOE/CustomizableObjectEditorUtilities.h"
@@ -51,9 +52,9 @@ void SCustomizableInstanceProperties::Construct(const FArguments& InArgs)
 	
 	NoInstanceMessage = LOCTEXT("Model not compiled", "Model not compiled");
 	
-	if (CustomInstance->IsSelectedParameterProfileDirty())
+	if (CustomInstance->GetPrivate()->IsSelectedParameterProfileDirty())
 	{
-		CustomInstance->SaveParametersToProfile(CustomInstance->SelectedProfileIndex);
+		CustomInstance->GetPrivate()->SaveParametersToProfile(CustomInstance->GetPrivate()->SelectedProfileIndex);
 	}
 	
 	CustomInstance->UpdatedNativeDelegate.AddSP(this, &SCustomizableInstanceProperties::InstanceUpdated);
@@ -66,7 +67,7 @@ void SCustomizableInstanceProperties::Construct(const FArguments& InArgs)
 	{
 		// Store the state data required for the ui.
 		const int numStates = CustomizableObject->GetStateCount();
-		const int currentState = CustomInstance->GetState();
+		const int currentState = CustomInstance->GetPrivate()->GetState();
 		TSharedPtr<FString> CurrentStateName = nullptr;
 
 		for (int i = 0; i < numStates; ++i)
@@ -127,7 +128,7 @@ void SCustomizableInstanceProperties::Construct(const FArguments& InArgs)
 		}
 
 		// Get values from TextureParameterDeclarations.
-		for (TObjectPtr<UTexture2D> Declaration : CustomInstance->GetTextureParameterDeclarations())
+		for (TObjectPtr<UTexture2D> Declaration : CustomInstance->TextureParameterDeclarations)
 		{
 			if (!Declaration)
 			{
@@ -147,7 +148,7 @@ void SCustomizableInstanceProperties::Construct(const FArguments& InArgs)
 		
 		int32 ProfileStringIndex = 0;
 
-		const int32 ProfileIdx = CustomInstance->SelectedProfileIndex;
+		const int32 ProfileIdx = CustomInstance->GetPrivate()->SelectedProfileIndex;
 
 		if (ProfileIdx != INDEX_NONE)
 		{
@@ -189,7 +190,7 @@ void SCustomizableInstanceProperties::Construct(const FArguments& InArgs)
 				[
 					SNew(STextComboBox)
 					.OptionsSource(&StateNames)
-					.InitiallySelectedItem((CustomInstance->GetState() != -1) ? CurrentStateName : StateNames[0])
+					.InitiallySelectedItem((CustomInstance->GetPrivate()->GetState() != -1) ? CurrentStateName : StateNames[0])
 					.OnSelectionChanged(this, &SCustomizableInstanceProperties::OnStateComboBoxSelectionChanged)
 				]
 			]
@@ -204,7 +205,7 @@ void SCustomizableInstanceProperties::Construct(const FArguments& InArgs)
 				.VAlign(VAlign_Center)
 				[
 					SNew(SCheckBox)
-					.IsChecked(CustomInstance->bShowOnlyRuntimeParameters ? ECheckBoxState::Checked : ECheckBoxState::Unchecked)
+					.IsChecked(CustomInstance->GetPrivate()->bShowOnlyRuntimeParameters ? ECheckBoxState::Checked : ECheckBoxState::Unchecked)
 					.OnCheckStateChanged(this, &SCustomizableInstanceProperties::OnShowOnlyRuntimeSelectionChanged)
 				]
 
@@ -227,7 +228,7 @@ void SCustomizableInstanceProperties::Construct(const FArguments& InArgs)
 				.VAlign(VAlign_Center)
 				[
 					SNew(SCheckBox)
-					.IsChecked(CustomInstance->bShowOnlyRelevantParameters ? ECheckBoxState::Checked : ECheckBoxState::Unchecked)
+					.IsChecked(CustomInstance->GetPrivate()->bShowOnlyRelevantParameters ? ECheckBoxState::Checked : ECheckBoxState::Unchecked)
 					.OnCheckStateChanged(this, &SCustomizableInstanceProperties::OnShowOnlyRelevantSelectionChanged)
 				]
 			
@@ -296,7 +297,7 @@ void SCustomizableInstanceProperties::Construct(const FArguments& InArgs)
 						SNew(SButton)
 						.Text(LOCTEXT("RemoveButtonLabel", " - "))
 						.ToolTipText(FText::FromString(HasAnyParameters() ? FString("Delete selected profile") : FString("Delete selected profile functionality is not available, no profile is selected.")))
-						.IsEnabled(CustomInstance->SelectedProfileIndex != INDEX_NONE)
+						.IsEnabled(CustomInstance->GetPrivate()->SelectedProfileIndex != INDEX_NONE)
 						.IsFocusable(false)
 						.OnClicked(this, &SCustomizableInstanceProperties::RemoveParameterProfile)
 					]
@@ -365,7 +366,7 @@ void SCustomizableInstanceProperties::Construct(const FArguments& InArgs)
 				[
 					SNew(SSearchBox)
 					.HintText(LOCTEXT("SearchHint", "Search Properties"))
-					.InitialText(CustomInstance->ParametersSearchFilter)
+					.InitialText(CustomInstance->GetPrivate()->ParametersSearchFilter)
 					.OnTextChanged(this, &SCustomizableInstanceProperties::OnFilterTextChanged)
 				]
 			]
@@ -452,14 +453,14 @@ void SCustomizableInstanceProperties::ResetParamBox()
 
 	if (UCustomizableObject* CustomizableObject = CustomInstance->GetCustomizableObject())
 	{
-		if (CustomInstance->bShowOnlyRuntimeParameters)
+		if (CustomInstance->GetPrivate()->bShowOnlyRuntimeParameters)
 		{
-			const int32 NumStateParameters = CustomizableObject->GetStateParameterCount(CustomInstance->GetState());
+			const int32 NumStateParameters = CustomizableObject->GetStateParameterCount(CustomInstance->GetPrivate()->GetState());
 			TArray<FCustomizableInstancePropertyArrayElem> ParamIndexesInState;
 			for (int32 ParamIndexInState = 0; ParamIndexInState < NumStateParameters; ++ParamIndexInState)
 			{
 				FCustomizableInstancePropertyArrayElem ParameterSortInfo;
-				ParameterSortInfo.PropertyCOIndex = CustomizableObject->GetStateParameterIndex(CustomInstance->GetState(), ParamIndexInState);
+				ParameterSortInfo.PropertyCOIndex = CustomizableObject->GetStateParameterIndex(CustomInstance->GetPrivate()->GetState(), ParamIndexInState);
 				if (CustomInstance->IsParameterRelevant(ParameterSortInfo.PropertyCOIndex))
 				{
 					ParameterSortInfo.UIOrder = CustomizableObject->GetParameterUIMetadataFromIndex(ParameterSortInfo.PropertyCOIndex).ParamUIMetadata.UIOrder;
@@ -485,7 +486,7 @@ void SCustomizableInstanceProperties::ResetParamBox()
 			{
 				FCustomizableInstancePropertyArrayElem ParameterSortInfo;
 				ParameterSortInfo.PropertyCOIndex = ParamIndexInObject;
-				if (!CustomInstance->bShowOnlyRelevantParameters || CustomInstance->IsParameterRelevant(ParameterSortInfo.PropertyCOIndex))
+				if (!CustomInstance->GetPrivate()->bShowOnlyRelevantParameters || CustomInstance->IsParameterRelevant(ParameterSortInfo.PropertyCOIndex))
 				{
 					ParameterSortInfo.UIOrder = CustomizableObject->GetParameterUIMetadataFromIndex(ParameterSortInfo.PropertyCOIndex).ParamUIMetadata.UIOrder;
 					ParameterSortInfo.ParameterName = CustomizableObject->GetParameterName(ParameterSortInfo.PropertyCOIndex);
@@ -567,7 +568,7 @@ void SCustomizableInstanceProperties::AddParameter(int32 ParamIndexInObject)
 		|| (ParamName.EndsWith(FMultilayerProjector::OPACITY_PARAMETER_POSTFIX) && CustomizableObject->IsParameterMultidimensional(ParamIndexInObject))
 		|| (ParamName.EndsWith(FMultilayerProjector::POSE_PARAMETER_POSTFIX));
 	
-	const FString Name = CustomInstance->ParametersSearchFilter.ToString();
+	const FString Name = CustomInstance->GetPrivate()->ParametersSearchFilter.ToString();
 
 	//If SearchBox text is not empty
 	if (!Name.IsEmpty())
@@ -596,7 +597,7 @@ void SCustomizableInstanceProperties::AddParameter(int32 ParamIndexInObject)
 			}
 		}
 
-		if (CustomInstance->bShowOnlyRelevantParameters)
+		if (CustomInstance->GetPrivate()->bShowOnlyRelevantParameters)
 		{
 			FString *Value = UIData.ParamUIMetadata.ExtraInformation.Find(FString("__DisplayWhenParentValueEquals"));
 			if (Value && CustomInstance->GetIntParameterSelectedOption(*ParentName) != *Value)
@@ -611,7 +612,7 @@ void SCustomizableInstanceProperties::AddParameter(int32 ParamIndexInObject)
 		TSharedPtr<SExpandableArea> ChildParamExpandableArea;
 		TSharedPtr<SVerticalBox> VerticalBox;
 
-		bool* bIsExpanded = CustomInstance->ParamNameToExpandedMap.Find(ParamName);
+		bool* bIsExpanded = CustomInstance->GetPrivate()->ParamNameToExpandedMap.Find(ParamName);
 
 		ActualParamBox->AddSlot()
 			.AutoHeight()
@@ -1314,7 +1315,7 @@ void SCustomizableInstanceProperties::Tick(const FGeometry& AllottedGeometry, co
 						FloatParameters[i].ParameterRangeValues[RangeIndex] = NewValue;
 					}
 
-					CustomInstance->SetSelectedParameterProfileDirty();
+					CustomInstance->GetPrivate()->SetSelectedParameterProfileDirty();
 
 					CustomInstance->PostEditChange();
 					CustomInstance->UpdateSkeletalMeshAsync(true, true);
@@ -1329,14 +1330,14 @@ void SCustomizableInstanceProperties::Tick(const FGeometry& AllottedGeometry, co
 
 void SCustomizableInstanceProperties::OnShowOnlyRuntimeSelectionChanged(ECheckBoxState InCheckboxState)
 {
-	CustomInstance->bShowOnlyRuntimeParameters = InCheckboxState == ECheckBoxState::Checked;
+	CustomInstance->GetPrivate()->bShowOnlyRuntimeParameters = InCheckboxState == ECheckBoxState::Checked;
 	ResetParamBox();
 }
 
 
 void SCustomizableInstanceProperties::OnShowOnlyRelevantSelectionChanged(ECheckBoxState InCheckboxState)
 {
-	CustomInstance->bShowOnlyRelevantParameters = InCheckboxState == ECheckBoxState::Checked;
+	CustomInstance->GetPrivate()->bShowOnlyRelevantParameters = InCheckboxState == ECheckBoxState::Checked;
 	ResetParamBox();
 }
 
@@ -1396,7 +1397,7 @@ FReply SCustomizableInstanceProperties::OnPasteAllParameters()
 
 		CustomInstance->LoadDescriptor(FromBinary);
 
-		CustomInstance->SetSelectedParameterProfileDirty();
+		CustomInstance->GetPrivate()->SetSelectedParameterProfileDirty();
 
 		CustomInstance->UpdateSkeletalMeshAsync(true, true);
 
@@ -1431,7 +1432,7 @@ FReply SCustomizableInstanceProperties::OnResetAllParameters()
 	// Non-continuous change: collect garbage.
 	GEngine->ForceGarbageCollection();
 
-	CustomInstance->SelectedProfileIndex = INDEX_NONE;
+	CustomInstance->GetPrivate()->SelectedProfileIndex = INDEX_NONE;
 
 	InstanceDetails.Pin()->Refresh();
 
@@ -1450,11 +1451,11 @@ void SCustomizableInstanceProperties::OnAreaExpansionChanged(bool bExpanded, FSt
 {
 	if (bExpanded)
 	{
-		CustomInstance->ParamNameToExpandedMap.Add(ParamName, true);
+		CustomInstance->GetPrivate()->ParamNameToExpandedMap.Add(ParamName, true);
 	}
 	else
 	{
-		CustomInstance->ParamNameToExpandedMap.Remove(ParamName);
+		CustomInstance->GetPrivate()->ParamNameToExpandedMap.Remove(ParamName);
 	}
 }
 
@@ -1465,7 +1466,7 @@ void SCustomizableInstanceProperties::OnBoolParameterChanged(ECheckBoxState InCh
 
 	CustomInstance->SetBoolParameterSelectedOption(ParamName, InCheckboxState == ECheckBoxState::Checked);
 
-	CustomInstance->SetSelectedParameterProfileDirty();
+	CustomInstance->GetPrivate()->SetSelectedParameterProfileDirty();
 
 	CustomInstance->UpdateSkeletalMeshAsync(true, true);
 
@@ -1591,7 +1592,7 @@ void SCustomizableInstanceProperties::OnIntParameterComboBoxChanged(TSharedPtr<F
 		}
 	}
 
-	CustomInstance->SetSelectedParameterProfileDirty();
+	CustomInstance->GetPrivate()->SetSelectedParameterProfileDirty();
 
 	CustomInstance->PreEditChange(nullptr);
 
@@ -1624,7 +1625,7 @@ void SCustomizableInstanceProperties::OnTextureParameterComboBoxSelectionChanged
 		}
 	}
 
-	CustomInstance->SetSelectedParameterProfileDirty();
+	CustomInstance->GetPrivate()->SetSelectedParameterProfileDirty();
 
 	CustomInstance->UpdateSkeletalMeshAsync(true, true);
 
@@ -1637,7 +1638,7 @@ void SCustomizableInstanceProperties::OnTextureParameterComboBoxSelectionChanged
 
 void SCustomizableInstanceProperties::OnFilterTextChanged(const FText & InFilterText)
 {
-	CustomInstance->ParametersSearchFilter = InFilterText;
+	CustomInstance->GetPrivate()->ParametersSearchFilter = InFilterText;
 	ResetParamBox();
 }
 
@@ -1669,7 +1670,7 @@ void SCustomizableInstanceProperties::OnSetColorFromColorPicker(FLinearColor New
 
 	CustomInstance->SetColorParameterSelectedOption(PickerParamName, NewColor);
 	
-	CustomInstance->SetSelectedParameterProfileDirty();
+	CustomInstance->GetPrivate()->SetSelectedParameterProfileDirty();
 
 	PickerParamName = FString();
 
@@ -1703,7 +1704,7 @@ FReply SCustomizableInstanceProperties::OnProjectorLayerAdded(FString ParamName)
 	check(CustomInstance->FindFloatParameterNameIndex(OpacitySliderParamName) != INDEX_NONE);
 	check(NumLayers == FloatParameters[CustomInstance->FindFloatParameterNameIndex(OpacitySliderParamName)].ParameterRangeValues.Num());
 
-	CustomInstance->SetSelectedParameterProfileDirty();
+	CustomInstance->GetPrivate()->SetSelectedParameterProfileDirty();
 
 	CustomInstance->PreEditChange(nullptr);
 
@@ -1743,7 +1744,7 @@ FReply SCustomizableInstanceProperties::OnProjectorLayerRemoved(const FString Pa
 	check(CustomInstance->FindFloatParameterNameIndex(OpacitySliderParamName) != INDEX_NONE);
 	check(NumLayers == FloatParameters[CustomInstance->FindFloatParameterNameIndex(OpacitySliderParamName)].ParameterRangeValues.Num());
 
-	CustomInstance->SetSelectedParameterProfileDirty();
+	CustomInstance->GetPrivate()->SetSelectedParameterProfileDirty();
 	
 	CustomInstance->PreEditChange(nullptr);
 
@@ -1811,7 +1812,7 @@ void SCustomizableInstanceProperties::OnProjectorTextureParameterComboBoxChanged
 		}
 	}
 
-	CustomInstance->SetSelectedParameterProfileDirty();
+	CustomInstance->GetPrivate()->SetSelectedParameterProfileDirty();
 
 	CustomInstance->UpdateSkeletalMeshAsync(true, true);
 
@@ -1928,7 +1929,7 @@ FReply SCustomizableInstanceProperties::OnProjectorResetTransform(const FString
 
 void SCustomizableInstanceProperties::InstanceUpdated(UCustomizableObjectInstance* Instance) const
 {
-	if (CustomInstance->bShowOnlyRelevantParameters &&
+	if (CustomInstance->GetPrivate()->bShowOnlyRelevantParameters &&
 		bShouldResetParamBox)
 	{
 		if (const TSharedPtr<FCustomizableInstanceDetails> Details = InstanceDetails.Pin())
@@ -2006,7 +2007,7 @@ FReply SCustomizableInstanceProperties::RemoveParameterProfile()
 {
 	UCustomizableObject* CustomizableObject = CustomInstance->GetCustomizableObject(); 
 
-	const int32 ProfileIdx = CustomInstance->SelectedProfileIndex;
+	const int32 ProfileIdx = CustomInstance->GetPrivate()->SelectedProfileIndex;
 	if (ProfileIdx == INDEX_NONE)
 	{
 		return FReply::Handled();
@@ -2015,7 +2016,7 @@ FReply SCustomizableInstanceProperties::RemoveParameterProfile()
 	TArray<FProfileParameterDat>& Profiles = CustomizableObject->InstancePropertiesProfiles;
 	
 	Profiles.RemoveAt(ProfileIdx);
-	CustomInstance->SelectedProfileIndex = INDEX_NONE;
+	CustomInstance->GetPrivate()->SelectedProfileIndex = INDEX_NONE;
 	CustomizableObject->Modify();
 	if (const TSharedPtr<FCustomizableInstanceDetails> Details = InstanceDetails.Pin())
 	{
@@ -2037,15 +2038,15 @@ FReply SCustomizableInstanceProperties::RemoveParameterProfile()
 
 void SCustomizableInstanceProperties::OnProfileSelectedChanged(TSharedPtr<FString> Selection, ESelectInfo::Type SelectInfo)
 {
-	const int32 ProfileIdx = CustomInstance->SelectedProfileIndex;
-	if (CustomInstance->IsSelectedParameterProfileDirty())
+	const int32 ProfileIdx = CustomInstance->GetPrivate()->SelectedProfileIndex;
+	if (CustomInstance->GetPrivate()->IsSelectedParameterProfileDirty())
 	{
-		CustomInstance->SaveParametersToProfile(ProfileIdx);
+		CustomInstance->GetPrivate()->SaveParametersToProfile(ProfileIdx);
 	}
 
 	if (*Selection == "None")
 	{
-		CustomInstance->SelectedProfileIndex = INDEX_NONE;
+		CustomInstance->GetPrivate()->SelectedProfileIndex = INDEX_NONE;
 		return;
 	}	
 
@@ -2055,12 +2056,12 @@ void SCustomizableInstanceProperties::OnProfileSelectedChanged(TSharedPtr<FStrin
 	{
 		if (Profiles[Idx].ProfileName == *Selection)
 		{
-			CustomInstance->SelectedProfileIndex = Idx;
+			CustomInstance->GetPrivate()->SelectedProfileIndex = Idx;
 			break;
 		}
 	}
 	
-	CustomInstance->LoadParametersFromProfile(CustomInstance->SelectedProfileIndex);
+	CustomInstance->GetPrivate()->LoadParametersFromProfile(CustomInstance->GetPrivate()->SelectedProfileIndex);
 
 	CustomInstance->PreEditChange(nullptr);
 
@@ -2222,7 +2223,7 @@ FReply SCustomizableInstanceProperties::OnResetParameterButtonClicked(int32 Para
 	{
 		CustomInstance->PreEditChange(nullptr);
 		SetParameterValueToDefault(ParameterIndex);
-		CustomInstance->SetSelectedParameterProfileDirty();
+		CustomInstance->GetPrivate()->SetSelectedParameterProfileDirty();
 		CustomInstance->UpdateSkeletalMeshAsync(true, true);
 		CustomInstance->PostEditChange();
 
@@ -2327,11 +2328,11 @@ FReply SCreateProfileParameters::OnButtonClick(EAppReturnType::Type ButtonID)
 		UCustomizableObject* CustomizableObject = CustomInstance->GetCustomizableObject(); 
 		CustomizableObject->GetPrivate()->AddNewParameterProfile(GetFileName(), *CustomInstance.Get());
 
-		if (CustomInstance->bSelectedProfileDirty && CustomInstance->SelectedProfileIndex != INDEX_NONE)
+		if (CustomInstance->GetPrivate()->bSelectedProfileDirty && CustomInstance->GetPrivate()->SelectedProfileIndex != INDEX_NONE)
 		{
-			CustomInstance->SaveParametersToProfile(CustomInstance->SelectedProfileIndex);
+			CustomInstance->GetPrivate()->SaveParametersToProfile(CustomInstance->GetPrivate()->SelectedProfileIndex);
 		}
-		CustomInstance->SelectedProfileIndex = CustomizableObject->InstancePropertiesProfiles.Num() - 1;
+		CustomInstance->GetPrivate()->SelectedProfileIndex = CustomizableObject->InstancePropertiesProfiles.Num() - 1;
 
 		if (CIProperties)
 		{
diff --git a/Engine/Plugins/Experimental/Mutable/Source/CustomizableObjectEditor/Private/MuCOE/SCustomizableObjectEditorPerformanceReport.cpp b/Engine/Plugins/Experimental/Mutable/Source/CustomizableObjectEditor/Private/MuCOE/SCustomizableObjectEditorPerformanceReport.cpp
--- a/Engine/Plugins/Experimental/Mutable/Source/CustomizableObjectEditor/Private/MuCOE/SCustomizableObjectEditorPerformanceReport.cpp
+++ b/Engine/Plugins/Experimental/Mutable/Source/CustomizableObjectEditor/Private/MuCOE/SCustomizableObjectEditorPerformanceReport.cpp
@@ -15,6 +15,7 @@
 #include "MuCOE/SCustomizableObjectEditorTextureAnalyzer.h"
 #include "MuCOE/UnrealEditorPortabilityHelpers.h"
 #include "PropertyCustomizationHelpers.h"
+#include "MuCO/CustomizableObjectInstancePrivate.h"
 #include "MuCO/CustomizableObjectPrivate.h"
 #include "Widgets/Input/SComboButton.h"
 #include "Widgets/Input/SNumericEntryBox.h"
@@ -720,7 +721,7 @@ void SCustomizableObjecEditorPerformanceReport::WorstTimeFound(float CurrentInst
 		WorstCase->SetInstance(Instance);
 		if (WorstCase->WorstCaseInstance)
 		{
-			WorstCase->WorstCaseInstance->SetState(LongesTimeStateIndex);
+			WorstCase->WorstCaseInstance->GetPrivate()->SetState(LongesTimeStateIndex);
 		}
 		WorstCase->LongestUpdateTimeStateIndex = LongesTimeStateIndex;
 		WorstCase->LongestUpdateTimeParameterIndexInCO = LongesTimeParameterIndexInCO;
diff --git a/Engine/Plugins/Experimental/Mutable/Source/CustomizableObjectEditor/Private/MuCOE/SMutableImageViewer.cpp b/Engine/Plugins/Experimental/Mutable/Source/CustomizableObjectEditor/Private/MuCOE/SMutableImageViewer.cpp
--- a/Engine/Plugins/Experimental/Mutable/Source/CustomizableObjectEditor/Private/MuCOE/SMutableImageViewer.cpp
+++ b/Engine/Plugins/Experimental/Mutable/Source/CustomizableObjectEditor/Private/MuCOE/SMutableImageViewer.cpp
@@ -7,7 +7,7 @@
 #include "TextureResource.h"
 #include "MuCO/CustomizableObject.h"
 #include "MuCO/CustomizableObjectPrivate.h"
-#include "MuCO/CustomizableObjectInstance.h"
+#include "MuCO/CustomizableObjectInstancePrivate.h"
 #include "Widgets/Input/SNumericEntryBox.h"
 #include "Widgets/Input/SSegmentedControl.h"
 
diff --git a/Engine/Plugins/Experimental/Mutable/Source/CustomizableObjectEditor/Private/MuCOE/StatePerformanceTesting.cpp b/Engine/Plugins/Experimental/Mutable/Source/CustomizableObjectEditor/Private/MuCOE/StatePerformanceTesting.cpp
--- a/Engine/Plugins/Experimental/Mutable/Source/CustomizableObjectEditor/Private/MuCOE/StatePerformanceTesting.cpp
+++ b/Engine/Plugins/Experimental/Mutable/Source/CustomizableObjectEditor/Private/MuCOE/StatePerformanceTesting.cpp
@@ -8,6 +8,7 @@
 #include "Logging/MessageLog.h"
 #include "MuCO/CustomizableObject.h"
 #include "MuCO/CustomizableObjectInstance.h"
+#include "MuCO/CustomizableObjectInstancePrivate.h"
 #include "MuCO/CustomizableSkeletalComponent.h"
 
 
@@ -185,7 +186,7 @@ void FTestingCustomizableObject::FinishTest()
 
 void FTestingCustomizableObject::InstanceUpdateEnd()
 {
-	const int32 MutableRuntimeCycles = Instance->LastUpdateMutableRuntimeCycles;
+	const int32 MutableRuntimeCycles = Instance->GetPrivate()->LastUpdateMutableRuntimeCycles;
 
 	if ((ActiveBatch != nullptr) && !ActiveBatch->TestFinished)
 	{
@@ -252,7 +253,7 @@ void FTestingCustomizableObject::LaunchNextBatch()
 
 	if (ActiveBatch->IsStateChange)
 	{
-		Instance->SetState(ActiveBatch->State);
+		Instance->GetPrivate()->SetState(ActiveBatch->State);
 	}
 	else
 	{
@@ -283,7 +284,7 @@ void FTestingCustomizableObject::LaunchNextBatch()
 				// For projectors we have small random offsets to apply to the position
 				FVector Position = Instance->GetProjectorPosition(ActiveBatch->Name);
 				Position *= ActiveBatch->VectorValue;
-				Instance->SetProjectorPosition(ActiveBatch->Name, (FVector3f)Position);
+				Instance->SetProjectorPosition(ActiveBatch->Name, Position);
 				break;
 			}
 			case EMutableParameterType::Texture:
diff --git a/Engine/Plugins/Experimental/Mutable/Source/MutableValidation/Private/MuV/CustomizableObjectValidationCommandlet.cpp b/Engine/Plugins/Experimental/Mutable/Source/MutableValidation/Private/MuV/CustomizableObjectValidationCommandlet.cpp
--- a/Engine/Plugins/Experimental/Mutable/Source/MutableValidation/Private/MuV/CustomizableObjectValidationCommandlet.cpp
+++ b/Engine/Plugins/Experimental/Mutable/Source/MutableValidation/Private/MuV/CustomizableObjectValidationCommandlet.cpp
@@ -6,6 +6,7 @@
 #include "Components/SkeletalMeshComponent.h"
 #include "Containers/Array.h"
 #include "MuCO/CustomizableObject.h"
+#include "MuCO/CustomizableObjectInstancePrivate.h"
 #include "MuCO/CustomizableObjectPrivate.h"
 #include "MuCO/CustomizableObjectSystem.h"
 #include "MuR/Model.h"
@@ -117,7 +118,7 @@ int32 UCustomizableObjectValidationCommandlet::Main(const FString& Params)
 					// Force generation of all LODS
 					TArray<uint16> RequestedLodLevels{};
 					RequestedLodLevels.Init(MAX_uint8, GeneratedInstance->GetNumComponents());
-					GeneratedInstance->GetDescriptor().SetRequestedLODLevels(RequestedLodLevels);
+					GeneratedInstance->GetPrivate()->GetDescriptor().SetRequestedLODLevels(RequestedLodLevels);
 					
 					// Randomize instance values
 					GeneratedInstance->SetRandomValuesFromStream(RandomizationStream);
