[Mutable] Moved CO non-public headers to private.

Moved headers which game users should not have access to. Move to:
- Private if only used by CO
- Internal if used by COE

#rnx

diff --git a/Engine/Plugins/Experimental/Mutable/Source/CustomizableObject/Classes/MuCO/CustomizableObjectIdentifier.h b/Engine/Plugins/Experimental/Mutable/Source/CustomizableObject/Classes/MuCO/CustomizableObjectIdentifier.h
deleted file mode 100644
diff --git a/Engine/Plugins/Experimental/Mutable/Source/CustomizableObject/Internal/MuCO/CustomizableObjectClothingTypes.h b/Engine/Plugins/Experimental/Mutable/Source/CustomizableObject/Internal/MuCO/CustomizableObjectClothingTypes.h
new file mode 100644
--- /dev/null
+++ b/Engine/Plugins/Experimental/Mutable/Source/CustomizableObject/Internal/MuCO/CustomizableObjectClothingTypes.h
@@ -0,0 +1,96 @@
+// Copyright Epic Games, Inc. All Rights Reserved.
+
+#pragma once
+
+#include "ClothCollisionPrim.h"
+#include "ClothLODData.h"
+
+#include "ClothLODData_Legacy.h"
+#include "ClothVertBoneData.h"
+#include "SkeletalMeshTypes.h"
+#include "CustomizableObjectClothingTypes.generated.h"
+
+
+USTRUCT()
+struct FCustomizableObjectClothConfigData
+{
+	GENERATED_BODY()
+	
+	UPROPERTY()
+	FString ClassPath;
+	
+	UPROPERTY()
+	FName ConfigName;
+
+	UPROPERTY()
+	TArray<uint8> ConfigBytes;
+
+	friend FArchive& operator<<(FArchive& Ar, FCustomizableObjectClothConfigData& ClothConfig)
+	{
+		Ar << ClothConfig.ClassPath;
+		Ar << ClothConfig.ConfigName;
+		Ar << ClothConfig.ConfigBytes;
+
+		return Ar;
+	}
+};
+
+USTRUCT()
+struct FCustomizableObjectClothingAssetData
+{
+	GENERATED_BODY()
+	
+	UPROPERTY()
+	TArray<FClothLODDataCommon> LodData;
+
+	UPROPERTY()
+	TArray<int32> LodMap;
+
+	UPROPERTY()
+	TArray<FName> UsedBoneNames;
+
+	UPROPERTY()
+	TArray<int32> UsedBoneIndices;
+
+	UPROPERTY()
+	int32 ReferenceBoneIndex = 0;
+
+	UPROPERTY()
+	TArray<FCustomizableObjectClothConfigData> ConfigsData;
+	
+	UPROPERTY()
+	FString PhysicsAssetPath;
+
+	UPROPERTY()
+	FName Name;
+	
+	UPROPERTY()
+	FGuid OriginalAssetGuid = FGuid{};
+	
+	friend FArchive& operator<<(FArchive& Ar, FCustomizableObjectClothingAssetData& ClothData)
+	{	
+		int32 LodDataCount = ClothData.LodData.Num();
+		Ar << LodDataCount;
+
+		if ( Ar.IsLoading() )
+		{
+			ClothData.LodData.SetNum(LodDataCount);
+		}
+
+		for ( FClothLODDataCommon& ClothLodData :  ClothData.LodData )
+		{
+			ClothLodData.Serialize( Ar );
+		}
+
+		Ar << ClothData.LodMap;
+		Ar << ClothData.UsedBoneNames;
+		Ar << ClothData.UsedBoneIndices;
+		Ar << ClothData.ReferenceBoneIndex;
+		Ar << ClothData.ConfigsData;
+		Ar << ClothData.PhysicsAssetPath;
+		Ar << ClothData.Name;
+		Ar << ClothData.OriginalAssetGuid;
+		
+		return Ar;
+	}
+};
diff --git a/Engine/Plugins/Experimental/Mutable/Source/CustomizableObject/Internal/MuCO/CustomizableObjectCustomVersion.h b/Engine/Plugins/Experimental/Mutable/Source/CustomizableObject/Internal/MuCO/CustomizableObjectCustomVersion.h
new file mode 100644
--- /dev/null
+++ b/Engine/Plugins/Experimental/Mutable/Source/CustomizableObject/Internal/MuCO/CustomizableObjectCustomVersion.h
@@ -0,0 +1,127 @@
+// Copyright Epic Games, Inc. All Rights Reserved.
+
+#pragma once
+
+#include "HAL/Platform.h"
+
+struct FGuid;
+
+// Custom serialization version for assets/classes in the CustomizableObject and CustomizableObjectEditor modules
+struct CUSTOMIZABLEOBJECT_API FCustomizableObjectCustomVersion
+{
+	enum Type
+	{
+		// Before any version changes were made in the plugin
+		BeforeCustomVersionWasAdded = 0,
+
+		FixBlueprintPinsUseRealNumbers,
+
+		NodeExposePinOnNameChangedDelegateAdded,
+
+		GroupProjectorPinTypeAdded,
+
+		AutomaticNodeMaterial,
+		
+		AutomaticNodeMaterialPerformance,
+
+		LayoutClassAdded,
+
+		AutomaticNodeMaterialPerformanceBug,
+
+		PinsNamesImageToTexture,
+
+		BugPinsSwitch,
+
+		PostLoadToCustomVersion, // Wrapped old backwards compatible code that was located at PostLoads inside this custom version.
+
+		AutomaticNodeMaterialUXImprovements,
+
+		ExtendMaterialOnlyMutableModeParameters,
+
+		ExtendMaterialOnlyMutableModeParametersBug,
+		
+		ExtendMaterialRemoveImages,
+		
+		EditMaterialOnlyMutableModeParameters, // Force refresh to avoid having Images which are not in Mutable mode.
+
+		GroupProjectorIntToScalarIndex,
+
+		FixBlueprintPinsUseRealNumbersAgain,
+
+		NodeMaterialAddTablePin,
+
+		MorphMaterialAddFactorPin,
+
+		NodeSkeletalMeshCorruptedPinRef,
+
+		CustomizableObjectInstanceDescriptor,
+		
+		DescriptorBuildParameterDecorations,
+
+		DescriptorMultilayerProjectors,
+
+		DeformSkeletonOptionsAdded,
+
+		GroupProjectorImagePin,
+
+		GroupProjectorImagePinRemoved,
+
+		StateTextureCompressionStrategyEnum,
+
+		AddedTableNodesTextureMode,
+		
+		ConvertAnimationSlotToFName,
+
+		RemovedParameterDecorations,
+
+		AutomaticNodeSkeletalMesh,
+
+		AutomaticNodeSkeletalMeshPinDataOuter,
+
+		AutomaticNodeSkeletalMeshPinDataUProperty,
+
+		EditMaterialMaskPinDesync, // MTBL-979
+		
+		IgnoreDisabledSections,
+
+		SkeletalMeshNodeDefaultPinWithoutPinData,
+		
+		AddedTableMaterialSwitch,
+
+		FixPinsNamesImageToTexture2,
+
+		MeshReshapeVertexColorUsageSelection,
+
+		DeduplicateNodeVariant,
+
+		TextureCompressionEnum,
+
+		NodeTextureParameterDefaultToReferenceValue,
+
+		NodeVariationSerializationIssue, // The parent class and child classes had a member with the same name. Unreal was serialized the member from the child class but deserialized it to the parent class.
+
+		RegenerateNodeObjectsIds,
+
+		NodeMaterialPinDataImageDetails,
+
+		CustomizableObjectStateHasSeparateNeverStreamFlag,
+
+		CustomizableObjectNodeHasSeparateNeverStreamFlag,
+
+		AddedColumnIdDataToTableNodePins,
+
+		AddedAnyTextureTypeToPassThroughTextures,
+
+		ProjectorNodesDefaultValueFix,
+
+		// -----<new versions can be added above this line>-------------------------------------------------
+		VersionPlusOne,
+		LatestVersion = VersionPlusOne - 1
+	};
+
+	// The GUID for this custom version number
+	const static FGuid GUID;
+
+	private:
+		FCustomizableObjectCustomVersion() {}
+};
\ No newline at end of file
diff --git a/Engine/Plugins/Experimental/Mutable/Source/CustomizableObject/Internal/MuCO/CustomizableObjectIdentifier.h b/Engine/Plugins/Experimental/Mutable/Source/CustomizableObject/Internal/MuCO/CustomizableObjectIdentifier.h
new file mode 100644
--- /dev/null
+++ b/Engine/Plugins/Experimental/Mutable/Source/CustomizableObject/Internal/MuCO/CustomizableObjectIdentifier.h
@@ -0,0 +1,56 @@
+// Copyright Epic Games, Inc. All Rights Reserved.
+
+#pragma once
+
+
+#include "CustomizableObjectIdentifier.generated.h"
+
+USTRUCT(BlueprintType)
+struct CUSTOMIZABLEOBJECT_API FCustomizableObjectIdPair
+{
+	GENERATED_USTRUCT_BODY()
+
+	UPROPERTY(Category = CustomizableObject, BlueprintReadOnly, EditDefaultsOnly)
+	FString CustomizableObjectGroupName;
+
+	UPROPERTY(Category = CustomizableObject, BlueprintReadOnly, EditDefaultsOnly)
+	FString CustomizableObjectName;
+
+	FCustomizableObjectIdPair() { };
+	FCustomizableObjectIdPair(FString ObjectGroupName, FString ObjectName);
+
+	bool operator ==(const FCustomizableObjectIdPair& Other) const
+	{
+		return CustomizableObjectGroupName == Other.CustomizableObjectGroupName && CustomizableObjectName == Other.CustomizableObjectName;
+	}
+
+	friend FArchive& operator <<(FArchive& Ar, FCustomizableObjectIdPair& IdPair)
+	{
+		Ar << IdPair.CustomizableObjectGroupName;
+		Ar << IdPair.CustomizableObjectName;
+
+		return Ar;
+	}
+};
+
+
+USTRUCT(BlueprintType)
+struct CUSTOMIZABLEOBJECT_API FCustomizableObjectIdentifier
+{
+	GENERATED_USTRUCT_BODY()
+
+private:
+	UPROPERTY()
+	FString CustomizableObjectGroupName; // Deprecated
+
+	UPROPERTY()
+	FString CustomizableObjectName; // Deprecated
+
+public:
+	UPROPERTY(Category = CustomizableObject, BlueprintReadOnly, EditDefaultsOnly)
+	FString Guid;
+
+	FCustomizableObjectIdentifier();
+	FCustomizableObjectIdentifier(FString ObjectGroupName, FString ObjectName);
+};
+
diff --git a/Engine/Plugins/Experimental/Mutable/Source/CustomizableObject/Internal/MuCO/CustomizableObjectInstanceAssetUserData.h b/Engine/Plugins/Experimental/Mutable/Source/CustomizableObject/Internal/MuCO/CustomizableObjectInstanceAssetUserData.h
new file mode 100644
--- /dev/null
+++ b/Engine/Plugins/Experimental/Mutable/Source/CustomizableObject/Internal/MuCO/CustomizableObjectInstanceAssetUserData.h
@@ -0,0 +1,36 @@
+// Copyright Epic Games, Inc. All Rights Reserved.
+
+#pragma once
+
+#include "Engine/AssetUserData.h"
+#include "GameplayTagContainer.h"
+
+#include "CustomizableObjectInstanceAssetUserData.generated.h"
+
+
+UCLASS(BlueprintType)
+class CUSTOMIZABLEOBJECT_API UCustomizableObjectInstanceUserData : public UAssetUserData
+{
+	GENERATED_BODY()
+
+public:
+
+	/** Return the list of tags for this instance. */
+	UFUNCTION(BlueprintCallable, Category = CustomizableObjectInstance)
+	const FGameplayTagContainer& GetAnimationGameplayTags() const
+	{
+		return AnimBPGameplayTags;
+	};
+
+	/** Sets the list of tags for this instance. */
+	UFUNCTION(BlueprintCallable, Category = CustomizableObjectInstance)
+	void SetAnimationGameplayTags(const FGameplayTagContainer& InstanceTags)
+	{
+		AnimBPGameplayTags = InstanceTags;
+	};
+
+private:
+
+	UPROPERTY()
+	FGameplayTagContainer AnimBPGameplayTags;
+};
diff --git a/Engine/Plugins/Experimental/Mutable/Source/CustomizableObject/Internal/MuCO/CustomizableObjectMipDataProvider.h b/Engine/Plugins/Experimental/Mutable/Source/CustomizableObject/Internal/MuCO/CustomizableObjectMipDataProvider.h
new file mode 100644
--- /dev/null
+++ b/Engine/Plugins/Experimental/Mutable/Source/CustomizableObject/Internal/MuCO/CustomizableObjectMipDataProvider.h
@@ -0,0 +1,169 @@
+// Copyright Epic Games, Inc. All Rights Reserved.
+
+#pragma once
+
+#include "MuCO/FMutableTaskGraph.h"
+#include "Async/TaskGraphInterfaces.h"
+#include "Engine/TextureMipDataProviderFactory.h"
+#include "MuCO/CustomizableObjectSystem.h"
+#include "MuR/Image.h"
+#include "Streaming/TextureMipDataProvider.h"
+
+#include "MuCO/DefaultImageProvider.h"
+#include "MuR/System.h"
+
+#include "CustomizableObjectMipDataProvider.generated.h"
+
+enum EPixelFormat : uint8;
+namespace mu { class Parameters; }
+namespace mu { class System; }
+
+class FThreadSafeCounter;
+class UCustomizableObjectInstance;
+class UObject;
+class UTexture;
+
+/** This struct stores the data relevant for the construction of a specific texture. 
+* This includes all the data required to rebuild the image (or any of its mips).
+*/
+class FMutableUpdateContext
+{
+public:
+	FMutableUpdateContext() = default;
+	
+	FMutableUpdateContext(const FString& InCustomizableObjectPathName, const FString& InInstancePathName, mu::Ptr<mu::System> InSystem,
+		TSharedPtr<mu::Model, ESPMode::ThreadSafe> InModel, mu::Ptr<const mu::Parameters> InParameters, int32 InState);
+
+	~FMutableUpdateContext();
+
+	const FString& GetCustomizableObjectPathName() const;
+
+	const FString& GetInstancePathName() const;
+	
+	mu::Ptr<mu::System> GetSystem() const;
+	
+	TSharedPtr<mu::Model, ESPMode::ThreadSafe> GetModel() const;
+
+	mu::Ptr<const mu::Parameters> GetParameters() const;
+	
+	int32 GetState() const;
+
+	const TArray<mu::Ptr<const mu::Image>>& GetImageParameterValues() const;
+	
+private:
+	FString CustomizableObjectPathName;
+	FString InstancePathName;
+
+	mu::Ptr<mu::System> System;
+	TSharedPtr<mu::Model, ESPMode::ThreadSafe> Model;
+	mu::Ptr<const mu::Parameters> Parameters;
+	int32 State = -1;
+
+	TArray<mu::Ptr<const mu::Image>> ImageParameterValues;
+};
+
+
+struct FMutableMipUpdateLevel
+{
+	FMutableMipUpdateLevel(int32 InMipLevel, void* InDest, int32 InSizeX, int32 InSizeY, int32 InDataSize, EPixelFormat InFormat) :
+		Dest(InDest), MipLevel(InMipLevel), SizeX(InSizeX), SizeY(InSizeY), DataSize(InDataSize), Format(InFormat) {}
+
+	void* Dest; // Only access from the FMutableTextureMipDataProvider, owned by the FTextureMipInfoArray so don't delete
+	int32 MipLevel;
+	int32 SizeX;
+	int32 SizeY;
+	int32 DataSize;
+	EPixelFormat Format;
+};
+
+
+/** Runtime data used during a mutable image mipmap update */
+struct FMutableImageOperationData
+{
+	/** This option comes from the operation request. It is used to reduce the number of mipmaps that mutable must generate for images.  */
+	int32 MipsToSkip = 0;
+	FMutableImageReference RequestedImage;
+
+	TSharedPtr<FMutableUpdateContext> UpdateContext;
+
+	mu::ImagePtrConst Result;
+
+	TArray<FMutableMipUpdateLevel> Levels;
+
+	// Used to sync with the FMutableTextureMipDataProvider and FRenderAssetUpdate::Tick
+	FThreadSafeCounter* Counter;
+	FTextureUpdateSyncOptions::FCallback RescheduleCallback;
+
+	/** Access to the Counter must be protected with this because it may be accessed from another thread to null it. */
+	FCriticalSection CounterTaskLock;
+
+	// Image Update Memory stats
+	int64 ImageUpdateStartBytes = 0;
+};
+
+
+class FMutableTextureMipDataProvider : public FTextureMipDataProvider
+{
+public:
+	FMutableTextureMipDataProvider(const UTexture* Texture, UCustomizableObjectInstance* InCustomizableObjectInstance, const FMutableImageReference& InImageRef);
+
+	virtual void Init(const FTextureUpdateContext& Context, const FTextureUpdateSyncOptions& SyncOptions) override;
+	virtual int32 GetMips(const FTextureUpdateContext& Context, int32 StartingMipIndex, const FTextureMipInfoArray& MipInfos, const FTextureUpdateSyncOptions& SyncOptions) override;
+	virtual bool PollMips(const FTextureUpdateSyncOptions& SyncOptions) override;
+	virtual void CleanUp(const FTextureUpdateSyncOptions& SyncOptions) override;
+	virtual void Cancel(const FTextureUpdateSyncOptions& SyncOptions) override;
+	virtual ETickThread GetCancelThread() const override;
+	void AbortPollMips() override;
+
+private:
+	void CancelCounterSafely();
+	void PrintWarningAndAdvanceToCleanup();
+
+public:
+	// Todo: Simplify by replacing the reference to the Instance with some static parametrization or hash with enough info to reconstruct the texture
+	UPROPERTY(Transient)
+	UCustomizableObjectInstance* CustomizableObjectInstance = nullptr;
+
+	FMutableImageReference ImageRef;
+	TSharedPtr<FMutableUpdateContext> UpdateContext;
+
+	bool bRequestAborted = false;
+
+	TSharedPtr<FMutableImageOperationData> OperationData;
+
+	uint32 MutableTaskId = FMutableTaskGraph::INVALID_ID;
+};
+
+
+UCLASS(hidecategories=Object)
+class CUSTOMIZABLEOBJECT_API UMutableTextureMipDataProviderFactory : public UTextureMipDataProviderFactory
+{
+	GENERATED_UCLASS_BODY()
+
+public:
+	virtual FTextureMipDataProvider* AllocateMipDataProvider(UTexture* Asset) override
+	{
+		check(ImageRef.ImageID > 0);
+		FMutableTextureMipDataProvider* Result = new FMutableTextureMipDataProvider(Asset, CustomizableObjectInstance, ImageRef);
+		Result->UpdateContext = UpdateContext;		
+		return Result;
+	}
+
+	virtual bool WillProvideMipDataWithoutDisk() const override
+	{ 
+		return true;
+	}
+
+	// Todo: Simplify by replacing the reference to the Instance with some static parametrization or hash with enough info to reconstruct the texture
+	UPROPERTY(Transient)
+	TObjectPtr<UCustomizableObjectInstance> CustomizableObjectInstance = nullptr;
+
+	FMutableImageReference ImageRef;
+	TSharedPtr<FMutableUpdateContext> UpdateContext;
+
+};
+
+#if UE_ENABLE_INCLUDE_ORDER_DEPRECATED_IN_5_2
+#include "MuR/Parameters.h"
+#include "MuR/System.h"
+#endif
diff --git a/Engine/Plugins/Experimental/Mutable/Source/CustomizableObject/Internal/MuCO/CustomizableObjectResourceDataTypes.h b/Engine/Plugins/Experimental/Mutable/Source/CustomizableObject/Internal/MuCO/CustomizableObjectResourceDataTypes.h
new file mode 100644
--- /dev/null
+++ b/Engine/Plugins/Experimental/Mutable/Source/CustomizableObject/Internal/MuCO/CustomizableObjectResourceDataTypes.h
@@ -0,0 +1,30 @@
+// Copyright Epic Games, Inc. All Rights Reserved.
+
+#pragma once
+
+#include "UObject/Object.h"
+#include "UObject/ObjectPtr.h"
+#include "UObject/SoftObjectPtr.h"
+
+#include "CustomizableObjectResourceDataTypes.generated.h"
+
+class UAssetUserData;
+
+/** Simple struct used to store streamed resources of type AssetUserData */
+USTRUCT()
+struct FCustomizableObjectAssetUserData
+{
+	GENERATED_BODY()
+
+public:
+	
+	// Valid in Cooked builds
+	UPROPERTY()
+	TObjectPtr<UAssetUserData> AssetUserData;
+
+#if WITH_EDITORONLY_DATA
+	// Valid in the Editor
+	UPROPERTY(Transient)
+	TObjectPtr<UAssetUserData> AssetUserDataEditor;
+#endif
+};
\ No newline at end of file
diff --git a/Engine/Plugins/Experimental/Mutable/Source/CustomizableObject/Internal/MuCO/EditorImageProvider.h b/Engine/Plugins/Experimental/Mutable/Source/CustomizableObject/Internal/MuCO/EditorImageProvider.h
new file mode 100644
--- /dev/null
+++ b/Engine/Plugins/Experimental/Mutable/Source/CustomizableObject/Internal/MuCO/EditorImageProvider.h
@@ -0,0 +1,23 @@
+// Copyright Epic Games, Inc. All Rights Reserved.
+
+#pragma once
+
+
+#include "MuCO/CustomizableObjectSystem.h"
+
+#include "EditorImageProvider.generated.h"
+
+class UTexture2D;
+class UCustomizableObjectInstance;
+
+
+UCLASS()
+class CUSTOMIZABLEOBJECT_API UEditorImageProvider : public UCustomizableSystemImageProvider
+{
+	GENERATED_BODY()
+
+public:
+	// UCustomizableSystemImageProvider interface
+	virtual ValueType HasTextureParameterValue(const FName& ID) override;
+	virtual UTexture2D* GetTextureParameterValue(const FName& ID) override;
+};
\ No newline at end of file
diff --git a/Engine/Plugins/Experimental/Mutable/Source/CustomizableObject/Internal/MuCO/FMutableTaskGraph.h b/Engine/Plugins/Experimental/Mutable/Source/CustomizableObject/Internal/MuCO/FMutableTaskGraph.h
new file mode 100644
--- /dev/null
+++ b/Engine/Plugins/Experimental/Mutable/Source/CustomizableObject/Internal/MuCO/FMutableTaskGraph.h
@@ -0,0 +1,103 @@
+// Copyright Epic Games, Inc. All Rights Reserved.
+
+#pragma once
+
+#include "MuR/System.h" // For the MUTABLE_USE_NEW_TASKGRAPH define
+
+#include "Tasks/Task.h"
+#include "Async/TaskGraphInterfaces.h"
+
+
+/** Mutable Tasks System.
+ *
+ * Allows launching tasks in different threads:
+ * - Mutable Thread
+ * - Game Thread (code currently in UCustomizableObjectSystem)
+ * - Any Thread
+ *
+ * Concurrency between tasks in the Mutable Thread is forbidden by chaining all tasks through their prerequisites.
+ *
+ *
+ * LOW PRIORITY MUTABLE TASKS:
+ *
+ * - Only one Low Priority task can be launched at the same time.
+ *   This is because once a Task Graph task is launched, it can not be canceled.
+ *   To allow canceling them, the system holds them until it can ensure that its execution will be imminent (no other tasks running).
+ *
+ * - A Low Priority task will not be launched if one of the follow conditions is true (in order):
+ *     1. There is a task Low Priority task running.
+ *     2. Flag bLaunchMutableTaskLowPriority is false.
+ *     3. There is a Normal Priority task running (unless time limit).
+ */
+class FMutableTaskGraph
+{
+	struct FTask
+	{
+		uint32 Id;
+		FString DebugName;
+		TFunction<void()> Body;
+		double CreationTime = FPlatformTime::Seconds();
+	};
+	
+public:
+	using TaskType = 
+#ifdef MUTABLE_USE_NEW_TASKGRAPH
+		UE::Tasks::FTask;
+#else
+		FGraphEventRef;
+#endif
+
+	/** Create and launch a task on the Mutable Thread with Normal priority. */
+	TaskType AddMutableThreadTask(const TCHAR* DebugName, TUniqueFunction<void()>&& TaskBody);
+
+	/** Create and launch a task on the Mutable Thread with Low priority. */
+	uint32 AddMutableThreadTaskLowPriority(const TCHAR* DebugName, TFunction<void()>&& TaskBody);
+
+	/** Cancel, if not already launched, a Mutable Thread with Low priority. */
+	void CancelMutableThreadTaskLowPriority(uint32 Id);
+	
+	/** Create and launch a task on Any Thread. */
+	void AddAnyThreadTask(const TCHAR* DebugName, TUniqueFunction<void()>&& TaskBody) const;
+
+	/** Wait for all Mutable Thread tasks. */
+	void WaitForMutableTasks();
+
+	/** Allow or disallow launching Mutable Tasks with Low priority.
+	 @param bFromMutableTask true if called from a Mutable Task. */
+	void AllowLaunchingMutableTaskLowPriority(bool bAllow, bool bFromMutableTask);
+	
+	void Tick();
+	
+private:
+	/** A Mutable Task Low Priority will only be launched if:
+	 * - No other low priority task is running.
+	 * - Is allowed to launch low priority tasks.
+	 *
+	 * @param bFromMutableTask true if called from a Mutable Task. */
+	void TryLaunchMutableTaskLowPriority(bool bFromMutableTask);
+
+	/** Return true if the task is completed (or is no longer valid). */
+	bool IsTaskCompleted(const TaskType& Task) const;
+
+	mutable FCriticalSection MutableTaskLock;
+
+	/** Allow or disallow launching low priority tasks. */
+	bool bAllowLaunchMutableTaskLowPriority = true;
+
+	/** Queue of low priority tasks. FIFO. */
+	TArray<FTask> QueueMutableTasksLowPriority;
+
+public:
+	inline static constexpr uint32 INVALID_ID = 0;
+
+private:
+	/** Incremental task ID generator. */
+	uint32 TaskIdGenerator = INVALID_ID;
+
+	/** Last Mutable Task Low Priority launched to the TaskGraph system. */
+	TaskType LastMutableTaskLowPriority = {};
+
+	/** Last Mutable Task launched to the TaskGraph system. Low and normal priority. */
+	TaskType LastMutableTask = {};
+};
+
diff --git a/Engine/Plugins/Experimental/Mutable/Source/CustomizableObject/Internal/MuCO/MutableMeshBufferUtils.h b/Engine/Plugins/Experimental/Mutable/Source/CustomizableObject/Internal/MuCO/MutableMeshBufferUtils.h
new file mode 100644
--- /dev/null
+++ b/Engine/Plugins/Experimental/Mutable/Source/CustomizableObject/Internal/MuCO/MutableMeshBufferUtils.h
@@ -0,0 +1,79 @@
+// Copyright Epic Games, Inc. All Rights Reserved.
+
+#pragma once
+
+#include "HAL/Platform.h"
+
+namespace mu { class FMeshBufferSet; }
+
+/** Publicly accessible methods designed to prepare the Mutable Mesh buffers with the internal structure  required by Unreal and
+ * our systems to be able to work with the mesh owning those buffers.
+ */
+namespace  MutableMeshBufferUtils
+{
+	// VERTEX BUFFERS
+	
+	/** Determines the max amount of channels that can inhabit a Texture Coordinate buffer. */
+	constexpr int32 MaxTexCordChannelCount = 4;
+	
+	/**
+	 * Sets up all the channels that ought to be part of the Vertex Position Buffer of the Mutable mesh owning the provided MeshBufferSet
+	 * @param InCurrentVertexBuffer Index of the Vertex Buffer being set up inside the OutTargetVertexBuffers.
+	 * @param OutTargetVertexBuffers Mesh buffer to be set up. 
+	 */
+	CUSTOMIZABLEOBJECT_API void SetupVertexPositionsBuffer(const int32& InCurrentVertexBuffer, mu::FMeshBufferSet& OutTargetVertexBuffers);
+
+	/**
+	 * Sets up all the channels that ought to be part of the Vertex Tangent Buffer of the Mutable mesh owning the provided MeshBufferSet
+	 * @param InCurrentVertexBuffer Index of the Vertex Buffer being set up inside the OutTargetVertexBuffers.
+	 * @param OutTargetVertexBuffers Mesh buffer to be set up. 
+	 */
+	CUSTOMIZABLEOBJECT_API void SetupTangentBuffer(const int32& InCurrentVertexBuffer, mu::FMeshBufferSet& OutTargetVertexBuffers);
+
+	/**
+	 * Sets up all the channels that ought to be part of the Vertex Texture Coordinates Buffer of the Mutable mesh owning the provided MeshBufferSet
+	 * @param InCurrentVertexBuffer Index of the Vertex Buffer being set up inside the OutTargetVertexBuffers.
+	 * @param InChannelCount Determines the amount of TexCord channels able to be added onto the buffer.
+	 * @param OutTargetVertexBuffers Mesh buffer to be set up. 
+	 * @param InTextureSemanticIndicesOverride (Optional) If you are certain of the texture semantic indices that should be used populate this argument, if not, leave it as null
+	 */
+	CUSTOMIZABLEOBJECT_API void SetupTexCoordinatesBuffer(const int32& InCurrentVertexBuffer,const int32& InChannelCount, mu::FMeshBufferSet& OutTargetVertexBuffers, const int32* InTextureSemanticIndicesOverride = nullptr);
+
+	/**
+	 * Sets up all the channels that ought to be part of the Vertex Skin Buffer of the Mutable mesh owning the provided MeshBufferSet
+	 * @param InCurrentVertexBuffer Index of the Vertex Buffer being set up inside the OutTargetVertexBuffers.
+	 * @param MaxBoneIndexTypeSizeBytes Maximum size of the vertex bone index type.
+	 * @param MaxNumBonesPerVertex Amount of bone influences a vertex can have.
+	 * @param OutTargetVertexBuffers Mesh buffer to be set up. 
+	 */
+	CUSTOMIZABLEOBJECT_API void SetupSkinBuffer(const int32& InCurrentVertexBuffer,const int32& MaxBoneIndexTypeSizeBytes, const int32& MaxBoneWeightTypeSizeBytes,
+		const int32& MaxNumBonesPerVertex, mu::FMeshBufferSet& OutTargetVertexBuffers);
+	
+	/**
+	 * Sets up all the channels that ought to be part of the Vertex Color Buffer of the Mutable mesh owning the provided MeshBufferSet
+	 * @param InCurrentVertexBuffer Index of the Vertex Buffer being set up inside the OutTargetVertexBuffers.
+	 * @param OutTargetVertexBuffers Mesh buffer to be set up. 
+	 */
+	CUSTOMIZABLEOBJECT_API void SetupVertexColorBuffer(const int32& InCurrentVertexBuffer, mu::FMeshBufferSet& OutTargetVertexBuffers);
+
+	// INDEX BUFFERS
+
+	/**
+	 * Sets up all the channels that ought to be part of the Index Buffer of the Mutable mesh owning the provided MeshBufferSet
+	 * @param OutTargetIndexBuffers Index buffer to be set up. 
+	 */
+	CUSTOMIZABLEOBJECT_API void SetupIndexBuffer(mu::FMeshBufferSet& OutTargetIndexBuffers);
+
+	/**
+	 * Sets up all the channels that ought to be part of the Vertex SkinWeightProfile Buffer of the Mutable mesh owning the provided MeshBufferSet
+	 * @param InCurrentVertexBuffer Index of the Vertex Buffer being set up inside the OutTargetVertexBuffers.
+	 * @param MaxBoneIndexTypeSizeBytes Maximum size of the vertex bone index type.
+	 * @param MaxBoneWeightTypeSizeBytes Maximum size of the vertex bone weight type.
+	 * @param MaxNumBonesPerVertex Amount of bone influences a vertex can have.
+	 * @param SemanticsIndex External Id that identifies a SkinWeightProfile.
+	 * @param OutTargetVertexBuffers Mesh buffer to be set up.
+	 */
+	CUSTOMIZABLEOBJECT_API void SetupSkinWeightProfileBuffer(const int32& InCurrentVertexBuffer, const int32& MaxBoneIndexTypeSizeBytes,const int32& MaxBoneWeightTypeSizeBytes, const int32& MaxNumBonesPerVertex,
+		const int32 SemanticsIndex, mu::FMeshBufferSet& OutTargetVertexBuffers);
+	
+}
diff --git a/Engine/Plugins/Experimental/Mutable/Source/CustomizableObject/Internal/MuCO/MutableProjectorTypeUtils.h b/Engine/Plugins/Experimental/Mutable/Source/CustomizableObject/Internal/MuCO/MutableProjectorTypeUtils.h
new file mode 100644
--- /dev/null
+++ b/Engine/Plugins/Experimental/Mutable/Source/CustomizableObject/Internal/MuCO/MutableProjectorTypeUtils.h
@@ -0,0 +1,24 @@
+// Copyright Epic Games, Inc. All Rights Reserved.
+ 
+#pragma once
+
+#include "MuCO/CustomizableObjectParameterTypeDefinitions.h"
+#include "MuR/Parameters.h"
+
+namespace ProjectorUtils
+{
+	/** Get with the equivalent type of CO projector provided a mutable internal projector type.
+	 * @param ProjectorType The mutable internal projector type to get the unreal's equivalent projector type.
+	 * @return The equivalent projector type to the one provided to be used with the CO and COI systems.
+	 */
+	CUSTOMIZABLEOBJECT_API
+	ECustomizableObjectProjectorType GetEquivalentProjectorType (mu::PROJECTOR_TYPE ProjectorType);
+
+	/** Get the equivalent mutable projector type to the CO mutable projector type.
+	 * @param ProjectorType The mutable CO projector type to get the mutable's internal equivalent projector type.
+	 * @return the projector type provided but as a mutable core projector type.
+	 */
+	CUSTOMIZABLEOBJECT_API
+	mu::PROJECTOR_TYPE GetEquivalentProjectorType (ECustomizableObjectProjectorType ProjectorType);
+}
+
diff --git a/Engine/Plugins/Experimental/Mutable/Source/CustomizableObject/Internal/MuCO/StateMachine.h b/Engine/Plugins/Experimental/Mutable/Source/CustomizableObject/Internal/MuCO/StateMachine.h
new file mode 100644
--- /dev/null
+++ b/Engine/Plugins/Experimental/Mutable/Source/CustomizableObject/Internal/MuCO/StateMachine.h
@@ -0,0 +1,49 @@
+// Copyright Epic Games, Inc. All Rights Reserved.
+
+#pragma once
+
+#include "HAL/Platform.h"
+#include "Delegates/Delegate.h" 
+
+
+/** Simple state machine with transitions check.
+ *
+ * Type T must implement the following:
+ * - EState: Enum of valid states.
+ * - StartState: Variable defining the initial state.
+ * - ValidTransitions: 2D bool array of valid transitions. */
+template<typename T>
+class FStateMachine
+{
+public:
+	using EState = typename T::EState;
+
+	static_assert(static_cast<int32>(EState::Count) * static_cast<int32>(EState::Count) == sizeof(T::ValidTransitions), "Size of EState and ValidTransitions do not match.");
+	
+	void NextState(EState State)
+	{
+		check(T::ValidTransitions[static_cast<int32>(CurrentState)][static_cast<int32>(State)])
+
+		EState PreviousState = CurrentState; 
+		CurrentState = State;
+		OnStateChangedDelegate.Broadcast(PreviousState, CurrentState);
+	}
+
+	EState Get() const
+	{
+		return CurrentState;
+	}
+
+	DECLARE_MULTICAST_DELEGATE_TwoParams(FOnStateChangedDelegate, EState, EState)
+
+	FOnStateChangedDelegate& GetOnStateChangedDelegate()
+	{
+		return OnStateChangedDelegate;
+	}
+	
+private:
+	FOnStateChangedDelegate OnStateChangedDelegate;
+
+	EState CurrentState = T::StartState;
+};
+
diff --git a/Engine/Plugins/Experimental/Mutable/Source/CustomizableObject/Internal/MuCO/UnrealConversionUtils.h b/Engine/Plugins/Experimental/Mutable/Source/CustomizableObject/Internal/MuCO/UnrealConversionUtils.h
new file mode 100644
--- /dev/null
+++ b/Engine/Plugins/Experimental/Mutable/Source/CustomizableObject/Internal/MuCO/UnrealConversionUtils.h
@@ -0,0 +1,101 @@
+// Copyright Epic Games, Inc. All Rights Reserved.
+
+#pragma once
+
+//! Order of the unreal vertex buffers when in mutable data
+#define MUTABLE_VERTEXBUFFER_POSITION	0
+#define MUTABLE_VERTEXBUFFER_TANGENT	1
+#define MUTABLE_VERTEXBUFFER_TEXCOORDS	2
+
+#include "MuR/Ptr.h"
+#include "Containers/Array.h"
+
+namespace mu
+{
+	class Mesh;
+	class FMeshBufferSet;
+}
+struct FReferenceSkeleton;
+
+class FSkeletalMeshLODRenderData;
+class USkeletalMesh;
+class USkeleton;
+
+
+namespace UnrealConversionUtils
+{
+	/*
+	 * Shared methods between mutable instance update and viewport mesh generation
+	 * These are the result of stripping up parts of the reference code to be able to share it on the two
+	 * current pipelines (instance update and USkeletal mesh generation for a mesh viewport)
+	 */
+	
+	/**
+	 * Prepares the render sections found on the InSkeletalMesh and sets them up accordingly what the InMutableMesh requires
+	 * @param OutSkeletalMesh - The Skeletal mesh whose sections are ought to be updated
+	 * @param InMutableMesh - Mutable mesh to be used as reference for the section data update on the Skeletal Mesh
+	 * @param MeshLODIndex - The Mesh lod to be working with
+	 * @param InBoneMap - The bones to be set as part of the sections.
+	 */
+	CUSTOMIZABLEOBJECT_API void SetupRenderSections(
+		const USkeletalMesh* OutSkeletalMesh,
+		const mu::Ptr<const mu::Mesh> InMutableMesh,
+		const int32 InMeshLODIndex,
+		const TArray<uint16>& InBoneMap,
+		const int32 InFirstBoneMapIndex);
+
+
+	/** Performs a copy of the data found on the vertex buffers on the mutable mesh to the buffers of the skeletal mesh
+	 * @param OutSkeletalMesh - Skeletal Mesh to be updated with new buffer data
+	 * @param InMutableMesh - Mutable mesh to be used as reference for the section data update on the Skeletal Mesh
+	 * @param MeshLODIndex - The LOD index we are working with.
+	 */
+	CUSTOMIZABLEOBJECT_API void CopyMutableVertexBuffers(
+		USkeletalMesh* OutSkeletalMesh,
+		const mu::Ptr<const mu::Mesh> InMutableMesh,
+		const int32 InMeshLODIndex);
+
+	
+	/**
+	 *Performs a copy of the data found on the index buffers on the mutable mesh to the buffers of the skeletal mesh
+	 * @param InMutableMesh - The mutable mesh whose index buffers you want to work with
+	 * @param OutLODModel - The LOD model to be updated.
+	 * @return True if the operation could be performed successfully, false if not.
+	 */
+	CUSTOMIZABLEOBJECT_API bool CopyMutableIndexBuffers(mu::Ptr<const mu::Mesh> InMutableMesh,
+	                                                    FSkeletalMeshLODRenderData& OutLODModel);
+	
+
+	/**
+	 *Performs a copy of the data found on the index buffers on the mutable mesh to the buffers of the skeletal mesh
+	 * @param OutLODModel - The LOD model to be updated.
+	 * @param InProfileName - Name of the profile to generate.
+	 * @param InMutableMeshVertexBuffers - The mutable buffers to be reading data from
+	 * @param InBoneIndexBuffer - The buffer containing the indices for the bones.
+	 * @return True if the operation could be performed successfully,	 false if not.
+	 */
+	CUSTOMIZABLEOBJECT_API void CopyMutableSkinWeightProfilesBuffers(
+		FSkeletalMeshLODRenderData& OutLODModel,
+		const FName InProfileName,
+		const mu::FMeshBufferSet& InMutableMeshVertexBuffers,
+		const int32 InBoneIndexBuffer);
+
+
+	/**
+	 *Performs a copy of the render data of a specific Skeletal Mesh LOD to another Skeletal Mesh
+	 * @param SrcSkeletalMesh - Skeletal Mesh with the data to copy.
+	 * @param DestSkeletalMesh - Skeletal Mesh where the data will be copied.
+	 * @param SrcLODIndex - Index of the LOD to copy.
+	 * @param DestSkeletalMesh - Index of the LOD that will be copyed to.
+	 */
+	CUSTOMIZABLEOBJECT_API void CopySkeletalMeshLODRenderData(
+		const USkeletalMesh* SrcSkeletalMesh,
+		USkeletalMesh* DestSkeletalMesh,
+		int32 SrcLODIndex,
+		int32 DestLODIndex);
+
+}
+
+#if UE_ENABLE_INCLUDE_ORDER_DEPRECATED_IN_5_2
+#include "ReferenceSkeleton.h"
+#endif
diff --git a/Engine/Plugins/Experimental/Mutable/Source/CustomizableObject/Internal/MuCO/UnrealMutableModelDiskStreamer.h b/Engine/Plugins/Experimental/Mutable/Source/CustomizableObject/Internal/MuCO/UnrealMutableModelDiskStreamer.h
new file mode 100644
--- /dev/null
+++ b/Engine/Plugins/Experimental/Mutable/Source/CustomizableObject/Internal/MuCO/UnrealMutableModelDiskStreamer.h
@@ -0,0 +1,126 @@
+// Copyright Epic Games, Inc. All Rights Reserved.
+
+#pragma once
+
+#include "MuCO/CustomizableObject.h"
+#include "MuR/Serialisation.h"
+
+class FArchive;
+class IAsyncReadFileHandle;
+class IAsyncReadRequest;
+class UCustomizableObject;
+namespace mu { class Model; }
+struct FMutableStreamableBlock;
+
+class UnrealMutableInputStream : public mu::InputStream
+{
+public:
+
+	UnrealMutableInputStream(FArchive& ar);
+
+	// mu::InputStream interface
+	void Read(void* pData, uint64 size) override;
+
+private:
+	FArchive& m_ar;
+};
+
+
+// Implementation of a mutable streamer using bulk storage.
+class CUSTOMIZABLEOBJECT_API FUnrealMutableModelBulkReader : public mu::ModelReader
+{
+public:
+	// 
+	~FUnrealMutableModelBulkReader();
+
+	// Own interface
+
+	/** Make sure that the provided object can stream data. */
+	bool PrepareStreamingForObject(UCustomizableObject* Object);
+
+#if WITH_EDITOR
+	/** Cancel any further streaming operations for the given object. This is necessary if the object compiled data is
+	 * going to be modified. This can only happen in the editor, when recompiling.
+	 * Any additional streaming requests for this object will fail.
+	 */
+	void CancelStreamingForObject(const UCustomizableObject* CustomizableObject);
+
+	/** Checks if there are any streaming operations for the parameter object.
+	* @return true if there are streaming operations in flight
+	*/
+	bool AreTherePendingStreamingOperationsForObject(const UCustomizableObject* CustomizableObject) const;
+#endif
+
+	/** Release all the pending resources. This disables treamings for all objects. */
+	void EndStreaming();
+
+	// mu::ModelReader interface
+	OPERATION_ID BeginReadBlock(const mu::Model*, uint64 key0, void* pBuffer, uint64 size, TFunction<void(bool bSuccess)>* CompletionCallback) override;
+	bool IsReadCompleted(OPERATION_ID) override;
+	void EndRead(OPERATION_ID) override;
+
+protected:
+
+	struct FReadRequest
+	{
+		TSharedPtr<IAsyncReadRequest> ReadRequest = nullptr;
+		TSharedPtr<FAsyncFileCallBack> FileCallback;
+	};
+	
+	/** Streaming data for one object. */
+	struct FObjectData
+	{
+		TWeakPtr<const mu::Model> Model;
+		TArray<TSharedPtr<IAsyncReadFileHandle>> ReadFileHandles;
+		TMap<OPERATION_ID, FReadRequest> CurrentReadRequests;
+		TMap<uint64, FMutableStreamableBlock> StreamableBlocks;
+	};
+
+	TArray<FObjectData> Objects;
+
+	/** This is used to generate unique ids for read requests. */
+	OPERATION_ID LastOperationID = 0;
+};
+
+
+#if WITH_EDITOR
+
+class UnrealMutableOutputStream : public mu::OutputStream
+{
+public:
+
+	UnrealMutableOutputStream(FArchive& ar);
+
+	// mu::OutputStream interface
+	void Write(const void* pData, uint64 size) override;
+
+private:
+	FArchive& m_ar;
+};
+
+
+// Implementation of a mutable streamer using bulk storage.
+class CUSTOMIZABLEOBJECT_API FUnrealMutableModelBulkWriter : public mu::ModelWriter
+{
+public:
+	// 
+	FUnrealMutableModelBulkWriter(FArchive* InMainDataArchive = nullptr, FArchive* InStreamedDataArchive = nullptr);
+
+	// mu::ModelWriter interface
+	void OpenWriteFile(uint64 key0) override;
+	void Write(const void* pBuffer, uint64 size) override;
+	void CloseWriteFile() override;
+
+protected:
+
+	// Non-owned pointer to an archive where we'll store the main model data (non-streamable)
+	FArchive* MainDataArchive = nullptr;
+
+	// Non-owned pointer to an archive where we'll store the resouces (streamable)
+	FArchive* StreamedDataArchive = nullptr;
+
+	FArchive* CurrentWriteFile = nullptr;
+
+};
+
+#endif
diff --git a/Engine/Plugins/Experimental/Mutable/Source/CustomizableObject/Internal/MuCO/UnrealPortabilityHelpers.h b/Engine/Plugins/Experimental/Mutable/Source/CustomizableObject/Internal/MuCO/UnrealPortabilityHelpers.h
new file mode 100644
--- /dev/null
+++ b/Engine/Plugins/Experimental/Mutable/Source/CustomizableObject/Internal/MuCO/UnrealPortabilityHelpers.h
@@ -0,0 +1,52 @@
+// Copyright Epic Games, Inc. All Rights Reserved.
+
+#pragma once
+
+
+#include "Runtime/Launch/Resources/Version.h"
+
+/** Helpers to ease portability across unreal engine versions */
+
+#if ENGINE_MAJOR_VERSION==5 && ENGINE_MINOR_VERSION<3
+
+#define ObjectPtrWrap(X)	X
+#define ObjectPtrDecay(X)	X
+
+#endif
+
+
+#if ENGINE_MAJOR_VERSION==5 && ENGINE_MINOR_VERSION>=1
+
+#include "Engine/SkinnedAssetCommon.h"
+#include "Rendering/SkinWeightVertexBuffer.h"
+#include "Rendering/SkeletalMeshRenderData.h"
+
+#define UE_MUTABLE_GET_CLASS_PATHNAME(X)		X->GetClassPathName()
+#define UE_MUTABLE_GET_CLASSPATHS(X)			X.ClassPaths
+#define UE_MUTABLE_OBJECTPATH(X)				FSoftObjectPath(X)
+#define UE_MUTABLE_TOPLEVELASSETPATH(X,Y)		FTopLevelAssetPath(X, Y)
+#define UE_MUTABLE_ASSETCLASS(X)				X.AssetClassPath
+#define UE_MUTABLE_GETOBJECTPATH(X)				X.GetObjectPathString()
+
+#define UE_MUTABLE_GETSKINNEDASSET(X)			X->GetSkinnedAsset()
+#define UE_MUTABLE_SETSKINNEDASSET(X,Y)			X->SetSkinnedAsset(Y)
+#define UE_MUTABLE_GETSKELETALMESHASSET(X)		X->GetSkeletalMeshAsset()
+
+#elif ENGINE_MAJOR_VERSION==5 && ENGINE_MINOR_VERSION==0
+
+#define USkinnedAsset USkeletalMesh
+
+#define UE_MUTABLE_GET_CLASS_PATHNAME(X)		X->GetFName()
+#define UE_MUTABLE_GET_CLASSPATHS(X)			X.ClassNames
+#define UE_MUTABLE_OBJECTPATH(X)				*X
+#define UE_MUTABLE_TOPLEVELASSETPATH(X,Y)		(Y)
+#define UE_MUTABLE_ASSETCLASS(X)				X.AssetClass
+#define UE_MUTABLE_GETOBJECTPATH(X)				X.ObjectPath.ToString()
+
+#define UE_MUTABLE_GETSKINNEDASSET(X)			X->SkeletalMesh
+#define UE_MUTABLE_SETSKINNEDASSET(X,Y)			X->SetSkeletalMesh(Y)
+#define UE_MUTABLE_GETSKELETALMESHASSET(X)		X->SkeletalMesh
+
+#endif
+
+
diff --git a/Engine/Plugins/Experimental/Mutable/Source/CustomizableObject/Private/MuCO/CustomizableObjectDGGUI.h b/Engine/Plugins/Experimental/Mutable/Source/CustomizableObject/Private/MuCO/CustomizableObjectDGGUI.h
new file mode 100644
--- /dev/null
+++ b/Engine/Plugins/Experimental/Mutable/Source/CustomizableObject/Private/MuCO/CustomizableObjectDGGUI.h
@@ -0,0 +1,28 @@
+// Copyright Epic Games, Inc. All Rights Reserved.
+
+#pragma once
+
+#include "Blueprint/UserWidget.h"
+
+#include "CustomizableObjectDGGUI.generated.h"
+
+class UAnimInstance;
+class UCustomizableObjectInstanceUsage;
+class UObject;
+class UWorld;
+
+
+UCLASS(Abstract, BlueprintType, Blueprintable)
+class CUSTOMIZABLEOBJECT_API UDGGUI : public UUserWidget
+{
+	GENERATED_BODY()
+public:
+	UFUNCTION(BlueprintImplementableEvent, Category = "DGGUI")
+	class UCustomizableObjectInstanceUsage* GetCustomizableObjectInstanceUsage();
+
+	UFUNCTION(BlueprintImplementableEvent, Category = "DGGUI")
+	void SetCustomizableObjectInstanceUsage(class UCustomizableObjectInstanceUsage* CustomizableObjectInstanceUsage);
+
+	static void OpenDGGUI(const int32 SlotID, UCustomizableObjectInstanceUsage* SelectedCustomizableObjectInstanceUsage, const UWorld* CurrentWorld, const int32 PlayerIndex = 0);
+	static bool CloseExistingDGGUI(const UWorld* CurrentWorld);
+};
diff --git a/Engine/Plugins/Experimental/Mutable/Source/CustomizableObject/Public/MuCO/CustomizableObjectClothingTypes.h b/Engine/Plugins/Experimental/Mutable/Source/CustomizableObject/Public/MuCO/CustomizableObjectClothingTypes.h
deleted file mode 100644
diff --git a/Engine/Plugins/Experimental/Mutable/Source/CustomizableObject/Public/MuCO/CustomizableObjectCustomVersion.h b/Engine/Plugins/Experimental/Mutable/Source/CustomizableObject/Public/MuCO/CustomizableObjectCustomVersion.h
deleted file mode 100644
diff --git a/Engine/Plugins/Experimental/Mutable/Source/CustomizableObject/Public/MuCO/CustomizableObjectDGGUI.h b/Engine/Plugins/Experimental/Mutable/Source/CustomizableObject/Public/MuCO/CustomizableObjectDGGUI.h
deleted file mode 100644
diff --git a/Engine/Plugins/Experimental/Mutable/Source/CustomizableObject/Public/MuCO/CustomizableObjectInstanceAssetUserData.h b/Engine/Plugins/Experimental/Mutable/Source/CustomizableObject/Public/MuCO/CustomizableObjectInstanceAssetUserData.h
deleted file mode 100644
diff --git a/Engine/Plugins/Experimental/Mutable/Source/CustomizableObject/Public/MuCO/CustomizableObjectMipDataProvider.h b/Engine/Plugins/Experimental/Mutable/Source/CustomizableObject/Public/MuCO/CustomizableObjectMipDataProvider.h
deleted file mode 100644
diff --git a/Engine/Plugins/Experimental/Mutable/Source/CustomizableObject/Public/MuCO/CustomizableObjectResourceDataTypes.h b/Engine/Plugins/Experimental/Mutable/Source/CustomizableObject/Public/MuCO/CustomizableObjectResourceDataTypes.h
deleted file mode 100644
diff --git a/Engine/Plugins/Experimental/Mutable/Source/CustomizableObject/Public/MuCO/EditorImageProvider.h b/Engine/Plugins/Experimental/Mutable/Source/CustomizableObject/Public/MuCO/EditorImageProvider.h
deleted file mode 100644
diff --git a/Engine/Plugins/Experimental/Mutable/Source/CustomizableObject/Public/MuCO/FMutableTaskGraph.h b/Engine/Plugins/Experimental/Mutable/Source/CustomizableObject/Public/MuCO/FMutableTaskGraph.h
deleted file mode 100644
diff --git a/Engine/Plugins/Experimental/Mutable/Source/CustomizableObject/Public/MuCO/MutableMeshBufferUtils.h b/Engine/Plugins/Experimental/Mutable/Source/CustomizableObject/Public/MuCO/MutableMeshBufferUtils.h
deleted file mode 100644
diff --git a/Engine/Plugins/Experimental/Mutable/Source/CustomizableObject/Public/MuCO/MutableProjectorTypeUtils.h b/Engine/Plugins/Experimental/Mutable/Source/CustomizableObject/Public/MuCO/MutableProjectorTypeUtils.h
deleted file mode 100644
diff --git a/Engine/Plugins/Experimental/Mutable/Source/CustomizableObject/Public/MuCO/StateMachine.h b/Engine/Plugins/Experimental/Mutable/Source/CustomizableObject/Public/MuCO/StateMachine.h
deleted file mode 100644
diff --git a/Engine/Plugins/Experimental/Mutable/Source/CustomizableObject/Public/MuCO/UnrealConversionUtils.h b/Engine/Plugins/Experimental/Mutable/Source/CustomizableObject/Public/MuCO/UnrealConversionUtils.h
deleted file mode 100644
diff --git a/Engine/Plugins/Experimental/Mutable/Source/CustomizableObject/Public/MuCO/UnrealMutableModelDiskStreamer.h b/Engine/Plugins/Experimental/Mutable/Source/CustomizableObject/Public/MuCO/UnrealMutableModelDiskStreamer.h
deleted file mode 100644
diff --git a/Engine/Plugins/Experimental/Mutable/Source/CustomizableObject/Public/MuCO/UnrealPortabilityHelpers.h b/Engine/Plugins/Experimental/Mutable/Source/CustomizableObject/Public/MuCO/UnrealPortabilityHelpers.h
deleted file mode 100644
