[mutable] Fixed states view for child objects.
+ Fixed widget appearence too.

#rnx

diff --git a/Engine/Plugins/Experimental/Mutable/Source/CustomizableObject/Internal/MuCO/CustomizableObjectPrivate.h b/Engine/Plugins/Experimental/Mutable/Source/CustomizableObject/Internal/MuCO/CustomizableObjectPrivate.h
--- a/Engine/Plugins/Experimental/Mutable/Source/CustomizableObject/Internal/MuCO/CustomizableObjectPrivate.h
+++ b/Engine/Plugins/Experimental/Mutable/Source/CustomizableObject/Internal/MuCO/CustomizableObjectPrivate.h
@@ -48,9 +48,6 @@ public:
 	// Return true if this object doesn't reference a parent object.
 	virtual bool IsRootObject(const class UCustomizableObject* Object) const { return true; }
 
-	/** Returns the Customizable Object that does start the CO tree */
-	virtual UCustomizableObject* GetRootObject( class UCustomizableObject* Object) = 0;
-
 	/** Provides the caller with the warning and error messages produced during compilation */
 	virtual void GetCompilationMessages(TArray<FText>& OutWarningMessages, TArray<FText>& OutErrorMessages) const = 0;
 	
diff --git a/Engine/Plugins/Experimental/Mutable/Source/CustomizableObjectEditor/Private/MuCOE/CustomizableObjectCompiler.cpp b/Engine/Plugins/Experimental/Mutable/Source/CustomizableObjectEditor/Private/MuCOE/CustomizableObjectCompiler.cpp
--- a/Engine/Plugins/Experimental/Mutable/Source/CustomizableObjectEditor/Private/MuCOE/CustomizableObjectCompiler.cpp
+++ b/Engine/Plugins/Experimental/Mutable/Source/CustomizableObjectEditor/Private/MuCOE/CustomizableObjectCompiler.cpp
@@ -451,24 +451,6 @@ void FCustomizableObjectCompiler::DisplayOrphanNodesWarning(FMutableGraphGenerat
 }
 
 
-UCustomizableObject* FCustomizableObjectCompiler::GetRootObject( UCustomizableObject* InObject)
-{
-	// Grab a node to start the search -> Get the root since it should be always present
-	bool bMultipleBaseObjectsFound = false;
-	UCustomizableObjectNodeObject* ObjectRootNode = GetRootNode(InObject, bMultipleBaseObjectsFound);
-
-	if (ObjectRootNode && ObjectRootNode->ParentObject)
-	{
-		TArray<UCustomizableObject*> VisitedNodes;
-		return GetFullGraphRootObject(ObjectRootNode,VisitedNodes);
-	}
-
-	// No parent object found, return input as the parent of the graph
-	// This can also mean the ObjectRootNode does not exist because it has not been opened yet (so no nodes have been generated)
-	return InObject;
-}
-
-
 mu::NodeObjectPtr FCustomizableObjectCompiler::GenerateMutableRoot( 
 	UCustomizableObject* Object, 
 	FMutableGraphGenerationContext& GenerationContext, 
diff --git a/Engine/Plugins/Experimental/Mutable/Source/CustomizableObjectEditor/Private/MuCOE/CustomizableObjectCompiler.h b/Engine/Plugins/Experimental/Mutable/Source/CustomizableObjectEditor/Private/MuCOE/CustomizableObjectCompiler.h
--- a/Engine/Plugins/Experimental/Mutable/Source/CustomizableObjectEditor/Private/MuCOE/CustomizableObjectCompiler.h
+++ b/Engine/Plugins/Experimental/Mutable/Source/CustomizableObjectEditor/Private/MuCOE/CustomizableObjectCompiler.h
@@ -32,14 +32,6 @@ public:
 	/** Check for pending compilation process. Returns true if an object has been updated. */
 	CUSTOMIZABLEOBJECTEDITOR_API virtual bool Tick() override;
 
-	/** Provided a CO object it provides the root CO it is connected. In other words : it returns the root of the entire
-	 * mutable graph.
-	 * @param InObject Customizable object whose root CO we are asking for.
-	 * @return The CO that is the root of the provided Customizable Object. It can be equal to InObject if the provided
-	 * object does not have any parent.
-	 */
-	virtual UCustomizableObject* GetRootObject(UCustomizableObject* InObject) override;
-
 	/** Generate the Mutable Graph from the Unreal Graph. */
 	mu::NodePtr Export(UCustomizableObject* Object, const FCompilationOptions& Options, TArray<TSoftObjectPtr<UTexture>>& OutRuntimeReferencedTextures, TArray<TSoftObjectPtr<UTexture>>& OutCompilerReferencedTextures);
 
diff --git a/Engine/Plugins/Experimental/Mutable/Source/CustomizableObjectEditor/Private/MuCOE/CustomizableObjectDetails.cpp b/Engine/Plugins/Experimental/Mutable/Source/CustomizableObjectEditor/Private/MuCOE/CustomizableObjectDetails.cpp
--- a/Engine/Plugins/Experimental/Mutable/Source/CustomizableObjectEditor/Private/MuCOE/CustomizableObjectDetails.cpp
+++ b/Engine/Plugins/Experimental/Mutable/Source/CustomizableObjectEditor/Private/MuCOE/CustomizableObjectDetails.cpp
@@ -7,6 +7,8 @@
 #include "IDetailsView.h"
 #include "MuCO/CustomizableObject.h"
 #include "MuCO/CustomizableObjectPrivate.h"
+#include "MuCOE/GraphTraversal.h"
+#include "MuCOE/Nodes/CustomizableObjectNodeObject.h"
 #include "MuR/Model.h"
 #include "MuR/Parameters.h"
 #include "MuR/Ptr.h"
@@ -76,11 +78,9 @@ public:
 
 		return ChildrenList;
 	}
-
 };
 
 
-
 void FCustomizableObjectDetails::CustomizeDetails( IDetailLayoutBuilder& DetailBuilder )
 {
 	const IDetailsView* DetailsView = DetailBuilder.GetDetailsView();
@@ -111,7 +111,7 @@ void FCustomizableObjectDetails::CustomizeDetails( IDetailLayoutBuilder& DetailB
 		+ SVerticalBox::Slot()
 		.Padding( 2.0f )
 		[
-			SAssignNew(StatesTree, STreeView<TSharedPtr< FStateDetailsNode > >)
+			SAssignNew(StatesTree, STreeView<TSharedPtr<FStateDetailsNode>>)
 			.SelectionMode(ESelectionMode::Single)
 			.TreeItemsSource( &RootTreeItems )
 			// Called to child items for any given parent item
@@ -120,7 +120,6 @@ void FCustomizableObjectDetails::CustomizeDetails( IDetailLayoutBuilder& DetailB
 			.OnGenerateRow( this, &FCustomizableObjectDetails::OnGenerateRowForStateTree ) 
 			// Allow for some spacing between items with a larger item height.
 			.ItemHeight(20.0f)
-
 			.HeaderRow
 			(
 				SNew(SHeaderRow)
@@ -148,17 +147,20 @@ void FCustomizableObjectDetails::CustomizeDetails( IDetailLayoutBuilder& DetailB
 	}
 }
 
+
 void FCustomizableObjectDetails::UpdateTree()
 {
-	if ( CustomizableObject &&  CustomizableObject->GetPrivate()->GetModel() )
+	UCustomizableObject* RootObject = GetRootObject(CustomizableObject);
+
+	if (RootObject && RootObject->GetPrivate()->GetModel() )
 	{
 		RootTreeItems.SetNumUninitialized(0);
-		const uint32 NumElements = CustomizableObject->GetPrivate()->GetModel()->GetStateCount();
+		const uint32 NumElements = RootObject->GetPrivate()->GetModel()->GetStateCount();
 		for ( uint32 i=0; i<NumElements; ++i )
 		{
 			TSharedPtr<FStateDetailsNode> SlateDetailsNode = MakeShareable( new FStateDetailsNode );
 
-			SlateDetailsNode->Model = CustomizableObject->GetPrivate()->GetModel();
+			SlateDetailsNode->Model = RootObject->GetPrivate()->GetModel();
 			SlateDetailsNode->StateIndex = i;
 
 			RootTreeItems.Add( SlateDetailsNode );
@@ -170,6 +172,8 @@ void FCustomizableObjectDetails::UpdateTree()
 		StatesTree->RequestTreeRefresh();
 	}
 }
+
+
 /** The item used for visualizing the class in the tree. */
 class SStateItem : public STableRow< TSharedPtr<FString> >
 {
@@ -205,25 +209,36 @@ public:
 		StateName = InArgs._StateName;
 		AssociatedNode = InArgs._AssociatedNode;
 
+		FMargin Padding = { 0.0f, 3.0f, 6.0f, 3.0f };
+		
+		if (AssociatedNode->ParameterIndex != INDEX_NONE)
+		{
+			Padding = { 10.0f, 3.0f, 6.0f, 3.0f };
+		}
+
 		this->ChildSlot
 		[
 			SNew(SHorizontalBox)
 				
 			+SHorizontalBox::Slot()
-				[
-					SNew( SExpanderArrow, SharedThis(this) )
-				]
+			.AutoWidth()
+			.VAlign(VAlign_Center)
+			.HAlign(HAlign_Right)
+			[
+				SNew( SExpanderArrow, SharedThis(this) )
+			]
 
 			+SHorizontalBox::Slot()
-				.FillWidth(1.0f)
-				.Padding( 0.0f, 3.0f, 6.0f, 3.0f )
-				.VAlign(VAlign_Center)
-				[
-					SNew( STextBlock )
-						.Text( FText::FromString(*StateName.Get()) )
-						//.HighlightString(*InArgs._HighlightText)
-						//.ColorAndOpacity( this, &SClassItem::GetTextColor)
-				]
+			.AutoWidth()
+			.Padding( Padding )
+			.VAlign(VAlign_Center)
+			.HAlign(HAlign_Left)
+			[
+				SNew( STextBlock )
+				.Text( FText::FromString(*StateName.Get()) )
+				//.HighlightString(*InArgs._HighlightText)
+				//.ColorAndOpacity( this, &SClassItem::GetTextColor)
+			]
 
 			//+SHorizontalBox::Slot()
 			//	.HAlign(HAlign_Right)
diff --git a/Engine/Plugins/Experimental/Mutable/Source/CustomizableObjectEditor/Private/MuCOE/GraphTraversal.cpp b/Engine/Plugins/Experimental/Mutable/Source/CustomizableObjectEditor/Private/MuCOE/GraphTraversal.cpp
--- a/Engine/Plugins/Experimental/Mutable/Source/CustomizableObjectEditor/Private/MuCOE/GraphTraversal.cpp
+++ b/Engine/Plugins/Experimental/Mutable/Source/CustomizableObjectEditor/Private/MuCOE/GraphTraversal.cpp
@@ -304,6 +304,24 @@ UCustomizableObject* GetRootObject(const UCustomizableObjectNode& Node)
 }
 
 
+UCustomizableObject* GetRootObject(UCustomizableObject* ChildObject)
+{
+	// Grab a node to start the search -> Get the root since it should be always present
+	bool bMultipleBaseObjectsFound = false;
+	UCustomizableObjectNodeObject* ObjectRootNode = GetRootNode(ChildObject, bMultipleBaseObjectsFound);
+
+	if (ObjectRootNode && ObjectRootNode->ParentObject)
+	{
+		TArray<UCustomizableObject*> VisitedNodes;
+		return GetFullGraphRootObject(ObjectRootNode, VisitedNodes);
+	}
+
+	// No parent object found, return input as the parent of the graph
+	// This can also mean the ObjectRootNode does not exist because it has not been opened yet (so no nodes have been generated)
+	return ChildObject;
+}
+
+
 UCustomizableObjectNodeObject* GetFullGraphRootNodeObject(UCustomizableObjectNodeObject* Node, TArray<UCustomizableObject*>& VisitedObjects)
 {
 	if (Node->ParentObject != nullptr)
diff --git a/Engine/Plugins/Experimental/Mutable/Source/CustomizableObjectEditor/Public/MuCOE/GraphTraversal.h b/Engine/Plugins/Experimental/Mutable/Source/CustomizableObjectEditor/Public/MuCOE/GraphTraversal.h
--- a/Engine/Plugins/Experimental/Mutable/Source/CustomizableObjectEditor/Public/MuCOE/GraphTraversal.h
+++ b/Engine/Plugins/Experimental/Mutable/Source/CustomizableObjectEditor/Public/MuCOE/GraphTraversal.h
@@ -52,6 +52,14 @@ UCustomizableObject* GetFullGraphRootObject(UCustomizableObjectNodeObject* Node,
 /** Return Customizable Object of the given node. */
 UCustomizableObject* GetRootObject(const UCustomizableObjectNode& Node);
 
+/** Provided a CO object it provides the root CO it is connected. In other words : it returns the root of the entire
+	 * mutable graph.
+	 * @param InObject Customizable object whose root CO we are asking for.
+	 * @return The CO that is the root of the provided Customizable Object. It can be equal to InObject if the provided
+	 * object does not have any parent.
+*/
+CUSTOMIZABLEOBJECTEDITOR_API UCustomizableObject* GetRootObject(UCustomizableObject* ChildObject);
+
 /** Return the full graph Customizable Object Node root of the node given as parameter */
 UCustomizableObjectNodeObject* GetFullGraphRootNodeObject(UCustomizableObjectNodeObject* Node, TArray<UCustomizableObject*>& VisitedObjects);
 
diff --git a/Engine/Plugins/Experimental/Mutable/Source/MutableValidation/Private/MuV/AssetValidator_CustomizableObjects.cpp b/Engine/Plugins/Experimental/Mutable/Source/MutableValidation/Private/MuV/AssetValidator_CustomizableObjects.cpp
--- a/Engine/Plugins/Experimental/Mutable/Source/MutableValidation/Private/MuV/AssetValidator_CustomizableObjects.cpp
+++ b/Engine/Plugins/Experimental/Mutable/Source/MutableValidation/Private/MuV/AssetValidator_CustomizableObjects.cpp
@@ -13,6 +13,7 @@
 #include "MuCO/CustomizableObject.h"
 #include "MuCO/CustomizableObjectPrivate.h"
 #include "MuCO/CustomizableObjectSystem.h"
+#include "MuCOE/GraphTraversal.h"
 #include "UObject/NameTypes.h"
 #include "UObject/Object.h"
 
@@ -92,7 +93,7 @@ EDataValidationResult UAssetValidator_CustomizableObjects::IsCustomizableObjectV
 		TUniquePtr<FCustomizableObjectCompilerBase>(UCustomizableObjectSystem::GetInstance()->GetNewCompiler());
 	
 	// Find out which is the root for this CO (it may be itself but that is OK)
-	UCustomizableObject* RootObject = Compiler->GetRootObject(InCustomizableObject);
+	UCustomizableObject* RootObject = GetRootObject(InCustomizableObject);
 	check (RootObject);
 	
 	// Check that the object to be compiled has not already been compiled
