[mutable] Fixed differences in Mutable textures depending on the optimization level. These differences were related to the use of color constant nodes and texture swizzling.
Textures generated from color constants didn't inherit the correct size from the reference texture, so were 4x4 and forced the other textures involved in a swizzle to be resized to 4x4 when compiled without optimization.
Compiling with optimizition, when a texture generated from a color constant was not the first one in a swizzle node with other textures with more that 1 mip level, it didn't force the resulting swizzled texture to have only 1 mip level. That meant that the resulting texture had a correct swizzled mip level 0, but the rest were blank.
#rnx
#tests PIE

diff --git a/Engine/Plugins/Experimental/Mutable/Source/CustomizableObjectEditor/Private/MuCOE/GenerateMutableSource/GenerateMutableSourceImage.cpp b/Engine/Plugins/Experimental/Mutable/Source/CustomizableObjectEditor/Private/MuCOE/GenerateMutableSource/GenerateMutableSourceImage.cpp
--- a/Engine/Plugins/Experimental/Mutable/Source/CustomizableObjectEditor/Private/MuCOE/GenerateMutableSource/GenerateMutableSourceImage.cpp
+++ b/Engine/Plugins/Experimental/Mutable/Source/CustomizableObjectEditor/Private/MuCOE/GenerateMutableSource/GenerateMutableSourceImage.cpp
@@ -477,6 +477,11 @@ mu::NodeImagePtr GenerateMutableSourceImage(const UEdGraphPin* Pin, FMutableGrap
 		{
 			ImageFromColour->SetColour(color);
 		}
+
+		if (ReferenceTextureSize > 0)
+		{
+			ImageFromColour->SetSize(ReferenceTextureSize, ReferenceTextureSize);
+		}
 	}
 
 	else if (const UCustomizableObjectNodeTextureFromChannels* TypedNodeFrom = Cast<UCustomizableObjectNodeTextureFromChannels>(Node))
diff --git a/Engine/Plugins/Experimental/Mutable/Source/MutableRuntime/Private/MuR/CodeRunnerTasked.cpp b/Engine/Plugins/Experimental/Mutable/Source/MutableRuntime/Private/MuR/CodeRunnerTasked.cpp
--- a/Engine/Plugins/Experimental/Mutable/Source/MutableRuntime/Private/MuR/CodeRunnerTasked.cpp
+++ b/Engine/Plugins/Experimental/Mutable/Source/MutableRuntime/Private/MuR/CodeRunnerTasked.cpp
@@ -1469,6 +1469,15 @@ namespace mu
 			}
 		}
 
+		// If any source has only 1 LOD, then the result has to have 1 LOD and the rest be regenerated later on
+		for (int32 SourceIndex = 0; SourceIndex < MUTABLE_OP_MAX_SWIZZLE_CHANNELS; ++SourceIndex)
+		{
+			if (Sources[SourceIndex] && Sources[SourceIndex]->GetLODCount() == 1)
+			{
+				ResultLODs = 1;
+			}
+		}
+
 		Result = Runner->CreateImage(Sources[0]->GetSizeX(), Sources[0]->GetSizeY(), ResultLODs, Args.format, EInitializationType::Black);
 		return true;
 	}
