[Mutable] Fix pose re-initialization with every mesh update.

#rnx

diff --git a/Engine/Plugins/Experimental/Mutable/Source/CustomizableObject/Private/MuCO/CustomizableObjectInstanceUsage.cpp b/Engine/Plugins/Experimental/Mutable/Source/CustomizableObject/Private/MuCO/CustomizableObjectInstanceUsage.cpp
--- a/Engine/Plugins/Experimental/Mutable/Source/CustomizableObject/Private/MuCO/CustomizableObjectInstanceUsage.cpp
+++ b/Engine/Plugins/Experimental/Mutable/Source/CustomizableObject/Private/MuCO/CustomizableObjectInstanceUsage.cpp
@@ -195,13 +195,50 @@ USkeletalMesh* UCustomizableObjectInstanceUsage::GetSkeletalMesh() const
 }
 
 
+bool RequiresReinitPose(USkeletalMesh* CurrentSkeletalMesh, USkeletalMesh* SkeletalMesh)
+{
+	if (CurrentSkeletalMesh == SkeletalMesh)
+	{
+		return false;
+	}
+
+	if (!CurrentSkeletalMesh || !SkeletalMesh)
+	{
+		return SkeletalMesh != nullptr;
+	}
+
+	if (CurrentSkeletalMesh->GetLODNum() != SkeletalMesh->GetLODNum())
+	{
+		return true;
+	}
+
+	const FSkeletalMeshRenderData* CurrentRenderData = CurrentSkeletalMesh->GetResourceForRendering();
+	const FSkeletalMeshRenderData* NewRenderData = SkeletalMesh->GetResourceForRendering();
+	if (!CurrentRenderData || !NewRenderData)
+	{
+		return false;
+	}
+
+	const int32 NumLODs = SkeletalMesh->GetLODNum();
+	for (int32 LODIndex = 0; LODIndex < NumLODs; ++LODIndex)
+	{
+		if (CurrentRenderData->LODRenderData[LODIndex].RequiredBones != NewRenderData->LODRenderData[LODIndex].RequiredBones)
+		{
+			return true;
+		}
+	}
+
+	return false;
+}
+
+
 void UCustomizableObjectInstanceUsage::SetSkeletalMesh(USkeletalMesh* SkeletalMesh)
 {
 	USkeletalMeshComponent* Parent = Cast<USkeletalMeshComponent>(GetAttachParent());
 
 	if (Parent)
 	{
-		Parent->SetSkeletalMesh(SkeletalMesh, true);
+		Parent->SetSkeletalMesh(SkeletalMesh, RequiresReinitPose(Parent->GetSkeletalMeshAsset(), SkeletalMesh));
 
 		const UCustomizableObjectInstance* Instance = GetCustomizableObjectInstance();
 		const UCustomizableObject* CustomizableObject = Instance ? Instance->GetCustomizableObject() : nullptr;
@@ -500,7 +537,7 @@ void UCustomizableObjectInstanceUsage::Tick(float DeltaTime)
 		// Set SkeletalMesh
 		if (bInstanceGenerated || SkeletalMesh)
 		{
-			Parent->SetSkeletalMesh(SkeletalMesh);
+			Parent->SetSkeletalMesh(SkeletalMesh, RequiresReinitPose(Parent->GetSkeletalMeshAsset(), SkeletalMesh));
 
 			if (Parent->HasOverrideMaterials())
 			{
