[Mutable] Fix node ClipMorphMesh not working if the Skeleton bones are not ordered

#rnx

diff --git a/Engine/Plugins/Experimental/Mutable/Source/MutableRuntime/Private/MuR/OpMeshClipMorphPlane.h b/Engine/Plugins/Experimental/Mutable/Source/MutableRuntime/Private/MuR/OpMeshClipMorphPlane.h
--- a/Engine/Plugins/Experimental/Mutable/Source/MutableRuntime/Private/MuR/OpMeshClipMorphPlane.h
+++ b/Engine/Plugins/Experimental/Mutable/Source/MutableRuntime/Private/MuR/OpMeshClipMorphPlane.h
@@ -100,27 +100,37 @@ namespace mu
 
         Ptr<const Skeleton> BaseSkeleton = pBase->GetSkeleton();
 
-		TArray<bool> AffectedBones;
+		TArray<bool> AffectedBoneMapIndices;
 		const int32 BaseBoneIndex = BaseSkeleton ? BaseSkeleton->FindBone(BoneId) : INDEX_NONE;
         if (BaseBoneIndex != INDEX_NONE)
 		{
-			AffectedBones.SetNum(BaseSkeleton->GetBoneCount());
-			AffectedBones[BaseBoneIndex] = true;
+			const TArray<uint16>& BoneMap = pBase->BoneMap;
+			AffectedBoneMapIndices.SetNum(BoneMap.Num());
+
+			const int32 BoneCount = BaseSkeleton->GetBoneCount();
+			TArray<bool> AffectedSkeletonBones;
+			AffectedSkeletonBones.SetNum(BoneCount);
 
-            for (int32 BoneIndex = 0; BoneIndex < BaseSkeleton->GetBoneCount(); ++BoneIndex)
+            for (int32 BoneIndex = 0; BoneIndex < BoneCount; ++BoneIndex)
 			{
-				int32 CurrentBoneIndex = BoneIndex;
-				while (CurrentBoneIndex >= 0)
+				const int32 ParentBoneIndex = BaseSkeleton->GetBoneParent(BoneIndex);
+				check(ParentBoneIndex < BoneIndex);
+				
+				const bool bIsBoneAffected = (AffectedSkeletonBones.IsValidIndex(ParentBoneIndex) && AffectedSkeletonBones[ParentBoneIndex])
+					|| BoneIndex == BaseBoneIndex;
+				
+				if (!bIsBoneAffected)
 				{
-                    if (BaseSkeleton->GetBoneId(CurrentBoneIndex) == BaseBoneIndex)
-					{
-						AffectedBones[BoneIndex] = true;
-						break;
-					}
-					else
-					{
-						CurrentBoneIndex = BaseSkeleton->GetBoneParent(CurrentBoneIndex);
-					}
+					continue;
+				}
+				
+				AffectedSkeletonBones[BoneIndex] = true;
+
+				const uint16 AffectedBoneId = BaseSkeleton->GetBoneId(BoneIndex);
+				const int32 AffectedBoneMapIndex = BoneMap.Find(AffectedBoneId);
+				if (AffectedBoneMapIndex != INDEX_NONE)
+				{
+					AffectedBoneMapIndices[AffectedBoneMapIndex] = true;
 				}
 			}
 
@@ -128,7 +138,6 @@ namespace mu
 			vertex_info.SetNum(vcount);
 			//int firstCount = pBase->GetVertexBuffers().GetElementCount();
 
-			const TArray<uint16>& BoneMap = pBase->BoneMap;
 			for (int32 vb = 0; vb < pBase->GetVertexBuffers().m_buffers.Num(); ++vb)
 			{
 				const MESH_BUFFER& result = pBase->GetVertexBuffers().m_buffers[vb];
@@ -162,7 +171,7 @@ namespace mu
 								
 								for (int j = 0; j < components; ++j)
 								{
-									vertex_info[i].bone_indices.Add(BoneMap[pD[j]]);
+									vertex_info[i].bone_indices.Add(pD[j]);
 								}
 
 								secondOffset += elemSize;
@@ -176,7 +185,7 @@ namespace mu
 
 								for (int j = 0; j < components; ++j)
 								{
-									vertex_info[i].bone_indices.Add(BoneMap[pD[j]]);
+									vertex_info[i].bone_indices.Add(pD[j]);
 								}
 
 								secondOffset += elemSize;
@@ -190,7 +199,7 @@ namespace mu
 
 								for (int j = 0; j < components; ++j)
 								{
-									vertex_info[i].bone_indices.Add(BoneMap[pD[j]]);
+									vertex_info[i].bone_indices.Add(pD[j]);
 								}
 
 								secondOffset += elemSize;
@@ -305,7 +314,7 @@ namespace mu
 						const bool bIsVertexAffectedBone = 
 								BaseBoneIndex != INDEX_NONE &&
 								VertexIsInMaxRadius(vertex, origin, vertexSelectionBoneMaxRadius) &&
-								VertexIsAffectedByBone(v, AffectedBones, vertex_info);
+								VertexIsAffectedByBone(v, AffectedBoneMapIndices, vertex_info);
 						
 						const bool bIsVertexAffectedNoShape = 
 								BaseBoneIndex == INDEX_NONE && 
