[Mutable] Fix Populations class crash on Characteristic's parameter set.

Texture was no longer used due to ParameterDecorations being deprecated.

#rnx

diff --git a/Engine/Plugins/Experimental/Mutable/Source/CustomizableObjectPopulationEditor/Private/MuCOPE/CustomizableObjectPopulationClassDetails.cpp b/Engine/Plugins/Experimental/Mutable/Source/CustomizableObjectPopulationEditor/Private/MuCOPE/CustomizableObjectPopulationClassDetails.cpp
--- a/Engine/Plugins/Experimental/Mutable/Source/CustomizableObjectPopulationEditor/Private/MuCOPE/CustomizableObjectPopulationClassDetails.cpp
+++ b/Engine/Plugins/Experimental/Mutable/Source/CustomizableObjectPopulationEditor/Private/MuCOPE/CustomizableObjectPopulationClassDetails.cpp
@@ -1918,50 +1918,26 @@ void SRangeSquare::Construct(const FArguments& InArgs)
 {
 	SquareColor	= InArgs._SquareColor;
 	Ranges = InArgs._Ranges;
-	Texture = InArgs._Texture;
 	bDiscrete = InArgs._bDiscrete;
-		
-	if (Texture)
-	{
-		Brush.SetResourceObject(Texture);
-		Brush.TintColor = FSlateColor(FLinearColor::White);
-		Brush.ImageSize.X = 400.0f;
-		Brush.ImageSize.Y = 20.0f;
-		Brush.DrawAs = ESlateBrushDrawType::Image;
-		
-		ChildSlot
-		[
-			SNew(SHorizontalBox)
-			+ SHorizontalBox::Slot()
-			.AutoWidth()
-			[
-				SNew(SImage)
-				.Image(&Brush)
-			]
-		];
-	}
-	else
-	{
-		Brush.TintColor = FSlateColor(FLinearColor::Gray);
-		Brush.ImageSize.X = 400.0f;
-		Brush.ImageSize.Y = 10.0f;
-		Brush.DrawAs = ESlateBrushDrawType::Image;
+	
+	Brush.TintColor = FSlateColor(FLinearColor::Gray);
+	Brush.ImageSize.X = 400.0f;
+	Brush.ImageSize.Y = 10.0f;
+	Brush.DrawAs = ESlateBrushDrawType::Image;
 
-		ChildSlot
+	ChildSlot
+	[
+		SNew(SHorizontalBox)
+		+ SHorizontalBox::Slot()
+		.AutoWidth()
 		[
-			SNew(SHorizontalBox)
-			+ SHorizontalBox::Slot()
-			.AutoWidth()
-			[
-				SNew(SImage)
-				.Image(&Brush)
-			]
-		];
-	}
+			SNew(SImage)
+			.Image(&Brush)
+		]
+	];
 	
 	bMouseDownMin = bMouseDownMax = false;
 	TextureRectangle = Brush.ImageSize;
-
 }
 
 
diff --git a/Engine/Plugins/Experimental/Mutable/Source/CustomizableObjectPopulationEditor/Private/MuCOPE/CustomizableObjectPopulationClassDetails.h b/Engine/Plugins/Experimental/Mutable/Source/CustomizableObjectPopulationEditor/Private/MuCOPE/CustomizableObjectPopulationClassDetails.h
--- a/Engine/Plugins/Experimental/Mutable/Source/CustomizableObjectPopulationEditor/Private/MuCOPE/CustomizableObjectPopulationClassDetails.h
+++ b/Engine/Plugins/Experimental/Mutable/Source/CustomizableObjectPopulationEditor/Private/MuCOPE/CustomizableObjectPopulationClassDetails.h
@@ -35,7 +35,6 @@ public:
 	SLATE_ARGUMENT(FColor, SquareColor)
 	SLATE_ARGUMENT(TArray<FConstraintRanges>*, Ranges)
 	SLATE_ARGUMENT(bool, bDiscrete)
-	SLATE_ARGUMENT(UTexture*, Texture)
 	SLATE_END_ARGS()
 
 	void Construct(const FArguments& InArgs);
@@ -71,8 +70,6 @@ private:
 
 	/** Brush to draw the decorations of a parameter option */
 	FSlateBrush Brush;
-
-	UTexture* Texture;
 };
 
 
