[mutable] Disabled cast shadows for generated meshes.

#rnx

diff --git a/Engine/Plugins/Experimental/Mutable/Source/CustomizableObject/Private/MuCO/CustomizableObjectInstance.cpp b/Engine/Plugins/Experimental/Mutable/Source/CustomizableObject/Private/MuCO/CustomizableObjectInstance.cpp
--- a/Engine/Plugins/Experimental/Mutable/Source/CustomizableObject/Private/MuCO/CustomizableObjectInstance.cpp
+++ b/Engine/Plugins/Experimental/Mutable/Source/CustomizableObject/Private/MuCO/CustomizableObjectInstance.cpp
@@ -6784,7 +6784,9 @@ void UCustomizableInstancePrivate::RegenerateImportedModels()
 				SectionUserData.CorrespondClothAssetIndex = RenderSection.CorrespondClothAssetIndex;
 				SectionUserData.ClothingData.AssetGuid = RenderSection.ClothingData.AssetGuid;
 				SectionUserData.ClothingData.AssetLodIndex = RenderSection.ClothingData.AssetLodIndex;
-				SectionUserData.bCastShadow = RenderSection.bCastShadow;
+				
+				//TODO(Max): MTBL-1779
+				//SectionUserData.bCastShadow = RenderSection.bCastShadow;
 
 				ImportedModel->LODModels[LODIndex].SyncronizeUserSectionsDataArray();
 
diff --git a/Engine/Plugins/Experimental/Mutable/Source/CustomizableObject/Private/MuCO/UnrealConversionUtils.cpp b/Engine/Plugins/Experimental/Mutable/Source/CustomizableObject/Private/MuCO/UnrealConversionUtils.cpp
--- a/Engine/Plugins/Experimental/Mutable/Source/CustomizableObject/Private/MuCO/UnrealConversionUtils.cpp
+++ b/Engine/Plugins/Experimental/Mutable/Source/CustomizableObject/Private/MuCO/UnrealConversionUtils.cpp
@@ -177,7 +177,9 @@ namespace UnrealConversionUtils
 			Section.BaseVertexIndex = FirstVertex;
 			Section.MaxBoneInfluences = NumBoneInfluences;
 			Section.NumVertices = VertexCount;
-			Section.bCastShadow = bCastShadow;
+
+			//TODO(Max): MTBL-1779
+			//Section.bCastShadow = bCastShadow;
 
 			// InBoneMaps may contain bonemaps from other sections. Copy the bones belonging to this mesh.
 			FirstBone += InFirstBoneMapIndex;
@@ -479,7 +481,9 @@ namespace UnrealConversionUtils
 					 DestSection->MaxBoneInfluences = SrcSection.MaxBoneInfluences;
 					 DestSection->NumVertices = SrcSection.NumVertices;
 					 DestSection->BoneMap = SrcSection.BoneMap;
-					 DestSection->bCastShadow = SrcSection.bCastShadow;
+
+					 //TODO(Max): MTBL-1779
+					 //DestSection->bCastShadow = SrcSection.bCastShadow;
 				 }
 			 }
 		 }
