[Mutable] Avoid the reuse of materials if the mesh has changed.

#rnx
#tests Frontend and locker

diff --git a/Engine/Plugins/Experimental/Mutable/Source/CustomizableObject/Internal/MuCO/CustomizableObjectInstancePrivate.h b/Engine/Plugins/Experimental/Mutable/Source/CustomizableObject/Internal/MuCO/CustomizableObjectInstancePrivate.h
--- a/Engine/Plugins/Experimental/Mutable/Source/CustomizableObject/Internal/MuCO/CustomizableObjectInstancePrivate.h
+++ b/Engine/Plugins/Experimental/Mutable/Source/CustomizableObject/Internal/MuCO/CustomizableObjectInstancePrivate.h
@@ -280,7 +280,7 @@ private:
 
 	void InvalidateGeneratedData();
 
-	bool DoComponentsNeedUpdate(UCustomizableObjectInstance* CustomizableObjectInstance, const TSharedRef<FUpdateContextPrivate>& OperationData, TArray<bool>& OutComponentNeedsUpdate, bool& bOutEmptyMesh);
+	bool DoComponentsNeedUpdate(UCustomizableObjectInstance* CustomizableObjectInstance, const TSharedRef<FUpdateContextPrivate>& OperationData, bool& bOutEmptyMesh);
 
 	mu::FResourceID GetLastMeshId(int32 ComponentIndex, int32 LODIndex) const;
 	void SetLastMeshId(int32 ComponentIndex, int32 LODIndex, mu::FResourceID MeshId);
diff --git a/Engine/Plugins/Experimental/Mutable/Source/CustomizableObject/Internal/MuCO/CustomizableObjectSystemPrivate.h b/Engine/Plugins/Experimental/Mutable/Source/CustomizableObject/Internal/MuCO/CustomizableObjectSystemPrivate.h
--- a/Engine/Plugins/Experimental/Mutable/Source/CustomizableObject/Internal/MuCO/CustomizableObjectSystemPrivate.h
+++ b/Engine/Plugins/Experimental/Mutable/Source/CustomizableObject/Internal/MuCO/CustomizableObjectSystemPrivate.h
@@ -553,6 +553,11 @@ public:
 
 	/** Mutable Meshes required for each component. Outermost index is the component, inner index is the LOD. */
 	TArray<TArray<mu::FResourceID>> MeshDescriptors;
+
+	/** Used to know if the updated instances' meshes are different from the previous ones. 
+	  * The index of the array is the component's index.
+	  * @return true if the mesh is new or new to this instance (e.g. mesh cached by another instance). */
+	TArray<bool> MeshChanged;
 	
 	bool UpdateStarted = false;
 
diff --git a/Engine/Plugins/Experimental/Mutable/Source/CustomizableObject/Private/MuCO/CustomizableObjectInstance.cpp b/Engine/Plugins/Experimental/Mutable/Source/CustomizableObject/Private/MuCO/CustomizableObjectInstance.cpp
--- a/Engine/Plugins/Experimental/Mutable/Source/CustomizableObject/Private/MuCO/CustomizableObjectInstance.cpp
+++ b/Engine/Plugins/Experimental/Mutable/Source/CustomizableObject/Private/MuCO/CustomizableObjectInstance.cpp
@@ -1673,7 +1673,7 @@ void SetMeshUVChannelDensity(FMeshUVChannelInfo& UVChannelInfo, float Density =
 }
 
 
-bool UCustomizableInstancePrivate::DoComponentsNeedUpdate(UCustomizableObjectInstance* Public, const TSharedRef<FUpdateContextPrivate>& OperationData, TArray<bool>& OutComponentNeedsUpdate, bool& bHasInvalidMesh)
+bool UCustomizableInstancePrivate::DoComponentsNeedUpdate(UCustomizableObjectInstance* Public, const TSharedRef<FUpdateContextPrivate>& OperationData, bool& bHasInvalidMesh)
 {
 	UCustomizableObject* CustomizableObject = Public->GetCustomizableObject();
 	check(CustomizableObject);
@@ -1724,15 +1724,17 @@ bool UCustomizableInstancePrivate::DoComponentsNeedUpdate(UCustomizableObjectIns
 	}
 
 	// Find which components need an update
-	OutComponentNeedsUpdate.AddDefaulted(NumComponents);
+	OperationData->MeshChanged.AddDefaulted(NumComponents);
 
 	for (int32 ComponentIndex = 0; ComponentIndex < NumComponents; ++ComponentIndex)
 	{
 		if (OperationData->bUseMeshCache)
 		{
-			if (CustomizableObject->GetPrivate()->MeshCache.Get(OperationData->MeshDescriptors[ComponentIndex]))
+			USkeletalMesh* CachedMesh = CustomizableObject->GetPrivate()->MeshCache.Get(OperationData->MeshDescriptors[ComponentIndex]);
+			if (CachedMesh)
 			{
-				OutComponentNeedsUpdate[ComponentIndex] = true;
+				const bool bMeshNeedsUpdate = !SkeletalMeshes.IsValidIndex(ComponentIndex) || (SkeletalMeshes[ComponentIndex] != CachedMesh);
+				OperationData->MeshChanged[ComponentIndex] = bMeshNeedsUpdate;
 				ComponentWithMesh[ComponentIndex] = true;
 				continue;
 			}
@@ -1750,7 +1752,7 @@ bool UCustomizableInstancePrivate::DoComponentsNeedUpdate(UCustomizableObjectIns
 
 		// Update the component if there is a mesh and it shouldn't, or the other way around.
 		const bool bHasSkeletalMesh = SkeletalMeshes.IsValidIndex(ComponentIndex) && SkeletalMeshes[ComponentIndex];
-		OutComponentNeedsUpdate[ComponentIndex] = (ComponentWithMesh[ComponentIndex] != bHasSkeletalMesh);
+		OperationData->MeshChanged[ComponentIndex] = (ComponentWithMesh[ComponentIndex] != bHasSkeletalMesh);
 
 		const FCustomizableInstanceComponentData* ComponentData = GetComponentData(ComponentIndex);
 		if (!ComponentData) // Could be nullptr if the component has not been generated.
@@ -1760,13 +1762,13 @@ bool UCustomizableInstancePrivate::DoComponentsNeedUpdate(UCustomizableObjectIns
 
 		// Update if MeshIDs are different
 		const int32 ComponentOffset = ComponentIndex * MAX_MESH_LOD_COUNT;
-		for (int32 MeshIndex = 0; !OutComponentNeedsUpdate[ComponentIndex] && MeshIndex < MAX_MESH_LOD_COUNT; ++MeshIndex)
+		for (int32 MeshIndex = 0; !OperationData->MeshChanged[ComponentIndex] && MeshIndex < MAX_MESH_LOD_COUNT; ++MeshIndex)
 		{
-			OutComponentNeedsUpdate[ComponentIndex] = MeshIDs[ComponentOffset + MeshIndex] != ComponentData->LastMeshIdPerLOD[MeshIndex];
+			OperationData->MeshChanged[ComponentIndex] = MeshIDs[ComponentOffset + MeshIndex] != ComponentData->LastMeshIdPerLOD[MeshIndex];
 		}
 	}
 
-	return !bHasInvalidMesh && OutComponentNeedsUpdate.Find(true) != INDEX_NONE;
+	return !bHasInvalidMesh && OperationData->MeshChanged.Find(true) != INDEX_NONE;
 }
 
 
@@ -1776,8 +1778,7 @@ bool UCustomizableInstancePrivate::UpdateSkeletalMesh_PostBeginUpdate0(UCustomiz
 
 	bool bHasInvalidMesh = false;
 
-	TArray<bool> ComponentNeedsUpdate;
-	bool bUpdateMeshes = DoComponentsNeedUpdate(Public, OperationData, ComponentNeedsUpdate, bHasInvalidMesh);
+	bool bUpdateMeshes = DoComponentsNeedUpdate(Public, OperationData, bHasInvalidMesh);
 
 	// We can not handle empty meshes, clear any generated mesh and return
 	if (bHasInvalidMesh)
@@ -1817,7 +1818,7 @@ bool UCustomizableInstancePrivate::UpdateSkeletalMesh_PostBeginUpdate0(UCustomiz
 		const FInstanceUpdateData::FComponent& Component = OperationData->InstanceUpdateData.Components[ComponentIndex];
 
 		// If the component doesn't need an update copy the previously generated mesh.
-		if (!ComponentNeedsUpdate[Component.Id])
+		if (!OperationData->MeshChanged[Component.Id])
 		{
 			if (OldSkeletalMeshes.IsValidIndex(Component.Id))
 			{
@@ -2061,7 +2062,7 @@ bool UCustomizableInstancePrivate::UpdateSkeletalMesh_PostBeginUpdate0(UCustomiz
 				}
 			}
 			
-			if (!ComponentNeedsUpdate[ComponentIndex])
+			if (!OperationData->MeshChanged[ComponentIndex])
 			{
 				continue;
 			}
@@ -5358,6 +5359,8 @@ void UCustomizableInstancePrivate::BuildMaterials(const TSharedRef<FUpdateContex
 			continue;
 		}
 
+		const bool bReuseMaterials = !OperationData->MeshChanged[ComponentIndex];
+
 		// It is not safe to replace the materials of a SkeletalMesh whose resources are initialized. Use overrides instead.
 		const bool bUseOverrideMaterialsOnly = OperationData->bUseMeshCache && SkeletalMesh->GetResourceForRendering()->IsInitialized();
 
@@ -5804,7 +5807,7 @@ void UCustomizableInstancePrivate::BuildMaterials(const TSharedRef<FUpdateContex
 
 					UMaterialInstanceDynamic* MaterialInstance = nullptr;
 					
-					if (const int32 OldMaterialIndex = OldGeneratedMaterials.Find(Material); OldMaterialIndex != INDEX_NONE)
+					if (const int32 OldMaterialIndex = OldGeneratedMaterials.Find(Material); bReuseMaterials && OldMaterialIndex != INDEX_NONE)
 					{
 						const FGeneratedMaterial& OldMaterial = OldGeneratedMaterials[OldMaterialIndex];
 						MaterialInstance = Cast<UMaterialInstanceDynamic>(OldMaterial.MaterialInterface);
