[mutable] Fixed check failing when compiling with poses added to a group projector.
#rnx
#tests Editor

diff --git a/Engine/Plugins/Experimental/Mutable/Source/CustomizableObjectEditor/Classes/MuCOE/Nodes/CustomizableObjectNodeGroupProjectorParameter.h b/Engine/Plugins/Experimental/Mutable/Source/CustomizableObjectEditor/Classes/MuCOE/Nodes/CustomizableObjectNodeGroupProjectorParameter.h
--- a/Engine/Plugins/Experimental/Mutable/Source/CustomizableObjectEditor/Classes/MuCOE/Nodes/CustomizableObjectNodeGroupProjectorParameter.h
+++ b/Engine/Plugins/Experimental/Mutable/Source/CustomizableObjectEditor/Classes/MuCOE/Nodes/CustomizableObjectNodeGroupProjectorParameter.h
@@ -76,6 +76,7 @@ public:
 	UPROPERTY(EditAnywhere, Category = ProjectorGroup)
 	TArray<FGroupProjectorParameterImage> OptionImages;
 
+	/** Poses will only affect component 0 of the CO */ // TODO UE-206803
 	UPROPERTY(EditAnywhere, Category = ProjectorGroup)
 	TArray<FGroupProjectorParameterPose> OptionPoses;
 
diff --git a/Engine/Plugins/Experimental/Mutable/Source/CustomizableObjectEditor/Private/MuCOE/GenerateMutableSource/GenerateMutableSourceGroupProjector.cpp b/Engine/Plugins/Experimental/Mutable/Source/CustomizableObjectEditor/Private/MuCOE/GenerateMutableSource/GenerateMutableSourceGroupProjector.cpp
--- a/Engine/Plugins/Experimental/Mutable/Source/CustomizableObjectEditor/Private/MuCOE/GenerateMutableSource/GenerateMutableSourceGroupProjector.cpp
+++ b/Engine/Plugins/Experimental/Mutable/Source/CustomizableObjectEditor/Private/MuCOE/GenerateMutableSource/GenerateMutableSourceGroupProjector.cpp
@@ -404,30 +404,40 @@ bool GenerateMutableSourceGroupProjector(const UEdGraphPin* Pin, FMutableGraphGe
 				FString msg = FString::Printf(TEXT("The group projection node must have at least one option image connected to a texture or at least one valid element in Option Images Data Table."));
 				GenerationContext.Compiler->CompilerLog(FText::FromString(msg), ProjParamNode, EMessageSeverity::Error, true);
 				return false;
-			}
-
-			PoseEnumParameterNode->SetValue(0, 0.f, "Default pose");
+			}			
 
-			for (int PoseIndex = 0; PoseIndex < ProjParamNode->OptionPoses.Num(); ++PoseIndex)
+			if (GenerationContext.ComponentInfos.IsValidIndex(0))
 			{
-				PoseEnumParameterNode->SetValue(PoseIndex + 1, (float)PoseIndex + 1.f, ProjParamNode->OptionPoses[PoseIndex].PoseName);
+				// Poses will only affect component 0 of the CO,
+				// TODO UE-206803
+				int32 OldCurrentMeshComponent = GenerationContext.CurrentMeshComponent;
+				GenerationContext.CurrentMeshComponent = 0;
+
+				PoseEnumParameterNode->SetValue(0, 0.f, "Default pose");
 
-				TArray<FString> ArrayBoneName;
-				TArray<FTransform> ArrayTransform;
-				UPoseAsset* PoseAsset = ProjParamNode->OptionPoses[PoseIndex].OptionPose;
-				if (PoseAsset == nullptr) // Check if the slot has a selected pose. Could be left empty by the user
+				for (int32 PoseIndex = 0; PoseIndex < ProjParamNode->OptionPoses.Num(); ++PoseIndex)
 				{
-					FString msg = FString::Printf(TEXT("The group projection node must have a pose assigned on each Option Poses element."));
-					GenerationContext.Compiler->CompilerLog(FText::FromString(msg), ProjParamNode, EMessageSeverity::Error, true);
-					return false;
+					PoseEnumParameterNode->SetValue(PoseIndex + 1, (float)PoseIndex + 1.f, ProjParamNode->OptionPoses[PoseIndex].PoseName);
+
+					TArray<FString> ArrayBoneName;
+					TArray<FTransform> ArrayTransform;
+					UPoseAsset* PoseAsset = ProjParamNode->OptionPoses[PoseIndex].OptionPose;
+
+					if (PoseAsset == nullptr) // Check if the slot has a selected pose. Could be left empty by the user
+					{
+						FString msg = FString::Printf(TEXT("The group projection node must have a pose assigned on each Option Poses element."));
+						GenerationContext.Compiler->CompilerLog(FText::FromString(msg), ProjParamNode, EMessageSeverity::Error, true);
+						return false;
+					}
+
+					check(GroupProjectorTempData.PoseBoneDataArray.Num() == PoseIndex);
+					GroupProjectorTempData.PoseBoneDataArray.AddDefaulted(1);
+					UCustomizableObjectNodeAnimationPose::StaticRetrievePoseInformation(PoseAsset, GenerationContext.GetCurrentComponentInfo().RefSkeletalMesh,
+						GroupProjectorTempData.PoseBoneDataArray[PoseIndex].ArrayBoneName, GroupProjectorTempData.PoseBoneDataArray[PoseIndex].ArrayTransform);
 				}
 
-				check(GroupProjectorTempData.PoseBoneDataArray.Num() == PoseIndex);
-				GroupProjectorTempData.PoseBoneDataArray.AddDefaulted(1);
-				UCustomizableObjectNodeAnimationPose::StaticRetrievePoseInformation(PoseAsset, GenerationContext.GetCurrentComponentInfo().RefSkeletalMesh,
-					GroupProjectorTempData.PoseBoneDataArray[PoseIndex].ArrayBoneName, GroupProjectorTempData.PoseBoneDataArray[PoseIndex].ArrayTransform);
+				GenerationContext.CurrentMeshComponent = OldCurrentMeshComponent;
 			}
-
 		
 			mu::NodeScalarEnumParameterPtr EnumParameterNode = new mu::NodeScalarEnumParameter;
 			FString NodeEnumParamName = ProjParamNode->ParameterName + FMultilayerProjector::IMAGE_PARAMETER_POSTFIX;
