[mutable] Fixed 3 missing icons/brushes from the high resolution screenshot screen in the Preview Instance tab.
#rnx
#tests Editor

diff --git a/Engine/Plugins/Experimental/Mutable/Source/CustomizableObjectEditor/Private/MuCOE/SCustomizableObjectHighresScreenshot.cpp b/Engine/Plugins/Experimental/Mutable/Source/CustomizableObjectEditor/Private/MuCOE/SCustomizableObjectHighresScreenshot.cpp
--- a/Engine/Plugins/Experimental/Mutable/Source/CustomizableObjectEditor/Private/MuCOE/SCustomizableObjectHighresScreenshot.cpp
+++ b/Engine/Plugins/Experimental/Mutable/Source/CustomizableObjectEditor/Private/MuCOE/SCustomizableObjectHighresScreenshot.cpp
@@ -157,7 +157,7 @@ void SCustomizableObjectHighresScreenshot::Construct( const FArguments& InArgs )
 				.AutoHeight()
 				[
 					SNew(SImage)
-					.Image(UE_MUTABLE_GET_BRUSH("HighresScreenshot.WarningStrip"))
+					.Image(UE_MUTABLE_GET_BRUSH("WarningStripe"))
 				]
 				+SVerticalBox::Slot()
 				.AutoHeight()
@@ -174,7 +174,7 @@ void SCustomizableObjectHighresScreenshot::Construct( const FArguments& InArgs )
 				.AutoHeight()
 				[
 					SNew(SImage)
-					.Image(UE_MUTABLE_GET_BRUSH("HighresScreenshot.WarningStrip"))
+					.Image(UE_MUTABLE_GET_BRUSH("WarningStripe"))
 				]
 			]
 			+SVerticalBox::Slot()
@@ -195,7 +195,7 @@ void SCustomizableObjectHighresScreenshot::Construct( const FArguments& InArgs )
 						.OnClicked( this, &SCustomizableObjectHighresScreenshot::OnCaptureClicked )
 						[
 							SNew(SImage)
-							.Image(UE_MUTABLE_GET_BRUSH("HighresScreenshot.Capture"))
+							.Image(UE_MUTABLE_GET_BRUSH("ClassIcon.SceneCaptureComponent"))
 						]
 					]
 				]
