[Mutable] Fix crash when clicking on a Table Node.

Right mouse button.

- Also fixed unsafe pointer to an UObject.

#rnx

diff --git a/Engine/Plugins/Experimental/Mutable/Source/CustomizableObjectEditor/Private/MuCOE/Nodes/CustomizableObjectNodeTableDetails.cpp b/Engine/Plugins/Experimental/Mutable/Source/CustomizableObjectEditor/Private/MuCOE/Nodes/CustomizableObjectNodeTableDetails.cpp
--- a/Engine/Plugins/Experimental/Mutable/Source/CustomizableObjectEditor/Private/MuCOE/Nodes/CustomizableObjectNodeTableDetails.cpp
+++ b/Engine/Plugins/Experimental/Mutable/Source/CustomizableObjectEditor/Private/MuCOE/Nodes/CustomizableObjectNodeTableDetails.cpp
@@ -37,7 +37,7 @@ void FCustomizableObjectNodeTableDetails::CustomizeDetails(const TSharedPtr<IDet
 		Node = Cast<UCustomizableObjectNodeTable>(DetailsView->GetSelectedObjects()[0].Get());
 	}
 
-	if (Node)
+	if (Node.IsValid())
 	{
 		IDetailCategoryBuilder& CustomizableObjectCategory = DetailBuilder->EditCategory("TableProperties");
 		IDetailCategoryBuilder& UICategory = DetailBuilder->EditCategory("UI");
@@ -95,14 +95,20 @@ void FCustomizableObjectNodeTableDetails::CustomizeDetails(const TSharedPtr<IDet
 					SNew(STextBlock)
 					.Text(LOCTEXT("AnimBPText", "Animation Blueprint Column: "))
 					.ToolTipText(LOCTEXT("AnimBlueprintColumnTooltip", "Select an animation blueprint column from the Data Table that will be applied to the mesh selected"))
-					.Visibility_Lambda([this]() -> EVisibility
+					.Visibility_Lambda([WeakDetails = SharedThis(this).ToWeakPtr()]() -> EVisibility
 					{
-						if (!AnimComboBox.IsValid())
+						const TSharedPtr<FCustomizableObjectNodeTableDetails> Details = WeakDetails.Pin();
+						if (!Details)
+						{
+							return EVisibility::Collapsed;
+						}
+						
+						if (!Details->AnimComboBox.IsValid())
 						{
 							return EVisibility::Collapsed;
 						}
 
-						return AnimComboBox->GetVisibility();
+						return Details->AnimComboBox->GetVisibility();
 					})
 				]
 
@@ -133,14 +139,20 @@ void FCustomizableObjectNodeTableDetails::CustomizeDetails(const TSharedPtr<IDet
 					SNew(STextBlock)
 					.Text(LOCTEXT("AnimSlotText", "Animation Slot Column: "))
 					.ToolTipText(LOCTEXT("AnimSlotColumnTooltip", "Select an animation slot column from the Data Table that will set to the slot value of the animation blueprint"))
-					.Visibility_Lambda([this]() -> EVisibility 
+					.Visibility_Lambda([WeakDetails = SharedThis(this).ToWeakPtr()]() -> EVisibility 
 						{
-							if (!AnimSlotComboBox.IsValid())
+							const TSharedPtr<FCustomizableObjectNodeTableDetails> Details = WeakDetails.Pin();
+							if (!Details)
 							{
 								return EVisibility::Collapsed;
 							}
 
-							return AnimSlotComboBox->GetVisibility();
+							if (!Details->AnimSlotComboBox.IsValid())
+							{
+								return EVisibility::Collapsed;
+							}
+
+							return Details->AnimSlotComboBox->GetVisibility();
 						})
 				]
 
@@ -171,14 +183,20 @@ void FCustomizableObjectNodeTableDetails::CustomizeDetails(const TSharedPtr<IDet
 					SNew(STextBlock)
 					.Text(LOCTEXT("AnimTagsText", "Animation Tags Column: "))
 					.ToolTipText(LOCTEXT("AnimTagColumnTooltip", "Select an animation tag column from the Data Table that will set to the animation tags of the animation blueprint"))
-					.Visibility_Lambda([this]() -> EVisibility 
+					.Visibility_Lambda([WeakDetails = SharedThis(this).ToWeakPtr()]() -> EVisibility 
 						{
-							if (!AnimTagsComboBox.IsValid())
+							const TSharedPtr<FCustomizableObjectNodeTableDetails> Details = WeakDetails.Pin();
+							if (!Details)
+							{
+								return EVisibility::Collapsed;
+							}
+
+							if (!Details->AnimTagsComboBox.IsValid())
 							{
 								return EVisibility::Collapsed;
 							}
 
-							return AnimTagsComboBox->GetVisibility();
+							return Details->AnimTagsComboBox->GetVisibility();								
 						})
 				]
 
diff --git a/Engine/Plugins/Experimental/Mutable/Source/CustomizableObjectEditor/Private/MuCOE/Nodes/CustomizableObjectNodeTableDetails.h b/Engine/Plugins/Experimental/Mutable/Source/CustomizableObjectEditor/Private/MuCOE/Nodes/CustomizableObjectNodeTableDetails.h
--- a/Engine/Plugins/Experimental/Mutable/Source/CustomizableObjectEditor/Private/MuCOE/Nodes/CustomizableObjectNodeTableDetails.h
+++ b/Engine/Plugins/Experimental/Mutable/Source/CustomizableObjectEditor/Private/MuCOE/Nodes/CustomizableObjectNodeTableDetails.h
@@ -4,6 +4,8 @@
 
 #include "IDetailCustomization.h"
 
+#include "UObject/WeakObjectPtr.h"
+
 namespace ESelectInfo { enum Type : int; }
 
 class FString;
@@ -12,6 +14,7 @@ class STextComboBox;
 class FReply;
 class SButton;
 class SCustomizableObjectNodeLayoutBlocksEditor;
+class UCustomizableObjectNodeTable;
 
 /** Copy Material node details panel. Hides all properties from the inheret Material node. */
 class FCustomizableObjectNodeTableDetails : public IDetailCustomization
@@ -65,10 +68,8 @@ private:
 	// Callback to clear the animation combobox selections
 	FReply OnClearButtonPressed();
 
-private:
-
 	// Pointer to the node represented in this details
-	class UCustomizableObjectNodeTable* Node;
+	TWeakObjectPtr<UCustomizableObjectNodeTable> Node;
 
 	// ComboBox widget to select a column from the NodeTable
 	TSharedPtr<STextComboBox> LayoutMeshColumnComboBox;
@@ -108,5 +109,4 @@ private:
 
 	// Pointer to the Detail Builder to force the refresh on recontruct the node
 	TWeakPtr<IDetailLayoutBuilder> DetailBuilderPtr = nullptr;
-
 };
