[Mutable] Fix Material Variations not working if a child is directly connected to a root CO.

#rnx

diff --git a/Engine/Plugins/Experimental/Mutable/Source/MutableTools/Private/MuT/CodeGenerator_SecondPass.cpp b/Engine/Plugins/Experimental/Mutable/Source/MutableTools/Private/MuT/CodeGenerator_SecondPass.cpp
--- a/Engine/Plugins/Experimental/Mutable/Source/MutableTools/Private/MuT/CodeGenerator_SecondPass.cpp
+++ b/Engine/Plugins/Experimental/Mutable/Source/MutableTools/Private/MuT/CodeGenerator_SecondPass.cpp
@@ -336,6 +336,14 @@ namespace mu
 				positiveTagsVisited,
 				negativeTagsVisited);
 
+			// No condition is equal to a conditional with a true constant
+			if (!tagCondition)
+			{
+				ASTOpConstantBool* ConstOp = new ASTOpConstantBool();
+				ConstOp->value = true;
+				tagCondition = ConstOp;
+			}
+
 			// TODO: Optimise the tag condition here
 
 			// If the tag is a constant ...
@@ -410,6 +418,12 @@ namespace mu
 			set<size_t> empty;
 			Ptr<ASTOp> tagCondition = GenerateTagCondition(tagIndex, empty, empty, empty, empty);
 
+			if (!tagCondition)
+			{
+				// This tag is unconditionally activated, so there's no condition logic to add
+				continue;
+			}
+
 			// TODO: Optimise the tag condition here
 
 			// If the tag is a constant ...
@@ -469,6 +483,14 @@ namespace mu
 				set<size_t> empty;
 				Ptr<ASTOp> tagCondition = GenerateTagCondition(tagIndex, empty, empty, empty, empty);
 
+				// No condition is equal to a conditional with a true constant
+				if (!tagCondition)
+				{
+					ASTOpConstantBool* ConstOp = new ASTOpConstantBool();
+					ConstOp->value = true;
+					tagCondition = ConstOp;
+				}
+
 				// TODO: Optimise the tag condition here
 
 				// If the tag is a constant ...
