[mutable] Added argument to determine the target platform we want to compile the COs for
- It only affects the COValidation tests.
- Setting a target compilation platform != than the current machine one will prevent the generation of instances.

diff --git a/Engine/Plugins/Experimental/Mutable/Source/MutableValidation/Private/MuV/COIBakingTestCommandlet.cpp b/Engine/Plugins/Experimental/Mutable/Source/MutableValidation/Private/MuV/COIBakingTestCommandlet.cpp
--- a/Engine/Plugins/Experimental/Mutable/Source/MutableValidation/Private/MuV/COIBakingTestCommandlet.cpp
+++ b/Engine/Plugins/Experimental/Mutable/Source/MutableValidation/Private/MuV/COIBakingTestCommandlet.cpp
@@ -12,6 +12,7 @@
 #include "HAL/FileManager.h"
 #include "MuCO/CustomizableObjectPrivate.h"
 #include "MuT/UnrealPixelFormatOverride.h"
+#include "Interfaces/ITargetPlatformManagerModule.h"
 
 
 /** Flag useful to know if we are currently updating an instance or not */
@@ -76,8 +77,13 @@ int32 UCOIBakingTestCommandlet::Main(const FString& Params)
 		UE_LOG(LogMutable,Error,TEXT("The instance %s does not have a CO to compile : Exitting commandlet."), *TargetInstance->GetName());
 		return 1;
 	}
+
 	
-	if (!CompileCustomizableObject(InstanceCustomizableObject))
+	// Set the target platform to be using for the compilation. Must not be a nullptr
+	FCompilationOptions CompilationOptions = InstanceCustomizableObject->CompileOptions;
+	ITargetPlatformManagerModule& TPM = GetTargetPlatformManagerRef();
+	CompilationOptions.TargetPlatform = TPM.GetRunningTargetPlatform();
+	if (!CompileCustomizableObject(InstanceCustomizableObject,true, &CompilationOptions))
 	{
 		UE_LOG(LogMutable,Error,TEXT("Failed to compile the target CO. Exitting commandlet."));
 		return 1;
diff --git a/Engine/Plugins/Experimental/Mutable/Source/MutableValidation/Private/MuV/COIBulkUpdateTestCommandlet.cpp b/Engine/Plugins/Experimental/Mutable/Source/MutableValidation/Private/MuV/COIBulkUpdateTestCommandlet.cpp
--- a/Engine/Plugins/Experimental/Mutable/Source/MutableValidation/Private/MuV/COIBulkUpdateTestCommandlet.cpp
+++ b/Engine/Plugins/Experimental/Mutable/Source/MutableValidation/Private/MuV/COIBulkUpdateTestCommandlet.cpp
@@ -7,10 +7,14 @@
 #include "AssetRegistry/AssetRegistryModule.h"
 #include "MuCO/CustomizableObject.h"
 #include "MuCO/CustomizableObjectInstance.h"
+#include "Interfaces/ITargetPlatformManagerModule.h"
 
 
 int32 UCOIBulkUpdateTestCommandlet::Main(const FString& Params)
 {
+	// Ensure we do not show any OK dialog since we are not an user that can interact with them
+	GIsRunningUnattendedScript = true;
+
 	// Get the path where to look for the Customizable Object Instances we want to validate
 	FString InstancesPackagePath;
 	{
@@ -114,16 +118,24 @@ int32 UCOIBulkUpdateTestCommandlet::Main(const FString& Params)
 	// Create the updater object so we can later call for the update a target instance
 	InstanceUpdater = NewObject<UCOIUpdater>();
 	
+	// Cache the target compilation platform so we can override the compilation configs of the target COs
+	ITargetPlatformManagerModule& TPM = GetTargetPlatformManagerRef();
+	const ITargetPlatform* TargetCompilationPlatform = TPM.GetRunningTargetPlatform();
+	
 	// Compile all found COs one by one
 	uint32 CurrentInstanceIndex = 1;
 	for (TTuple<UCustomizableObject*, TSet<UCustomizableObjectInstance*>>& MutableResourceTuple : MutableResources)
 	{
 		UCustomizableObject* CustomizableObjectToCompile = MutableResourceTuple.Key;
 		check (CustomizableObjectToCompile);
+
+		// Set the compilation platform based on what the system is currently running on
+		FCompilationOptions CompilationOptions = CustomizableObjectToCompile->CompileOptions;
+		CompilationOptions.TargetPlatform = TargetCompilationPlatform;
 		
 		// Compile the current CO object
 		const FString CustomizableObjectName = CustomizableObjectToCompile->GetName();
-		if (!CompileCustomizableObject(CustomizableObjectToCompile, false))	// Do not log mutable data since mongoDB will not be able to handle it correctly 
+		if (!CompileCustomizableObject(CustomizableObjectToCompile, false, &CompilationOptions))	// Do not log mutable data since mongoDB will not be able to handle it correctly 
 		{
 			UE_LOG(LogMutable,Error,TEXT("The CO %s could not be compiled succesfully. Skipping the update of all COIs that use it."), *CustomizableObjectName )
 			continue;
diff --git a/Engine/Plugins/Experimental/Mutable/Source/MutableValidation/Private/MuV/CustomizableObjectValidationCommandlet.cpp b/Engine/Plugins/Experimental/Mutable/Source/MutableValidation/Private/MuV/CustomizableObjectValidationCommandlet.cpp
--- a/Engine/Plugins/Experimental/Mutable/Source/MutableValidation/Private/MuV/CustomizableObjectValidationCommandlet.cpp
+++ b/Engine/Plugins/Experimental/Mutable/Source/MutableValidation/Private/MuV/CustomizableObjectValidationCommandlet.cpp
@@ -10,12 +10,17 @@
 #include "MuCO/CustomizableObjectPrivate.h"
 #include "MuCO/CustomizableObjectSystem.h"
 #include "MuR/Model.h"
+#include "Interfaces/ITargetPlatformManagerModule.h"
+#include "Interfaces/ITargetPlatform.h"
 
 
 int32 UCustomizableObjectValidationCommandlet::Main(const FString& Params)
 {
 	// Execution arguments for commandlet from IDE
 	// -run=CustomizableObjectValidation -CustomizableObject=(PathToCO)
+
+	// Ensure we do not show any OK dialog since we are not an user that can interact with them
+	GIsRunningUnattendedScript = true;
 	
 	// Get the package name of the CO to test
 	FString CustomizableObjectAssetPath = "";
@@ -33,8 +38,6 @@ int32 UCustomizableObjectValidationCommandlet::Main(const FString& Params)
 	}
 	UE_LOG(LogMutable, Log,TEXT("(int) instances_to_generate_count : %u "), InstancesToGenerate);
 	
-	// TODO: Detect target compilation platform based on argument value (string)
-	
 	// Load the resource
 	UObject* FoundObject = FSoftObjectPath(CustomizableObjectAssetPath).TryLoad();
 	if (!FoundObject)
@@ -50,6 +53,14 @@ int32 UCustomizableObjectValidationCommandlet::Main(const FString& Params)
 		UE_LOG(LogMutable,Error,TEXT("Failed to cast found UObject to UCustomizableObject."));
 		return 1;
 	}
+
+	// What platform we want to compile the CO against
+	ITargetPlatform* TargetCompilationPlatform = ParseCompilationPlatform(Params);
+	if (!TargetCompilationPlatform)
+	{
+		UE_LOG(LogMutable,Error,TEXT("No Target Compilation Platform was provided for the compilation of the CO."));
+		return 1;
+	}
 	
 	// Perform a blocking search to ensure all assets used by mutable are reachable using the AssetRegistry
 	PrepareAssetRegistry();
@@ -63,15 +74,37 @@ int32 UCustomizableObjectValidationCommandlet::Main(const FString& Params)
 		CompilationOptions.OptimizationLevel = 2;			// Set the optimization level to the max
 		CompilationOptions.TextureCompression = ECustomizableObjectTextureCompression::Fast;
 
+		// Set the target compilation platform based on what the caller wants
+		CompilationOptions.TargetPlatform = TargetCompilationPlatform;
+
 		bWasCoCompilationSuccessful = CompileCustomizableObject(ToTestCustomizableObject, true, &CompilationOptions);
 	}
 	// -------------------------------------------------------------------------------------------------------------- //
+
+	if (!bWasCoCompilationSuccessful)
+	{
+		UE_LOG(LogMutable,Error,TEXT("The compilation of the Customizable object was not succesfull : No instances will be generated."));
+		return 1;		// Validation failed
+	}
+
+	UE_LOG(LogMutable,Display,TEXT("Customizable Object was compiled succesfully."));
 	
-	if (bWasCoCompilationSuccessful)
+	// Skip instances updating if no instances should be updated 
+	if (InstancesToGenerate <= 0)
+	{
+		UE_LOG(LogMutable,Display,TEXT("Instances to generate are 0 : No instances will be generated."));
+		return 0;	// No instances are targeted for generation, this will be taken as compilation only test.
+	}
+
+	// Do not generate instances if the selected platform is not the running platform
+	if (TargetCompilationPlatform != GetTargetPlatformManagerRef().GetRunningTargetPlatform())
+	{
+		UE_LOG(LogMutable,Display,TEXT("RunningPlatform != UserProvidedCompilationPlatform : No instances will be generated."));
+		return 0;
+	}
+	
+	// GHet the total size of the streaming data of the model ---------------------------------------------- //
 	{
-		UE_LOG(LogMutable,Display,TEXT("Customizable Object was compiled succesfully."));
-		
-		// GHet the total size of the streaming data of the model ---------------------------------------------- //
 		const TSharedPtr<const mu::Model> MutableModel = ToTestCustomizableObject->GetPrivate()->GetModel();
 		check (MutableModel);
 
@@ -94,90 +127,123 @@ int32 UCustomizableObjectValidationCommandlet::Main(const FString& Params)
 		{
 			TArray<uint8> EmbeddedDataBytes{};
 			FMemoryWriter SerializationTarget{EmbeddedDataBytes, false};
-			
+		
 			ToTestCustomizableObject->GetPrivate()->SaveEmbeddedData(SerializationTarget);
 			const int64 COEmbeddedDataSizeBytes = EmbeddedDataBytes.Num();
-			
+		
 			UE_LOG(LogMutable, Log,TEXT("(int) co_embedded_data_bytes : %lld "), COEmbeddedDataSizeBytes);
 		}
+	}
+	
 
-		// Generate target random instances to be tested ------------------------------------------------------------ //
-		bool bWasInstancesCreationSuccessful = true;
-		{
-			UE_LOG(LogMutable,Display,TEXT("Generating %i random instances..."), InstancesToGenerate);	
+	// Generate target random instances to be tested ------------------------------------------------------------ //
+	bool bWasInstancesCreationSuccessful = true;
+	{
+		UE_LOG(LogMutable,Display,TEXT("Generating %i random instances..."), InstancesToGenerate);	
 
-			// Create randomization stream for the parameters of the instance
-			FRandomStream RandomizationStream = FRandomStream(0);
-			
-			// Generate a series of instances to later update
-			for (uint32 CustomizableObjectInstanceIndex = 0; CustomizableObjectInstanceIndex < InstancesToGenerate; CustomizableObjectInstanceIndex++)
-			{
-				UCustomizableObjectInstance* GeneratedInstance = ToTestCustomizableObject->CreateInstance();
-				if (GeneratedInstance)
-				{
-					// Force generation of all LODS
-					TArray<uint16> RequestedLodLevels{};
-					RequestedLodLevels.Init(MAX_uint8, GeneratedInstance->GetNumComponents());
-					GeneratedInstance->GetPrivate()->GetDescriptor().SetRequestedLODLevels(RequestedLodLevels);
-					
-					// Randomize instance values
-					GeneratedInstance->SetRandomValuesFromStream(RandomizationStream);
-					InstancesToProcess.Push(GeneratedInstance);
-				}
-				else
-				{
-					UE_LOG(LogMutable,Error,TEXT("Failed to generate COI for the %s CO."),*ToTestCustomizableObject->GetName());
-					bWasInstancesCreationSuccessful = false;
-				}
-			}
-		}
-		// ---------------------------------------------------------------------------------------------------------- //
+		// Create randomization stream for the parameters of the instance
+		FRandomStream RandomizationStream = FRandomStream(0);
 		
-		// Update the instances generated --------------------------------------------------------------------------- //
-		UE_LOG(LogMutable,Display,TEXT("Updating generated instances..."));
-		bool bInstanceFailedUpdate = false;
-		const double InstancesUpdateStartSeconds = FPlatformTime::Seconds();
+		// Generate a series of instances to later update
+		for (uint32 CustomizableObjectInstanceIndex = 0; CustomizableObjectInstanceIndex < InstancesToGenerate; CustomizableObjectInstanceIndex++)
 		{
-			InstanceUpdater = NewObject<UCOIUpdater>();
-			for (UCustomizableObjectInstance* InstanceToUpdate : InstancesToProcess)
+			UCustomizableObjectInstance* GeneratedInstance = ToTestCustomizableObject->CreateInstance();
+			if (GeneratedInstance)
+			{
+				// Force generation of all LODS
+				TArray<uint16> RequestedLodLevels{};
+				RequestedLodLevels.Init(MAX_uint8, GeneratedInstance->GetNumComponents());
+				GeneratedInstance->GetPrivate()->GetDescriptor().SetRequestedLODLevels(RequestedLodLevels);
+				
+				// Randomize instance values
+				GeneratedInstance->SetRandomValuesFromStream(RandomizationStream);
+				InstancesToProcess.Push(GeneratedInstance);
+			}
+			else
 			{
-				if (InstanceUpdater && !InstanceUpdater->UpdateInstance(InstanceToUpdate))
-				{
-					bInstanceFailedUpdate = true;
-				}
+				UE_LOG(LogMutable,Error,TEXT("Failed to generate COI for the %s CO."),*ToTestCustomizableObject->GetName());
+				bWasInstancesCreationSuccessful = false;
 			}
 		}
-		const double InstancesUpdateEndSeconds = FPlatformTime::Seconds();
-		
-		// Notify and log time required by the instances to get updated
-		const double CombinedInstanceUpdateSeconds = InstancesUpdateEndSeconds - InstancesUpdateStartSeconds;
-		UE_LOG(LogMutable, Log,TEXT("(double) combined_update_time_ms : %f "), CombinedInstanceUpdateSeconds * 1000);
-
-		check(InstancesToGenerate > 0);
-		const double AverageInstanceUpdateSeconds = CombinedInstanceUpdateSeconds / InstancesToGenerate;
-		UE_LOG(LogMutable, Log,TEXT("(double) avg_update_time_ms : %f "), AverageInstanceUpdateSeconds * 1000);
-
-		UE_LOG(LogMutable,Display,TEXT("Generation of Customizable object instances took %f seconds (%f seconds avg)."), CombinedInstanceUpdateSeconds, AverageInstanceUpdateSeconds);
-		// ---------------------------------------------------------------------------------------------------------- //
-
-		// Compute instance update result
-		const bool bInstancesTestedSuccessfully = !bInstanceFailedUpdate && bWasInstancesCreationSuccessful;
-		if (bInstancesTestedSuccessfully)
-        {
-        	UE_LOG(LogMutable,Display,TEXT("Generation of Customizable object instances was succesfull."));
-        }
-        else
-        {
-        	UE_LOG(LogMutable,Error,TEXT("The generation of Customizable object instances  was not succesfull."));
-        }
 	}
-	else
+	// ---------------------------------------------------------------------------------------------------------- //
+	
+	// Update the instances generated --------------------------------------------------------------------------- //
+	UE_LOG(LogMutable,Display,TEXT("Updating generated instances..."));
+	bool bInstanceFailedUpdate = false;
+	const double InstancesUpdateStartSeconds = FPlatformTime::Seconds();
 	{
-		UE_LOG(LogMutable,Error,TEXT("The compilation of the Customizable object was not succesfull : No instances will be generated."));
+		InstanceUpdater = NewObject<UCOIUpdater>();
+		for (UCustomizableObjectInstance* InstanceToUpdate : InstancesToProcess)
+		{
+			if (InstanceUpdater && !InstanceUpdater->UpdateInstance(InstanceToUpdate))
+			{
+				bInstanceFailedUpdate = true;
+			}
+		}
 	}
+	const double InstancesUpdateEndSeconds = FPlatformTime::Seconds();
+	
+	// Notify and log time required by the instances to get updated
+	const double CombinedInstanceUpdateSeconds = InstancesUpdateEndSeconds - InstancesUpdateStartSeconds;
+	UE_LOG(LogMutable, Log,TEXT("(double) combined_update_time_ms : %f "), CombinedInstanceUpdateSeconds * 1000);
+
+	check(InstancesToGenerate > 0);
+	const double AverageInstanceUpdateSeconds = CombinedInstanceUpdateSeconds / InstancesToGenerate;
+	UE_LOG(LogMutable, Log,TEXT("(double) avg_update_time_ms : %f "), AverageInstanceUpdateSeconds * 1000);
+
+	UE_LOG(LogMutable,Display,TEXT("Generation of Customizable object instances took %f seconds (%f seconds avg)."), CombinedInstanceUpdateSeconds, AverageInstanceUpdateSeconds);
+	// ---------------------------------------------------------------------------------------------------------- //
+
+	// Compute instance update result
+	const bool bInstancesTestedSuccessfully = !bInstanceFailedUpdate && bWasInstancesCreationSuccessful;
+	if (bInstancesTestedSuccessfully)
+    {
+        UE_LOG(LogMutable,Display,TEXT("Generation of Customizable object instances was succesfull."));
+    }
+    else
+    {
+        UE_LOG(LogMutable,Error,TEXT("The generation of Customizable object instances  was not succesfull."));
+    }
 	
 	
 	// If something failed then fail the commandlet execution
 	UE_LOG(LogMutable,Display,TEXT("Mutable commandlet finished."));
 	return 0;
 }
+
+
+ITargetPlatform* UCustomizableObjectValidationCommandlet::ParseCompilationPlatform(const FString& Params) const
+{
+	// Get the package name of the CO to test
+	FString TargetPlatformName = "";
+	if (!FParse::Value(*Params, TEXT("CompilationPlatformName="), TargetPlatformName))
+	{
+		UE_LOG(LogMutable,Error,TEXT("Failed to parse the target compilation platformm."))
+		return nullptr;
+	}
+
+	// Set the target platform in the context. For now it is the current platform.
+	ITargetPlatformManagerModule* TPM = GetTargetPlatformManager();
+	check(TPM);
+	
+	ITargetPlatform* TargetCompilationPlatform = nullptr;
+	const TArray<ITargetPlatform*> TPMPlatforms = TPM->GetTargetPlatforms();
+	for (ITargetPlatform* Platform : TPMPlatforms)
+	{
+		FString PlatformName = Platform->PlatformName();
+		if (PlatformName.Compare(TargetPlatformName) == 0)
+		{
+			// We have found the platform provided
+			TargetCompilationPlatform = Platform;
+			break;
+		}
+	}
+	
+	if (!TargetCompilationPlatform)
+	{
+		UE_LOG(LogMutable, Error, TEXT("Unable to relate the provided platform name (%s) with the available platforms in this machine."), *TargetPlatformName);
+	}
+
+	return TargetCompilationPlatform;
+}
diff --git a/Engine/Plugins/Experimental/Mutable/Source/MutableValidation/Private/MuV/ValidationUtils.cpp b/Engine/Plugins/Experimental/Mutable/Source/MutableValidation/Private/MuV/ValidationUtils.cpp
--- a/Engine/Plugins/Experimental/Mutable/Source/MutableValidation/Private/MuV/ValidationUtils.cpp
+++ b/Engine/Plugins/Experimental/Mutable/Source/MutableValidation/Private/MuV/ValidationUtils.cpp
@@ -9,6 +9,7 @@
 #include "MuCO/CustomizableObjectInstance.h"
 #include "MuCO/CustomizableObjectPrivate.h"
 #include "MuCO/CustomizableObjectSystem.h"
+#include "Interfaces/ITargetPlatform.h"
 
 void PrepareAssetRegistry()
 {
@@ -40,6 +41,14 @@ bool CompileCustomizableObject(UCustomizableObject* InCustomizableObject, const
 		CompilationOptions = InCustomizableObject->CompileOptions;
 		UE_LOG(LogMutable,Display,TEXT("Compiling CO using it's own compilation options."));
 	}
+
+	// Ensure that the user has provided a target compilation platform.
+	// Mutable is able to run without one but we want to be explicit in the context of testing.
+	if (!CompilationOptions.TargetPlatform)
+	{
+		UE_LOG(LogMutable, Error, TEXT("The compilation of the %s model could not be started : No explicit platform was provided."), *InCustomizableObject->GetName());
+		return false;
+	}
 	
 	// Request a compiler to be able to locate the root and to compile it
 	const TUniquePtr<FCustomizableObjectCompilerBase> Compiler =
@@ -80,8 +89,11 @@ bool CompileCustomizableObject(UCustomizableObject* InCustomizableObject, const
 		UE_LOG(LogMutable, Log, TEXT("(int) model_optimization_level : %d "), CompilationOptions.OptimizationLevel);
 		UE_LOG(LogMutable, Log, TEXT("(string) model_texture_compression : %s "), *UEnum::GetValueAsString(CompilationOptions.TextureCompression));
 		UE_LOG(LogMutable, Log, TEXT("(string) model_disk_compilation : %s "), CompilationOptions.bUseDiskCompilation ? TEXT("true") : TEXT("false"));
+		UE_LOG(LogMutable, Log, TEXT("(string) model_compile_platform_name : %s "), *CompilationOptions.TargetPlatform->PlatformName());
+
 		UE_LOG(LogMutable, Log, TEXT("(double) model_compile_time_ms : %f "), CompilationEndSeconds * 1000);
 		UE_LOG(LogMutable, Log, TEXT("(string) model_compile_end_state : %s "), *UEnum::GetValueAsString(CompilationEndResult));
+
 		// TODO: Add logs for the other relevant configs of the model being compiled
 	}
 
@@ -200,4 +212,4 @@ void UCOIUpdater::OnInstanceUpdateResult(const FUpdateContext& Result)
 		// Tell the system the instance finished it's update so we can continue the execution without waitting for the mips to stream in
 		bIsInstanceBeingUpdated = false;
 	}	
-}
\ No newline at end of file
+}
diff --git a/Engine/Plugins/Experimental/Mutable/Source/MutableValidation/Public/MuV/CustomizableObjectValidationCommandlet.h b/Engine/Plugins/Experimental/Mutable/Source/MutableValidation/Public/MuV/CustomizableObjectValidationCommandlet.h
--- a/Engine/Plugins/Experimental/Mutable/Source/MutableValidation/Public/MuV/CustomizableObjectValidationCommandlet.h
+++ b/Engine/Plugins/Experimental/Mutable/Source/MutableValidation/Public/MuV/CustomizableObjectValidationCommandlet.h
@@ -11,6 +11,7 @@
 // Forward declarations
 class UCustomizableObject;
 class UCOIUpdater;
+class ITargetPlatform;
 
 UCLASS()
 class UCustomizableObjectValidationCommandlet : public UCommandlet 
@@ -21,6 +22,15 @@ public:
 	virtual int32 Main(const FString& Params) override;
 	
 private:
+	
+	/**
+	 * Extracts the targeted compilation platform provided by the user. It will look for "-CompilationPlatformName="PlatformName".
+	 * Examples : -CompilationPlatformName=WindowsEditor or -CompilationPlatformName=Switch
+	 * @param Params The arguments provided to this commandlet.
+	 * @return The target platform to be used for the CO compilation.
+	 */
+	ITargetPlatform* ParseCompilationPlatform(const FString& Params) const;
+
 	/** Customizable Object to be tested */
 	UPROPERTY()
 	TObjectPtr<UCustomizableObject> ToTestCustomizableObject = nullptr;
@@ -28,7 +38,8 @@ private:
 	/** Array of COI to be generated with randomized parameter values */
 	UPROPERTY()
 	TArray<TObjectPtr<UCustomizableObjectInstance>> InstancesToProcess;
-	
+
+	/** Helper object designed to aid in the update of the CO Instances.*/
 	UPROPERTY()
 	TObjectPtr<UCOIUpdater> InstanceUpdater;
 };
\ No newline at end of file
