[Mutable] Fix possible unset memory on realtime morph vertex indices causing packaging indeterminism.

#rnx

diff --git a/Engine/Plugins/Experimental/Mutable/Source/CustomizableObjectEditor/Private/MuCOE/GenerateMutableSource/GenerateMutableSourceMesh.cpp b/Engine/Plugins/Experimental/Mutable/Source/CustomizableObjectEditor/Private/MuCOE/GenerateMutableSource/GenerateMutableSourceMesh.cpp
--- a/Engine/Plugins/Experimental/Mutable/Source/CustomizableObjectEditor/Private/MuCOE/GenerateMutableSource/GenerateMutableSourceMesh.cpp
+++ b/Engine/Plugins/Experimental/Mutable/Source/CustomizableObjectEditor/Private/MuCOE/GenerateMutableSource/GenerateMutableSourceMesh.cpp
@@ -1153,6 +1153,10 @@ mu::MeshPtr ConvertSkeletalMeshToMutable(const USkeletalMesh* InSkeletalMesh, co
 		}
 
 		TArrayView<int32>  VertexMorphsInfoIndexBufferView(reinterpret_cast<int32*>(MutableMesh->GetVertexBuffers().GetBufferData(1)), VertexCount);
+		for (int32& Elem : VertexMorphsInfoIndexBufferView)
+		{
+			Elem = -1;
+		}
 
 		if (UsedMorphTargets.Num())
 		{
