[mutable] Avoid generating mips at link-time when linking for constant optimization during compilation.
#rnx

diff --git a/Engine/Plugins/Experimental/Mutable/Source/MutableTools/Private/MuT/AST.h b/Engine/Plugins/Experimental/Mutable/Source/MutableTools/Private/MuT/AST.h
--- a/Engine/Plugins/Experimental/Mutable/Source/MutableTools/Private/MuT/AST.h
+++ b/Engine/Plugins/Experimental/Mutable/Source/MutableTools/Private/MuT/AST.h
@@ -222,6 +222,11 @@ namespace mu
 
 		int32 MinTextureResidentMipCount = 0;
 
+		/** This flag controls the splitting of image data into mips to store separately. It is usually necessary to
+		* be able to generate progressive textures (for texture streaming).
+		*/
+		bool bSeparateImageMips = true;
+
 		/** Structure used to speedup mesh constant comparison. */
 		struct FDeduplicationMeshFuncs : TDefaultMapHashableKeyFuncs<Ptr<const Mesh>, int32, false>
 		{
diff --git a/Engine/Plugins/Experimental/Mutable/Source/MutableTools/Private/MuT/ASTOpConstantResource.cpp b/Engine/Plugins/Experimental/Mutable/Source/MutableTools/Private/MuT/ASTOpConstantResource.cpp
--- a/Engine/Plugins/Experimental/Mutable/Source/MutableTools/Private/MuT/ASTOpConstantResource.cpp
+++ b/Engine/Plugins/Experimental/Mutable/Source/MutableTools/Private/MuT/ASTOpConstantResource.cpp
@@ -70,37 +70,45 @@ namespace mu
 			MUTABLE_CPUPROFILER_SCOPE(AddConstantImage);
 
 			check(pImage->GetSizeX() * pImage->GetSizeY() > 0);
-
+			
 			// Mips to store
 			int32 MipsToStore = 1;
 
-			// We may want the full mipmaps for fragments of images, regardless of the resident mip size, for intermediate operations.
-			// \TODO: Calculate the mip ranges that makes sense to store.
-			int32 MaxMipmaps = Image::GetMipmapCount(pImage->GetSizeX(), pImage->GetSizeY());
-			MipsToStore = MaxMipmaps;
-
 			int32 FirstLODIndexIndex = Program.m_constantImageLODIndices.Num();
 
-			// Some images cannot be resized or mipmaped
-			bool bCannotBeScaled = pImage->m_flags & Image::IF_CANNOT_BE_SCALED;
-			if (bCannotBeScaled)
-			{
-				// Store only the mips that we have already calculated. We assume we have calculated them correctly.
-				MipsToStore = pImage->GetLODCount();
-			}
-
 			FImageOperator& ImOp = Options.ImageOperator;
-
-			// TODO: If the image already has mips, we will be duplicating them...
 			Ptr<const Image> pMip;
-			if (pImage->GetLODCount() == 1)
+
+			if (!Options.bSeparateImageMips)
 			{
 				pMip = pImage;
 			}
 			else
 			{
-				pMip = ImOp.ExtractMip(pImage.get(), 0);
+				// We may want the full mipmaps for fragments of images, regardless of the resident mip size, for intermediate operations.
+				// \TODO: Calculate the mip ranges that makes sense to store.
+				int32 MaxMipmaps = Image::GetMipmapCount(pImage->GetSizeX(), pImage->GetSizeY());
+				MipsToStore = MaxMipmaps;
+
+				// Some images cannot be resized or mipmaped
+				bool bCannotBeScaled = pImage->m_flags & Image::IF_CANNOT_BE_SCALED;
+				if (bCannotBeScaled)
+				{
+					// Store only the mips that we have already calculated. We assume we have calculated them correctly.
+					MipsToStore = pImage->GetLODCount();
+				}
+
+				// TODO: If the image already has mips, we will be duplicating them...
+				if (pImage->GetLODCount() == 1)
+				{
+					pMip = pImage;
+				}
+				else
+				{
+					pMip = ImOp.ExtractMip(pImage.get(), 0);
+				}
 			}
+
 			for (int Mip = 0; Mip < MipsToStore; ++Mip)
 			{
 				check(pMip->GetFormat() == pImage->GetFormat());
@@ -108,22 +116,22 @@ namespace mu
 				// Ensure unique at mip level
 				int32 MipIndex = -1;
 
-				// Use a map-based deduplication
+				// Use a map-based deduplication only if we are splitting mips.
+				if (Options.bSeparateImageMips)
 				{
 					MUTABLE_CPUPROFILER_SCOPE(Deduplicate);
 
-					Ptr<const mu::Image> Key = pMip;
-					const int32* IndexPtr = Options.ImageConstantMipMap.Find(Key);
+					const int32* IndexPtr = Options.ImageConstantMipMap.Find(pMip);
 					if (IndexPtr)
 					{
 						MipIndex = *IndexPtr;
 					}
-					else
-					{
-						check(pMip->GetLODCount() == 1);
-						MipIndex = Program.m_constantImageLODs.Add(TPair<int32, Ptr<const Image>>(-1, pMip));
-						Options.ImageConstantMipMap.Add(Key, MipIndex);
-					}
+				}
+
+				if (MipIndex<0)
+				{
+					MipIndex = Program.m_constantImageLODs.Add(TPair<int32, Ptr<const Image>>(-1, pMip));
+					Options.ImageConstantMipMap.Add(pMip, MipIndex);
 				}
 
 				Program.m_constantImageLODIndices.Add(uint32(MipIndex));
diff --git a/Engine/Plugins/Experimental/Mutable/Source/MutableTools/Private/MuT/CodeOptimiser.cpp b/Engine/Plugins/Experimental/Mutable/Source/MutableTools/Private/MuT/CodeOptimiser.cpp
--- a/Engine/Plugins/Experimental/Mutable/Source/MutableTools/Private/MuT/CodeOptimiser.cpp
+++ b/Engine/Plugins/Experimental/Mutable/Source/MutableTools/Private/MuT/CodeOptimiser.cpp
@@ -504,6 +504,7 @@ namespace mu
 			// Don't generate mips during linking here.
 			FLinkerOptions LinkerOptions(ImOp);
 			LinkerOptions.MinTextureResidentMipCount = 255;
+			LinkerOptions.bSeparateImageMips = false;
 
 			TSharedPtr<const Model> model = MakeShared<Model>();
 			ASTOp::FullLink(SourceCloned, model->GetPrivate()->m_program, &LinkerOptions);
