[Mutable] Add support for multiple sections pointing to the same cloth asset when generating instances.

#rnx

diff --git a/Engine/Plugins/Experimental/Mutable/Source/CustomizableObject/Private/MuCO/CustomizableObjectInstance.cpp b/Engine/Plugins/Experimental/Mutable/Source/CustomizableObject/Private/MuCO/CustomizableObjectInstance.cpp
--- a/Engine/Plugins/Experimental/Mutable/Source/CustomizableObject/Private/MuCO/CustomizableObjectInstance.cpp
+++ b/Engine/Plugins/Experimental/Mutable/Source/CustomizableObject/Private/MuCO/CustomizableObjectInstance.cpp
@@ -3936,20 +3936,18 @@ void UCustomizableInstancePrivate::BuildOrCopyClothingData(const TSharedRef<FUpd
 	}
 
 	// Indices remaps for {Section, AssetLod}, needed to recreate the lod transition data.
-	TArray<TArray<TArray<int32>>> PhysicsSectionLodsIndicesRemaps;	
+	TMap<int32, TArray<TArray<int32>>> PhysicsSectionLodsIndicesRemaps;
 
 	check( SectionsWithCloth.Num() > 0 );
 	FSectionWithClothData* MaxSection = Algo::MaxElement(SectionsWithCloth, 
-			[](const FSectionWithClothData& A, const FSectionWithClothData& B) { return  A.Section < B.Section; } );
+			[](const FSectionWithClothData& A, const FSectionWithClothData& B) { return  A.ClothAssetIndex < B.ClothAssetIndex; } );
 	
-	PhysicsSectionLodsIndicesRemaps.SetNum(MaxSection->Section + 1);
+	PhysicsSectionLodsIndicesRemaps.Reserve(MaxSection->ClothAssetIndex + 1);
 		
 	for (FSectionWithClothData& SectionLods : SectionsWithCloth)
 	{
-		const int32 SectionLodNum = PhysicsSectionLodsIndicesRemaps[SectionLods.Section].Num();
-
-		PhysicsSectionLodsIndicesRemaps[SectionLods.Section].SetNum(
-				FMath::Max( SectionLodNum, SectionLods.ClothAssetLodIndex + 1 ));
+		TArray<TArray<int32>>& Value = PhysicsSectionLodsIndicesRemaps.FindOrAdd(SectionLods.ClothAssetIndex, {});
+		Value.SetNum(FMath::Max(Value.Num(), SectionLods.ClothAssetLodIndex + 1));
 	}
 
 	{
@@ -4019,7 +4017,7 @@ void UCustomizableInstancePrivate::BuildOrCopyClothingData(const TSharedRef<FUpd
 				}
 			}
 
-			TArray<int32>& IndexMap = PhysicsSectionLodsIndicesRemaps[SectionWithCloth.Section][SectionWithCloth.ClothAssetLodIndex];
+			TArray<int32>& IndexMap = PhysicsSectionLodsIndicesRemaps[SectionWithCloth.ClothAssetIndex][SectionWithCloth.ClothAssetLodIndex];
 			IndexMap.SetNumUninitialized(PhysicalMeshVerticesNum);
 
 			// Compute index remap and number of remaining physics vertices.
@@ -4073,15 +4071,16 @@ void UCustomizableInstancePrivate::BuildOrCopyClothingData(const TSharedRef<FUpd
 			CopyIfUsed(NewLodData.PhysicalMeshData.BoneData, SrcLodData.PhysicalMeshData.BoneData);
 			CopyIfUsed(NewLodData.PhysicalMeshData.InverseMasses, SrcLodData.PhysicalMeshData.InverseMasses);
 
-			const bool bNeedsTransitionUpData = SectionWithCloth.Lod - 1 >= 0;
+			const int32 PrevIndex = SectionWithCloth.Lod - 1;
+			const bool bNeedsTransitionUpData = NewClothingAssetData.LodMap.IsValidIndex(PrevIndex) && NewClothingAssetData.LodMap[PrevIndex] != INDEX_NONE;
 			if (bNeedsTransitionUpData)
 			{
 				NewLodData.TransitionUpSkinData.SetNum(SrcLodData.TransitionUpSkinData.Num() ? NewPhysicalMeshVerticesNum : 0);	
 				CopyIfUsed(NewLodData.TransitionUpSkinData, SrcLodData.TransitionUpSkinData);
 			}
 
-			check(SectionWithCloth.Section < PhysicsSectionLodsIndicesRemaps.Num());
-			const bool bNeedsTransitionDownData = SectionWithCloth.Lod + 1 < PhysicsSectionLodsIndicesRemaps[SectionWithCloth.Section].Num();
+			const int32 NextIndex = SectionWithCloth.Lod + 1;
+			const bool bNeedsTransitionDownData = NewClothingAssetData.LodMap.IsValidIndex(NextIndex) && NewClothingAssetData.LodMap[NextIndex] != INDEX_NONE;
 			if (bNeedsTransitionDownData)
 			{
 				NewLodData.TransitionDownSkinData.SetNum(SrcLodData.TransitionDownSkinData.Num() ? NewPhysicalMeshVerticesNum : 0);
@@ -4251,11 +4250,8 @@ void UCustomizableInstancePrivate::BuildOrCopyClothingData(const TSharedRef<FUpd
 			};
 
 			if (NewLodData.TransitionDownSkinData.Num() > 0)
-			{
-				check(SectionWithCloth.Section < PhysicsSectionLodsIndicesRemaps.Num() &&
-					  SectionWithCloth.ClothAssetLodIndex + 1 < PhysicsSectionLodsIndicesRemaps[SectionWithCloth.Section].Num() );
-				
-				TArray<int32>& IndexMap = PhysicsSectionLodsIndicesRemaps[SectionWithCloth.Section][SectionWithCloth.ClothAssetLodIndex + 1];
+			{	
+				TArray<int32>& IndexMap = PhysicsSectionLodsIndicesRemaps[SectionWithCloth.ClothAssetIndex][SectionWithCloth.ClothAssetLodIndex + 1];
 
 				if (IndexMap.Num())
 				{
@@ -4264,14 +4260,8 @@ void UCustomizableInstancePrivate::BuildOrCopyClothingData(const TSharedRef<FUpd
 			}
 			
 			if (NewLodData.TransitionUpSkinData.Num() > 0)
-			{
-				check( SectionWithCloth.ClothAssetLodIndex - 1 >= 0 )
-				
-				check(SectionWithCloth.Section < PhysicsSectionLodsIndicesRemaps.Num() &&
-					  SectionWithCloth.ClothAssetLodIndex - 1 < PhysicsSectionLodsIndicesRemaps[SectionWithCloth.Section].Num());
-				
-				TArray<int32>& IndexMap = PhysicsSectionLodsIndicesRemaps[SectionWithCloth.Section][SectionWithCloth.ClothAssetLodIndex - 1];
-			
+			{	
+				TArray<int32>& IndexMap = PhysicsSectionLodsIndicesRemaps[SectionWithCloth.ClothAssetIndex][SectionWithCloth.ClothAssetLodIndex - 1];
 				if (IndexMap.Num())
 				{
 					RemapTransitionMeshToMeshVertData( NewLodData.TransitionUpSkinData, IndexMap );
@@ -4382,11 +4372,11 @@ void UCustomizableInstancePrivate::BuildOrCopyClothingData(const TSharedRef<FUpd
 				NewClothingAssetData.LodData[SectionWithCloth.ClothAssetLodIndex].PhysicalMeshData.Normals,
 				NewClothingAssetData.LodData[SectionWithCloth.ClothAssetLodIndex].PhysicalMeshData.Indices };
 
-			const TArray< TArray<int32> >& SectionIndexRemaps = PhysicsSectionLodsIndicesRemaps[SectionWithCloth.Section];
-
+			const TArray<TArray<int32>>& SectionIndexRemaps = PhysicsSectionLodsIndicesRemaps[SectionWithCloth.ClothAssetIndex];
+			
 			if (SectionWithCloth.ClothAssetLodIndex < SectionIndexRemaps.Num() - 1)
 			{
-				const TArray<int32>& IndexMap = PhysicsSectionLodsIndicesRemaps[SectionWithCloth.Section][SectionWithCloth.ClothAssetLodIndex + 1];
+				const TArray<int32>& IndexMap = SectionIndexRemaps[SectionWithCloth.ClothAssetLodIndex + 1];
 				
 				const FMeshPhysicsDesc TransitionDownTarget {  
 					NewClothingAssetData.LodData[SectionWithCloth.ClothAssetLodIndex + 1].PhysicalMeshData.Vertices,
@@ -4398,7 +4388,7 @@ void UCustomizableInstancePrivate::BuildOrCopyClothingData(const TSharedRef<FUpd
 			
 			if (SectionWithCloth.ClothAssetLodIndex > 0)
 			{
-				const TArray<int32>& IndexMap = PhysicsSectionLodsIndicesRemaps[SectionWithCloth.Section][SectionWithCloth.ClothAssetLodIndex - 1];
+				const TArray<int32>& IndexMap = SectionIndexRemaps[SectionWithCloth.ClothAssetLodIndex - 1];
 				
 				FMeshPhysicsDesc TransitionUpTarget{  
 					NewClothingAssetData.LodData[SectionWithCloth.ClothAssetLodIndex - 1].PhysicalMeshData.Vertices,
@@ -4444,7 +4434,7 @@ void UCustomizableInstancePrivate::BuildOrCopyClothingData(const TSharedRef<FUpd
 					int32 FirstVertex, VerticesCount, FirstIndex, IndicesCount;
 					MutableMesh->GetSurface(Section, &FirstVertex, &VerticesCount, &FirstIndex, &IndicesCount, nullptr, nullptr, nullptr);
 
-					if (VerticesCount == 0 || IndicesCount == 0 )
+					if (VerticesCount == 0 || IndicesCount == 0)
 					{
 						continue;
 					}					
@@ -4531,8 +4521,9 @@ void UCustomizableInstancePrivate::BuildOrCopyClothingData(const TSharedRef<FUpd
 			{
 				continue;
 			}
-		
-			NewClothingAssets[I] = NewObject<UCustomizableObjectClothingAsset>( SkeletalMesh, NewClothingAssetsData[I].Name );
+	
+			FName UniqueClothingAssetName = FName(FString::Printf(TEXT("%s_%d"), *NewClothingAssetsData[I].Name.ToString(), I));
+			NewClothingAssets[I] = NewObject<UCustomizableObjectClothingAsset>(SkeletalMesh, UniqueClothingAssetName);
 			
 			// The data can be moved to the actual asset since it will not be used anymore.
 			NewClothingAssets[I]->LodMap = MoveTemp(NewClothingAssetsData[I].LodMap);
