[mutable] Fix corruption and possible crash for generic image pixel format conversion.
#rnx

diff --git a/Engine/Plugins/Experimental/Mutable/Source/MutableRuntime/Private/MuR/OpImagePixelFormat.cpp b/Engine/Plugins/Experimental/Mutable/Source/MutableRuntime/Private/MuR/OpImagePixelFormat.cpp
--- a/Engine/Plugins/Experimental/Mutable/Source/MutableRuntime/Private/MuR/OpImagePixelFormat.cpp
+++ b/Engine/Plugins/Experimental/Mutable/Source/MutableRuntime/Private/MuR/OpImagePixelFormat.cpp
@@ -1236,7 +1236,7 @@ void FImageOperator::ImagePixelFormat( bool& bOutSuccess, int32 CompressionQuali
 				{
 					FIntVector2 mipSize = pResult->CalculateMipSize(m);
 					miro::RGBA_to_BC1(
-						mipSize[0], mipSize[1], Base->GetMipData(baseLOD + m),
+						mipSize[0], mipSize[1], Temp->GetMipData(baseLOD + m),
 						pResult->GetMipData(m), CompressionQuality);
 				}
 				ReleaseImage(Temp);
@@ -1271,7 +1271,7 @@ void FImageOperator::ImagePixelFormat( bool& bOutSuccess, int32 CompressionQuali
 				{
 					FIntVector2 mipSize = pResult->CalculateMipSize(m);
 					miro::RGBA_to_BC2(
-						mipSize[0], mipSize[1], Base->GetMipData(baseLOD + m),
+						mipSize[0], mipSize[1], Temp->GetMipData(baseLOD + m),
 						pResult->GetMipData(m), CompressionQuality);
 				}
 				ReleaseImage(Temp);
@@ -1330,7 +1330,7 @@ void FImageOperator::ImagePixelFormat( bool& bOutSuccess, int32 CompressionQuali
 				{
 					FIntVector2 mipSize = pResult->CalculateMipSize(m);
 					miro::RGBA_to_BC3(
-						mipSize[0], mipSize[1], Base->GetMipData(baseLOD + m),
+						mipSize[0], mipSize[1], Temp->GetMipData(baseLOD + m),
 						pResult->GetMipData(m), CompressionQuality);
 				}
 				ReleaseImage(Temp);
