[mutable] Improve compilation options to add support for packaged file size control.
Packaged files are named after a hash of the content.
Added control of the size threshold of embedded data.
#rnx

diff --git a/Engine/Plugins/Experimental/Mutable/Source/CustomizableObject/Classes/MuCO/CustomizableObject.h b/Engine/Plugins/Experimental/Mutable/Source/CustomizableObject/Classes/MuCO/CustomizableObject.h
--- a/Engine/Plugins/Experimental/Mutable/Source/CustomizableObject/Classes/MuCO/CustomizableObject.h
+++ b/Engine/Plugins/Experimental/Mutable/Source/CustomizableObject/Classes/MuCO/CustomizableObject.h
@@ -132,12 +132,13 @@ enum class ECustomizableObjectTextureCompression : uint8
 	HighQuality
 };
 
+
 USTRUCT()
 struct FCompilationOptions
 {
 	GENERATED_USTRUCT_BODY()
 	
-	/** Enum to know what texture compression should be used. This compression is used on ly in manual compiles in editor. 
+	/** Enum to know what texture compression should be used. This compression is used only in manual compiles in editor. 
 	 *  When packaging, ECustomizableObjectTextureCompression::HighQuality is always used.
 	 */
 	UPROPERTY()
@@ -152,6 +153,16 @@ struct FCompilationOptions
 	UPROPERTY()
 	bool bUseDiskCompilation = false;
 
+	/** High limit of the size in bytes of the packaged data when cooking this object. 
+	* This limit is before any pak or filesystem compression. This limit will be broken if a single piece of data is bigger because data is not fragmented for packaging purposes.
+	*/
+	UPROPERTY()
+	uint64 PackagedDataBytesLimit = 256*1024*1024;
+
+	/** High (inclusive) limit of the size in bytes of a data block to be included into the compiled object directly instead of stored in a streamable file. */
+	UPROPERTY()
+	uint64 EmbeddedDataBytesLimit = 1024;
+
 	// Did we have the extra bones enabled when we compiled?
 	ECustomizableObjectNumBoneInfluences CustomizableObjectNumBoneInfluences = ECustomizableObjectNumBoneInfluences::Four;
 
@@ -363,20 +374,19 @@ struct CUSTOMIZABLEOBJECT_API FMutableStreamableBlock
 	GENERATED_USTRUCT_BODY()
 
 	UPROPERTY()
-	uint16 FileIndex = 0;
-
-	UPROPERTY()
-	uint64 Offset = 0;
+	uint32 FileId = 0;
 
 	UPROPERTY()
 	uint32 Size = 0;
 
+	UPROPERTY()
+	uint64 Offset = 0;
 
 	friend FArchive& operator<<(FArchive& Ar, FMutableStreamableBlock& Data)
 	{
-		Ar << Data.FileIndex;
-		Ar << Data.Offset;
+		Ar << Data.FileId;
 		Ar << Data.Size;
+		Ar << Data.Offset;
 
 		return Ar;
 	}
@@ -389,7 +399,6 @@ struct FMutableLODSettings
 {
 	GENERATED_BODY()
 
-
 	/** Minimum LOD to render per Platform. */
 	UPROPERTY(EditAnywhere, Category = LODSettings, meta = (DisplayName = "Minimum LOD"))
 	FPerPlatformInt MinLOD;
@@ -442,21 +451,22 @@ class CUSTOMIZABLEOBJECT_API UCustomizableObjectBulk : public UObject
 public:
 	GENERATED_BODY()
 
+	//~ Begin UObject Interface
 	virtual void PostLoad() override;
+	//~ End UObject Interface
 
-	/** Creates and returns an array with IAsyncReadFileHandles for each BulkData file.
-	 * Used by Mutable to stream in resources when generating instances. Must be deleted by the caller. */
-	TArray<TSharedPtr<IAsyncReadFileHandle>> GetAsyncReadFileHandles() const;
+	/**  */
+	const FString& GetBulkFilePrefix() const { return BulkFilePrefix; }
 
 #if WITH_EDITOR
 
 	//~ Begin UObject Interface
-	virtual void CookAdditionalFilesOverride(const TCHAR* PackageFilename, const ITargetPlatform* TargetPlatform,
-		TFunctionRef<void(const TCHAR* Filename, void* Data, int64 Size)> WriteAdditionalFile) override;
+	virtual void CookAdditionalFilesOverride(const TCHAR*, const ITargetPlatform*, TFunctionRef<void(const TCHAR*, void*, int64)> ) override;
 	//~ End UObject Interface
 
 	/** Compute the number of files and sizes the BulkData will be split into and fix up 
-	 * HashToStreamableBlock's FileIndices and Offsets. */	
+	 * the HashToStreamableBlock's FileIds and Offsets. 
+	 */	
 	void PrepareBulkData(UCustomizableObject* InCustomizableObject, const ITargetPlatform* TargetPlatform);
 
 #endif
@@ -465,17 +475,37 @@ private:
 
 #if WITH_EDITOR
 
+	struct FBlock
+	{
+		// \TODO: needed?
+		uint32 Id;
+
+		/** Size of the data block. */
+		uint32 Size;
+
+		/** Offset in the full source streamed data file that is created when compiling. */
+		uint64 Offset;
+	};
+
+	struct FFile
+	{
+		/** Id generated from a hash of the file content + offset to avoid collisions. */
+		uint32 Id;
+
+		/** List of blocks that are contained in the file, in order. */
+		TArray<FBlock> Blocks;
+	};
+
 	/** Helper to store the size of each BulkData partition. Only valid while cooking */
-	TArray<uint64> BulkDataFilesSize;
+	TArray<FFile> BulkDataFiles;
 
 	/** Helper to retrieve the BulkData from within the CookAdditionalFilesOverride */
 	TObjectPtr<UCustomizableObject> CustomizableObject;
 
 #endif
 
-	/** BulkData file paths */
-	UPROPERTY()
-	TArray<FString> BulkDataFileNames;
+	/** Prefix to locate bulkfiles for loading, using the file ids in each FMutableStreamableBlock. */
+	FString BulkFilePrefix;
 };
 
 
diff --git a/Engine/Plugins/Experimental/Mutable/Source/CustomizableObject/Internal/MuCO/CustomizableObjectPrivate.h b/Engine/Plugins/Experimental/Mutable/Source/CustomizableObject/Internal/MuCO/CustomizableObjectPrivate.h
--- a/Engine/Plugins/Experimental/Mutable/Source/CustomizableObject/Internal/MuCO/CustomizableObjectPrivate.h
+++ b/Engine/Plugins/Experimental/Mutable/Source/CustomizableObject/Internal/MuCO/CustomizableObjectPrivate.h
@@ -685,8 +685,8 @@ public:
 	 * Updated each time the CO is compiled and saved in the Derived Data. */
 	TMap<FName, FGuid> ParticipatingObjects;
 
-	/** If the object is compiled, this flag is true if it was compiled with maximum optimizations. If the object is not compiled, its value is meaningless. */
-	bool bIsCompiledWithOptimization = true;
+	/** If the object is compiled, this flag is false unless it was compiled with maximum optimizations. If the object is not compiled, its value is meaningless. */
+	bool bIsCompiledWithoutOptimization = true;
 
 	/** This is a non-user-controlled flag to disable streaming (set at object compilation time, depending on optimization). */
 	bool bDisableTextureStreaming = false;
@@ -722,6 +722,6 @@ public:
 	// This is a manual version number for the binary blobs in this asset.
 	// Increasing it invalidates all the previously compiled models.
 	// Warning: If while merging code both versions have changed, take the highest+1.
-	static constexpr int32 CurrentSupportedVersion = 426;
+	static constexpr int32 CurrentSupportedVersion = 429;
 };
 
diff --git a/Engine/Plugins/Experimental/Mutable/Source/CustomizableObject/Internal/MuCO/UnrealMutableModelDiskStreamer.h b/Engine/Plugins/Experimental/Mutable/Source/CustomizableObject/Internal/MuCO/UnrealMutableModelDiskStreamer.h
--- a/Engine/Plugins/Experimental/Mutable/Source/CustomizableObject/Internal/MuCO/UnrealMutableModelDiskStreamer.h
+++ b/Engine/Plugins/Experimental/Mutable/Source/CustomizableObject/Internal/MuCO/UnrealMutableModelDiskStreamer.h
@@ -71,13 +71,17 @@ protected:
 	struct FObjectData
 	{
 		TWeakPtr<const mu::Model> Model;
-		TArray<TSharedPtr<IAsyncReadFileHandle>> ReadFileHandles;
 		TMap<OPERATION_ID, FReadRequest> CurrentReadRequests;
 		TMap<uint64, FMutableStreamableBlock> StreamableBlocks;
+		FString BulkFilePrefix;
+
+		TMap<uint32, TSharedPtr<IAsyncReadFileHandle>> ReadFileHandles;
 	};
 
 	TArray<FObjectData> Objects;
 
+	FCriticalSection FileHandlesCritical;
+
 	/** This is used to generate unique ids for read requests. */
 	OPERATION_ID LastOperationID = 0;
 };
diff --git a/Engine/Plugins/Experimental/Mutable/Source/CustomizableObject/Private/MuCO/CustomizableObject.cpp b/Engine/Plugins/Experimental/Mutable/Source/CustomizableObject/Private/MuCO/CustomizableObject.cpp
--- a/Engine/Plugins/Experimental/Mutable/Source/CustomizableObject/Private/MuCO/CustomizableObject.cpp
+++ b/Engine/Plugins/Experimental/Mutable/Source/CustomizableObject/Private/MuCO/CustomizableObject.cpp
@@ -565,7 +565,7 @@ void UCustomizableObjectPrivate::SaveCompiledData(FArchive& MemoryWriter, bool b
 
 	// Editor Only data
 	MemoryWriter << bDisableTextureStreaming;
-	MemoryWriter << bIsCompiledWithOptimization;
+	MemoryWriter << bIsCompiledWithoutOptimization;
 	MemoryWriter << CustomizableObjectPathMap;
 	MemoryWriter << GroupNodeMap;
 	MemoryWriter << ParticipatingObjects;
@@ -687,7 +687,7 @@ void UCustomizableObjectPrivate::LoadCompiledData(FArchive& MemoryReader, const
 		// Editor Only data
 		{
 			MemoryReader << bDisableTextureStreaming;
-			MemoryReader << bIsCompiledWithOptimization;
+			MemoryReader << bIsCompiledWithoutOptimization;
 			MemoryReader << CustomizableObjectPathMap;
 			MemoryReader << GroupNodeMap;
 			MemoryReader << ParticipatingObjects;
@@ -2004,41 +2004,9 @@ void UCustomizableObjectBulk::PostLoad()
 {
 	UObject::PostLoad();
 
-	const FString PackageFilename = FPackageName::LongPackageNameToFilename(GetOutermost()->GetName(), TEXT(".uasset"));
-	FString OuterPathName = FPaths::GetPath(PackageFilename);
-
-	for(FString& FilePath : BulkDataFileNames)
-	{
-		FilePath = OuterPathName / FilePath;
-	}
-
-}
-
-
-TArray<TSharedPtr<IAsyncReadFileHandle>> UCustomizableObjectBulk::GetAsyncReadFileHandles() const
-{
-	TArray<TSharedPtr<IAsyncReadFileHandle>> ReadFileHandles;
-	ReadFileHandles.Reserve(BulkDataFileNames.Num());
-
-	for (const FString& FilePath : BulkDataFileNames)
-	{
-		TSharedPtr<IAsyncReadFileHandle> ReadFileHandle = MakeShareable(FPlatformFileManager::Get().GetPlatformFile().OpenAsyncRead(*FilePath));
-		
-		if (!ReadFileHandle)
-		{
-			UE_LOG(LogMutable, Error, TEXT("Failed to create AsyncReadFileHandle. File Path [%s]."), *FilePath);
-			break;
-		}
-
-		ReadFileHandles.Add(ReadFileHandle);
-	}
-
-	if (ReadFileHandles.Num() != BulkDataFileNames.Num())
-	{
-		ReadFileHandles.Empty();
-	}
-
-	return ReadFileHandles;
+	const FString OutermostName = GetOutermost()->GetName();
+	const FString PackageFilename = FPackageName::LongPackageNameToFilename(OutermostName);
+	BulkFilePrefix = PackageFilename;
 }
 
 
@@ -2059,20 +2027,36 @@ void UCustomizableObjectBulk::CookAdditionalFilesOverride(const TCHAR* PackageFi
 	FMutableCachedPlatformData* PlatformData = CustomizableObject->GetPrivate()->CachedPlatformsData.Find(TargetPlatform->PlatformName());
 	check(PlatformData);
 
-	// Data to serialize in separate files
-	uint8* Data = PlatformData->StreamableData.GetData();
-
-	// Path to the asset
-	const FString CookedFilePath = FPaths::GetPath(PackageFilename);
+	// Source data pointer
+	const uint8* SourceData = PlatformData->StreamableData.GetData();
 
-	const uint16 NumBulkDataFiles = BulkDataFilesSize.Num();
-	for(uint16 FileIndex = 0; FileIndex < NumBulkDataFiles; ++FileIndex)
+	const int32 NumBulkDataFiles = BulkDataFiles.Num();
+	for(int32 FileIndex = 0; FileIndex < NumBulkDataFiles; ++FileIndex)
 	{
-		const FString CookedBulkFileName = CookedFilePath / BulkDataFileNames[FileIndex];
+		const FFile& CurrentFile = BulkDataFiles[FileIndex];
+
+		int64 FileSize = 0;
+		for ( const FBlock& Block: CurrentFile.Blocks )
+		{
+			FileSize += Block.Size;
+		}
+
+		// Generate the bulk data file in memory
+		TArray64<uint8> FileBulkData;
+		FileBulkData.SetNum( FileSize );
+		uint8* FileData = FileBulkData.GetData();
+
+		for (const FBlock& Block : CurrentFile.Blocks)
+		{
+			FMemory::Memcpy(FileData,SourceData+Block.Offset,Block.Size);
+			FileData += Block.Size;
+		}
 
-		const int64 NumBytes = BulkDataFilesSize[FileIndex];
-		WriteAdditionalFile(*CookedBulkFileName, (void*)Data, NumBytes);
-		Data += NumBytes;
+		// Path to the asset
+		const FString CookedFilePath = FPaths::GetPath(PackageFilename);
+
+		FString CookedBulkFileName = FString::Printf(TEXT("%s/%s-%08x.mut"), *CookedFilePath, *CustomizableObject->GetName(), CurrentFile.Id);
+		WriteAdditionalFile(*CookedBulkFileName, FileBulkData.GetData(), FileBulkData.Num());
 	}
 }
 
@@ -2082,52 +2066,129 @@ void UCustomizableObjectBulk::PrepareBulkData(UCustomizableObject* InOuter, cons
 	CustomizableObject = InOuter;
 	check(CustomizableObject);
 
-	BulkDataFilesSize.Empty();
-	BulkDataFileNames.Empty();
-	
-	// Split the Streamable data into several separate files and fix up FileIndex and Offset of each StreamableBlock
-	if (TSharedPtr<const mu::Model, ESPMode::ThreadSafe> Model = CustomizableObject->GetPrivate()->GetModel())
+	BulkDataFiles.Empty();
+
+	TSharedPtr<const mu::Model, ESPMode::ThreadSafe> Model = CustomizableObject->GetPrivate()->GetModel();
+	if (!Model)
+	{
+		return;
+	}
+
+	uint64 TargetBulkDataFileBytes = CustomizableObject->CompileOptions.PackagedDataBytesLimit;
+	const uint64 MaxChunkSize = UCustomizableObjectSystem::GetInstance()->GetMaxChunkSizeForPlatform(TargetPlatform);
+	TargetBulkDataFileBytes = FMath::Min(TargetBulkDataFileBytes, MaxChunkSize);
+
+	const int32 NumBlocks = Model->GetRomCount();
+
+	// TODO:
+	// To avoid influence of the order of the streamed data (their index), classify it recursively based on hash values
+	// until the tree leaves have either a single block, or a sum of blocks below the desired file size.
+	struct FClassifyNode
 	{
-		const uint64 MaxChunkSize = UCustomizableObjectSystem::GetInstance()->GetMaxChunkSizeForPlatform(TargetPlatform);
+		TArray<FBlock> Blocks;
+
+		//TSharedPtr<FClassifyNode> Child0;
+		//TSharedPtr<FClassifyNode> Child1;
+		//uint32 Depth=0;
+	};
 
-		const FString BulkFileName = CustomizableObject->GetName() + FString(TEXT("_Bulk"));
+	FClassifyNode RootNode;
+	RootNode.Blocks.Reserve(NumBlocks);
+
+	// Create blocks data, filtering out the ones that are too big and will go on its own file in any case.
+	uint64 SourceOffset = 0;
+	for (int32 BlockIndex = 0; BlockIndex < NumBlocks; ++BlockIndex)
+	{
+		uint32 BlockId = Model->GetRomId(BlockIndex);
+		const FMutableStreamableBlock& StreamableBlock = CustomizableObject->HashToStreamableBlock[BlockId];
+		const uint32 BlockSize = StreamableBlock.Size;
 
-		uint16 CurrentFileIndex = 0;
-		uint64 CurrentChunkSize = 0;
+		FBlock CurrentBlock = { BlockId, BlockSize, SourceOffset };
 
-		const int32 NumStreamingFiles = Model->GetRomCount();
-		for (size_t FileIndex = 0; FileIndex < NumStreamingFiles; ++FileIndex)
+		if (BlockSize > TargetBulkDataFileBytes)
 		{
-			const uint32 ResourceId = Model->GetRomId(FileIndex);
+			// It will go to its own file
+			BulkDataFiles.Add(FFile{ 0, {CurrentBlock} });
+		}
+		else
+		{
+			// It may merge with other small blocks
+			RootNode.Blocks.Add(CurrentBlock);
+		}
 
-			FMutableStreamableBlock& StreamableBlock = CustomizableObject->HashToStreamableBlock[ResourceId];
+		SourceOffset += BlockSize;
+	}
+	
+	// Temp: Group by order in the array
+	for (int32 BlockIndex = 0; BlockIndex < RootNode.Blocks.Num(); )
+	{
+		FFile CurrentFile;
+		int32 CurrentFileSize = 0;
+
+		while(BlockIndex < RootNode.Blocks.Num())
+		{
+			FBlock CurrentBlock = RootNode.Blocks[BlockIndex];
 
-			const uint32 BlockSize = StreamableBlock.Size;
-			if(CurrentChunkSize + BlockSize > MaxChunkSize)
+			// Next file?
+			if (CurrentFileSize>0 && CurrentFileSize + CurrentBlock.Size > TargetBulkDataFileBytes)
 			{
-				BulkDataFileNames.Add(BulkFileName + FString::Printf(TEXT("%d.mut"), CurrentFileIndex));
+				break;
+			}
 
-				if(CurrentChunkSize == 0)
+			// Add the block to the current file
+			CurrentFile.Blocks.Add(CurrentBlock);
+			CurrentFileSize += CurrentBlock.Size;
+
+			// Next block
+			++BlockIndex;			
+		}
+
+		BulkDataFiles.Add(MoveTemp(CurrentFile));
+	}
+
+	// Create the file list
+	for ( int32 FileIndex=0; FileIndex<BulkDataFiles.Num(); ++FileIndex )
+	{
+		// Generate the id for this file
+		FFile& CurrentFile = BulkDataFiles[FileIndex];
+		uint32 FileId = 0;
+		for (int32 FileBlockIndex = 0; FileBlockIndex < CurrentFile.Blocks.Num(); ++FileBlockIndex)
+		{
+			FBlock ThisBlock = CurrentFile.Blocks[FileBlockIndex];
+			FileId = (FileBlockIndex == 0) ? ThisBlock.Id : HashCombine( FileId, ThisBlock.Id);
+		}
+
+		// Ensure the FileId is unique
+		bool bUnique = false;
+		while (!bUnique)
+		{
+			bUnique = true;
+			for (int32 PreviousFileIndex = 0; PreviousFileIndex < FileIndex; ++PreviousFileIndex)
+			{
+				if (BulkDataFiles[PreviousFileIndex].Id == FileId)
 				{
-					BulkDataFilesSize.Add(BlockSize);
-					StreamableBlock.FileIndex = CurrentFileIndex++;
-					StreamableBlock.Offset = 0;
-					continue;
+					bUnique = false;
+					++FileId;
+					break;
 				}
-
-				BulkDataFilesSize.Add(CurrentChunkSize);
-			
-				CurrentFileIndex++;
-				CurrentChunkSize = 0;
 			}
-
-			StreamableBlock.FileIndex = CurrentFileIndex;
-			StreamableBlock.Offset = CurrentChunkSize;
-			CurrentChunkSize += BlockSize;			
 		}
 
-		BulkDataFilesSize.Add(CurrentChunkSize);
-		BulkDataFileNames.Add(BulkFileName + FString::Printf(TEXT("%d.mut"), CurrentFileIndex));
+		// Set it to the editor-only file descriptor
+		CurrentFile.Id = FileId;
+
+		// Set it to all streamable blocks
+		uint32 OffsetInFile = 0;
+		for (int32 FileBlockIndex = 0; FileBlockIndex < CurrentFile.Blocks.Num(); ++FileBlockIndex)
+		{
+			FBlock ThisBlock = CurrentFile.Blocks[FileBlockIndex];
+
+			FMutableStreamableBlock& StreamableBlock = CustomizableObject->HashToStreamableBlock[ThisBlock.Id];
+			check(StreamableBlock.Size == ThisBlock.Size);
+			StreamableBlock.FileId = FileId;
+			StreamableBlock.Offset = OffsetInFile;
+			OffsetInFile += ThisBlock.Size;
+		}
 	}
 }
 
diff --git a/Engine/Plugins/Experimental/Mutable/Source/CustomizableObject/Private/MuCO/CustomizableObjectSystem.cpp b/Engine/Plugins/Experimental/Mutable/Source/CustomizableObject/Private/MuCO/CustomizableObjectSystem.cpp
--- a/Engine/Plugins/Experimental/Mutable/Source/CustomizableObject/Private/MuCO/CustomizableObjectSystem.cpp
+++ b/Engine/Plugins/Experimental/Mutable/Source/CustomizableObject/Private/MuCO/CustomizableObjectSystem.cpp
@@ -3663,14 +3663,14 @@ void UCustomizableObjectSystemPrivate::ShowOnScreenCompileWarnings()
 	{
 		// Show a warning if the compilation was not done with optimizations.
 		const uint64 KeyCompiledWithOptimization = reinterpret_cast<uint64>(Object);
-		if (Object->GetPrivate()->bIsCompiledWithOptimization && // Quicker to check
+		if (Object->GetPrivate()->bIsCompiledWithoutOptimization && // Quicker to check
 			!GEngine->OnScreenDebugMessageExists(KeyCompiledWithOptimization))
 		{
 			FString Msg = FString::Printf(TEXT("Warning: Customizable Object [%s] was compiled without optimization."), *Object->GetName());
 			GEngine->AddOnScreenDebugMessage(KeyCompiledWithOptimization, 10.0f, FColor::Red, Msg);
 		}
 
-		const uint64 KeyCompiledOutOfData = reinterpret_cast<uint64>(Object) + KEY_OFFSET_COMPILATION_OUT_OF_DATE; // Offset added to avoid collision with bIsCompiledWithOptimization warning
+		const uint64 KeyCompiledOutOfData = reinterpret_cast<uint64>(Object) + KEY_OFFSET_COMPILATION_OUT_OF_DATE; // Offset added to avoid collision with bIsCompiledWithoutOptimization warning
 		if (!GEngine->OnScreenDebugMessageExists(KeyCompiledWithOptimization) && // Quicker to check
 			Object->GetPrivate()->IsCompilationOutOfDate())
 		{
@@ -3683,7 +3683,7 @@ void UCustomizableObjectSystemPrivate::ShowOnScreenCompileWarnings()
 
 void UCustomizableObjectSystemPrivate::HideOnScreenCompileWarnings(const UCustomizableObjectPrivate& ObjectPrivate)
 {
-	if (ObjectPrivate.bIsCompiledWithOptimization) 
+	if (ObjectPrivate.bIsCompiledWithoutOptimization)
 	{
 		GEngine->RemoveOnScreenDebugMessage(reinterpret_cast<uint64>(&ObjectPrivate));
 	}
diff --git a/Engine/Plugins/Experimental/Mutable/Source/CustomizableObject/Private/MuCO/UnrealMutableModelDiskStreamer.cpp b/Engine/Plugins/Experimental/Mutable/Source/CustomizableObject/Private/MuCO/UnrealMutableModelDiskStreamer.cpp
--- a/Engine/Plugins/Experimental/Mutable/Source/CustomizableObject/Private/MuCO/UnrealMutableModelDiskStreamer.cpp
+++ b/Engine/Plugins/Experimental/Mutable/Source/CustomizableObject/Private/MuCO/UnrealMutableModelDiskStreamer.cpp
@@ -64,6 +64,9 @@ bool FUnrealMutableModelBulkReader::PrepareStreamingForObject(UCustomizableObjec
 	// See if we can free previuously allocated resources
 	for (int32 ObjectIndex = 0; ObjectIndex < Objects.Num(); )
 	{
+		// Close open file handles
+		Objects[ObjectIndex].ReadFileHandles.Empty();
+
 		if (!Objects[ObjectIndex].Model.Pin() && Objects[ObjectIndex].CurrentReadRequests.IsEmpty())
 		{
 			Objects.RemoveAtSwap(ObjectIndex);
@@ -89,15 +92,7 @@ bool FUnrealMutableModelBulkReader::PrepareStreamingForObject(UCustomizableObjec
 #if WITH_EDITOR
 		FString FolderPath = CustomizableObject->GetPrivate()->GetCompiledDataFolderPath();
 		FString FullFileName = FolderPath + CustomizableObject->GetPrivate()->GetCompiledDataFileName(false, nullptr, true);
-
-		const TSharedPtr<IAsyncReadFileHandle> ReadFileHandle = MakeShareable(FPlatformFileManager::Get().GetPlatformFile().OpenAsyncRead(*FullFileName));
-		if (!ReadFileHandle)
-		{
-			UE_LOG(LogMutable, Warning, TEXT("Streaming: Customizable Object %s is missing the Editor BulkData."), *CustomizableObject->GetName());
-			return false;
-		}
-
-		NewData.ReadFileHandles.Add(ReadFileHandle);
+		NewData.BulkFilePrefix = *FullFileName;
 #else
 		const UCustomizableObjectBulk* BulkData = CustomizableObject->GetPrivate()->GetStreamableBulkData();
 		if (!BulkData)
@@ -106,18 +101,9 @@ bool FUnrealMutableModelBulkReader::PrepareStreamingForObject(UCustomizableObjec
 			return false;
 		}
 
-		NewData.ReadFileHandles = BulkData->GetAsyncReadFileHandles();
+		NewData.BulkFilePrefix = BulkData->GetBulkFilePrefix();
 #endif
 
-
-		if (NewData.ReadFileHandles.IsEmpty())
-		{
-			UE_LOG(LogMutable, Warning, TEXT("Streaming: Customizable Object %s read file handles empty."), *CustomizableObject->GetName());
-
-			check(false);
-			return false;
-		}
-
 		NewData.StreamableBlocks = CustomizableObject->HashToStreamableBlock;
 		if (NewData.StreamableBlocks.IsEmpty())
 		{
@@ -228,8 +214,8 @@ mu::ModelReader::OPERATION_ID FUnrealMutableModelBulkReader::BeginReadBlock(cons
 		return -1;
 	}
 
-	// this generally cannot fail because it is async
-	if (!ObjectData->StreamableBlocks.Contains(Key))
+	const FMutableStreamableBlock* Block = ObjectData->StreamableBlocks.Find(Key);
+	if (!Block)
 	{
 		// File Handle not found! This shouldn't really happen.
 		UE_LOG(LogMutable, Error, TEXT("Streaming Block not found!"));
@@ -244,8 +230,6 @@ mu::ModelReader::OPERATION_ID FUnrealMutableModelBulkReader::BeginReadBlock(cons
 
 	OPERATION_ID Result = ++LastOperationID;
 
-	const FMutableStreamableBlock& Block = ObjectData->StreamableBlocks[Key];
-
 	int32 BulkDataOffsetInFile = 0;
 #if WITH_EDITOR
 	BulkDataOffsetInFile = sizeof(MutableCompiledDataStreamHeader);
@@ -260,10 +244,35 @@ mu::ModelReader::OPERATION_ID FUnrealMutableModelBulkReader::BeginReadBlock(cons
 			CompletionCallbackCapture(!bWasCancelled);			
 		});		
 	}
+
+	TSharedPtr<IAsyncReadFileHandle> FileHandle;
+	{
+		FScopeLock Lock(&FileHandlesCritical);
+
+		TSharedPtr<IAsyncReadFileHandle>& Found = ObjectData->ReadFileHandles.FindOrAdd( Block->FileId );
+		if (!Found)
+		{
+#if WITH_EDITOR
+			FString FilePath = ObjectData->BulkFilePrefix;
+#else
+			FString FilePath = FString::Printf(TEXT("%s-%08x.mut"), *ObjectData->BulkFilePrefix, Block->FileId);
+#endif
+
+			Found = MakeShareable(FPlatformFileManager::Get().GetPlatformFile().OpenAsyncRead(*FilePath));
+
+			if (!Found)
+			{
+				UE_LOG(LogMutable, Error, TEXT("Failed to create AsyncReadFileHandle. File Path [%s]."), *FilePath);
+				check(false);
+				return -1;
+			}
+		}
+
+		FileHandle = Found;
+	}
 	
-	check(!ObjectData->ReadFileHandles.IsEmpty());
-	ReadRequest.ReadRequest = MakeShareable(ObjectData->ReadFileHandles[Block.FileIndex]->ReadRequest(
-		BulkDataOffsetInFile + Block.Offset,
+	ReadRequest.ReadRequest = MakeShareable(FileHandle->ReadRequest(
+		BulkDataOffsetInFile + Block->Offset,
 		size,
 		(EAsyncIOPriorityAndFlags)StreamPriority,
 		ReadRequest.FileCallback.Get(),
diff --git a/Engine/Plugins/Experimental/Mutable/Source/CustomizableObject/Private/MuCO/UnrealToMutableTextureConversionUtils.cpp b/Engine/Plugins/Experimental/Mutable/Source/CustomizableObject/Private/MuCO/UnrealToMutableTextureConversionUtils.cpp
--- a/Engine/Plugins/Experimental/Mutable/Source/CustomizableObject/Private/MuCO/UnrealToMutableTextureConversionUtils.cpp
+++ b/Engine/Plugins/Experimental/Mutable/Source/CustomizableObject/Private/MuCO/UnrealToMutableTextureConversionUtils.cpp
@@ -277,9 +277,9 @@ EUnrealToMutableConversionError ConvertTextureUnrealSourceToMutable(mu::Image* O
 			OutResult->Init(SizeX, SizeY, LODs, mu::EImageFormat::IF_RGBA_UBYTE, mu::EInitializationType::NotInitialized);
 			uint8* DataDest = OutResult->GetData();
 
-			// Convert to RGBA8 in place
+			// Convert to RGBA8 while copying
 			TArrayView64<FColor> ImageDataView = TempImage.AsBGRA8();
-			ParallelFor(ImageDataView.Num(),
+			ParallelFor(TEXT("MutableToRGBA"), ImageDataView.Num(), 16*1024,
 				[DataDest, &ImageDataView](uint32 p)
 				{
 					DataDest[4 * p + 0] = ImageDataView[p].R;
@@ -296,9 +296,9 @@ EUnrealToMutableConversionError ConvertTextureUnrealSourceToMutable(mu::Image* O
 			OutResult->Init(SizeX, SizeY, LODs, mu::EImageFormat::IF_RGB_UBYTE, mu::EInitializationType::NotInitialized);
 			uint8* DataDest = OutResult->GetData();
 
-			// Convert to RGB8 in place
+			// Convert to RGB8 while copying
 			TArrayView64<FColor> ImageDataView = TempImage.AsBGRA8();
-			ParallelFor(ImageDataView.Num(),
+			ParallelFor(TEXT("MutableToRGB"), ImageDataView.Num(), 16 * 1024,
 				[DataDest, &ImageDataView](uint32 p)
 				{
 					DataDest[3 * p + 0] = ImageDataView[p].R;
diff --git a/Engine/Plugins/Experimental/Mutable/Source/CustomizableObjectEditor/Private/MuCOE/CustomizableObjectCompileRunnable.cpp b/Engine/Plugins/Experimental/Mutable/Source/CustomizableObjectEditor/Private/MuCOE/CustomizableObjectCompileRunnable.cpp
--- a/Engine/Plugins/Experimental/Mutable/Source/CustomizableObjectEditor/Private/MuCOE/CustomizableObjectCompileRunnable.cpp
+++ b/Engine/Plugins/Experimental/Mutable/Source/CustomizableObjectEditor/Private/MuCOE/CustomizableObjectCompileRunnable.cpp
@@ -168,14 +168,10 @@ uint32 FCustomizableObjectCompileRunnable::Run()
 				return true;
 			}));
 
-	// Minimum resident mip count.
 	const int32 MinResidentMips = UTexture::GetStaticMinTextureResidentMipCount();
-	// Data smaller than this will always be loaded, as part of the customizable object compiled model.
-	const int32 MinRomSizeBytes = 128;
-	CompilerOptions->SetDataPackingStrategy(MinRomSizeBytes, MinResidentMips);
+	CompilerOptions->SetDataPackingStrategy( MinResidentMips, Options.EmbeddedDataBytesLimit, Options.PackagedDataBytesLimit );
 
-	// At object compilation time we don't know if we will want progressive images or not. Assume we will. 
-	// TODO: Per-state setting?
+	// We always compile for progressive image generation.
 	CompilerOptions->SetEnableProgressiveImages(true);
 	
 	CompilerOptions->SetImageTiling(Options.ImageTiling);
diff --git a/Engine/Plugins/Experimental/Mutable/Source/CustomizableObjectEditor/Private/MuCOE/CustomizableObjectCompiler.cpp b/Engine/Plugins/Experimental/Mutable/Source/CustomizableObjectEditor/Private/MuCOE/CustomizableObjectCompiler.cpp
--- a/Engine/Plugins/Experimental/Mutable/Source/CustomizableObjectEditor/Private/MuCOE/CustomizableObjectCompiler.cpp
+++ b/Engine/Plugins/Experimental/Mutable/Source/CustomizableObjectEditor/Private/MuCOE/CustomizableObjectCompiler.cpp
@@ -1018,7 +1018,7 @@ void FCustomizableObjectCompiler::CompileInternal(UCustomizableObject* Object, c
 			Object->GetPrivate()->bDisableTextureStreaming = false;
 		}
 		
-		Object->GetPrivate()->bIsCompiledWithOptimization = GenerationContext.Options.OptimizationLevel < UE_MUTABLE_MAX_OPTIMIZATION;
+		Object->GetPrivate()->bIsCompiledWithoutOptimization = GenerationContext.Options.OptimizationLevel < UE_MUTABLE_MAX_OPTIMIZATION;
 
 		Object->AlwaysLoadedExtensionData = MoveTemp(GenerationContext.AlwaysLoadedExtensionData);
 
@@ -1197,6 +1197,7 @@ void FCustomizableObjectCompiler::FinishCompilation()
 	TSharedPtr<mu::Model, ESPMode::ThreadSafe> Model = CompileTask->Model;
 
 	// Generate a map that using the resource id tells the offset and size of the resource inside the bulk data
+	// At this point it is assumed that all data goes into a single file.
 	if (Model)
 	{
 		uint64 Offset = 0;
@@ -1204,12 +1205,12 @@ void FCustomizableObjectCompiler::FinishCompilation()
 		const int32 NumStreamingFiles = Model->GetRomCount();
 		CurrentObject->HashToStreamableBlock.Empty(NumStreamingFiles);
 
-		for (size_t FileIndex = 0; FileIndex < NumStreamingFiles; ++FileIndex)
+		for (int32 FileIndex = 0; FileIndex < NumStreamingFiles; ++FileIndex)
 		{
 			const uint32 ResourceId = Model->GetRomId(FileIndex);
 			const uint32 ResourceSize = Model->GetRomSize(FileIndex);
 
-			CurrentObject->HashToStreamableBlock.Add(ResourceId, FMutableStreamableBlock({0, Offset, ResourceSize }));
+			CurrentObject->HashToStreamableBlock.Add(ResourceId, FMutableStreamableBlock{0, ResourceSize, Offset });
 			Offset += ResourceSize;
 		}
 	}
diff --git a/Engine/Plugins/Experimental/Mutable/Source/CustomizableObjectEditor/Private/MuCOE/CustomizableObjectEditor.cpp b/Engine/Plugins/Experimental/Mutable/Source/CustomizableObjectEditor/Private/MuCOE/CustomizableObjectEditor.cpp
--- a/Engine/Plugins/Experimental/Mutable/Source/CustomizableObjectEditor/Private/MuCOE/CustomizableObjectEditor.cpp
+++ b/Engine/Plugins/Experimental/Mutable/Source/CustomizableObjectEditor/Private/MuCOE/CustomizableObjectEditor.cpp
@@ -1388,8 +1388,64 @@ TSharedRef<SWidget> FCustomizableObjectEditor::GenerateCompileOptionsMenuContent
 	}
 	MenuBuilder.EndSection();
 
+	MenuBuilder.BeginSection("Packaging", LOCTEXT("MutableCompilePackagingHeading", "Packaging"));
+	{
+		// Unfortunately SNumericDropDown doesn't work with integers at the time of writing.
+		TArray<SNumericDropDown<float>::FNamedValue> EmbeddedOptions;
+		EmbeddedOptions.Add(SNumericDropDown<float>::FNamedValue(0, FText::FromString(TEXT("0")), FText::FromString(TEXT("Disabled"))));
+		EmbeddedOptions.Add(SNumericDropDown<float>::FNamedValue(16, FText::FromString(TEXT("16")), FText::FromString(TEXT("16"))));
+		EmbeddedOptions.Add(SNumericDropDown<float>::FNamedValue(64, FText::FromString(TEXT("64")), FText::FromString(TEXT("64"))));
+		EmbeddedOptions.Add(SNumericDropDown<float>::FNamedValue(256, FText::FromString(TEXT("256")), FText::FromString(TEXT("256"))));
+		EmbeddedOptions.Add(SNumericDropDown<float>::FNamedValue(512, FText::FromString(TEXT("512")), FText::FromString(TEXT("512"))));
+		EmbeddedOptions.Add(SNumericDropDown<float>::FNamedValue(1024, FText::FromString(TEXT("1024")), FText::FromString(TEXT("1024"))));
+		EmbeddedOptions.Add(SNumericDropDown<float>::FNamedValue(4096, FText::FromString(TEXT("4096")), FText::FromString(TEXT("4096"))));
+
+		EmbeddedDataLimitCombo = SNew(SNumericDropDown<float>)
+			.DropDownValues(EmbeddedOptions)
+			.Value_Lambda([this]()
+				{
+					return CustomizableObject ? float(CustomizableObject->CompileOptions.EmbeddedDataBytesLimit) : 0.0f;
+				})
+			.OnValueChanged_Lambda([this](float Value)
+				{
+					if (CustomizableObject)
+					{
+						CustomizableObject->CompileOptions.EmbeddedDataBytesLimit = uint64(Value);
+						CustomizableObject->Modify();
+					}
+				});
+			MenuBuilder.AddWidget(EmbeddedDataLimitCombo.ToSharedRef(), LOCTEXT("MutableCompileEmbeddedLimit", "Embedded Data Limit (Bytes)"));
+
+		// Packaging file size control.
+		TArray<SNumericDropDown<float>::FNamedValue> PackagedOptions;
+		PackagedOptions.Add(SNumericDropDown<float>::FNamedValue(0, FText::FromString(TEXT("0")), FText::FromString(TEXT("Split All"))));
+		PackagedOptions.Add(SNumericDropDown<float>::FNamedValue(16   * 1024, FText::FromString(TEXT("16 KB")), FText::FromString(TEXT("16 KB"))));
+		PackagedOptions.Add(SNumericDropDown<float>::FNamedValue(64   * 1024, FText::FromString(TEXT("64 KB")), FText::FromString(TEXT("64 KB"))));
+		PackagedOptions.Add(SNumericDropDown<float>::FNamedValue(1024 * 1024, FText::FromString(TEXT("1 MB")), FText::FromString(TEXT("1 MB"))));
+		PackagedOptions.Add(SNumericDropDown<float>::FNamedValue(64   * 1024 * 1024, FText::FromString(TEXT("64 MB")), FText::FromString(TEXT("64 MB"))));
+		PackagedOptions.Add(SNumericDropDown<float>::FNamedValue(256  * 1024 * 1024, FText::FromString(TEXT("256 MB")), FText::FromString(TEXT("256 MB"))));
+		PackagedOptions.Add(SNumericDropDown<float>::FNamedValue(1024 * 1024 * 1024, FText::FromString(TEXT("1 GB")), FText::FromString(TEXT("1 GB"))));
+
+		PackagedDataLimitCombo = SNew(SNumericDropDown<float>)
+			.DropDownValues(PackagedOptions)
+			.Value_Lambda([this]()
+				{
+					return CustomizableObject ? float(CustomizableObject->CompileOptions.PackagedDataBytesLimit) : 0.0f;
+				})
+			.OnValueChanged_Lambda([this](float Value)
+				{
+					if (CustomizableObject)
+					{
+						CustomizableObject->CompileOptions.PackagedDataBytesLimit = uint64(Value);
+						CustomizableObject->Modify();
+					}
+				});
+			MenuBuilder.AddWidget(PackagedDataLimitCombo.ToSharedRef(), LOCTEXT("MutableCompilePackagedLimit", "Packaged Data File Max Limit (Bytes)"));
+	}
+	MenuBuilder.EndSection();
+
 	// Debugging options
-	MenuBuilder.BeginSection("Debugger");
+	MenuBuilder.BeginSection("Debugger", LOCTEXT("MutableDebugger", "Debugger"));
 	{
 		MenuBuilder.AddMenuEntry(FCustomizableObjectEditorCommands::Get().Debug);
 	}
diff --git a/Engine/Plugins/Experimental/Mutable/Source/CustomizableObjectEditor/Private/MuCOE/CustomizableObjectEditor.h b/Engine/Plugins/Experimental/Mutable/Source/CustomizableObjectEditor/Private/MuCOE/CustomizableObjectEditor.h
--- a/Engine/Plugins/Experimental/Mutable/Source/CustomizableObjectEditor/Private/MuCOE/CustomizableObjectEditor.h
+++ b/Engine/Plugins/Experimental/Mutable/Source/CustomizableObjectEditor/Private/MuCOE/CustomizableObjectEditor.h
@@ -316,6 +316,8 @@ private:
 	TSharedPtr<STextComboBox> CompileTextureCompressionCombo;
 	TArray< TSharedPtr<FString> > CompileTextureCompressionStrings;
 	TSharedPtr<SNumericDropDown<float>> CompileTilingCombo;
+	TSharedPtr<SNumericDropDown<float>> EmbeddedDataLimitCombo;
+	TSharedPtr<SNumericDropDown<float>> PackagedDataLimitCombo;
 
 	void CompileOptions_UseDiskCompilation_Toggled();
 	bool CompileOptions_UseDiskCompilation_IsChecked();
diff --git a/Engine/Plugins/Experimental/Mutable/Source/MutableTools/Private/MuT/AST.h b/Engine/Plugins/Experimental/Mutable/Source/MutableTools/Private/MuT/AST.h
--- a/Engine/Plugins/Experimental/Mutable/Source/MutableTools/Private/MuT/AST.h
+++ b/Engine/Plugins/Experimental/Mutable/Source/MutableTools/Private/MuT/AST.h
@@ -220,6 +220,7 @@ namespace mu
 		{
 		}
 
+		// TODO: Unused?
 		int32 MinTextureResidentMipCount = 0;
 
 		/** This flag controls the splitting of image data into mips to store separately. It is usually necessary to
diff --git a/Engine/Plugins/Experimental/Mutable/Source/MutableTools/Private/MuT/Compiler.cpp b/Engine/Plugins/Experimental/Mutable/Source/MutableTools/Private/MuT/Compiler.cpp
--- a/Engine/Plugins/Experimental/Mutable/Source/MutableTools/Private/MuT/Compiler.cpp
+++ b/Engine/Plugins/Experimental/Mutable/Source/MutableTools/Private/MuT/Compiler.cpp
@@ -145,9 +145,10 @@ namespace mu
 
 
 	//---------------------------------------------------------------------------------------------
-	void CompilerOptions::SetDataPackingStrategy(int32 MinRomSize, int32 MinTextureResidentMipCount)
+	void CompilerOptions::SetDataPackingStrategy(int32 MinTextureResidentMipCount, uint64 EmbeddedDataBytesLimit, uint64 PackagedDataBytesLimit)
 	{
-		m_pD->MinRomSize = MinRomSize;
+		m_pD->EmbeddedDataBytesLimit = EmbeddedDataBytesLimit;
+		m_pD->PackagedDataBytesLimit = PackagedDataBytesLimit;
 		m_pD->MinTextureResidentMipCount = MinTextureResidentMipCount;
 	}
 
@@ -431,8 +432,8 @@ namespace mu
         m_pD->m_pErrorLog = genErrorLog.get();
 
 		// Pack data
-		int32 MinimumBytesPerRom = 1024; // \TODO: compilation parameter
-		m_pD->GenerateRoms(pResult.Get(),MinimumBytesPerRom);
+		uint64 EmbeddedDataFileBytesLimit = m_pD->m_options->GetPrivate()->EmbeddedDataBytesLimit;
+		m_pD->GenerateRoms(pResult.Get(), EmbeddedDataFileBytesLimit);
 
 		UE_LOG(LogMutableCore, Verbose, TEXT("(int) %s : %ld"), TEXT("program size"), int64(program.m_opAddress.Num()));
 
@@ -478,7 +479,7 @@ namespace mu
 
 
 	//---------------------------------------------------------------------------------------------
-	void Compiler::Private::GenerateRoms(Model* p, int32 MinRomSize)
+	void Compiler::Private::GenerateRoms(Model* p, int32 EmbeddedDataBytesLimit)
 	{
 		LLM_SCOPE_BYNAME(TEXT("MutableRuntime"));
 		MUTABLE_CPUPROFILER_SCOPE(Mutable_GenerateRoms);
@@ -488,78 +489,74 @@ namespace mu
 		int32 NumRoms = 0;
 		int32 NumEmbedded = 0;
 
-		// Serialise roms in a separate files if possible
-		if (MinRomSize >= 0)
+		// Save images and unload from memory
+		for (int32 ResourceIndex = 0; ResourceIndex < program.m_constantImageLODs.Num(); ++ResourceIndex)
 		{
-			// Save images and unload from memory
-			for (int32 ResourceIndex = 0; ResourceIndex < program.m_constantImageLODs.Num(); ++ResourceIndex)
-			{
-				TPair<int32, mu::ImagePtrConst>& ResData = program.m_constantImageLODs[ResourceIndex];
+			TPair<int32, mu::ImagePtrConst>& ResData = program.m_constantImageLODs[ResourceIndex];
 
-				// This shouldn't have been serialised with rom support before.
-				check(ResData.Key < 0);
+			// This shouldn't have been serialised with rom support before.
+			check(ResData.Key < 0);
 
-				// Serialize to memory, to find out final size of this rom
-				OutputMemoryStream MemStream(1024 * 1024);
-				OutputArchive MemoryArch(&MemStream);
-				Image::Serialise(ResData.Value.get(), MemoryArch);
+			// Serialize to memory, to find out final size of this rom
+			OutputMemoryStream MemStream(1024 * 1024);
+			OutputArchive MemoryArch(&MemStream);
+			Image::Serialise(ResData.Value.get(), MemoryArch);
 
-				// If the resource uses less memory than the threshold, don't save it in a separate rom.
-				if (MemStream.GetBufferSize() < MinRomSize)
-				{
-					NumEmbedded++;
-					continue;
-				}
+			// If the resource uses less memory than the threshold, don't save it in a separate rom.
+			if (MemStream.GetBufferSize() <= EmbeddedDataBytesLimit)
+			{
+				NumEmbedded++;
+				continue;
+			}
 
-				NumRoms++;
+			NumRoms++;
 
-				FRomData RomData;
-				RomData.ResourceType = DT_IMAGE;
-				RomData.ResourceIndex = ResourceIndex;
-				RomData.Size = MemStream.GetBufferSize();
+			FRomData RomData;
+			RomData.ResourceType = DT_IMAGE;
+			RomData.ResourceIndex = ResourceIndex;
+			RomData.Size = MemStream.GetBufferSize();
 
-				// Ensure that the Id is unique
-				RomData.Id = CityHash32(static_cast<const char*>(MemStream.GetBuffer()), MemStream.GetBufferSize());
-				EnsureUniqueRomId(RomData.Id, program);
+			// Ensure that the Id is unique
+			RomData.Id = CityHash32(static_cast<const char*>(MemStream.GetBuffer()), MemStream.GetBufferSize());
+			EnsureUniqueRomId(RomData.Id, program);
 
-				int32 RomIndex = program.m_roms.Add(RomData);
-				ResData.Key = RomIndex;
-			}
+			int32 RomIndex = program.m_roms.Add(RomData);
+			ResData.Key = RomIndex;
+		}
 
-			// Save meshes and unload from memory
-			for (int32 ResourceIndex = 0; ResourceIndex < program.m_constantMeshes.Num(); ++ResourceIndex)
-			{
-				TPair<int32, mu::MeshPtrConst>& ResData = program.m_constantMeshes[ResourceIndex];
+		// Save meshes and unload from memory
+		for (int32 ResourceIndex = 0; ResourceIndex < program.m_constantMeshes.Num(); ++ResourceIndex)
+		{
+			TPair<int32, mu::MeshPtrConst>& ResData = program.m_constantMeshes[ResourceIndex];
 
-				// This shouldn't have been serialised with rom support before.
-				check(ResData.Key < 0);
+			// This shouldn't have been serialised with rom support before.
+			check(ResData.Key < 0);
 
-				// Serialize to memory, to find out final size of this rom
-				OutputMemoryStream MemStream(1024 * 1024);
-				OutputArchive MemoryArch(&MemStream);
-				Mesh::Serialise(ResData.Value.get(), MemoryArch);
+			// Serialize to memory, to find out final size of this rom
+			OutputMemoryStream MemStream(1024 * 1024);
+			OutputArchive MemoryArch(&MemStream);
+			Mesh::Serialise(ResData.Value.get(), MemoryArch);
 
-				// If the resource uses less memory than the threshold, don't save it in a separate rom.
-				if (MemStream.GetBufferSize() < MinRomSize)
-				{
-					NumEmbedded++;
-					continue;
-				}
+			// If the resource uses less memory than the threshold, don't save it in a separate rom.
+			if (MemStream.GetBufferSize() <= EmbeddedDataBytesLimit)
+			{
+				NumEmbedded++;
+				continue;
+			}
 
-				NumRoms++;
+			NumRoms++;
 
-				FRomData RomData;
-				RomData.ResourceType = DT_MESH;
-				RomData.ResourceIndex = ResourceIndex;
-				RomData.Size = MemStream.GetBufferSize();
+			FRomData RomData;
+			RomData.ResourceType = DT_MESH;
+			RomData.ResourceIndex = ResourceIndex;
+			RomData.Size = MemStream.GetBufferSize();
 
-				// Ensure that the Id is unique
-				RomData.Id = CityHash32(static_cast<const char*>(MemStream.GetBuffer()), MemStream.GetBufferSize());
-				EnsureUniqueRomId(RomData.Id, program);
+			// Ensure that the Id is unique
+			RomData.Id = CityHash32(static_cast<const char*>(MemStream.GetBuffer()), MemStream.GetBufferSize());
+			EnsureUniqueRomId(RomData.Id, program);
 
-				int32 RomIndex = program.m_roms.Add(RomData);
-				ResData.Key = RomIndex;
-			}
+			int32 RomIndex = program.m_roms.Add(RomData);
+			ResData.Key = RomIndex;
 		}
 
 		UE_LOG(LogMutableCore, Log, TEXT("Generated roms for model with %d embedded constants and %d in roms."), NumEmbedded, NumRoms);
diff --git a/Engine/Plugins/Experimental/Mutable/Source/MutableTools/Private/MuT/CompilerPrivate.h b/Engine/Plugins/Experimental/Mutable/Source/MutableTools/Private/MuT/CompilerPrivate.h
--- a/Engine/Plugins/Experimental/Mutable/Source/MutableTools/Private/MuT/CompilerPrivate.h
+++ b/Engine/Plugins/Experimental/Mutable/Source/MutableTools/Private/MuT/CompilerPrivate.h
@@ -22,7 +22,10 @@ namespace mu
         //! Detailed optimization options
         FModelOptimizationOptions OptimisationOptions;
 
-		int32 MinRomSize = 0;
+		uint64 EmbeddedDataBytesLimit = 1024;
+		uint64 PackagedDataBytesLimit = 1024*1024*64;
+
+		// \TODO: Unused?
 		int32 MinTextureResidentMipCount = 3;
 
         int32 ImageCompressionQuality = 0;
diff --git a/Engine/Plugins/Experimental/Mutable/Source/MutableTools/Public/MuT/Compiler.h b/Engine/Plugins/Experimental/Mutable/Source/MutableTools/Public/MuT/Compiler.h
--- a/Engine/Plugins/Experimental/Mutable/Source/MutableTools/Public/MuT/Compiler.h
+++ b/Engine/Plugins/Experimental/Mutable/Source/MutableTools/Public/MuT/Compiler.h
@@ -74,7 +74,7 @@ namespace mu
 		void SetImageTiling(int32 Tiling);
 
         //! 
-        void SetDataPackingStrategy( int32 minRomSize, int32 MinTextureResidentMipCount );
+        void SetDataPackingStrategy( int32 MinTextureResidentMipCount, uint64 EmbeddedDataBytesLimit, uint64 PackagedDataBytesLimit);
 
 		/** If enabled it will make sure that the object is compile to generate smaller mips of the images. */
 		void SetEnableProgressiveImages(bool bEnabled);
