[mutable] Fix bug with composition of images with different pixel formats.
#rnx

diff --git a/Engine/Plugins/Experimental/Mutable/Source/MutableRuntime/Private/MuR/CodeRunner.cpp b/Engine/Plugins/Experimental/Mutable/Source/MutableRuntime/Private/MuR/CodeRunner.cpp
--- a/Engine/Plugins/Experimental/Mutable/Source/MutableRuntime/Private/MuR/CodeRunner.cpp
+++ b/Engine/Plugins/Experimental/Mutable/Source/MutableRuntime/Private/MuR/CodeRunner.cpp
@@ -4381,6 +4381,7 @@ namespace mu
 							ImOp.ImagePixelFormat(bSuccess, m_pSettings->ImageCompressionQuality, Formatted.get(), pB.get());
 							check(bSuccess);
 							Release(pB);
+							pB = Formatted;
 						}
 					}
 
