[Mutable] Added level lines to Debuger STreeView.

- Removed unused SImage.

#rnx

diff --git a/Engine/Plugins/Experimental/Mutable/Source/CustomizableObjectEditor/Private/MuCOE/SMutableCodeViewer.cpp b/Engine/Plugins/Experimental/Mutable/Source/CustomizableObjectEditor/Private/MuCOE/SMutableCodeViewer.cpp
--- a/Engine/Plugins/Experimental/Mutable/Source/CustomizableObjectEditor/Private/MuCOE/SMutableCodeViewer.cpp
+++ b/Engine/Plugins/Experimental/Mutable/Source/CustomizableObjectEditor/Private/MuCOE/SMutableCodeViewer.cpp
@@ -137,9 +137,6 @@ public:
 		// Primary column showing the name of the operation and tye type
 		if (ColumnName == MutableCodeTreeViewColumns::OperationsColumnID)
 		{
-			// TODO:
-			const FSlateBrush* IconBrush = nullptr;
-			
 			// Prepare a ui container for all the UI objects required by this row element
 			TSharedRef<SHorizontalBox> RowContainer = SNew(SHorizontalBox)
 				
@@ -159,26 +156,12 @@ public:
 				[
 					SNew(SHorizontalBox)
 					+ SHorizontalBox::Slot()
-					.VAlign(VAlign_Center)
 					.AutoWidth()
 					[
 						SNew(SExpanderArrow, SharedThis(this))
+						.ShouldDrawWires(true)
 					]
-
-					+ SHorizontalBox::Slot()
-					.VAlign(VAlign_Center)
-					.Padding(FMargin(5.f, 0.f, 5.f, 0.f))
-					.AutoWidth()
-					[
-						SNew(SOverlay)
-						+ SOverlay::Slot()
-						[
-							SNew(SImage)
-							.Image(IconBrush ? IconBrush : FCoreStyle::Get().GetDefaultBrush())
-							.ColorAndOpacity(IconBrush ? FLinearColor::White : FLinearColor::Transparent)
-						]
-					]
-
+					
 					+ SHorizontalBox::Slot()
 					[
 						SNew(STextBlock)
diff --git a/Engine/Plugins/Experimental/Mutable/Source/CustomizableObjectEditor/Private/MuCOE/SMutableGraphViewer.cpp b/Engine/Plugins/Experimental/Mutable/Source/CustomizableObjectEditor/Private/MuCOE/SMutableGraphViewer.cpp
--- a/Engine/Plugins/Experimental/Mutable/Source/CustomizableObjectEditor/Private/MuCOE/SMutableGraphViewer.cpp
+++ b/Engine/Plugins/Experimental/Mutable/Source/CustomizableObjectEditor/Private/MuCOE/SMutableGraphViewer.cpp
@@ -72,9 +72,6 @@ public:
 			MainLabel = FText::FromString( FString::Printf(TEXT("%s (Duplicated)"), StringCast<TCHAR>(TypeName).Get()));
 		}
 
-		// TODO
-		const FSlateBrush* IconBrush = nullptr; 
-
 		this->ChildSlot
 		[
 			SNew(SHorizontalBox)
@@ -83,20 +80,7 @@ public:
 			.AutoWidth()
 			[
 				SNew(SExpanderArrow, SharedThis(this))
-			]
-
-			+ SHorizontalBox::Slot()
-			.VAlign(VAlign_Center)
-			.Padding(FMargin(5.f, 0.f, 5.f, 0.f))
-			.AutoWidth()
-			[
-				SNew(SOverlay)
-				+ SOverlay::Slot()
-				[
-					SNew(SImage)
-					.Image(IconBrush ? IconBrush : FCoreStyle::Get().GetDefaultBrush())
-					.ColorAndOpacity(IconBrush ? FLinearColor::White : FLinearColor::Transparent)
-				]
+				.ShouldDrawWires(true)
 			]
 
 			+ SHorizontalBox::Slot()
