[mutable] Ensure header is included even in non-editor builds to fix compilation error.
#rnx

diff --git a/Engine/Plugins/Experimental/Mutable/Source/CustomizableObject/Private/MuCO/UnrealMutableModelDiskStreamer.cpp b/Engine/Plugins/Experimental/Mutable/Source/CustomizableObject/Private/MuCO/UnrealMutableModelDiskStreamer.cpp
--- a/Engine/Plugins/Experimental/Mutable/Source/CustomizableObject/Private/MuCO/UnrealMutableModelDiskStreamer.cpp
+++ b/Engine/Plugins/Experimental/Mutable/Source/CustomizableObject/Private/MuCO/UnrealMutableModelDiskStreamer.cpp
@@ -9,10 +9,8 @@
 #include "MuR/Model.h"
 #include "MuR/MutableTrace.h"
 #include "MuCO/LogBenchmarkUtil.h"
-
-#if WITH_EDITOR
 #include "HAL/PlatformFileManager.h"
-#endif
+
 
 DECLARE_DWORD_ACCUMULATOR_STAT(TEXT("Streaming Ops"), STAT_MutableStreamingOps, STATGROUP_Mutable);
 
