[Mutable] Fix missing ImageProperties for images coming from table nodes if the reuse of surfaces between LODs is disabled.
Simplify code to find the source material of a UCustomizableObjectNodeCopyMaterial.

#rnx

diff --git a/Engine/Plugins/Experimental/Mutable/Source/CustomizableObjectEditor/Private/MuCOE/CustomizableObjectCompiler.cpp b/Engine/Plugins/Experimental/Mutable/Source/CustomizableObjectEditor/Private/MuCOE/CustomizableObjectCompiler.cpp
--- a/Engine/Plugins/Experimental/Mutable/Source/CustomizableObjectEditor/Private/MuCOE/CustomizableObjectCompiler.cpp
+++ b/Engine/Plugins/Experimental/Mutable/Source/CustomizableObjectEditor/Private/MuCOE/CustomizableObjectCompiler.cpp
@@ -938,6 +938,10 @@ void FCustomizableObjectCompiler::CompileInternal(UCustomizableObject* Object, c
 
 		TArray<FGeneratedImageProperties> ImageProperties;
 		GenerationContext.ImageProperties.GenerateValueArray(ImageProperties);
+		
+		// Must sort image properties by ImagePropertiesIndex so that ImageNames point to the right properties.
+		ImageProperties.Sort([](const FGeneratedImageProperties& PropsA, const FGeneratedImageProperties& PropsB)
+			{ return PropsA.ImagePropertiesIndex < PropsB.ImagePropertiesIndex;	});
 
 		ModelResources.ImageProperties.Empty(ImageProperties.Num());
 
diff --git a/Engine/Plugins/Experimental/Mutable/Source/CustomizableObjectEditor/Private/MuCOE/GenerateMutableSource/GenerateMutableSource.cpp b/Engine/Plugins/Experimental/Mutable/Source/CustomizableObjectEditor/Private/MuCOE/GenerateMutableSource/GenerateMutableSource.cpp
--- a/Engine/Plugins/Experimental/Mutable/Source/CustomizableObjectEditor/Private/MuCOE/GenerateMutableSource/GenerateMutableSource.cpp
+++ b/Engine/Plugins/Experimental/Mutable/Source/CustomizableObjectEditor/Private/MuCOE/GenerateMutableSource/GenerateMutableSource.cpp
@@ -796,6 +796,23 @@ mu::NodeObjectPtr GenerateMutableSource(const UEdGraphPin * Pin, FMutableGraphGe
 			// Find the MinLOD available for the target platform
 			if (RefSkeletalMesh->IsMinLodQualityLevelEnable())
 			{
+				FSupportedQualityLevelArray SupportedQualityLevels = LODSettings.MinQualityLevelLOD.GetSupportedQualityLevels(*GenerationContext.Options.TargetPlatform->GetPlatformInfo().IniPlatformName.ToString());
+				
+				int32 MinValue = MAX_int32;
+				for (int32& QL : SupportedQualityLevels)
+				{
+					// check if have data for the supported quality level or set to default.
+					if (LODSettings.MinQualityLevelLOD.IsQualityLevelValid(QL))
+					{
+						MinValue = FMath::Min(LODSettings.MinQualityLevelLOD.GetValueForQualityLevel(QL), MinValue);
+					}
+					else 
+					{
+						MinValue = LODSettings.MinQualityLevelLOD.GetDefault();
+						break;
+					}
+				}
+
 				GenerationContext.FirstLODAvailable = LODSettings.MinQualityLevelLOD.GetValueForPlatform(GenerationContext.Options.TargetPlatform);
 			}
 			else
@@ -986,8 +1003,7 @@ mu::NodeObjectPtr GenerateMutableSource(const UEdGraphPin * Pin, FMutableGraphGe
 							continue;
 						}
 
-						FMutableGraphSurfaceGenerationData DummySurfaceData;
-						mu::NodeSurfacePtr SurfaceNode = GenerateMutableSourceSurface(ChildNodePin, GenerationContext, DummySurfaceData);
+						mu::NodeSurfacePtr SurfaceNode = GenerateMutableSourceSurface(ChildNodePin, GenerationContext);
 
 						mu::NodeComponentPtr& ComponentNode = ComponentNodes[MeshComponentIndex];
 
diff --git a/Engine/Plugins/Experimental/Mutable/Source/CustomizableObjectEditor/Private/MuCOE/GenerateMutableSource/GenerateMutableSource.h b/Engine/Plugins/Experimental/Mutable/Source/CustomizableObjectEditor/Private/MuCOE/GenerateMutableSource/GenerateMutableSource.h
--- a/Engine/Plugins/Experimental/Mutable/Source/CustomizableObjectEditor/Private/MuCOE/GenerateMutableSource/GenerateMutableSource.h
+++ b/Engine/Plugins/Experimental/Mutable/Source/CustomizableObjectEditor/Private/MuCOE/GenerateMutableSource/GenerateMutableSource.h
@@ -53,6 +53,9 @@ struct FGeneratedImageProperties
 	/** Name in the Material. */
 	FString TextureParameterName;
 
+	/** Name in the mu::Surface. */
+	int32 ImagePropertiesIndex;
+
 	TEnumAsByte<TextureCompressionSettings> CompressionSettings = TC_Default;
 
 	TEnumAsByte<TextureFilter> Filter = TF_Bilinear;
@@ -185,17 +188,6 @@ struct FGeneratedImagePropertiesKey
 	uint32 ImageIndex = 0;
 };
 
-/** Structure used to store the results of the recursive GenerateMutableSourceSurface calls. */
-struct FMutableGraphSurfaceGenerationData
-{
-	/** Pointer to the NodeMaterial which a pin is connected (directly or indirectly). Used to find the real NodeMaterial on the NodeCopyMaterial case. 
-	 * The NodeMaterial may be behind a export/import nodes, or already generated
-	 */
-	const UCustomizableObjectNodeMaterial* NodeMaterial = nullptr;
-	
-	FGeneratedImageProperties ImageProperties;
-};
-
 // Structure storing results to propagate up when generating mutable mesh node expressions.
 struct FMutableGraphMeshGenerationData
 {
diff --git a/Engine/Plugins/Experimental/Mutable/Source/CustomizableObjectEditor/Private/MuCOE/GenerateMutableSource/GenerateMutableSourceSurface.cpp b/Engine/Plugins/Experimental/Mutable/Source/CustomizableObjectEditor/Private/MuCOE/GenerateMutableSource/GenerateMutableSourceSurface.cpp
--- a/Engine/Plugins/Experimental/Mutable/Source/CustomizableObjectEditor/Private/MuCOE/GenerateMutableSource/GenerateMutableSourceSurface.cpp
+++ b/Engine/Plugins/Experimental/Mutable/Source/CustomizableObjectEditor/Private/MuCOE/GenerateMutableSource/GenerateMutableSourceSurface.cpp
@@ -214,7 +214,7 @@ void AddModifierToSharedSurface(FMutableGraphGenerationContext& GenerationContex
 }
 
 
-mu::NodeSurfacePtr GenerateMutableSourceSurface(const UEdGraphPin * Pin, FMutableGraphGenerationContext & GenerationContext, FMutableGraphSurfaceGenerationData& SurfaceData)
+mu::NodeSurfacePtr GenerateMutableSourceSurface(const UEdGraphPin * Pin, FMutableGraphGenerationContext & GenerationContext)
 {
 	check(Pin)
 	RETURN_ON_CYCLE(*Pin, GenerationContext)
@@ -226,7 +226,6 @@ mu::NodeSurfacePtr GenerateMutableSourceSurface(const UEdGraphPin * Pin, FMutabl
 	const FGeneratedKey Key(reinterpret_cast<void*>(&GenerateMutableSourceSurface), *Pin, *Node, GenerationContext, true);
 	if (const FGeneratedData* Generated = GenerationContext.Generated.Find(Key))
 	{
-		SurfaceData.NodeMaterial = Cast<UCustomizableObjectNodeMaterial>(Generated->Source);
 		return static_cast<mu::NodeSurface*>(Generated->Node.get());
 	}
 	
@@ -249,8 +248,6 @@ mu::NodeSurfacePtr GenerateMutableSourceSurface(const UEdGraphPin * Pin, FMutabl
 			return Result;
 		}
 
-		SurfaceData.NodeMaterial = TypedNodeMat; // Save the NodeMaterial for the calling recursive calls
-
 		// NodeCopyMaterial. Special case when the TypedNodeMat is a NodeCopyMaterial. The TypedNodeMat pointer now points to the parent NodeMaterial except when reading the mesh pin, which comes from the NodeCopyMaterial.
 		UCustomizableObjectNodeMaterial* TypedNodeMaterial = TypedNodeMat;
 		UCustomizableObjectNodeMaterial* TypedNodeCopyMaterial = TypedNodeMat;
@@ -260,17 +257,12 @@ mu::NodeSurfacePtr GenerateMutableSourceSurface(const UEdGraphPin * Pin, FMutabl
 			UEdGraphPin* MaterialPin = TypedDerivedNodeCopyMaterial->GetMaterialPin();
 			if (const UEdGraphPin* ConnectedPin = FollowInputPin(*MaterialPin))
 			{
-				FMutableGraphSurfaceGenerationData ParentSurfaceData;
-				GenerateMutableSourceSurface(ConnectedPin, GenerationContext, ParentSurfaceData);
-
-				if (ParentSurfaceData.NodeMaterial)
+				if (UCustomizableObjectNodeMaterial* ConnectedNodeMaterial = Cast<UCustomizableObjectNodeMaterial>(ConnectedPin->GetOwningNode()))
 				{
-					UEdGraphNode* NodeMaterial = const_cast<UCustomizableObjectNodeMaterial*>(ParentSurfaceData.NodeMaterial);
-
-					if (NodeMaterial->IsA(UCustomizableObjectNodeMaterial::StaticClass()) && !NodeMaterial->IsA(UCustomizableObjectNodeCopyMaterial::StaticClass()))
+					if (!ConnectedNodeMaterial->IsA(UCustomizableObjectNodeCopyMaterial::StaticClass()))
 					{
-						TypedNodeMaterial = static_cast<UCustomizableObjectNodeMaterial*>(NodeMaterial);
-						TypedNodeMat = TypedNodeMaterial;
+						TypedNodeMaterial = ConnectedNodeMaterial;
+						TypedNodeMat = ConnectedNodeMaterial;
 					}
 					else
 					{
@@ -441,24 +433,28 @@ mu::NodeSurfacePtr GenerateMutableSourceSurface(const UEdGraphPin * Pin, FMutabl
 			{
 				if (const UEdGraphPin* ConnectedPin = FollowInputPin(*ImagePin))
 				{
+					// Find or add Image properties
+					const FGeneratedImagePropertiesKey PropsKey(TypedNodeMat, (uint32)ImageIndex);
+					const bool bNewImageProps = !GenerationContext.ImageProperties.Contains(PropsKey);
+
+					FGeneratedImageProperties& Props = GenerationContext.ImageProperties.FindOrAdd(PropsKey);
+					if (bNewImageProps)
+					{
+						// We don't need a reference texture or props here, but we do need the parameter name.
+						Props.TextureParameterName = TypedNodeMat->GetParameterName(EMaterialParameterType::Texture, ImageIndex).ToString();
+						Props.ImagePropertiesIndex = GenerationContext.ImageProperties.Num() - 1;
+						Props.bIsPassThrough = true;
+					}
+
 					// This is a connected pass-through texture that simply has to be passed to the core
 					mu::Ptr<mu::NodeImage> PassThroughImagePtr = GenerateMutableSourceImage(ConnectedPin, GenerationContext, 0);
 					SurfNode->SetImage(ImageIndex, PassThroughImagePtr);
 
-					const FString ImageName = TypedNodeMat->GetParameterName(EMaterialParameterType::Texture, ImageIndex).ToString();
-					FString SurfNodeImageName = FString::Printf(TEXT("%d"), GenerationContext.ImageProperties.Num());
+					const FString SurfNodeImageName = FString::Printf(TEXT("%d"), Props.ImagePropertiesIndex);
 					SurfNode->SetImageName(ImageIndex, SurfNodeImageName);
-
 					SurfNode->SetImageLayoutIndex(ImageIndex, -1);
-					SurfNode->SetImageAdditionalNames(ImageIndex, TypedNodeMat->Material->GetName(), ImageName);
-
-					// We don't need a reference texture or props here, but we do need the parameter name.
-					FGeneratedImageProperties Props;
-					FGeneratedImagePropertiesKey PropsKey(TypedNodeMat,(uint32)ImageIndex);
-					Props.TextureParameterName = ImageName;
-					Props.bIsPassThrough = true;
-					GenerationContext.ImageProperties.Add(PropsKey, Props);
-					SurfaceData.ImageProperties = Props;
+					SurfNode->SetImageAdditionalNames(ImageIndex, TypedNodeMat->Material->GetName(), Props.TextureParameterName);
+
 				}
 			}
 			else
@@ -519,55 +515,60 @@ mu::NodeSurfacePtr GenerateMutableSourceSurface(const UEdGraphPin * Pin, FMutabl
 								}
 							}
 						}
+						
+						const FGeneratedImagePropertiesKey PropsKey(TypedNodeMat, ImageIndex);
+						const bool bNewImageProps = !GenerationContext.ImageProperties.Contains(PropsKey);
 
-						FGeneratedImageProperties Props;
-						FGeneratedImagePropertiesKey PropsKey(TypedNodeMat, ImageIndex);
-						if (ReferenceTexture)
+						FGeneratedImageProperties& Props = GenerationContext.ImageProperties.FindOrAdd(PropsKey);
+
+						if (bNewImageProps)
 						{
-							GenerationContext.AddParticipatingObject(*ReferenceTexture);
-							
-							// Store properties for the generated images
-							Props.TextureParameterName = ImageName;
-							Props.CompressionSettings = ReferenceTexture->CompressionSettings;
-							Props.Filter = ReferenceTexture->Filter;
-							Props.SRGB = ReferenceTexture->SRGB;
-							Props.LODBias = 0;
-							Props.MipGenSettings = ReferenceTexture->MipGenSettings;
-							Props.LODGroup = ReferenceTexture->LODGroup;
-							Props.AddressX = ReferenceTexture->AddressX;
-							Props.AddressY = ReferenceTexture->AddressY;
-							Props.bFlipGreenChannel = ReferenceTexture->bFlipGreenChannel;
-
-
-							// MaxTextureSize setting. Based on the ReferenceTexture and Platform settings.
-							const UTextureLODSettings& TextureLODSettings = GenerationContext.Options.TargetPlatform->GetTextureLODSettings();
-							Props.MaxTextureSize = GetMaxTextureSize(*ReferenceTexture, TextureLODSettings);
-							
-							// ReferenceTexture source size. Textures contributing to this Image should be equal to or smaller than TextureSize. 
-							// The LOD Bias applied to the root node will be applied on top of it.
-							Props.TextureSize = (int32)FMath::Max3(ReferenceTexture->Source.GetSizeX(), ReferenceTexture->Source.GetSizeY(), 1LL);
-
-							// TODO: MTBL-1081
-							// TextureGroup::TEXTUREGROUP_UI does not support streaming. If we generate a texture that requires streaming and set this group, it will crash when initializing the resource. 
-							// If LODGroup == TEXTUREGROUP_UI, UTexture::IsPossibleToStream() will return false and UE will assume all mips are loaded, when they're not, and crash.
-							if (Props.LODGroup == TEXTUREGROUP_UI)
+							if (ReferenceTexture)
 							{
-								Props.LODGroup = TextureGroup::TEXTUREGROUP_Character;
+								GenerationContext.AddParticipatingObject(*ReferenceTexture);
+
+								// Store properties for the generated images
+								Props.TextureParameterName = ImageName;
+								Props.ImagePropertiesIndex = GenerationContext.ImageProperties.Num() - 1;
+
+								Props.CompressionSettings = ReferenceTexture->CompressionSettings;
+								Props.Filter = ReferenceTexture->Filter;
+								Props.SRGB = ReferenceTexture->SRGB;
+								Props.LODBias = 0;
+								Props.MipGenSettings = ReferenceTexture->MipGenSettings;
+								Props.LODGroup = ReferenceTexture->LODGroup;
+								Props.AddressX = ReferenceTexture->AddressX;
+								Props.AddressY = ReferenceTexture->AddressY;
+								Props.bFlipGreenChannel = ReferenceTexture->bFlipGreenChannel;
+
+
+								// MaxTextureSize setting. Based on the ReferenceTexture and Platform settings.
+								const UTextureLODSettings& TextureLODSettings = GenerationContext.Options.TargetPlatform->GetTextureLODSettings();
+								Props.MaxTextureSize = GetMaxTextureSize(*ReferenceTexture, TextureLODSettings);
+
+								// ReferenceTexture source size. Textures contributing to this Image should be equal to or smaller than TextureSize. 
+								// The LOD Bias applied to the root node will be applied on top of it.
+								Props.TextureSize = (int32)FMath::Max3(ReferenceTexture->Source.GetSizeX(), ReferenceTexture->Source.GetSizeY(), 1LL);
+
+								// TODO: MTBL-1081
+								// TextureGroup::TEXTUREGROUP_UI does not support streaming. If we generate a texture that requires streaming and set this group, it will crash when initializing the resource. 
+								// If LODGroup == TEXTUREGROUP_UI, UTexture::IsPossibleToStream() will return false and UE will assume all mips are loaded, when they're not, and crash.
+								if (Props.LODGroup == TEXTUREGROUP_UI)
+								{
+									Props.LODGroup = TextureGroup::TEXTUREGROUP_Character;
 
-								FString msg = FString::Printf(TEXT("The Reference texture [%s] is using TEXTUREGROUP_UI which does not support streaming. Please set a different TEXTURE group."),
-									*ReferenceTexture->GetName(), *ImageName);
-								GenerationContext.Compiler->CompilerLog(FText::FromString(msg), Node, EMessageSeverity::Info);
+									FString msg = FString::Printf(TEXT("The Reference texture [%s] is using TEXTUREGROUP_UI which does not support streaming. Please set a different TEXTURE group."),
+										*ReferenceTexture->GetName(), *ImageName);
+									GenerationContext.Compiler->CompilerLog(FText::FromString(msg), Node, EMessageSeverity::Info);
+								}
+							}
+							else if (!GroupProjectionImg.get())
+							{
+								// warning!
+								FString msg = FString::Printf(TEXT("The Reference texture for material image [%s] is not set and it couldn't be found automatically."), *ImageName);
+								GenerationContext.Compiler->CompilerLog(FText::FromString(msg), Node);
 							}
 						}
-						else if (!GroupProjectionImg.get())
-						{
-							// warning!
-							FString msg = FString::Printf(TEXT("The Reference texture for material image [%s] is not set and it couldn't be found automatically."), *ImageName);
-							GenerationContext.Compiler->CompilerLog(FText::FromString(msg), Node);
-						}
-
-						GenerationContext.ImageProperties.Add(PropsKey, Props);
-						SurfaceData.ImageProperties = Props;
 
 						// Generate the texture nodes
 						mu::NodeImagePtr ImageNode = [&]()
@@ -781,17 +782,14 @@ mu::NodeSurfacePtr GenerateMutableSourceSurface(const UEdGraphPin * Pin, FMutabl
 
 						SurfNode->SetImage(ImageIndex, ImageNode);
 
-						FString SurfNodeImageName = FString::Printf(TEXT("%d"), GenerationContext.ImageProperties.Num() - 1);
+						const FString SurfNodeImageName = FString::Printf(TEXT("%d"), Props.ImagePropertiesIndex);
 
 						// Encoding material layer in mutable name
 						const int32 LayerIndex = TypedNodeMat->GetParameterLayerIndex(EMaterialParameterType::Texture, ImageIndex);
-						if (LayerIndex != -1)
-						{
-							SurfNodeImageName += "-MutableLayerParam:" + FString::FromInt(LayerIndex);
-						}
-
-						SurfNode->SetImageName(ImageIndex, SurfNodeImageName);
-						int32 UVLayout = TypedNodeMat->GetImageUVLayout(ImageIndex);
+						const FString LayerEncoding = LayerIndex != INDEX_NONE ? "-MutableLayerParam:" + FString::FromInt(LayerIndex) : "";
+						
+						SurfNode->SetImageName(ImageIndex, SurfNodeImageName + LayerEncoding);
+						const int32 UVLayout = TypedNodeMat->GetImageUVLayout(ImageIndex);
 						SurfNode->SetImageLayoutIndex(ImageIndex, UVLayout);
 						SurfNode->SetImageAdditionalNames(ImageIndex, TypedNodeMat->Material->GetName(), ImageName);
 
@@ -1048,8 +1046,7 @@ mu::NodeSurfacePtr GenerateMutableSourceSurface(const UEdGraphPin * Pin, FMutabl
 			}
 
 			// Parent, probably generated, will be retrieved from the cache
-			FMutableGraphSurfaceGenerationData DummySurfaceData;
-			mu::NodeSurfacePtr ParentNode = GenerateMutableSourceSurface(ParentMaterialNode->OutputPin(), GenerationContext, DummySurfaceData);
+			mu::NodeSurfacePtr ParentNode = GenerateMutableSourceSurface(ParentMaterialNode->OutputPin(), GenerationContext);
 			SurfNode->SetParent(ParentNode.get());
 
 			mu::NodeMeshPtr AddMeshNode;
@@ -1191,8 +1188,7 @@ mu::NodeSurfacePtr GenerateMutableSourceSurface(const UEdGraphPin * Pin, FMutabl
 		else
 		{
 			// Parent, probably generated, will be retrieved from the cache
-			FMutableGraphSurfaceGenerationData DummySurfaceData;
-			mu::NodeSurfacePtr ParentNode = GenerateMutableSourceSurface(ParentMaterialNode->OutputPin(), GenerationContext, DummySurfaceData);
+			mu::NodeSurfacePtr ParentNode = GenerateMutableSourceSurface(ParentMaterialNode->OutputPin(), GenerationContext);
 			SurfNode->SetParent(ParentNode.get());
 
 			mu::NodeMeshPtr RemoveMeshNode;
@@ -1229,8 +1225,7 @@ mu::NodeSurfacePtr GenerateMutableSourceSurface(const UEdGraphPin * Pin, FMutabl
 		else
 		{
 			// Parent, probably generated, will be retrieved from the cache
-			FMutableGraphSurfaceGenerationData DummySurfaceData;
-			mu::NodeSurfacePtr ParentNode = GenerateMutableSourceSurface(ParentMaterialNode->OutputPin(), GenerationContext, DummySurfaceData);
+			mu::NodeSurfacePtr ParentNode = GenerateMutableSourceSurface(ParentMaterialNode->OutputPin(), GenerationContext);
 			SurfNode->SetParent(ParentNode.get());
 
 			const UEdGraphPin* BaseSourcePin = FindMeshBaseSource(*ParentMaterialNode->OutputPin(), false);
@@ -1312,8 +1307,7 @@ mu::NodeSurfacePtr GenerateMutableSourceSurface(const UEdGraphPin * Pin, FMutabl
 		else
 		{
 			// Parent, probably generated, will be retrieved from the cache
-			FMutableGraphSurfaceGenerationData DummySurfaceData;
-			mu::NodeSurfacePtr ParentNode = GenerateMutableSourceSurface(ParentMaterialNode->OutputPin(), GenerationContext, DummySurfaceData);
+			mu::NodeSurfacePtr ParentNode = GenerateMutableSourceSurface(ParentMaterialNode->OutputPin(), GenerationContext);
 			SurfNode->SetParent(ParentNode.get());
 
 
@@ -1418,8 +1412,7 @@ mu::NodeSurfacePtr GenerateMutableSourceSurface(const UEdGraphPin * Pin, FMutabl
 		else
 		{
 			// Parent, probably generated, will be retrieved from the cache
-			FMutableGraphSurfaceGenerationData DummySurfaceData;
-			mu::NodeSurfacePtr ParentNode = GenerateMutableSourceSurface(ParentMaterialNode->OutputPin(), GenerationContext, DummySurfaceData);
+			mu::NodeSurfacePtr ParentNode = GenerateMutableSourceSurface(ParentMaterialNode->OutputPin(), GenerationContext);
 			SurfNode->SetParent(ParentNode.get());
 
 			const UEdGraphPin* BaseSourcePin = FindMeshBaseSource(*ParentMaterialNode->OutputPin(), false);
@@ -1506,8 +1499,7 @@ mu::NodeSurfacePtr GenerateMutableSourceSurface(const UEdGraphPin * Pin, FMutabl
 		for (const UEdGraphPin* ConnectedPin : FollowInputPinArray(*TypedNodeVar->DefaultPin()))
 		{
 			// Is it a modifier?
-			FMutableGraphSurfaceGenerationData DummySurfaceData;
-			mu::NodeSurfacePtr ChildNode = GenerateMutableSourceSurface(ConnectedPin, GenerationContext, DummySurfaceData);
+			mu::NodeSurfacePtr ChildNode = GenerateMutableSourceSurface(ConnectedPin, GenerationContext);
 			if (ChildNode)
 			{
 				SurfNode->AddDefaultSurface(ChildNode.get());
@@ -1530,8 +1522,7 @@ mu::NodeSurfacePtr GenerateMutableSourceSurface(const UEdGraphPin * Pin, FMutabl
 				for (const UEdGraphPin* ConnectedPin : FollowInputPinArray(*VariationPin))
 				{
 					// Is it a modifier?
-					FMutableGraphSurfaceGenerationData DummySurfaceData;
-					mu::NodeSurfacePtr ChildNode = GenerateMutableSourceSurface(ConnectedPin, GenerationContext, DummySurfaceData);
+					mu::NodeSurfacePtr ChildNode = GenerateMutableSourceSurface(ConnectedPin, GenerationContext);
 					if (ChildNode)
 					{
 						SurfNode->AddVariationSurface(VariationIndex, ChildNode.get());
@@ -1595,8 +1586,7 @@ mu::NodeSurfacePtr GenerateMutableSourceSurface(const UEdGraphPin * Pin, FMutabl
 				{
 					if (const UEdGraphPin* ConnectedPin = FollowInputPin(*TypedNodeSwitch->GetElementPin(SelectorIndex)))
 					{
-						FMutableGraphSurfaceGenerationData DummySurfaceData;
-						mu::NodeSurfacePtr ChildNode = GenerateMutableSourceSurface(ConnectedPin, GenerationContext, DummySurfaceData);
+						mu::NodeSurfacePtr ChildNode = GenerateMutableSourceSurface(ConnectedPin, GenerationContext);
 						if (ChildNode)
 						{
 							SwitchNode->SetOption(SelectorIndex, ChildNode.get());
diff --git a/Engine/Plugins/Experimental/Mutable/Source/CustomizableObjectEditor/Private/MuCOE/GenerateMutableSource/GenerateMutableSourceSurface.h b/Engine/Plugins/Experimental/Mutable/Source/CustomizableObjectEditor/Private/MuCOE/GenerateMutableSource/GenerateMutableSourceSurface.h
--- a/Engine/Plugins/Experimental/Mutable/Source/CustomizableObjectEditor/Private/MuCOE/GenerateMutableSource/GenerateMutableSourceSurface.h
+++ b/Engine/Plugins/Experimental/Mutable/Source/CustomizableObjectEditor/Private/MuCOE/GenerateMutableSource/GenerateMutableSourceSurface.h
@@ -6,6 +6,5 @@
 
 class UEdGraphPin;
 struct FMutableGraphGenerationContext;
-struct FMutableGraphSurfaceGenerationData;
 
-mu::NodeSurfacePtr GenerateMutableSourceSurface(const UEdGraphPin * Pin, FMutableGraphGenerationContext & GenerationContext, FMutableGraphSurfaceGenerationData& SurfaceData);
+mu::NodeSurfacePtr GenerateMutableSourceSurface(const UEdGraphPin * Pin, FMutableGraphGenerationContext & GenerationContext);
