[mutable] Added more default fallback conversions for pixel formats.
#rnx

diff --git a/Engine/Plugins/Experimental/Mutable/Source/MutableRuntime/Private/MuR/OpImagePixelFormat.cpp b/Engine/Plugins/Experimental/Mutable/Source/MutableRuntime/Private/MuR/OpImagePixelFormat.cpp
--- a/Engine/Plugins/Experimental/Mutable/Source/MutableRuntime/Private/MuR/OpImagePixelFormat.cpp
+++ b/Engine/Plugins/Experimental/Mutable/Source/MutableRuntime/Private/MuR/OpImagePixelFormat.cpp
@@ -172,8 +172,7 @@ void FImageOperator::ImagePixelFormat( bool& bOutSuccess, int32 CompressionQuali
 
             case EImageFormat::IF_RGB_UBYTE_RLE:
             {
-                ImagePtr pTempBase =
-                    ImagePixelFormat( CompressionQuality, Base, EImageFormat::IF_RGB_UBYTE );
+                ImagePtr pTempBase = ImagePixelFormat( CompressionQuality, Base, EImageFormat::IF_RGB_UBYTE );
                 if (OnlyLOD==-1)
                 {
                     pBaseBuf = pTempBase->GetData();
@@ -206,8 +205,7 @@ void FImageOperator::ImagePixelFormat( bool& bOutSuccess, int32 CompressionQuali
 
             case EImageFormat::IF_RGBA_UBYTE_RLE:
             {
-                ImagePtr pTempBase =
-                    ImagePixelFormat( CompressionQuality, Base, EImageFormat::IF_RGBA_UBYTE );
+                ImagePtr pTempBase = ImagePixelFormat( CompressionQuality, Base, EImageFormat::IF_RGBA_UBYTE );
                 if (OnlyLOD==-1)
                 {
                     pBaseBuf = pTempBase->GetData();
@@ -244,8 +242,7 @@ void FImageOperator::ImagePixelFormat( bool& bOutSuccess, int32 CompressionQuali
             {
                 for (int m = 0; m < resultLODCount; ++m)
                 {
-                    ImagePtr pTempBase =
-                        ImagePixelFormat( CompressionQuality, Base, EImageFormat::IF_RGB_UBYTE );
+                    ImagePtr pTempBase = ImagePixelFormat( CompressionQuality, Base, EImageFormat::IF_RGB_UBYTE );
                     if (OnlyLOD == -1)
                     {
                         pBaseBuf = pTempBase->GetData();
@@ -1359,21 +1356,20 @@ void FImageOperator::ImagePixelFormat( bool& bOutSuccess, int32 CompressionQuali
 			case EImageFormat::IF_RGB_UBYTE:
 			case EImageFormat::IF_RGBA_UBYTE:
 			case EImageFormat::IF_BGRA_UBYTE:
-			{	
-                ImagePtr TempBase = ImagePixelFormat( CompressionQuality, Base, EImageFormat::IF_L_UBYTE, -1);
+			default:
+			{
+				UE_LOG(LogMutableCore, Log, TEXT("Image format conversion not implemented. Expensive generic one used. "));
+				
+				Ptr<Image> TempBase = ImagePixelFormat( CompressionQuality, Base, EImageFormat::IF_L_UBYTE, -1);
                 for (int m = 0; m < resultLODCount; ++m)
                 {
                     FIntVector2 MipSize = TempBase->CalculateMipSize( m );
-                    miro::L_to_BC4( MipSize[0], MipSize[1], Base->GetMipData( baseLOD + m ), pResult->GetMipData( m ), CompressionQuality );
+                    miro::L_to_BC4( MipSize[0], MipSize[1], TempBase->GetMipData( baseLOD + m ), pResult->GetMipData( m ), CompressionQuality );
                 }
 				ReleaseImage(TempBase);
 				break;
 			}
 
-            default:
-                // Case not implemented
-                check( false );
-
             }
 
             break;
@@ -1403,9 +1399,19 @@ void FImageOperator::ImagePixelFormat( bool& bOutSuccess, int32 CompressionQuali
                 }
                 break;
 
-            default:
-                // Case not implemented
-                check( false );
+			default:
+			{
+				UE_LOG(LogMutableCore, Log, TEXT("Image format conversion not implemented. Expensive generic one used. "));
+
+				Ptr<Image> TempBase = ImagePixelFormat(CompressionQuality, Base, EImageFormat::IF_RGBA_UBYTE, -1);
+				for (int m = 0; m < resultLODCount; ++m)
+				{
+					FIntVector2 MipSize = TempBase->CalculateMipSize(m);
+					miro::RGBA_to_BC5(MipSize[0], MipSize[1], TempBase->GetMipData(baseLOD + m), pResult->GetMipData(m), CompressionQuality);
+				}
+				ReleaseImage(TempBase);
+				break;
+			}
 
             }
 
@@ -1456,9 +1462,19 @@ void FImageOperator::ImagePixelFormat( bool& bOutSuccess, int32 CompressionQuali
 				}
 				break;
 
-            default:
-                // Case not implemented
-                check( false );
+			default:
+			{
+				UE_LOG(LogMutableCore, Log, TEXT("Image format conversion not implemented. Expensive generic one used. "));
+
+				Ptr<Image> TempBase = ImagePixelFormat(CompressionQuality, Base, EImageFormat::IF_RGBA_UBYTE, -1);
+				for (int m = 0; m < resultLODCount; ++m)
+				{
+					FIntVector2 MipSize = TempBase->CalculateMipSize(m);
+					miro::RGBA_to_ASTC4x4RGBL(MipSize[0], MipSize[1], TempBase->GetMipData(baseLOD + m), pResult->GetMipData(m), CompressionQuality);
+				}
+				ReleaseImage(TempBase);
+				break;
+			}
 
             }
 
@@ -1499,9 +1515,19 @@ void FImageOperator::ImagePixelFormat( bool& bOutSuccess, int32 CompressionQuali
                 }
                 break;
 
-            default:
-                // Case not implemented
-                check( false );
+			default:
+			{
+				UE_LOG(LogMutableCore, Log, TEXT("Image format conversion not implemented. Expensive generic one used. "));
+
+				Ptr<Image> TempBase = ImagePixelFormat(CompressionQuality, Base, EImageFormat::IF_RGBA_UBYTE, -1);
+				for (int m = 0; m < resultLODCount; ++m)
+				{
+					FIntVector2 MipSize = TempBase->CalculateMipSize(m);
+					miro::RGBA_to_ASTC4x4RGBAL(MipSize[0], MipSize[1], TempBase->GetMipData(baseLOD + m), pResult->GetMipData(m), CompressionQuality);
+				}
+				ReleaseImage(TempBase);
+				break;
+			}
 
             }
 
@@ -1543,9 +1569,19 @@ void FImageOperator::ImagePixelFormat( bool& bOutSuccess, int32 CompressionQuali
                 }
                 break;
 
-            default:
-                // Case not implemented
-                check( false );
+			default:
+			{
+				UE_LOG(LogMutableCore, Log, TEXT("Image format conversion not implemented. Expensive generic one used. "));
+
+				Ptr<Image> TempBase = ImagePixelFormat(CompressionQuality, Base, EImageFormat::IF_RGBA_UBYTE, -1);
+				for (int m = 0; m < resultLODCount; ++m)
+				{
+					FIntVector2 MipSize = TempBase->CalculateMipSize(m);
+					miro::RGBA_to_ASTC4x4RGL(MipSize[0], MipSize[1], TempBase->GetMipData(baseLOD + m), pResult->GetMipData(m), CompressionQuality);
+				}
+				ReleaseImage(TempBase);
+				break;
+			}
 
             }
 
@@ -1597,8 +1633,18 @@ void FImageOperator::ImagePixelFormat( bool& bOutSuccess, int32 CompressionQuali
 			//	break;
 
 			default:
-				// Case not implemented
-				check(false);
+			{
+				UE_LOG(LogMutableCore, Log, TEXT("Image format conversion not implemented. Expensive generic one used. "));
+
+				Ptr<Image> TempBase = ImagePixelFormat(CompressionQuality, Base, EImageFormat::IF_RGBA_UBYTE, -1);
+				for (int m = 0; m < resultLODCount; ++m)
+				{
+					FIntVector2 MipSize = TempBase->CalculateMipSize(m);
+					miro::RGBA_to_ASTC6x6RGBL(MipSize[0], MipSize[1], TempBase->GetMipData(baseLOD + m), pResult->GetMipData(m), CompressionQuality);
+				}
+				ReleaseImage(TempBase);
+				break;
+			}
 
 			}
 
@@ -1640,8 +1686,18 @@ void FImageOperator::ImagePixelFormat( bool& bOutSuccess, int32 CompressionQuali
 			//	break;
 
 			default:
-				// Case not implemented
-				check(false);
+			{
+				UE_LOG(LogMutableCore, Log, TEXT("Image format conversion not implemented. Expensive generic one used. "));
+
+				Ptr<Image> TempBase = ImagePixelFormat(CompressionQuality, Base, EImageFormat::IF_RGBA_UBYTE, -1);
+				for (int m = 0; m < resultLODCount; ++m)
+				{
+					FIntVector2 MipSize = TempBase->CalculateMipSize(m);
+					miro::RGBA_to_ASTC6x6RGBAL(MipSize[0], MipSize[1], TempBase->GetMipData(baseLOD + m), pResult->GetMipData(m), CompressionQuality);
+				}
+				ReleaseImage(TempBase);
+				break;
+			}
 
 			}
 
@@ -1684,8 +1740,18 @@ void FImageOperator::ImagePixelFormat( bool& bOutSuccess, int32 CompressionQuali
 			//	break;
 
 			default:
-				// Case not implemented
-				check(false);
+			{
+				UE_LOG(LogMutableCore, Log, TEXT("Image format conversion not implemented. Expensive generic one used. "));
+
+				Ptr<Image> TempBase = ImagePixelFormat(CompressionQuality, Base, EImageFormat::IF_RGBA_UBYTE, -1);
+				for (int m = 0; m < resultLODCount; ++m)
+				{
+					FIntVector2 MipSize = TempBase->CalculateMipSize(m);
+					miro::RGBA_to_ASTC6x6RGL(MipSize[0], MipSize[1], TempBase->GetMipData(baseLOD + m), pResult->GetMipData(m), CompressionQuality);
+				}
+				ReleaseImage(TempBase);
+				break;
+			}
 
 			}
 
@@ -1737,8 +1803,18 @@ void FImageOperator::ImagePixelFormat( bool& bOutSuccess, int32 CompressionQuali
 				//	break;
 
 			default:
-				// Case not implemented
-				check(false);
+			{
+				UE_LOG(LogMutableCore, Log, TEXT("Image format conversion not implemented. Expensive generic one used. "));
+
+				Ptr<Image> TempBase = ImagePixelFormat(CompressionQuality, Base, EImageFormat::IF_RGBA_UBYTE, -1);
+				for (int m = 0; m < resultLODCount; ++m)
+				{
+					FIntVector2 MipSize = TempBase->CalculateMipSize(m);
+					miro::RGBA_to_ASTC8x8RGBL(MipSize[0], MipSize[1], TempBase->GetMipData(baseLOD + m), pResult->GetMipData(m), CompressionQuality);
+				}
+				ReleaseImage(TempBase);
+				break;
+			}
 
 			}
 
@@ -1780,8 +1856,18 @@ void FImageOperator::ImagePixelFormat( bool& bOutSuccess, int32 CompressionQuali
 			//	break;
 
 			default:
-				// Case not implemented
-				check(false);
+			{
+				UE_LOG(LogMutableCore, Log, TEXT("Image format conversion not implemented. Expensive generic one used. "));
+
+				Ptr<Image> TempBase = ImagePixelFormat(CompressionQuality, Base, EImageFormat::IF_RGBA_UBYTE, -1);
+				for (int m = 0; m < resultLODCount; ++m)
+				{
+					FIntVector2 MipSize = TempBase->CalculateMipSize(m);
+					miro::RGBA_to_ASTC8x8RGBAL(MipSize[0], MipSize[1], TempBase->GetMipData(baseLOD + m), pResult->GetMipData(m), CompressionQuality);
+				}
+				ReleaseImage(TempBase);
+				break;
+			}
 
 			}
 
@@ -1813,8 +1899,18 @@ void FImageOperator::ImagePixelFormat( bool& bOutSuccess, int32 CompressionQuali
 				break;
 
 			default:
-				// Case not implemented
-				check(false);
+			{
+				UE_LOG(LogMutableCore, Log, TEXT("Image format conversion not implemented. Expensive generic one used. "));
+
+				Ptr<Image> TempBase = ImagePixelFormat(CompressionQuality, Base, EImageFormat::IF_RGBA_UBYTE, -1);
+				for (int m = 0; m < resultLODCount; ++m)
+				{
+					FIntVector2 MipSize = TempBase->CalculateMipSize(m);
+					miro::RGBA_to_ASTC8x8RGL(MipSize[0], MipSize[1], TempBase->GetMipData(baseLOD + m), pResult->GetMipData(m), CompressionQuality);
+				}
+				ReleaseImage(TempBase);
+				break;
+			}
 
 			}
 
