[Mutable] Removed MutableTools PrivateIncludePaths to MutableRuntime.

Moved required MuR private headers to internal.

#rnx

diff --git a/Engine/Plugins/Experimental/Mutable/Source/MutableRuntime/Internal/MuR/BlockCompression/Miro/Miro.h b/Engine/Plugins/Experimental/Mutable/Source/MutableRuntime/Internal/MuR/BlockCompression/Miro/Miro.h
new file mode 100644
--- /dev/null
+++ b/Engine/Plugins/Experimental/Mutable/Source/MutableRuntime/Internal/MuR/BlockCompression/Miro/Miro.h
@@ -0,0 +1,210 @@
+// Copyright Epic Games, Inc. All Rights Reserved.
+
+#pragma once
+
+#include "HAL/Platform.h"
+#include "Math/IntVector.h"
+#include "Templates/Function.h"
+// Configure the compression library with these defines:
+
+// Support for BC1 to BC7, also known as DXTC1 to 5
+#ifndef MIRO_INCLUDE_BC
+#define MIRO_INCLUDE_BC                     1
+#endif
+
+// Support for ASTC
+#ifndef MIRO_INCLUDE_ASTC
+#define MIRO_INCLUDE_ASTC                   1
+#endif
+
+// Exclude high quality compression code. The index corresponds to the quality levels.
+#ifndef MIRO_EXCLUDE_QUALITY_ABOVE
+#define MIRO_EXCLUDE_QUALITY_ABOVE          4
+#endif
+
+
+namespace miro
+{
+    // MIRO is a runtime texture compression library for hardware-accelerated formats.
+
+    // Compression quality levels (q):
+    // 0 - Fastest for runtime
+    // 1 - Best for runtime
+    // 2 - Fast for tools
+    // 3 - Best for tools
+    // 4 - Maximum, with no time limits.
+
+    using FImageSize = UE::Math::TIntVector2<uint16>;
+
+    extern void initialize();
+    extern void finalize();
+
+#if MIRO_INCLUDE_BC
+
+    //! BC1 support
+    extern void RGB_to_BC1( uint32 sx, uint32 sy, const uint8* From, uint8* To, int32 Quality );
+    extern void RGBA_to_BC1( uint32 sx, uint32 sy, const uint8* From, uint8* To, int32 Quality );
+    extern void BC1_to_RGB( uint32 sx, uint32 sy, const uint8* From, uint8* To );
+    extern void BC1_to_RGBA( uint32 sx, uint32 sy, const uint8* From, uint8* To );
+
+    //! BC2 support
+    extern void RGB_to_BC2( uint32 sx, uint32 sy, const uint8* From, uint8* To, int32 Quality );
+    extern void RGBA_to_BC2( uint32 sx, uint32 sy, const uint8* From, uint8* To, int32 Quality );
+    extern void BC2_to_RGB( uint32 sx, uint32 sy, const uint8* From, uint8* To );
+    extern void BC2_to_RGBA( uint32 sx, uint32 sy, const uint8* From, uint8* To );
+
+    //! BC3 support
+    extern void RGB_to_BC3( uint32 sx, uint32 sy, const uint8* From, uint8* To, int32 Quality );
+    extern void RGBA_to_BC3( uint32 sx, uint32 sy, const uint8* From, uint8* To, int32 Quality );
+    extern void BC3_to_RGB( uint32 sx, uint32 sy, const uint8* From, uint8* To );
+    extern void BC3_to_RGBA( uint32 sx, uint32 sy, const uint8* From, uint8* To );
+    extern void BC1_to_BC3( uint32 sx, uint32 sy, const uint8* From, uint8* To, int32 Quality );
+
+    //! BC4 support
+    extern void L_to_BC4( uint32 sx, uint32 sy, const uint8* From, uint8* To, int32 Quality );
+    extern void BC4_to_L( uint32 sx, uint32 sy, const uint8* From, uint8* To );
+    extern void BC4_to_RGB( uint32 sx, uint32 sy, const uint8* From, uint8* To );
+    extern void BC4_to_RGBA( uint32 sx, uint32 sy, const uint8* From, uint8* To );
+
+    //! BC5 support
+    extern void RGBA_to_BC5( uint32 sx, uint32 sy, const uint8* From, uint8* To, int32 Quality );
+    extern void RGB_to_BC5( uint32 sx, uint32 sy, const uint8* From, uint8* To, int32 Quality );
+    extern void BC5_to_RGBA( uint32 sx, uint32 sy, const uint8* From, uint8* To );
+    extern void BC5_to_RGB( uint32 sx, uint32 sy, const uint8* From, uint8* To );
+
+#endif
+
+
+#if MIRO_INCLUDE_ASTC
+
+    //! ASTC 4x4 support
+    extern void RGBA_to_ASTC4x4RGBAL( uint32 sx, uint32 sy, const uint8* From, uint8* To, int32 Quality );
+    extern void RGB_to_ASTC4x4RGBAL( uint32 sx, uint32 sy, const uint8* From, uint8* To, int32 Quality );
+    extern void ASTC4x4RGBAL_to_RGBA( uint32 sx, uint32 sy, const uint8* From, uint8* To );
+    extern void ASTC4x4RGBAL_to_RGB( uint32 sx, uint32 sy, const uint8* From, uint8* To );
+
+    extern void RGB_to_ASTC4x4RGBL( uint32 sx, uint32 sy, const uint8* From, uint8* To, int32 Quality );
+    extern void RGBA_to_ASTC4x4RGBL( uint32 sx, uint32 sy, const uint8* From, uint8* To, int32 Quality );
+    extern void ASTC4x4RGBL_to_RGB( uint32 sx, uint32 sy, const uint8* From, uint8* To );
+    extern void ASTC4x4RGBL_to_RGBA( uint32 sx, uint32 sy, const uint8* From, uint8* To );
+	extern void L_to_ASTC4x4RGBL(uint32 sx, uint32 sy, const uint8* From, uint8* To, int32 Quality);
+
+    extern void RGB_to_ASTC4x4RGL( uint32 sx, uint32 sy, const uint8* From, uint8* To, int32 Quality );
+    extern void RGBA_to_ASTC4x4RGL( uint32 sx, uint32 sy, const uint8* From, uint8* To, int32 Quality );
+    extern void ASTC4x4RGL_to_RGB( uint32 sx, uint32 sy, const uint8* From, uint8* To );
+    extern void ASTC4x4RGL_to_RGBA( uint32 sx, uint32 sy, const uint8* From, uint8* To );
+
+	extern void ASTC4x4RGBL_to_ASTC4x4RGBAL(uint32 sx, uint32 sy, const uint8* From, uint8* To);
+	extern void ASTC4x4RGBAL_to_ASTC4x4RGBL(uint32 sx, uint32 sy, const uint8* From, uint8* To);
+
+	//! ASTC 6x6 support
+	extern void RGBA_to_ASTC6x6RGBAL(uint32 sx, uint32 sy, const uint8* From, uint8* To, int32 Quality);
+	extern void RGB_to_ASTC6x6RGBAL(uint32 sx, uint32 sy, const uint8* From, uint8* To, int32 Quality);
+	extern void ASTC6x6RGBAL_to_RGBA(uint32 sx, uint32 sy, const uint8* From, uint8* To);
+	extern void ASTC6x6RGBAL_to_RGB(uint32 sx, uint32 sy, const uint8* From, uint8* To);
+	
+	extern void RGB_to_ASTC6x6RGBL(uint32 sx, uint32 sy, const uint8* From, uint8* To, int32 Quality);
+	extern void RGBA_to_ASTC6x6RGBL(uint32 sx, uint32 sy, const uint8* From, uint8* To, int32 Quality);
+	extern void ASTC6x6RGBL_to_RGB(uint32 sx, uint32 sy, const uint8* From, uint8* To);
+	extern void ASTC6x6RGBL_to_RGBA(uint32 sx, uint32 sy, const uint8* From, uint8* To);
+
+	extern void RGB_to_ASTC6x6RGL(uint32 sx, uint32 sy, const uint8* From, uint8* To, int32 Quality);
+	extern void RGBA_to_ASTC6x6RGL(uint32 sx, uint32 sy, const uint8* From, uint8* To, int32 Quality);
+	extern void ASTC6x6RGL_to_RGB(uint32 sx, uint32 sy, const uint8* From, uint8* To);
+	extern void ASTC6x6RGL_to_RGBA(uint32 sx, uint32 sy, const uint8* From, uint8* To);
+
+	//! ASTC 8x8 support
+	extern void RGBA_to_ASTC8x8RGBAL(uint32 sx, uint32 sy, const uint8* From, uint8* To, int32 Quality);
+	extern void RGB_to_ASTC8x8RGBAL(uint32 sx, uint32 sy, const uint8* From, uint8* To, int32 Quality);
+	extern void ASTC8x8RGBAL_to_RGBA(uint32 sx, uint32 sy, const uint8* From, uint8* To);
+	extern void ASTC8x8RGBAL_to_RGB(uint32 sx, uint32 sy, const uint8* From, uint8* To);
+
+	extern void RGB_to_ASTC8x8RGBL(uint32 sx, uint32 sy, const uint8* From, uint8* To, int32 Quality);
+	extern void RGBA_to_ASTC8x8RGBL(uint32 sx, uint32 sy, const uint8* From, uint8* To, int32 Quality);
+	extern void ASTC8x8RGBL_to_RGB(uint32 sx, uint32 sy, const uint8* From, uint8* To);
+	extern void ASTC8x8RGBL_to_RGBA(uint32 sx, uint32 sy, const uint8* From, uint8* To);
+
+	extern void RGB_to_ASTC8x8RGL(uint32 sx, uint32 sy, const uint8* From, uint8* To, int32 Quality);
+	extern void RGBA_to_ASTC8x8RGL(uint32 sx, uint32 sy, const uint8* From, uint8* To, int32 Quality);
+	extern void ASTC8x8RGL_to_RGB(uint32 sx, uint32 sy, const uint8* From, uint8* To);
+	extern void ASTC8x8RGL_to_RGBA(uint32 sx, uint32 sy, const uint8* From, uint8* To);
+
+	//! ASTC 10x10 support
+	extern void ASTC10x10RGBAL_to_RGBA(uint32 sx, uint32 sy, const uint8* From, uint8* To);
+	extern void ASTC10x10RGBAL_to_RGB(uint32 sx, uint32 sy, const uint8* From, uint8* To);
+	extern void ASTC10x10RGBL_to_RGB(uint32 sx, uint32 sy, const uint8* From, uint8* To);
+	extern void ASTC10x10RGBL_to_RGBA(uint32 sx, uint32 sy, const uint8* From, uint8* To);
+	extern void ASTC10x10RGL_to_RGB(uint32 sx, uint32 sy, const uint8* From, uint8* To);
+	extern void ASTC10x10RGL_to_RGBA(uint32 sx, uint32 sy, const uint8* From, uint8* To);
+
+	//! ASTC 12x12 support
+	extern void ASTC12x12RGBAL_to_RGBA(uint32 sx, uint32 sy, const uint8* From, uint8* To);
+	extern void ASTC12x12RGBAL_to_RGB(uint32 sx, uint32 sy, const uint8* From, uint8* To);
+	extern void ASTC12x12RGBL_to_RGB(uint32 sx, uint32 sy, const uint8* From, uint8* To);
+	extern void ASTC12x12RGBL_to_RGBA(uint32 sx, uint32 sy, const uint8* From, uint8* To);
+	extern void ASTC12x12RGL_to_RGB(uint32 sx, uint32 sy, const uint8* From, uint8* To);
+	extern void ASTC12x12RGL_to_RGBA(uint32 sx, uint32 sy, const uint8* From, uint8* To);
+
+#endif
+
+}
+
+
+namespace miro::SubImageDecompression
+{
+    using FuncRefType = TFunctionRef<void(miro::FImageSize, miro::FImageSize, miro::FImageSize, const uint8*, uint8*)>;
+
+#if MIRO_INCLUDE_ASTC
+	 void ASTC4x4RGBAL_To_RGBASubImage(FImageSize FromSize, FImageSize ToSize, FImageSize SubSize, const uint8* RESTRICT From, uint8* RESTRICT To);
+	 void ASTC4x4RGBAL_To_RGBSubImage (FImageSize FromSize, FImageSize ToSize, FImageSize SubSize, const uint8* RESTRICT From, uint8* RESTRICT To);
+	 void ASTC4x4RGBL_To_RGBSubImage  (FImageSize FromSize, FImageSize ToSize, FImageSize SubSize, const uint8* RESTRICT From, uint8* RESTRICT To);
+	 void ASTC4x4RGBL_To_RGBASubImage (FImageSize FromSize, FImageSize ToSize, FImageSize SubSize, const uint8* RESTRICT From, uint8* RESTRICT To);
+	 void ASTC4x4RGL_To_RGBSubImage   (FImageSize FromSize, FImageSize ToSize, FImageSize SubSize, const uint8* RESTRICT From, uint8* RESTRICT To);
+	 void ASTC4x4RGL_To_RGBASubImage  (FImageSize FromSize, FImageSize ToSize, FImageSize SubSize, const uint8* RESTRICT From, uint8* RESTRICT To);
+
+	 void ASTC6x6RGBAL_To_RGBASubImage(FImageSize FromSize, FImageSize ToSize, FImageSize SubSize, const uint8* RESTRICT From, uint8* RESTRICT To);
+	 void ASTC6x6RGBAL_To_RGBSubImage (FImageSize FromSize, FImageSize ToSize, FImageSize SubSize, const uint8* RESTRICT From, uint8* RESTRICT To);
+	 void ASTC6x6RGBL_To_RGBSubImage  (FImageSize FromSize, FImageSize ToSize, FImageSize SubSize, const uint8* RESTRICT From, uint8* RESTRICT To);
+	 void ASTC6x6RGBL_To_RGBASubImage (FImageSize FromSize, FImageSize ToSize, FImageSize SubSize, const uint8* RESTRICT From, uint8* RESTRICT To);
+	 void ASTC6x6RGL_To_RGBSubImage   (FImageSize FromSize, FImageSize ToSize, FImageSize SubSize, const uint8* RESTRICT From, uint8* RESTRICT To);
+	 void ASTC6x6RGL_To_RGBASubImage  (FImageSize FromSize, FImageSize ToSize, FImageSize SubSize, const uint8* RESTRICT From, uint8* RESTRICT To);
+
+	 void ASTC8x8RGBAL_To_RGBASubImage(FImageSize FromSize, FImageSize ToSize, FImageSize SubSize, const uint8* RESTRICT From, uint8* RESTRICT To);
+	 void ASTC8x8RGBAL_To_RGBSubImage (FImageSize FromSize, FImageSize ToSize, FImageSize SubSize, const uint8* RESTRICT From, uint8* RESTRICT To);
+	 void ASTC8x8RGBL_To_RGBSubImage  (FImageSize FromSize, FImageSize ToSize, FImageSize SubSize, const uint8* RESTRICT From, uint8* RESTRICT To);
+	 void ASTC8x8RGBL_To_RGBASubImage (FImageSize FromSize, FImageSize ToSize, FImageSize SubSize, const uint8* RESTRICT From, uint8* RESTRICT To);
+	 void ASTC8x8RGL_To_RGBSubImage   (FImageSize FromSize, FImageSize ToSize, FImageSize SubSize, const uint8* RESTRICT From, uint8* RESTRICT To);
+	 void ASTC8x8RGL_To_RGBASubImage  (FImageSize FromSize, FImageSize ToSize, FImageSize SubSize, const uint8* RESTRICT From, uint8* RESTRICT To);
+
+	 void ASTC10x10RGBAL_To_RGBASubImage(FImageSize FromSize, FImageSize ToSize, FImageSize SubSize, const uint8* RESTRICT From, uint8* RESTRICT To);
+	 void ASTC10x10RGBAL_To_RGBSubImage (FImageSize FromSize, FImageSize ToSize, FImageSize SubSize, const uint8* RESTRICT From, uint8* RESTRICT To);
+	 void ASTC10x10RGBL_To_RGBSubImage  (FImageSize FromSize, FImageSize ToSize, FImageSize SubSize, const uint8* RESTRICT From, uint8* RESTRICT To);
+	 void ASTC10x10RGBL_To_RGBASubImage (FImageSize FromSize, FImageSize ToSize, FImageSize SubSize, const uint8* RESTRICT From, uint8* RESTRICT To);
+	 void ASTC10x10RGL_To_RGBSubImage   (FImageSize FromSize, FImageSize ToSize, FImageSize SubSize, const uint8* RESTRICT From, uint8* RESTRICT To);
+	 void ASTC10x10RGL_To_RGBASubImage  (FImageSize FromSize, FImageSize ToSize, FImageSize SubSize, const uint8* RESTRICT From, uint8* RESTRICT To);
+
+	 void ASTC12x12RGBAL_To_RGBASubImage(FImageSize FromSize, FImageSize ToSize, FImageSize SubSize, const uint8* RESTRICT From, uint8* RESTRICT To);
+	 void ASTC12x12RGBAL_To_RGBSubImage (FImageSize FromSize, FImageSize ToSize, FImageSize SubSize, const uint8* RESTRICT From, uint8* RESTRICT To);
+	 void ASTC12x12RGBL_To_RGBSubImage  (FImageSize FromSize, FImageSize ToSize, FImageSize SubSize, const uint8* RESTRICT From, uint8* RESTRICT To);
+	 void ASTC12x12RGBL_To_RGBASubImage (FImageSize FromSize, FImageSize ToSize, FImageSize SubSize, const uint8* RESTRICT From, uint8* RESTRICT To);
+	 void ASTC12x12RGL_To_RGBSubImage   (FImageSize FromSize, FImageSize ToSize, FImageSize SubSize, const uint8* RESTRICT From, uint8* RESTRICT To);
+	 void ASTC12x12RGL_To_RGBASubImage  (FImageSize FromSize, FImageSize ToSize, FImageSize SubSize, const uint8* RESTRICT From, uint8* RESTRICT To);
+
+#endif // MIRO_INCLUDE_ASTC
+	// BC Formats.
+
+#if MIRO_INCLUDE_BC
+	 void BC1_To_RGBSubImage (FImageSize FromSize, FImageSize ToSize, FImageSize SubSize, const uint8* RESTRICT From, uint8* RESTRICT To);
+	 void BC1_To_RGBASubImage(FImageSize FromSize, FImageSize ToSize, FImageSize SubSize, const uint8* RESTRICT From, uint8* RESTRICT To);
+	 void BC2_To_RGBASubImage(FImageSize FromSize, FImageSize ToSize, FImageSize SubSize, const uint8* RESTRICT From, uint8* RESTRICT To);
+	 void BC2_To_RGBSubImage (FImageSize FromSize, FImageSize ToSize, FImageSize SubSize, const uint8* RESTRICT From, uint8* RESTRICT To);
+	 void BC3_To_RGBASubImage(FImageSize FromSize, FImageSize ToSize, FImageSize SubSize, const uint8* RESTRICT From, uint8* RESTRICT To);
+	 void BC3_To_RGBSubImage (FImageSize FromSize, FImageSize ToSize, FImageSize SubSize, const uint8* RESTRICT From, uint8* RESTRICT To);
+	 void BC4_To_LSubImage   (FImageSize FromSize, FImageSize ToSize, FImageSize SubSize, const uint8* RESTRICT From, uint8* RESTRICT To);
+	 void BC4_To_RGBSubImage (FImageSize FromSize, FImageSize ToSize, FImageSize SubSize, const uint8* RESTRICT From, uint8* RESTRICT To);
+	 void BC4_To_RGBASubImage(FImageSize FromSize, FImageSize ToSize, FImageSize SubSize, const uint8* RESTRICT From, uint8* RESTRICT To);
+	 void BC5_To_RGBASubImage(FImageSize FromSize, FImageSize ToSize, FImageSize SubSize, const uint8* RESTRICT From, uint8* RESTRICT To);
+	 void BC5_To_RGBSubImage (FImageSize FromSize, FImageSize ToSize, FImageSize SubSize, const uint8* RESTRICT From, uint8* RESTRICT To);
+#endif
+}
+
diff --git a/Engine/Plugins/Experimental/Mutable/Source/MutableRuntime/Internal/MuR/CodeVisitor.h b/Engine/Plugins/Experimental/Mutable/Source/MutableRuntime/Internal/MuR/CodeVisitor.h
new file mode 100644
--- /dev/null
+++ b/Engine/Plugins/Experimental/Mutable/Source/MutableRuntime/Internal/MuR/CodeVisitor.h
@@ -0,0 +1,671 @@
+// Copyright Epic Games, Inc. All Rights Reserved.
+
+#pragma once
+
+#include "HAL/PlatformMath.h"
+#include "HAL/UnrealMemory.h"
+#include "Misc/AssertionMacros.h"
+#include "MuR/Model.h"
+#include "MuR/ModelPrivate.h"
+#include "MuR/MutableMemory.h"
+#include "MuR/Operations.h"
+#include "MuR/ParametersPrivate.h"
+#include "MuR/Ptr.h"
+#include "MuR/System.h"
+#include "MuR/SystemPrivate.h"
+
+namespace mu
+{
+
+    //---------------------------------------------------------------------------------------------
+    //! Decide what operations are an "add resource" since they are handled differently sometimes.
+    //---------------------------------------------------------------------------------------------
+    inline bool VisitorIsAddResource(const OP_TYPE& type)
+    {
+        return  type == OP_TYPE::IN_ADDIMAGE ||
+                type == OP_TYPE::IN_ADDMESH;
+    }
+
+
+    //---------------------------------------------------------------------------------------------
+    //! Code visitor that:
+    //! - is top-down
+    //! - cannot change the visited instructions.
+    //! - will not visit twice the same instruction with the same state.
+    //! - Its iterative
+    //---------------------------------------------------------------------------------------------
+    template<typename STATE=int>
+    class UniqueConstCodeVisitorIterative
+    {
+    public:
+
+        UniqueConstCodeVisitorIterative( bool skipResources=false )
+        {
+            // Default state
+            m_states.Add(STATE());
+            m_currentState = 0;
+            m_skipResources = skipResources;
+        }
+
+        //! Ensure virtual destruction
+        virtual ~UniqueConstCodeVisitorIterative() {}
+
+    protected:
+
+        //!
+        void SetDefaultState(const STATE& s)
+        {
+            m_states[0] = s;
+        }
+
+        //!
+        const STATE& GetDefaultState() const
+        {
+            return m_states[0];
+        }
+
+        //! Use this from visit to access the state at the time of processing the current
+        //! instruction.
+        STATE GetCurrentState() const
+        {
+            return m_states[m_currentState];
+        }
+
+        //! For manual recursion that changes the state for a specific path.
+        void RecurseWithState(OP::ADDRESS at, const STATE& newState)
+        {
+			int32 it = m_states.Find(newState);
+            if (it==INDEX_NONE)
+            {
+                m_states.Add(newState);
+            }
+            int stateIndex = m_states.IndexOfByKey(newState);
+
+            m_pending.Add( PENDING(at,stateIndex) );
+        }
+
+        //! For manual recursion that doesn't change the state for a specific path.
+        void RecurseWithCurrentState(OP::ADDRESS at)
+        {
+            m_pending.Add( PENDING(at,m_currentState) );
+        }
+
+        //! Can be called from visit to set the state to visit all children ops
+        void SetCurrentState(const STATE& newState)
+        {
+            int32 it = m_states.Find(newState);
+            if (it==INDEX_NONE)
+            {
+                m_states.Add(newState);
+            }
+            m_currentState = m_states.Find(newState);
+        }
+
+
+        void Traverse( OP::ADDRESS root, FProgram& program )
+        {
+            m_pending.Reserve( program.m_opAddress.Num() );
+
+            // Visit the given root
+            m_pending.Add( PENDING(root,0) );
+            Recurse( program );
+         }
+
+        void FullTraverse( FProgram& program )
+        {
+            // Visit all the state roots
+            for ( int32 p=0; p<program.m_states.Num(); ++p )
+            {
+                m_pending.Add(PENDING(program.m_states[p].m_root,0) );
+                Recurse( program );
+            }
+        }
+
+
+    private:
+
+        //! Do the actual work by overriding this in the derived classes.
+        //! Return true if the traverse has to continue with the children of "at"
+        virtual bool Visit( OP::ADDRESS at, FProgram& program ) = 0;
+
+        //! Operations to be processed
+        struct PENDING
+        {
+            PENDING()
+            {
+                at = 0;
+                stateIndex = 0;
+            }
+            PENDING(OP::ADDRESS _at, int _stateIndex)
+            {
+                at = _at;
+                stateIndex = _stateIndex;
+            }
+
+            OP::ADDRESS at;
+            int stateIndex;
+        };
+		TArray<PENDING> m_pending;
+
+        //! States found so far
+		TArray<STATE> m_states;
+
+        //! Index of the current state, from the m_states array.
+        int m_currentState;
+
+        //! If true, operations adding resources (meshes or images) will only
+        //! have the base operation recursed, but not the resources.
+        bool m_skipResources;
+
+        //! Array of states visited for each operation.
+        //! Empty array means operation not visited at all.
+		TArray<TArray<int>> m_visited;
+
+        //! Process all the pending operations and visit all children if necessary
+        void Recurse( FProgram& program )
+        {
+			m_visited.Empty();
+			m_visited.SetNum(program.m_opAddress.Num());
+
+            while ( m_pending.Num() )
+            {
+                OP::ADDRESS at = m_pending.Last().at;
+                m_currentState = m_pending.Last().stateIndex;
+                m_pending.Pop();
+
+                bool recurse = false;
+
+                bool visitedInThisState = m_visited[at].Contains(m_currentState);
+                if (!visitedInThisState)
+                {
+                    m_visited[at].Add(m_currentState);
+
+                    // Visit may change current state
+                    recurse = Visit( at, program );
+                }
+
+                if (recurse)
+                {
+                    ForEachReference( program, at, [&](OP::ADDRESS ref)
+                    {
+                        if (ref)
+                        {
+                            m_pending.Add( PENDING(ref,m_currentState) );
+                        }
+                    });
+                }
+            }
+        }
+
+    };
+
+
+    //---------------------------------------------------------------------------------------------
+    //! Code visitor that:
+    //! - is top-down
+    //! - cannot change the instructions.
+    //! - will repeat visits to the instructions that have multiple references.
+    //! - Its iterative
+    //---------------------------------------------------------------------------------------------
+    template<class STATE=int>
+    class RepeatConstCodeVisitorIterative
+    {
+    public:
+
+        RepeatConstCodeVisitorIterative( bool skipResources=false )
+        {
+            // Default state
+            m_states.Add(STATE());
+            m_currentState = 0;
+            m_skipResources = skipResources;
+        }
+
+        //! Ensure virtual destruction
+        virtual ~RepeatConstCodeVisitorIterative() {}
+
+    protected:
+
+        //!
+        void SetDefaultState(const STATE& s)
+        {
+            m_states[0] = s;
+        }
+
+        //!
+        const STATE& GetDefaultState() const
+        {
+            return m_states[0];
+        }
+
+        //! Use this from visit to access the state at the time of processing the current
+        //! instruction.
+        const STATE& GetCurrentState() const
+        {
+            return m_states[m_currentState];
+        }
+
+        //! For manual recursion that changes the state for a specific path.
+        void RecurseWithState(OP::ADDRESS At, const STATE& NewState)
+        {
+            int32 StateIndex = m_states.Find(NewState);
+            if (StateIndex ==INDEX_NONE)
+            {
+				StateIndex = m_states.Add(NewState);
+            }
+
+            //check(at<1000000);
+            m_pending.Add( PENDING(At,StateIndex) );
+        }
+
+        //! For manual recursion that doesn't change the state for a specific path.
+        void RecurseWithCurrentState(OP::ADDRESS at)
+        {
+            //check(at<1000000);
+            m_pending.Add( PENDING(at,m_currentState) );
+        }
+
+        //! Can be called from visit to set the state to visit all children ops
+        void SetCurrentState(const STATE& NewState)
+        {
+			int32 StateIndex = m_states.Find(NewState);
+			if (StateIndex == INDEX_NONE)
+			{
+				StateIndex = m_states.Add(NewState);
+			}
+
+			m_currentState = StateIndex;
+        }
+
+
+        void Traverse( OP::ADDRESS Root, FProgram& Program )
+        {
+            m_pending.reserve(Program.m_opAddress.Num() );
+
+            // Visit the given root
+            m_pending.Add( PENDING(Root,0) );
+            Recurse(Program);
+        }
+
+        void FullTraverse( FProgram& Program)
+        {
+            // Visit all the state roots
+            for ( int32 p=0; p<Program.m_states.Num(); ++p )
+            {
+                m_pending.Add(PENDING(Program.m_states[p].m_root,0) );
+                Recurse(Program);
+            }
+        }
+
+
+    private:
+
+        //! Do the actual work by overriding this in the derived classes.
+        //! Return true if the traverse has to continue with the children of "at"
+        virtual bool Visit( OP::ADDRESS at, FProgram& program ) = 0;
+
+        //! Operations to be processed
+        struct PENDING
+        {
+            PENDING()
+            {
+                at = 0;
+                stateIndex = 0;
+            }
+            PENDING(OP::ADDRESS _at, int _stateIndex)
+            {
+                at = _at;
+                stateIndex = _stateIndex;
+            }
+
+            OP::ADDRESS at;
+            int stateIndex;
+        };
+        TArray<PENDING> m_pending;
+
+        //! States found so far
+		TArray<STATE> m_states;
+
+        //! Index of the current state, from the m_states array.
+        int m_currentState;
+
+        //! If true, operations adding resources (meshes or images) will only
+        //! have the base operation recursed, but not the resources.
+        bool m_skipResources;
+
+
+        //! Process all the pending operations and visit all children if necessary
+        void Recurse( FProgram& program )
+        {
+            while ( m_pending.Num() )
+            {
+                OP::ADDRESS at = m_pending.back().at;
+                m_currentState = m_pending.back().stateIndex;
+                m_pending.pop_back();
+
+                bool recurse = false;
+
+                // Visit may change current state
+                recurse = Visit( at, program );
+
+                if (recurse)
+                {
+                    if (m_skipResources && VisitorIsAddResource(program.GetOpType(at)))
+                    {
+						OP::InstanceAddArgs args = program.GetOpArgs<OP::InstanceAddArgs>(at);
+
+                        // Recurse only the base
+                        OP::ADDRESS base = args.instance;
+                        if (base)
+                        {
+							check(base<program.m_opAddress.Num());
+                            m_pending.Add( PENDING(base,m_currentState) );
+                        }
+                    }
+                    else
+                    {
+                        ForEachReference( program, at, [&](OP::ADDRESS ref)
+                        {
+                            if (ref)
+                            {
+								check(ref<program.m_opAddress.Num());
+                                m_pending.Add( PENDING(ref,m_currentState) );
+                            }
+                        });
+                    }
+                }
+            }
+        }
+
+    };
+
+
+    //---------------------------------------------------------------------------------------------
+    //! Code visitor template for visitors that:
+    //! - only traverses the operations that are relevant for a given set of parameter values. It
+    //! only considers the discrete parameters like integers and booleans. In the case of forks
+    //! caused by continuous parameters like float weights for interpolation, all the branches are
+    //! traversed.
+    //! - cannot change the instructinos
+    //---------------------------------------------------------------------------------------------
+    struct COVERED_CODE_VISITOR_STATE
+    {
+        uint16 m_underResourceCount = 0;
+
+        bool operator==(const COVERED_CODE_VISITOR_STATE& o) const
+        {
+            return m_underResourceCount==o.m_underResourceCount;
+        }
+    };
+
+    template<typename PARENT,typename STATE>
+    class DiscreteCoveredCodeVisitorBase : public PARENT
+    {
+    public:
+
+        DiscreteCoveredCodeVisitorBase
+            (
+                System::Private* pSystem,
+                const TSharedPtr<const Model>& pModel,
+                const ParametersPtrConst& pParams,
+                unsigned lodMask,
+                bool skipResources=false
+            )
+            : PARENT(skipResources)
+        {
+            m_pSystem = pSystem;
+            m_pModel = pModel;
+            m_pParams = pParams.get();
+            m_lodMask = lodMask;
+
+            // Visiting state
+            PARENT::SetDefaultState( STATE() );
+        }
+
+        void Run( OP::ADDRESS at  )
+        {
+            PARENT::SetDefaultState( STATE() );
+
+            PARENT::Traverse( at, m_pModel->GetPrivate()->m_program );
+        }
+
+    protected:
+
+        virtual bool Visit( OP::ADDRESS at, FProgram& program )
+        {
+            bool recurse = true;
+
+            OP_TYPE type = program.GetOpType(at);
+
+            switch ( type )
+            {
+            case OP_TYPE::NU_CONDITIONAL:
+            case OP_TYPE::SC_CONDITIONAL:
+            case OP_TYPE::CO_CONDITIONAL:
+            case OP_TYPE::IM_CONDITIONAL:
+            case OP_TYPE::ME_CONDITIONAL:
+            case OP_TYPE::LA_CONDITIONAL:
+            case OP_TYPE::IN_CONDITIONAL:
+			case OP_TYPE::ED_CONDITIONAL:
+            {
+				OP::ConditionalArgs args = program.GetOpArgs<OP::ConditionalArgs>(at);
+
+                recurse = false;
+
+                PARENT::RecurseWithCurrentState( args.condition );
+
+                // If there is no expression, we'll assume true.
+                bool value = true;
+
+                if (args.condition)
+                {
+                    value = m_pSystem->BuildBool(m_pModel, m_pParams, args.condition);
+                }
+
+                if (value)
+                {
+                    PARENT::RecurseWithCurrentState( args.yes );
+                }
+                else
+                {
+                    PARENT::RecurseWithCurrentState( args.no );
+                }
+                break;
+            }
+
+            case OP_TYPE::NU_SWITCH:
+            case OP_TYPE::SC_SWITCH:
+            case OP_TYPE::CO_SWITCH:
+            case OP_TYPE::IM_SWITCH:
+            case OP_TYPE::ME_SWITCH:
+            case OP_TYPE::LA_SWITCH:
+            case OP_TYPE::IN_SWITCH:
+			case OP_TYPE::ED_SWITCH:
+            {
+                recurse = false;
+
+				const uint8* data = program.GetOpArgsPointer(at);
+				
+				OP::ADDRESS VarAddress;
+				FMemory::Memcpy( &VarAddress, data, sizeof(OP::ADDRESS));
+				data += sizeof(OP::ADDRESS);
+
+                if (VarAddress)
+                {
+					OP::ADDRESS DefAddress;
+					FMemory::Memcpy( &DefAddress, data, sizeof(OP::ADDRESS));
+					data += sizeof(OP::ADDRESS);
+
+					uint32 CaseCount;
+					FMemory::Memcpy( &CaseCount, data, sizeof(uint32));
+					data += sizeof(uint32);
+
+                    PARENT::RecurseWithCurrentState( VarAddress );
+
+                    int var = m_pSystem->BuildInt( m_pModel, m_pParams, VarAddress );
+
+					OP::ADDRESS valueAt = DefAddress;
+					for (uint32 C = 0; C < CaseCount; ++C)
+					{
+						int32 Condition;
+						FMemory::Memcpy( &Condition, data, sizeof(int32));		
+						data += sizeof(int32);
+
+						OP::ADDRESS At;
+						FMemory::Memcpy( &At, data, sizeof(OP::ADDRESS));
+						data += sizeof(OP::ADDRESS);
+
+						if (At && var == (int)Condition)
+						{
+							valueAt = At;
+							break;	
+						}
+					}
+
+					PARENT::RecurseWithCurrentState( valueAt );
+                }
+
+                break;
+            }
+
+
+            case OP_TYPE::IN_ADDLOD:
+            {
+				OP::InstanceAddLODArgs args = program.GetOpArgs<OP::InstanceAddLODArgs>(at);
+
+                recurse = false;
+
+                STATE newState = PARENT::GetCurrentState();
+                for (int32 t=0;t<MUTABLE_OP_MAX_ADD_COUNT;++t)
+                {
+                    OP::ADDRESS lodAt = args.lod[t];
+                    if (lodAt)
+                    {
+                        bool selected = ( (1<<t) & m_lodMask ) != 0;
+                        if ( selected )
+                        {
+                            PARENT::RecurseWithState( lodAt, newState );
+                        }
+                    }
+                }
+                break;
+            }
+
+
+            case OP_TYPE::IN_ADDMESH:
+            {
+				OP::InstanceAddArgs args = program.GetOpArgs<OP::InstanceAddArgs>(at);
+
+                recurse = false;
+
+                PARENT::RecurseWithCurrentState(args.instance);
+
+                STATE newState = PARENT::GetCurrentState();
+                newState.m_underResourceCount=1;
+
+                OP::ADDRESS meshAt = args.value;
+                if (meshAt)
+                {
+                    PARENT::RecurseWithState(meshAt, newState);
+                }
+                break;
+            }
+
+
+            case OP_TYPE::IN_ADDIMAGE:
+            {
+				OP::InstanceAddArgs args = program.GetOpArgs<OP::InstanceAddArgs>(at);
+
+                recurse = false;
+
+                PARENT::RecurseWithCurrentState(args.instance);
+
+                STATE newState = PARENT::GetCurrentState();
+                newState.m_underResourceCount=1;
+
+                OP::ADDRESS imageAt = args.value;
+                if (imageAt)
+                {
+                    PARENT::RecurseWithState(imageAt, newState);
+                }
+                break;
+            }
+
+            default:
+                break;
+            }
+
+            return recurse;
+        }
+
+
+    protected:
+        System::Private* m_pSystem = nullptr;
+		TSharedPtr<const Model> m_pModel;
+        const Parameters* m_pParams = nullptr;
+        unsigned m_lodMask = 0;
+    };
+
+
+    //---------------------------------------------------------------------------------------------
+    //! Code visitor that:
+    //! - only traverses the operations that are relevant for a given set of parameter values. It
+    //! only considers the discrete parameters like integers and booleans. In the case of forks
+    //! caused by continuous parameters like float weights for interpolation, all the branches are
+    //! traversed.
+    //! - cannot change the instructions
+    //! - will not repeat visits to instructions with the same state
+    //! - the state has to be a compatible with COVERED_CODE_VISITOR_STATE
+    //---------------------------------------------------------------------------------------------
+    template<typename COVERED_STATE=COVERED_CODE_VISITOR_STATE>
+    class UniqueDiscreteCoveredCodeVisitor :
+            public DiscreteCoveredCodeVisitorBase
+            <
+            UniqueConstCodeVisitorIterative<COVERED_STATE>,
+            COVERED_STATE
+            >
+    {
+        using PARENT=DiscreteCoveredCodeVisitorBase<UniqueConstCodeVisitorIterative<COVERED_STATE>, COVERED_STATE>;
+
+    public:
+
+        UniqueDiscreteCoveredCodeVisitor
+            (
+                System::Private* InSystem,
+				const TSharedPtr<const Model>& InModel,
+                const ParametersPtrConst& InParams,
+                uint32 InLodMask
+            )
+            : PARENT(InSystem, InModel, InParams, InLodMask )
+        {
+        }
+
+    };
+
+
+
+    //---------------------------------------------------------------------------------------------
+    //! Calculate all the parameters found under a particular operation
+    //! It has an internal cache, so don't reuse objects of this class if the program changes.
+    //---------------------------------------------------------------------------------------------
+    class MUTABLERUNTIME_API SubtreeParametersVisitor
+    {
+    public:
+
+        void Run( OP::ADDRESS Root, const FProgram& );
+
+        //! After Run, list of relevant parameters.
+		TArray<int32> RelevantParams;
+
+    private:
+
+		TArray<int32> CurrentParams;
+		TArray<uint8> Visited;
+		TArray<OP::ADDRESS> Pending;
+
+        // Result cache
+        TMap< OP::ADDRESS, TArray<int32> > ResultCache;
+    };
+
+
+}
+
diff --git a/Engine/Plugins/Experimental/Mutable/Source/MutableRuntime/Internal/MuR/ConvertData.h b/Engine/Plugins/Experimental/Mutable/Source/MutableRuntime/Internal/MuR/ConvertData.h
new file mode 100644
--- /dev/null
+++ b/Engine/Plugins/Experimental/Mutable/Source/MutableRuntime/Internal/MuR/ConvertData.h
@@ -0,0 +1,1339 @@
+// Copyright Epic Games, Inc. All Rights Reserved.
+
+#pragma once
+
+#include "MuR/Mesh.h"
+#include "MuR/Platform.h"
+#include "MuR/MutableMath.h"
+
+
+namespace mu
+{
+
+	//---------------------------------------------------------------------------------------------
+	//! Convert one channel element
+	//---------------------------------------------------------------------------------------------
+	inline void ConvertData
+		(
+			int channel,
+			void* pResult, MESH_BUFFER_FORMAT resultFormat,
+			const void* pSource, MESH_BUFFER_FORMAT sourceFormat
+		)
+	{
+		switch ( resultFormat )
+		{
+		case MBF_FLOAT64:
+		{
+			double* pTypedResult = reinterpret_cast<double*>(pResult);
+			uint8* pByteResult = reinterpret_cast<uint8*>(pResult);
+			const uint8* pByteSource = reinterpret_cast<const uint8*>(pSource);
+
+			switch (sourceFormat)
+			{
+			case MBF_FLOAT64:
+			{
+				// Just dereferencing is not safe in all architectures. some like ARM require the
+				// floats to be 4-byte aligned and it may not be the case. We need memcpy.
+				memcpy(pByteResult + 8 * channel, pByteSource + 8* channel, 8);
+				break;
+			}
+
+			case MBF_FLOAT32:
+			{
+				const float* pTypedSource = reinterpret_cast<const float*>(pSource);
+				pTypedResult[channel] = pTypedSource[channel];
+				break;
+			}
+
+			case MBF_FLOAT16:
+			{
+				const float16* pTypedSource = reinterpret_cast<const float16*>(pSource);
+				pTypedResult[channel] = halfToFloat(pTypedSource[channel]);
+				break;
+			}
+
+			case MBF_INT32:
+			{
+				const int32* pTypedSource = reinterpret_cast<const int32*>(pSource);
+				pTypedResult[channel] = (double)(pTypedSource[channel]);
+				break;
+			}
+
+			case MBF_UINT32:
+			{
+				const uint32* pTypedSource = reinterpret_cast<const uint32*>(pSource);
+				pTypedResult[channel] = (double)(pTypedSource[channel]);
+				break;
+			}
+
+			case MBF_INT16:
+			{
+				const int16* pTypedSource = reinterpret_cast<const int16*>(pSource);
+				pTypedResult[channel] = (double)(pTypedSource[channel]);
+				break;
+			}
+
+			case MBF_UINT16:
+			{
+				const uint16* pTypedSource = reinterpret_cast<const uint16*>(pSource);
+				pTypedResult[channel] = (double)(pTypedSource[channel]);
+				break;
+			}
+
+			case MBF_INT8:
+			{
+				const int8* pTypedSource = reinterpret_cast<const int8*>(pSource);
+				pTypedResult[channel] = (double)(pTypedSource[channel]);
+				break;
+			}
+
+			case MBF_UINT8:
+			{
+				const uint8* pTypedSource = reinterpret_cast<const uint8*>(pSource);
+				pTypedResult[channel] = (double)(pTypedSource[channel]);
+				break;
+			}
+
+			case MBF_NINT32:
+			{
+				const int32* pTypedSource = reinterpret_cast<const int32*>(pSource);
+				pTypedResult[channel] = (double)(pTypedSource[channel]);
+				pTypedResult[channel] /= 65536.0f * 65536.0f / 2.0f;
+				break;
+			}
+
+			case MBF_NUINT32:
+			{
+				const uint32* pTypedSource = reinterpret_cast<const uint32*>(pSource);
+				pTypedResult[channel] = (double)(pTypedSource[channel]);
+				pTypedResult[channel] /= 65536.0f * 65536.0f - 1.0f;
+				break;
+			}
+
+			case MBF_NINT16:
+			{
+				const int16* pTypedSource = reinterpret_cast<const int16*>(pSource);
+				pTypedResult[channel] = (double)(pTypedSource[channel]);
+				pTypedResult[channel] /= 32768.0f;
+				break;
+			}
+
+			case MBF_NUINT16:
+			{
+				const uint16* pTypedSource = reinterpret_cast<const uint16*>(pSource);
+				pTypedResult[channel] = (double)(pTypedSource[channel]);
+				pTypedResult[channel] /= 65535.0f;
+				break;
+			}
+
+			case MBF_NINT8:
+			{
+				const int8* pTypedSource = reinterpret_cast<const int8*>(pSource);
+				pTypedResult[channel] = (double)(pTypedSource[channel]);
+				pTypedResult[channel] /= 128.0f;
+				break;
+			}
+
+			case MBF_NUINT8:
+			{
+				const uint8* pTypedSource = reinterpret_cast<const uint8*>(pSource);
+				pTypedResult[channel] = (double)(pTypedSource[channel]);
+				pTypedResult[channel] /= 255.0f;
+				break;
+			}
+
+			case MBF_PACKEDDIR8:
+			case MBF_PACKEDDIR8_W_TANGENTSIGN:
+			{
+				const uint8* pTypedSource = reinterpret_cast<const uint8*>(pSource);
+				pTypedResult[channel] = (double)(pTypedSource[channel]);
+				pTypedResult[channel] /= 127.5f;
+				pTypedResult[channel] -= 1.0f;
+				break;
+			}
+
+			case MBF_PACKEDDIRS8:
+			case MBF_PACKEDDIRS8_W_TANGENTSIGN:
+			{
+				const int8* pTypedSource = reinterpret_cast<const int8*>(pSource);
+				pTypedResult[channel] = (double)(pTypedSource[channel]);
+				pTypedResult[channel] /= 127.5f;
+				break;
+			}
+
+			default:
+				checkf(false, TEXT("Conversion not implemented."));
+				break;
+			}
+			break;
+		}
+
+		case MBF_FLOAT32:
+		{
+			float* pTypedResult = reinterpret_cast<float*>( pResult );
+            uint8* pByteResult = reinterpret_cast<uint8*>( pResult );
+            const uint8* pByteSource = reinterpret_cast<const uint8*>( pSource );
+
+			switch ( sourceFormat )
+			{
+			case MBF_FLOAT64:
+			{
+				const double* pTypedSource = reinterpret_cast<const double*>(pSource);
+				pTypedResult[channel] = float(pTypedSource[channel]);
+				break;
+			}
+
+			case MBF_FLOAT32:
+			{
+				// Just dereferencing is not safe in all architectures. some like ARM require the
+				// floats to be 4-byte aligned and it may not be the case. We need memcpy.
+				memcpy(pByteResult + 4 * channel, pByteSource + 4 * channel, 4);
+				break;
+			}
+
+			case MBF_FLOAT16:
+			{
+				const float16* pTypedSource = reinterpret_cast<const float16*>( pSource );
+				pTypedResult[channel] = halfToFloat(pTypedSource[channel]);
+				break;
+			}
+
+			case MBF_INT32:
+			{
+                const int32* pTypedSource = reinterpret_cast<const int32*>( pSource );
+				pTypedResult[channel] = (float)(pTypedSource[channel]);
+				break;
+			}
+
+			case MBF_UINT32:
+			{
+                const uint32* pTypedSource = reinterpret_cast<const uint32*>( pSource );
+				pTypedResult[channel] = (float)(pTypedSource[channel]);
+				break;
+			}
+
+			case MBF_INT16:
+			{
+                const int16* pTypedSource = reinterpret_cast<const int16*>( pSource );
+				pTypedResult[channel] = (float)(pTypedSource[channel]);
+				break;
+			}
+
+			case MBF_UINT16:
+			{
+                const uint16* pTypedSource = reinterpret_cast<const uint16*>( pSource );
+				pTypedResult[channel] = (float)(pTypedSource[channel]);
+				break;
+			}
+
+			case MBF_INT8:
+			{
+                const int8* pTypedSource = reinterpret_cast<const int8*>( pSource );
+				pTypedResult[channel] = (float)(pTypedSource[channel]);
+				break;
+			}
+
+			case MBF_UINT8:
+			{
+                const uint8* pTypedSource = reinterpret_cast<const uint8*>( pSource );
+				pTypedResult[channel] = (float)(pTypedSource[channel]);
+				break;
+			}
+
+			case MBF_NINT32:
+			{
+                const int32* pTypedSource = reinterpret_cast<const int32*>( pSource );
+				pTypedResult[channel] = (float)(pTypedSource[channel]);
+				pTypedResult[channel] /= 65536.0f*65536.0f/2.0f;
+				break;
+			}
+
+			case MBF_NUINT32:
+			{
+                const uint32* pTypedSource = reinterpret_cast<const uint32*>( pSource );
+				pTypedResult[channel] = (float)(pTypedSource[channel]);
+				pTypedResult[channel] /= 65536.0f*65536.0f-1.0f;
+				break;
+			}
+
+			case MBF_NINT16:
+			{
+                const int16* pTypedSource = reinterpret_cast<const int16*>( pSource );
+				pTypedResult[channel] = (float)(pTypedSource[channel]);
+				pTypedResult[channel] /= 32768.0f;
+				break;
+			}
+
+			case MBF_NUINT16:
+			{
+                const uint16* pTypedSource = reinterpret_cast<const uint16*>( pSource );
+				pTypedResult[channel] = (float)(pTypedSource[channel]);
+				pTypedResult[channel] /= 65535.0f;
+				break;
+			}
+
+			case MBF_NINT8:
+			{
+                const int8* pTypedSource = reinterpret_cast<const int8*>( pSource );
+				pTypedResult[channel] = (float)(pTypedSource[channel]);
+				pTypedResult[channel] /= 128.0f;
+				break;
+			}
+
+			case MBF_NUINT8:
+			{
+                const uint8* pTypedSource = reinterpret_cast<const uint8*>( pSource );
+				pTypedResult[channel] = (float)(pTypedSource[channel]);
+				pTypedResult[channel] /= 255.0f;
+				break;
+			}
+
+            case MBF_PACKEDDIR8:
+            case MBF_PACKEDDIR8_W_TANGENTSIGN:
+            {
+                const uint8* pTypedSource = reinterpret_cast<const uint8*>( pSource );
+                pTypedResult[channel] = (float)(pTypedSource[channel]);
+                pTypedResult[channel] /= 127.5f;
+                pTypedResult[channel] -= 1.0f;
+                break;
+            }
+
+            case MBF_PACKEDDIRS8:
+            case MBF_PACKEDDIRS8_W_TANGENTSIGN:
+            {
+                const int8* pTypedSource = reinterpret_cast<const int8*>( pSource );
+                pTypedResult[channel] = (float)(pTypedSource[channel]);
+                pTypedResult[channel] /= 127.5f;
+                break;
+            }
+
+			default:
+				checkf( false, TEXT("Conversion not implemented.") );
+				break;
+			}
+			break;
+		}
+
+		//-----------------------------------------------------------------------------------------
+		case MBF_FLOAT16:
+		{
+			float16* pTypedResult = reinterpret_cast<float16*>( pResult );
+
+			switch ( sourceFormat )
+			{
+			case MBF_FLOAT32:
+			{
+				const float* pTypedSource = reinterpret_cast<const float*>( pSource );
+				pTypedResult[channel] = floatToHalf(pTypedSource[channel]);
+				break;
+			}
+
+			case MBF_FLOAT16:
+			{
+				const float16* pTypedSource = reinterpret_cast<const float16*>( pSource );
+				pTypedResult[channel] = pTypedSource[channel];
+				break;
+			}
+
+			case MBF_UINT32:
+			{
+                const uint32* pTypedSource = reinterpret_cast<const uint32*>( pSource );
+				pTypedResult[channel] = floatToHalf((float)(pTypedSource[channel]));
+				break;
+			}
+
+			case MBF_INT32:
+			{
+                const int32* pTypedSource = reinterpret_cast<const int32*>( pSource );
+				pTypedResult[channel] = floatToHalf((float)(pTypedSource[channel]));
+				break;
+			}
+
+			case MBF_UINT16:
+			{
+                const uint16* pTypedSource = reinterpret_cast<const uint16*>( pSource );
+				pTypedResult[channel] = floatToHalf((float)(pTypedSource[channel]));
+				break;
+			}
+
+			case MBF_INT16:
+			{
+                const int16* pTypedSource = reinterpret_cast<const int16*>( pSource );
+				pTypedResult[channel] = floatToHalf((float)(pTypedSource[channel]));
+				break;
+			}
+
+			case MBF_UINT8:
+			{
+                const uint8* pTypedSource = reinterpret_cast<const uint8*>( pSource );
+				pTypedResult[channel] = floatToHalf((float)(pTypedSource[channel]));
+				break;
+			}
+
+			case MBF_INT8:
+			{
+                const int8* pTypedSource = reinterpret_cast<const int8*>( pSource );
+				pTypedResult[channel] = floatToHalf((float)(pTypedSource[channel]));
+				break;
+			}
+
+			default:
+				checkf( false, TEXT("Conversion not implemented.") );
+				break;
+			}
+			break;
+		}
+
+		//-----------------------------------------------------------------------------------------
+		case MBF_UINT8:
+		{
+            uint8* pTypedResult = reinterpret_cast<uint8*>( pResult );
+
+			switch ( sourceFormat )
+			{
+			case MBF_FLOAT32:
+			{
+				const float* pTypedSource = reinterpret_cast<const float*>( pSource );
+                pTypedResult[channel] = (uint8)
+                        FMath::Min<uint32>(
+							0xFF,
+                            FMath::Max<uint32>(
+								0,
+                                (uint32)(pTypedSource[channel])
+								)
+							);
+				break;
+			}
+
+			case MBF_FLOAT16:
+			{
+				const float16* pTypedSource = reinterpret_cast<const float16*>( pSource );
+                pTypedResult[channel] = (uint8)
+                        FMath::Min<uint32>(
+							0xFF,
+                            FMath::Max<uint32>(
+								0,
+                                (uint32)halfToFloat(pTypedSource[channel])
+								)
+							);
+				break;
+			}
+
+			case MBF_INT8:
+			{
+                const int8* pTypedSource = reinterpret_cast<const int8*>( pSource );
+                pTypedResult[channel] =  (uint8)FMath::Max<int8>( 0, pTypedSource[channel] );
+				break;
+			}
+
+			case MBF_UINT8:
+			{
+                const uint8* pTypedSource = reinterpret_cast<const uint8*>( pSource );
+				pTypedResult[channel] = pTypedSource[channel];
+				break;
+			}
+
+			case MBF_INT16:
+			{
+                const int16* pTypedSource = reinterpret_cast<const int16*>( pSource );
+                pTypedResult[channel] =  (uint8)
+                        FMath::Min<int16>(
+							0xFF,
+                            FMath::Max<int16>(
+								0,
+								pTypedSource[channel]
+								)
+							);
+				break;
+			}
+
+			case MBF_UINT16:
+			{
+				// Clamp
+                const uint16* pTypedSource = reinterpret_cast<const uint16*>( pSource );
+                pTypedResult[channel] = (uint8)
+                        FMath::Min<uint16>(
+							0xFF,
+                            FMath::Max<uint16>(
+								0,
+								pTypedSource[channel]
+								)
+							);
+				break;
+			}
+
+			case MBF_INT32:
+			{
+                const int32* pTypedSource = reinterpret_cast<const int32*>( pSource );
+                pTypedResult[channel] =  (uint8)
+                        FMath::Min<int32>(
+							0xFF,
+                            FMath::Max<int32>(
+								0,
+								pTypedSource[channel]
+								)
+							);
+				break;
+			}
+
+			case MBF_UINT32:
+			{
+				// Clamp
+                const uint32* pTypedSource = reinterpret_cast<const uint32*>( pSource );
+                pTypedResult[channel] = (uint8)
+                        FMath::Min<uint32>(
+							0xFF,
+                            FMath::Max<uint32>(
+								0,
+								pTypedSource[channel]
+								)
+							);
+				break;
+			}
+
+			case MBF_NUINT8:
+			case MBF_NUINT16:
+			case MBF_NUINT32:
+			case MBF_NINT8:
+			case MBF_NINT16:
+			case MBF_NINT32:
+				pTypedResult = 0;
+				break;
+
+			default:
+				checkf( false, TEXT("Conversion not implemented." ) );
+				break;
+			}
+			break;
+		}
+
+		case MBF_UINT16:
+		{
+            uint16* pTypedResult = reinterpret_cast<uint16*>( pResult );
+
+			switch ( sourceFormat )
+			{
+			case MBF_FLOAT32:
+			{
+				const float* pTypedSource = reinterpret_cast<const float*>( pSource );
+                pTypedResult[channel] = (uint16)
+                        FMath::Min<uint32>(
+							0xFFFF,
+                            FMath::Max<uint32>(
+								0,
+                                (uint32)(pTypedSource[channel])
+								)
+							);
+				break;
+			}
+
+			case MBF_FLOAT16:
+			{
+				const float16* pTypedSource = reinterpret_cast<const float16*>( pSource );
+                pTypedResult[channel] = (uint16)
+                        FMath::Min<uint32>(
+							0xFF,
+                            FMath::Max<uint32>(
+								0,
+                                (uint32)halfToFloat(pTypedSource[channel])
+								)
+							);
+				break;
+			}
+
+			case MBF_UINT8:
+			{
+                const uint8* pTypedSource = reinterpret_cast<const uint8*>( pSource );
+				pTypedResult[channel] =  pTypedSource[channel];
+				break;
+			}
+
+			case MBF_INT8:
+			{
+                const int8* pTypedSource = reinterpret_cast<const int8*>( pSource );
+                pTypedResult[channel] =  (uint16)FMath::Max<int8>( 0, pTypedSource[channel] );
+				break;
+			}
+
+			case MBF_UINT16:
+			{
+                const uint16* pTypedSource = reinterpret_cast<const uint16*>( pSource );
+				pTypedResult[channel] =  pTypedSource[channel];
+				break;
+			}
+
+			case MBF_INT16:
+			{
+                const int16* pTypedSource = reinterpret_cast<const int16*>( pSource );
+                pTypedResult[channel] =  (uint16)FMath::Max<int16>( 0, pTypedSource[channel] );
+				break;
+			}
+
+			case MBF_UINT32:
+			{
+				// Clamp
+                const uint32* pTypedSource = reinterpret_cast<const uint32*>( pSource );
+                pTypedResult[channel] = (uint16)
+                        FMath::Min<uint32>(
+							0xFFFF,
+                            FMath::Max<uint32>(
+								0,
+								pTypedSource[channel]
+								)
+							);
+				break;
+			}
+
+			case MBF_INT32:
+			{
+				// Clamp
+                const int32* pTypedSource = reinterpret_cast<const int32*>( pSource );
+                pTypedResult[channel] = (uint16)
+                        FMath::Min<int32>(
+							0xFFFF,
+                            FMath::Max<int32>(
+								0,
+								pTypedSource[channel]
+								)
+							);
+				break;
+			}
+
+			case MBF_NUINT8:
+			case MBF_NUINT16:
+			case MBF_NUINT32:
+			case MBF_NINT8:
+			case MBF_NINT16:
+			case MBF_NINT32:
+				pTypedResult = 0;
+				break;
+
+			default:
+				checkf( false, TEXT("Conversion not implemented." ) );
+				break;
+			}
+			break;
+		}
+
+		case MBF_UINT32:
+		{
+            uint32* pTypedResult = reinterpret_cast<uint32*>( pResult );
+
+			switch ( sourceFormat )
+			{
+			case MBF_FLOAT32:
+			{
+				const float* pTypedSource = reinterpret_cast<const float*>( pSource );
+				pTypedResult[channel] =
+                        FMath::Min<uint32>(
+							0xFFFF,
+                            FMath::Max<uint32>(
+								0,
+                                (uint32)(pTypedSource[channel])
+								)
+							);
+				break;
+			}
+
+			case MBF_FLOAT16:
+			{
+				const float16* pTypedSource = reinterpret_cast<const float16*>( pSource );
+				pTypedResult[channel] =
+                        FMath::Min<uint32>(
+							0xFF,
+                            FMath::Max<uint32>(
+								0,
+                                (uint32)halfToFloat(pTypedSource[channel])
+								)
+							);
+				break;
+			}
+
+			case MBF_UINT8:
+			{
+                const uint8* pTypedSource = reinterpret_cast<const uint8*>( pSource );
+				pTypedResult[channel] =  pTypedSource[channel];
+				break;
+			}
+
+			case MBF_INT8:
+			{
+                const int8* pTypedSource = reinterpret_cast<const int8*>( pSource );
+                pTypedResult[channel] =  (uint16)FMath::Max<int8>( 0, pTypedSource[channel] );
+				break;
+			}
+
+			case MBF_UINT16:
+			{
+                const uint16* pTypedSource = reinterpret_cast<const uint16*>( pSource );
+				pTypedResult[channel] =  pTypedSource[channel];
+				break;
+			}
+
+			case MBF_INT16:
+			{
+                const int16* pTypedSource = reinterpret_cast<const int16*>( pSource );
+                pTypedResult[channel] =  (uint16)FMath::Max<int16>( 0, pTypedSource[channel] );
+				break;
+			}
+
+			case MBF_UINT32:
+			{
+                const uint32* pTypedSource = reinterpret_cast<const uint32*>( pSource );
+				pTypedResult[channel] = pTypedSource[channel];
+				break;
+			}
+
+			case MBF_INT32:
+			{
+				// Clamp
+                const int32* pTypedSource = reinterpret_cast<const int32*>( pSource );
+                pTypedResult[channel] = (uint32)FMath::Max<int32>( 0, pTypedSource[channel] );
+				break;
+			}
+
+			case MBF_NUINT8:
+			case MBF_NUINT16:
+			case MBF_NUINT32:
+			case MBF_NINT8:
+			case MBF_NINT16:
+			case MBF_NINT32:
+				pTypedResult = 0;
+				break;
+
+			default:
+				checkf( false, TEXT("Conversion not implemented." ) );
+				break;
+			}
+			break;
+		}
+
+		case MBF_INT8:
+		{
+            int8* pTypedResult = reinterpret_cast<int8*>( pResult );
+
+			switch ( sourceFormat )
+			{
+			case MBF_FLOAT32:
+			{
+				const float* pTypedSource = reinterpret_cast<const float*>( pSource );
+                pTypedResult[channel] = (int8)
+                        FMath::Min<int32>(
+							127,
+                            FMath::Max<int32>(
+								-128,
+                                (int32)(pTypedSource[channel])
+								)
+							);
+				break;
+			}
+
+			case MBF_FLOAT16:
+			{
+				const float16* pTypedSource = reinterpret_cast<const float16*>( pSource );
+                pTypedResult[channel] = (int8)
+                        FMath::Min<int32>(
+							127,
+                            FMath::Max<int32>(
+								-128,
+                                (int32)halfToFloat(pTypedSource[channel])
+								)
+							);
+				break;
+			}
+
+			case MBF_INT8:
+			{
+                const int8* pTypedSource = reinterpret_cast<const int8*>( pSource );
+				pTypedResult[channel] = pTypedSource[channel];
+				break;
+			}
+
+			case MBF_NUINT8:
+			case MBF_NUINT16:
+			case MBF_NUINT32:
+			case MBF_NINT8:
+			case MBF_NINT16:
+			case MBF_NINT32:
+				pTypedResult = 0;
+				break;
+
+			default:
+				checkf( false, TEXT("Conversion not implemented." ) );
+				break;
+			}
+			break;
+		}
+
+		case MBF_INT16:
+		{
+            int16* pTypedResult = reinterpret_cast<int16*>( pResult );
+
+			switch ( sourceFormat )
+			{
+			case MBF_FLOAT32:
+			{
+				const float* pTypedSource = reinterpret_cast<const float*>( pSource );
+                pTypedResult[channel] = (int16)
+                        FMath::Min<int32>(
+							32767,
+                            FMath::Max<int32>(
+								-32768,
+                                (int32)(pTypedSource[channel])
+								)
+							);
+				break;
+			}
+
+			case MBF_FLOAT16:
+			{
+				const float16* pTypedSource = reinterpret_cast<const float16*>( pSource );
+                pTypedResult[channel] = (int16)
+                        FMath::Min<int32>(
+							32767,
+                            FMath::Max<int32>(
+								-32768,
+                                (int32)halfToFloat(pTypedSource[channel])
+								)
+							);
+				break;
+			}
+
+			case MBF_INT8:
+			{
+                const int8* pTypedSource = reinterpret_cast<const int8*>( pSource );
+                pTypedResult[channel] = (int16)pTypedSource[channel];
+				break;
+			}
+
+			case MBF_UINT8:
+			{
+                const uint8* pTypedSource = reinterpret_cast<const uint8*>( pSource );
+                pTypedResult[channel] = (int16)pTypedSource[channel];
+				break;
+			}
+
+			case MBF_UINT16:
+			{
+                const uint16* pTypedSource = reinterpret_cast<const uint16*>( pSource );
+                pTypedResult[channel] = (int16)
+                        FMath::Min<int32>(
+                            32767, (int32)pTypedSource[channel]
+							);
+				break;
+			}
+
+			case MBF_INT32:
+			{
+                const int32* pTypedSource = reinterpret_cast<const int32*>( pSource );
+                pTypedResult[channel] = (int16)
+                        FMath::Min<int32>(
+							32767,
+                            FMath::Max<int32>(
+								-32768,
+								pTypedSource[channel]
+								)
+							);
+				break;
+			}
+
+			case MBF_UINT32:
+			{
+                const uint32* pTypedSource = reinterpret_cast<const uint32*>( pSource );
+                pTypedResult[channel] = (int16)
+                        FMath::Min<int32>(
+							32767, pTypedSource[channel]
+							);
+				break;
+			}
+
+			case MBF_NUINT8:
+			case MBF_NUINT16:
+			case MBF_NUINT32:
+			case MBF_NINT8:
+			case MBF_NINT16:
+			case MBF_NINT32:
+				pTypedResult = 0;
+				break;
+
+			default:
+				checkf( false, TEXT("Conversion not implemented." ) );
+				break;
+			}
+			break;
+		}
+
+		case MBF_INT32:
+		{
+            int32* pTypedResult = reinterpret_cast<int32*>( pResult );
+
+			switch ( sourceFormat )
+			{
+			case MBF_FLOAT32:
+			{
+				const float* pTypedSource = reinterpret_cast<const float*>( pSource );
+                pTypedResult[channel] = (int32)(pTypedSource[channel]);
+				break;
+			}
+
+			case MBF_FLOAT16:
+			{
+				const float16* pTypedSource = reinterpret_cast<const float16*>( pSource );
+                pTypedResult[channel] = (int32)halfToFloat(pTypedSource[channel]);
+				break;
+			}
+
+			case MBF_INT8:
+			{
+                const int8* pTypedSource = reinterpret_cast<const int8*>( pSource );
+                pTypedResult[channel] = (int32)pTypedSource[channel];
+				break;
+			}
+
+			case MBF_UINT8:
+			{
+                const uint8* pTypedSource = reinterpret_cast<const uint8*>( pSource );
+                pTypedResult[channel] = (int32)pTypedSource[channel];
+				break;
+			}
+
+			case MBF_INT16:
+			{
+                const int16* pTypedSource = reinterpret_cast<const int16*>( pSource );
+                pTypedResult[channel] = (int32)pTypedSource[channel];
+				break;
+			}
+
+			case MBF_UINT16:
+			{
+                const uint16* pTypedSource = reinterpret_cast<const uint16*>( pSource );
+                pTypedResult[channel] = (int32)pTypedSource[channel];
+				break;
+			}
+
+			case MBF_UINT32:
+			{
+                const uint32* pTypedSource = reinterpret_cast<const uint32*>( pSource );
+                pTypedResult[channel] = (int32)pTypedSource[channel];
+				break;
+			}
+
+			case MBF_INT32:
+			{
+                const int32* pTypedSource = reinterpret_cast<const int32*>( pSource );
+                pTypedResult[channel] = pTypedSource[channel];
+				break;
+			}
+
+			case MBF_NUINT8:
+			case MBF_NUINT16:
+			case MBF_NUINT32:
+			case MBF_NINT8:
+			case MBF_NINT16:
+			case MBF_NINT32:
+				pTypedResult = 0;
+				break;
+
+			default:
+				checkf( false, TEXT("Conversion not implemented." ) );
+				break;
+			}
+			break;
+		}
+
+		//-----------------------------------------------------------------------------------------
+		case MBF_NUINT8:
+		{
+            uint8* pTypedResult = reinterpret_cast<uint8*>( pResult );
+
+			switch ( sourceFormat )
+			{
+			case MBF_NUINT8:
+			{
+				const uint8* pTypedSource = reinterpret_cast<const uint8*>(pSource);
+				pTypedResult[channel] = pTypedSource[channel];
+				break;
+			}
+
+			case MBF_NUINT16:
+			{
+				const uint16* pTypedSource = reinterpret_cast<const uint16*>(pSource);
+				pTypedResult[channel] = pTypedSource[channel] / (65535 / 255);
+				break;
+			}
+
+			case MBF_FLOAT32:
+			{
+				const float* pTypedSource = reinterpret_cast<const float*>(pSource);
+				pTypedResult[channel] = (uint8)
+					FMath::Min<uint32>(
+						0xFF,
+						FMath::Max<uint32>(
+							0,
+							(uint32)(((float)0xFF)*pTypedSource[channel] + 0.5f)
+							)
+						);
+				break;
+			}
+
+			case MBF_FLOAT16:
+			{
+				const float16* pTypedSource = reinterpret_cast<const float16*>( pSource );
+                pTypedResult[channel] = (uint8)
+                        FMath::Min<uint32>(
+							0xFF,
+                            FMath::Max<uint32>(
+								0,
+                                (uint32)(((float)0xFF)*halfToFloat(pTypedSource[channel])+0.5f)
+								)
+							);
+				break;
+			}
+
+			case MBF_UINT8:
+			case MBF_UINT16:
+			case MBF_UINT32:
+			case MBF_INT8:
+			case MBF_INT16:
+			case MBF_INT32:
+				// TODO: 1 or 0
+				pTypedResult = 0;
+				break;
+
+			default:
+				checkf( false, TEXT("Conversion not implemented." ) );
+				break;
+			}
+			break;
+		}
+
+		case MBF_NUINT16:
+		{
+            uint16* pTypedResult = reinterpret_cast<uint16*>( pResult );
+
+			switch ( sourceFormat )
+			{
+			case MBF_NUINT16:
+			{
+				const uint16* pTypedSource = reinterpret_cast<const uint16*>(pSource);
+				pTypedResult[channel] = pTypedSource[channel];
+				break;
+			}
+
+			case MBF_NUINT8:
+			{
+				const uint8* pTypedSource = reinterpret_cast<const uint8*>(pSource);
+				pTypedResult[channel] = pTypedSource[channel] * (65535 / 255);
+				break;
+			}
+
+			case MBF_FLOAT32:
+			{
+				const float* pTypedSource = reinterpret_cast<const float*>( pSource );
+                pTypedResult[channel] = (uint16)
+                        FMath::Min<uint32>(
+							0xFFFF,
+                            FMath::Max<uint32>(
+								0,
+                                (uint32)(((float)0xFFFF)*pTypedSource[channel]+0.5f)
+								)
+							);
+				break;
+			}
+
+			case MBF_FLOAT16:
+			{
+				const float16* pTypedSource = reinterpret_cast<const float16*>( pSource );
+                pTypedResult[channel] = (uint16)
+                        FMath::Min<uint32>(
+							0xFFFF,
+                            FMath::Max<uint32>(
+								0,
+                                (uint32)(((float)0xFFFF)*halfToFloat(pTypedSource[channel])+0.5f)
+								)
+							);
+				break;
+			}
+
+			case MBF_UINT8:
+			case MBF_UINT16:
+			case MBF_UINT32:
+			case MBF_INT8:
+			case MBF_INT16:
+			case MBF_INT32:
+				// TODO: 1 or 0
+				pTypedResult = 0;
+				break;
+
+			default:
+				checkf( false, TEXT("Conversion not implemented." ) );
+				break;
+			}
+			break;
+		}
+
+		case MBF_NUINT32:
+		{
+            uint32* pTypedResult = reinterpret_cast<uint32*>( pResult );
+
+			switch ( sourceFormat )
+			{
+			case MBF_FLOAT32:
+			{
+				const float* pTypedSource = reinterpret_cast<const float*>( pSource );
+                pTypedResult[channel] = (uint32)
+                        FMath::Min<uint32>(
+							0xFFFFFFFF,
+                            FMath::Max<uint32>(
+								0,
+                                (uint32)(((float)0xFFFFFFFF)*pTypedSource[channel])
+								)
+							);
+				break;
+			}
+
+			case MBF_FLOAT16:
+			{
+				const float16* pTypedSource = reinterpret_cast<const float16*>( pSource );
+                pTypedResult[channel] = (uint32)
+                        FMath::Min<uint32>(
+							0xFFFFFFFF,
+                            FMath::Max<uint32>(
+								0,
+                                (uint32)(((float)0xFFFFFFFF)*halfToFloat(pTypedSource[channel])+0.5f)
+								)
+							);
+				break;
+			}
+
+			case MBF_UINT8:
+			case MBF_UINT16:
+			case MBF_UINT32:
+			case MBF_INT8:
+			case MBF_INT16:
+			case MBF_INT32:
+				// TODO: 1 or 0
+				pTypedResult = 0;
+				break;
+
+			default:
+				checkf( false, TEXT("Conversion not implemented." ) );
+				break;
+			}
+			break;
+		}
+
+		case MBF_NINT8:
+		{
+            int8* pTypedResult = reinterpret_cast<int8*>( pResult );
+
+			switch ( sourceFormat )
+			{
+			case MBF_FLOAT32:
+			{
+				const float* pTypedSource = reinterpret_cast<const float*>( pSource );
+                pTypedResult[channel] = (int8)
+                        FMath::Min<int32>(
+							127,
+                            FMath::Max<int32>(
+								-128,
+                                (int32)(128.0f*pTypedSource[channel]+0.5f)
+								)
+							);
+				break;
+			}
+
+			case MBF_FLOAT16:
+			{
+				const float16* pTypedSource = reinterpret_cast<const float16*>( pSource );
+                pTypedResult[channel] = (int8)
+                        FMath::Min<int32>(
+							127,
+                            FMath::Max<int32>(
+								-128,
+                                (int32)(128.0f*halfToFloat(pTypedSource[channel])+0.5f)
+								)
+							);
+				break;
+			}
+
+			case MBF_UINT8:
+			case MBF_UINT16:
+			case MBF_UINT32:
+			case MBF_INT8:
+			case MBF_INT16:
+			case MBF_INT32:
+				// TODO: -1, 1 or 0
+				pTypedResult = 0;
+				break;
+
+			default:
+				checkf( false, TEXT("Conversion not implemented." ) );
+				break;
+			}
+			break;
+		}
+
+		case MBF_NINT16:
+		{
+            int16* pTypedResult = reinterpret_cast<int16*>( pResult );
+
+			switch ( sourceFormat )
+			{
+			case MBF_FLOAT32:
+			{
+				const float* pTypedSource = reinterpret_cast<const float*>( pSource );
+                pTypedResult[channel] = (int16)
+                        FMath::Min<int32>(
+							32767,
+                            FMath::Max<int32>(
+								-32768,
+                                (int32)(32768.0f*pTypedSource[channel]+0.5f)
+								)
+							);
+				break;
+			}
+
+			case MBF_FLOAT16:
+			{
+				const float16* pTypedSource = reinterpret_cast<const float16*>( pSource );
+                pTypedResult[channel] = (int16)
+                        FMath::Min<int32>(
+							32767,
+                            FMath::Max<int32>(
+								-32768,
+                                (int32)(32768.0f*halfToFloat(pTypedSource[channel])+0.5f)
+								)
+							);
+				break;
+			}
+
+			case MBF_UINT8:
+			case MBF_UINT16:
+			case MBF_UINT32:
+			case MBF_INT8:
+			case MBF_INT16:
+			case MBF_INT32:
+				// TODO: -1, 1 or 0
+				pTypedResult = 0;
+				break;
+
+			default:
+				checkf( false, TEXT("Conversion not implemented." ) );
+				break;
+			}
+			break;
+		}
+
+
+		case MBF_NINT32:
+		{
+            int32* pTypedResult = reinterpret_cast<int32*>( pResult );
+
+			switch ( sourceFormat )
+			{
+			case MBF_FLOAT32:
+			{
+				const float* pTypedSource = reinterpret_cast<const float*>( pSource );
+                pTypedResult[channel] = (int32)(2147483648.0f*pTypedSource[channel]+0.5f);
+				break;
+			}
+
+			case MBF_FLOAT16:
+			{
+				const float16* pTypedSource = reinterpret_cast<const float16*>( pSource );
+                pTypedResult[channel] = (int32)(2147483648.0f*halfToFloat(pTypedSource[channel])+0.5f);
+				break;
+			}
+
+			case MBF_UINT8:
+			case MBF_UINT16:
+			case MBF_UINT32:
+			case MBF_INT8:
+			case MBF_INT16:
+			case MBF_INT32:
+				// TODO: -1, 1 or 0
+				pTypedResult = 0;
+				break;
+
+			default:
+				checkf( false, TEXT("Conversion not implemented." ) );
+				break;
+			}
+			break;
+		}
+
+        case MBF_PACKEDDIR8:
+        case MBF_PACKEDDIR8_W_TANGENTSIGN:
+        {
+            uint8* pTypedResult = reinterpret_cast<uint8*>( pResult );
+
+            switch ( sourceFormat )
+            {
+            case MBF_PACKEDDIR8:
+            case MBF_PACKEDDIR8_W_TANGENTSIGN:
+            {
+                const uint8* pTypedSource = reinterpret_cast<const uint8*>( pSource );
+                uint8 source = pTypedSource[channel];
+                pTypedResult[channel] = source;
+                break;
+            }
+
+            case MBF_FLOAT32:
+            {
+                const float* pTypedSource = reinterpret_cast<const float*>( pSource );
+                float source = pTypedSource[channel];
+                source = (source*0.5f+0.5f)*255.0f;
+                pTypedResult[channel] =
+                        (uint8)FMath::Min<float>( 255.0f, FMath::Max<float>( 0.0f, source ) );
+                break;
+            }
+
+            default:
+                checkf( false, TEXT("Conversion not implemented." ) );
+                break;
+            }
+            break;
+        }
+
+        case MBF_PACKEDDIRS8:
+        case MBF_PACKEDDIRS8_W_TANGENTSIGN:
+        {
+            int8* pTypedResult = reinterpret_cast<int8*>( pResult );
+
+            switch ( sourceFormat )
+            {
+            case MBF_PACKEDDIRS8:
+            case MBF_PACKEDDIRS8_W_TANGENTSIGN:
+            {
+                const int8* pTypedSource = reinterpret_cast<const int8*>( pSource );
+                int8 source = pTypedSource[channel];
+                pTypedResult[channel] = source;
+                break;
+            }
+
+            case MBF_FLOAT32:
+            {
+                const float* pTypedSource = reinterpret_cast<const float*>( pSource );
+                float source = pTypedSource[channel];
+                source = source*0.5f*255.0f;
+                pTypedResult[channel] =
+                        (int8)FMath::Min<float>( 127.0f, FMath::Max<float>( -128.0f, source ) );
+                break;
+            }
+
+            default:
+                checkf( false, TEXT("Conversion not implemented." ) );
+                break;
+            }
+            break;
+        }
+
+		default:
+			checkf( false, TEXT("Conversion not implemented." ) );
+			break;
+		}
+
+	}
+
+}
diff --git a/Engine/Plugins/Experimental/Mutable/Source/MutableRuntime/Internal/MuR/ImagePrivate.h b/Engine/Plugins/Experimental/Mutable/Source/MutableRuntime/Internal/MuR/ImagePrivate.h
new file mode 100644
--- /dev/null
+++ b/Engine/Plugins/Experimental/Mutable/Source/MutableRuntime/Internal/MuR/ImagePrivate.h
@@ -0,0 +1,398 @@
+// Copyright Epic Games, Inc. All Rights Reserved.
+
+#pragma once
+
+#include "MuR/Image.h"
+#include "MuR/ImageTypes.h"
+
+#include "MuR/SerialisationPrivate.h"
+#include "MuR/MutableMath.h"
+
+#include "MuR/BlockCompression/Miro/Miro.h"
+
+namespace mu
+{
+	MUTABLE_DEFINE_ENUM_SERIALISABLE( EMipmapFilterType );
+	MUTABLE_DEFINE_ENUM_SERIALISABLE( EImageFormat );
+	MUTABLE_DEFINE_ENUM_SERIALISABLE( EAddressMode );
+
+
+    //---------------------------------------------------------------------------------------------
+    //!
+    //---------------------------------------------------------------------------------------------
+	inline EImageFormat GetMostGenericFormat(EImageFormat FormatA, EImageFormat FormatB)
+    {
+		if (FormatA == FormatB)
+		{
+			return FormatA;
+		}
+
+		if (FormatA == EImageFormat::IF_NONE)
+		{
+			return FormatA;
+		}
+
+		if (FormatB == EImageFormat::IF_NONE)
+		{
+			return FormatB;
+		}
+
+        if (GetImageFormatData(FormatA).Channels > GetImageFormatData(FormatB).Channels)
+		{	
+			return FormatA;
+		}
+
+        if (GetImageFormatData(FormatB).Channels > GetImageFormatData(FormatA).Channels)
+		{
+			return FormatB;
+		}
+
+        if (FormatA == EImageFormat::IF_BC2 || FormatA == EImageFormat::IF_BC3 
+			|| FormatA == EImageFormat::IF_ASTC_4x4_RGBA_LDR || FormatA == EImageFormat::IF_ASTC_6x6_RGBA_LDR || FormatA == EImageFormat::IF_ASTC_8x8_RGBA_LDR || FormatA == EImageFormat::IF_ASTC_10x10_RGBA_LDR)
+		{	
+			return FormatA;
+		}
+
+        if (FormatB == EImageFormat::IF_BC2 || FormatB == EImageFormat::IF_BC3
+			|| FormatB == EImageFormat::IF_ASTC_4x4_RGBA_LDR || FormatB == EImageFormat::IF_ASTC_6x6_RGBA_LDR || FormatB == EImageFormat::IF_ASTC_8x8_RGBA_LDR || FormatB == EImageFormat::IF_ASTC_10x10_RGBA_LDR)
+
+		{
+			return FormatB;
+		}
+
+        return FormatA;
+    }
+
+
+    //---------------------------------------------------------------------------------------------
+    //!
+    //---------------------------------------------------------------------------------------------
+	inline EImageFormat GetRGBOrRGBAFormat(EImageFormat InFormat)
+    {
+		InFormat = GetUncompressedFormat(InFormat);
+
+		if (InFormat == EImageFormat::IF_NONE)
+		{
+			return InFormat;
+		}
+
+		switch (InFormat)
+		{
+		case EImageFormat::IF_L_UBYTE: 
+		{
+			return EImageFormat::IF_RGB_UBYTE;
+		}
+		case EImageFormat::IF_RGB_UBYTE:
+		case EImageFormat::IF_RGBA_UBYTE:
+		case EImageFormat::IF_BGRA_UBYTE:
+		{
+			return InFormat;
+		}
+		default:
+		{
+			unimplemented();
+		}
+		}
+
+		return EImageFormat::IF_NONE;
+    }
+
+    //---------------------------------------------------------------------------------------------
+    //!
+    //---------------------------------------------------------------------------------------------
+    inline bool IsCompressedFormat(EImageFormat Format)
+    {
+        return Format != GetUncompressedFormat(Format);
+    }
+
+
+	inline bool IsBlockCompressedFormat(EImageFormat Format)
+	{
+		return GetImageFormatData(Format).PixelsPerBlockX > 1 && Format != GetUncompressedFormat(Format);
+	}
+
+	namespace Private
+	{
+		inline void DecompressionFuncNotFoundFunc(miro::FImageSize, miro::FImageSize, miro::FImageSize, const uint8*, uint8*)
+		{
+			check(false);
+		}
+	}
+
+	inline miro::SubImageDecompression::FuncRefType SelectDecompressionFunction(EImageFormat DestFormat, EImageFormat SrcFormat)
+	{
+		using DecompressionFuncRefType = miro::SubImageDecompression::FuncRefType;
+
+		switch(SrcFormat)
+		{
+			case EImageFormat::IF_BC1:
+			{
+				if (DestFormat == EImageFormat::IF_RGBA_UBYTE)
+				{
+					return DecompressionFuncRefType(miro::SubImageDecompression::BC1_To_RGBASubImage);
+				}
+				else if (DestFormat == EImageFormat::IF_RGB_UBYTE)
+				{
+					return DecompressionFuncRefType(miro::SubImageDecompression::BC1_To_RGBSubImage);
+				}
+				break;
+			}
+
+			case EImageFormat::IF_BC2:
+			{
+				if (DestFormat == EImageFormat::IF_RGBA_UBYTE)
+				{
+					return DecompressionFuncRefType(miro::SubImageDecompression::BC2_To_RGBASubImage);
+				}
+				else if (DestFormat == EImageFormat::IF_RGB_UBYTE)
+				{
+					return DecompressionFuncRefType(miro::SubImageDecompression::BC2_To_RGBSubImage);
+				}
+				break;
+			}
+
+			case EImageFormat::IF_BC3:
+			{
+				if (DestFormat == EImageFormat::IF_RGBA_UBYTE)
+				{
+					return DecompressionFuncRefType(miro::SubImageDecompression::BC3_To_RGBASubImage);
+				}
+				else if (DestFormat == EImageFormat::IF_RGB_UBYTE)
+				{
+					return DecompressionFuncRefType(miro::SubImageDecompression::BC3_To_RGBSubImage);
+				}
+				break;
+			}
+
+			case EImageFormat::IF_BC4:
+			{
+				if (DestFormat == EImageFormat::IF_RGBA_UBYTE)
+				{
+					return DecompressionFuncRefType(miro::SubImageDecompression::BC4_To_RGBASubImage);
+				}
+				else if (DestFormat == EImageFormat::IF_RGB_UBYTE)
+				{
+					return DecompressionFuncRefType(miro::SubImageDecompression::BC4_To_RGBSubImage);
+				}
+				if (DestFormat == EImageFormat::IF_L_UBYTE)
+				{
+					return DecompressionFuncRefType(miro::SubImageDecompression::BC4_To_LSubImage);
+				}
+				break;
+			}
+
+			case EImageFormat::IF_BC5:
+			{
+				if (DestFormat == EImageFormat::IF_RGBA_UBYTE)
+				{
+					return DecompressionFuncRefType(miro::SubImageDecompression::BC5_To_RGBASubImage);
+				}
+				else if (DestFormat == EImageFormat::IF_RGB_UBYTE)
+				{
+					return DecompressionFuncRefType(miro::SubImageDecompression::BC5_To_RGBSubImage);
+				}
+				break;
+			}
+	
+			case EImageFormat::IF_ASTC_4x4_RGBA_LDR:
+			{
+				if (DestFormat == EImageFormat::IF_RGBA_UBYTE)
+				{
+					return DecompressionFuncRefType(miro::SubImageDecompression::ASTC4x4RGBAL_To_RGBASubImage);
+				}
+				else if (DestFormat == EImageFormat::IF_RGB_UBYTE)
+				{
+					return DecompressionFuncRefType(miro::SubImageDecompression::ASTC4x4RGBAL_To_RGBSubImage);
+				}
+				break;
+			}
+
+			case EImageFormat::IF_ASTC_4x4_RGB_LDR:
+			{
+				if (DestFormat == EImageFormat::IF_RGBA_UBYTE)
+				{
+					return DecompressionFuncRefType(miro::SubImageDecompression::ASTC4x4RGBL_To_RGBASubImage);
+				}
+				else if (DestFormat == EImageFormat::IF_RGB_UBYTE)
+				{
+					return DecompressionFuncRefType(miro::SubImageDecompression::ASTC4x4RGBL_To_RGBSubImage);
+				}
+				break;
+			}
+
+			case EImageFormat::IF_ASTC_4x4_RG_LDR:
+			{
+				if (DestFormat == EImageFormat::IF_RGBA_UBYTE)
+				{
+					return DecompressionFuncRefType(miro::SubImageDecompression::ASTC4x4RGL_To_RGBASubImage);
+				}
+				else if (DestFormat == EImageFormat::IF_RGB_UBYTE)
+				{
+					return DecompressionFuncRefType(miro::SubImageDecompression::ASTC4x4RGL_To_RGBSubImage);
+				}
+				break;
+			}
+
+			case EImageFormat::IF_ASTC_6x6_RGBA_LDR:
+			{
+				if (DestFormat == EImageFormat::IF_RGBA_UBYTE)
+				{
+					return DecompressionFuncRefType(miro::SubImageDecompression::ASTC6x6RGBAL_To_RGBASubImage);
+				}
+				else if (DestFormat == EImageFormat::IF_RGB_UBYTE)
+				{
+					return DecompressionFuncRefType(miro::SubImageDecompression::ASTC6x6RGBAL_To_RGBSubImage);
+				}
+				break;
+			}
+
+			case EImageFormat::IF_ASTC_6x6_RGB_LDR:
+			{
+				if (DestFormat == EImageFormat::IF_RGBA_UBYTE)
+				{
+					return DecompressionFuncRefType(miro::SubImageDecompression::ASTC6x6RGBL_To_RGBASubImage);
+				}
+				else if (DestFormat == EImageFormat::IF_RGB_UBYTE)
+				{
+					return DecompressionFuncRefType(miro::SubImageDecompression::ASTC6x6RGBL_To_RGBSubImage);
+				}
+				break;
+			}
+
+			case EImageFormat::IF_ASTC_6x6_RG_LDR:
+			{
+				if (DestFormat == EImageFormat::IF_RGBA_UBYTE)
+				{
+					return DecompressionFuncRefType(miro::SubImageDecompression::ASTC6x6RGL_To_RGBASubImage);
+				}
+				else if (DestFormat == EImageFormat::IF_RGB_UBYTE)
+				{
+					return DecompressionFuncRefType(miro::SubImageDecompression::ASTC6x6RGL_To_RGBSubImage);
+				}
+				break;
+			}
+
+			case EImageFormat::IF_ASTC_8x8_RGBA_LDR:
+			{
+				if (DestFormat == EImageFormat::IF_RGBA_UBYTE)
+				{
+					return DecompressionFuncRefType(miro::SubImageDecompression::ASTC8x8RGBAL_To_RGBASubImage);
+				}
+				else if (DestFormat == EImageFormat::IF_RGB_UBYTE)
+				{
+					return DecompressionFuncRefType(miro::SubImageDecompression::ASTC8x8RGBAL_To_RGBSubImage);
+				}
+				break;
+			}
+
+			case EImageFormat::IF_ASTC_8x8_RGB_LDR:
+			{
+				if (DestFormat == EImageFormat::IF_RGBA_UBYTE)
+				{
+					return DecompressionFuncRefType(miro::SubImageDecompression::ASTC8x8RGBL_To_RGBASubImage);
+				}
+				else if (DestFormat == EImageFormat::IF_RGB_UBYTE)
+				{
+					return DecompressionFuncRefType(miro::SubImageDecompression::ASTC8x8RGBL_To_RGBSubImage);
+				}
+				break;
+			}
+
+			case EImageFormat::IF_ASTC_8x8_RG_LDR:
+			{
+				if (DestFormat == EImageFormat::IF_RGBA_UBYTE)
+				{
+					return DecompressionFuncRefType(miro::SubImageDecompression::ASTC8x8RGL_To_RGBASubImage);
+				}
+				else if (DestFormat == EImageFormat::IF_RGB_UBYTE)
+				{
+					return DecompressionFuncRefType(miro::SubImageDecompression::ASTC8x8RGL_To_RGBSubImage);
+				}
+				break;
+			}
+
+			case EImageFormat::IF_ASTC_10x10_RGBA_LDR:
+			{
+				if (DestFormat == EImageFormat::IF_RGBA_UBYTE)
+				{
+					return DecompressionFuncRefType(miro::SubImageDecompression::ASTC10x10RGBAL_To_RGBASubImage);
+				}
+				else if (DestFormat == EImageFormat::IF_RGB_UBYTE)
+				{
+					return DecompressionFuncRefType(miro::SubImageDecompression::ASTC10x10RGBAL_To_RGBSubImage);
+				}
+				break;
+			}
+
+			case EImageFormat::IF_ASTC_10x10_RGB_LDR:
+			{
+				if (DestFormat == EImageFormat::IF_RGBA_UBYTE)
+				{
+					return DecompressionFuncRefType(miro::SubImageDecompression::ASTC10x10RGBL_To_RGBASubImage);
+				}
+				else if (DestFormat == EImageFormat::IF_RGB_UBYTE)
+				{
+					return DecompressionFuncRefType(miro::SubImageDecompression::ASTC10x10RGBL_To_RGBSubImage);
+				}
+				break;
+			}
+
+			case EImageFormat::IF_ASTC_10x10_RG_LDR:
+			{
+				if (DestFormat == EImageFormat::IF_RGBA_UBYTE)
+				{
+					return DecompressionFuncRefType(miro::SubImageDecompression::ASTC10x10RGL_To_RGBASubImage);
+				}
+				else if (DestFormat == EImageFormat::IF_RGB_UBYTE)
+				{
+					return DecompressionFuncRefType(miro::SubImageDecompression::ASTC10x10RGL_To_RGBSubImage);
+				}
+				break;
+			}
+
+			case EImageFormat::IF_ASTC_12x12_RGBA_LDR:
+			{
+				if (DestFormat == EImageFormat::IF_RGBA_UBYTE)
+				{
+					return DecompressionFuncRefType(miro::SubImageDecompression::ASTC12x12RGBAL_To_RGBASubImage);
+				}
+				else if (DestFormat == EImageFormat::IF_RGB_UBYTE)
+				{
+					return DecompressionFuncRefType(miro::SubImageDecompression::ASTC12x12RGBAL_To_RGBSubImage);
+				}
+				break;
+			}
+
+			case EImageFormat::IF_ASTC_12x12_RGB_LDR:
+			{
+				if (DestFormat == EImageFormat::IF_RGBA_UBYTE)
+				{
+					return DecompressionFuncRefType(miro::SubImageDecompression::ASTC12x12RGBL_To_RGBASubImage);
+				}
+				else if (DestFormat == EImageFormat::IF_RGB_UBYTE)
+				{
+					return DecompressionFuncRefType(miro::SubImageDecompression::ASTC12x12RGBL_To_RGBSubImage);
+				}
+				break;
+			}
+
+			case EImageFormat::IF_ASTC_12x12_RG_LDR:
+			{
+				if (DestFormat == EImageFormat::IF_RGBA_UBYTE)
+				{
+					return DecompressionFuncRefType(miro::SubImageDecompression::ASTC12x12RGL_To_RGBASubImage);
+				}
+				else if (DestFormat == EImageFormat::IF_RGB_UBYTE)
+				{
+					return DecompressionFuncRefType(miro::SubImageDecompression::ASTC12x12RGL_To_RGBSubImage);
+				}
+				break;
+			}
+		}
+
+		checkf(false, TEXT("Decompression not supported."));
+		return DecompressionFuncRefType(Private::DecompressionFuncNotFoundFunc);
+	}
+
+}
+
diff --git a/Engine/Plugins/Experimental/Mutable/Source/MutableRuntime/Internal/MuR/MeshPrivate.h b/Engine/Plugins/Experimental/Mutable/Source/MutableRuntime/Internal/MuR/MeshPrivate.h
new file mode 100644
--- /dev/null
+++ b/Engine/Plugins/Experimental/Mutable/Source/MutableRuntime/Internal/MuR/MeshPrivate.h
@@ -0,0 +1,665 @@
+// Copyright Epic Games, Inc. All Rights Reserved.
+
+#pragma once
+
+#include "MuR/Mesh.h"
+
+#include "MuR/Layout.h"
+#include "MuR/Skeleton.h"
+#include "MuR/SerialisationPrivate.h"
+#include "MuR/ConvertData.h"
+
+#include "CoreFwd.h"
+
+
+namespace mu
+{
+
+	//---------------------------------------------------------------------------------------------
+	inline void GetMeshBuf
+		(
+			Mesh* pMesh,
+			MESH_BUFFER_SEMANTIC semantic,
+			MESH_BUFFER_FORMAT expectedFormat,
+			int expectedComponents,
+            uint8_t*& pBuf,
+			int& elemSize
+		)
+	{
+		// Avoid unreferenced parameter warnings
+		(void)expectedFormat;
+		(void)expectedComponents;
+
+		int buffer = -1;
+		int channel = -1;
+		pMesh->GetVertexBuffers().FindChannel( semantic, 0, &buffer, &channel );
+		check( buffer>=0 && channel>=0 );
+
+		MESH_BUFFER_SEMANTIC realSemantic = MBS_NONE;
+		int realSemanticIndex = 0;
+		MESH_BUFFER_FORMAT format = MBF_NONE;
+		int components = 0;
+		int offset = 0;
+		pMesh->GetVertexBuffers().GetChannel( buffer, channel, &realSemantic, &realSemanticIndex, &format, &components, &offset );
+		check( realSemantic == semantic );
+		check( format == expectedFormat );
+		check( components == expectedComponents );
+		elemSize = pMesh->GetVertexBuffers().GetElementSize( buffer );
+		pBuf = pMesh->GetVertexBuffers().GetBufferData( buffer );
+		pBuf += offset;
+	}
+
+
+	//---------------------------------------------------------------------------------------------
+	inline void GetMeshBuf
+		(
+			const Mesh* pMesh,
+			MESH_BUFFER_SEMANTIC semantic,
+			MESH_BUFFER_FORMAT expectedFormat,
+			int expectedComponents,
+            const uint8_t*& pBuf,
+			int& elemSize
+		)
+	{
+		// Avoid unreferenced parameter warnings
+		(void)expectedFormat;
+		(void)expectedComponents;
+
+		int buffer = -1;
+		int channel = -1;
+		pMesh->GetVertexBuffers().FindChannel( semantic, 0, &buffer, &channel );
+		check( buffer>=0 && channel>=0 );
+
+		MESH_BUFFER_SEMANTIC realSemantic = MBS_NONE;
+		int realSemanticIndex = 0;
+		MESH_BUFFER_FORMAT format = MBF_NONE;
+		int components = 0;
+		int offset = 0;
+		pMesh->GetVertexBuffers().GetChannel
+				( buffer, channel, &realSemantic, &realSemanticIndex, &format, &components, &offset );
+		check( realSemantic == semantic );
+		check( format == expectedFormat );
+		check( components == expectedComponents );
+		elemSize = pMesh->GetVertexBuffers().GetElementSize( buffer );
+		pBuf = pMesh->GetVertexBuffers().GetBufferData( buffer );
+		pBuf += offset;
+	}
+
+
+	//---------------------------------------------------------------------------------------------
+	//! Class to iterate a specific buffer channel of unknown type
+	//---------------------------------------------------------------------------------------------
+	class UntypedMeshBufferIteratorConst;
+
+	class MUTABLERUNTIME_API UntypedMeshBufferIterator 
+	{
+	public:
+		inline UntypedMeshBufferIterator()
+		{
+			m_format = MBF_NONE;
+			m_components = 0;
+			m_elementSize = 0;
+			m_pBuf = 0;
+		}
+
+
+		inline UntypedMeshBufferIterator
+			(
+				FMeshBufferSet& bufferSet,
+				MESH_BUFFER_SEMANTIC semantic,
+				int semanticIndex = 0
+			)
+		{
+			int buffer = -1;
+			int channel = -1;
+			bufferSet.FindChannel( semantic, semanticIndex, &buffer, &channel );
+
+			if ( buffer>=0 && channel>=0 )
+			{
+				MESH_BUFFER_SEMANTIC realSemantic = MBS_NONE;
+				int realSemanticIndex = 0;
+				m_format = MBF_NONE;
+				m_components = 0;
+				int offset = 0;
+				bufferSet.GetChannel
+						(
+							buffer, channel,
+							&realSemantic, &realSemanticIndex,
+							&m_format, &m_components,
+							&offset
+						);
+				check( realSemantic == semantic );
+				check( realSemanticIndex == semanticIndex );
+
+				m_elementSize = bufferSet.GetElementSize( buffer );
+
+				m_pBuf = bufferSet.GetBufferData( buffer );
+				m_pBuf += offset;
+			}
+			else
+			{
+				m_format = MBF_NONE;
+				m_components = 0;
+				m_elementSize = 0;
+				m_pBuf = 0;
+			}
+		}
+
+        inline uint8_t* ptr() const
+		{
+			return m_pBuf;
+		}
+
+		inline void operator++()
+		{
+			m_pBuf += m_elementSize;
+		}
+
+		inline void operator++(int)
+		{
+			m_pBuf += m_elementSize;
+		}
+
+		inline void operator+=( int c )
+		{
+			m_pBuf += c*m_elementSize;
+		}
+
+		inline UntypedMeshBufferIterator operator+(int c) const
+		{
+			UntypedMeshBufferIterator res = *this;
+			res += c;
+
+			return res;
+		}
+
+		inline std::size_t operator-( const UntypedMeshBufferIterator& other ) const
+		{
+			// Special degenerate case.
+			if (m_elementSize == 0)
+			{
+				return 0;
+			}
+			check( other.m_elementSize == m_elementSize );
+			check( (ptr()-other.ptr()) % m_elementSize == 0 );
+			return (ptr()-other.ptr()) / m_elementSize;
+		}
+
+		inline int GetElementSize() const
+		{
+			return m_elementSize;
+		}
+
+		inline MESH_BUFFER_FORMAT GetFormat() const
+		{
+			return m_format;
+		}
+
+		inline int GetComponents() const
+		{
+			return m_components;
+		}
+
+        FVector4f GetAsVec4f() const
+        {
+			FVector4f res;
+
+            for (int c=0; c< FMath::Min(m_components,4); ++c)
+            {
+                ConvertData( c, &res[0], MBF_FLOAT32, ptr(), m_format );
+            }
+
+            return res;
+        }
+
+		FVector3f GetAsVec3f() const
+		{
+			FVector3f res = { 0,0,0 };
+			for (int c = 0; c < FMath::Min(m_components, 3); ++c)
+			{
+				ConvertData(c, &res[0], MBF_FLOAT32, ptr(), m_format);
+			}
+			return res;
+		}
+
+		FVector3d GetAsVec3d() const
+		{
+			FVector3d res = { 0,0,0 };
+			for (int c = 0; c < FMath::Min(m_components, 3); ++c)
+			{
+				ConvertData(c, &res[0], MBF_FLOAT64, ptr(), m_format);
+			}
+			return res;
+		}
+		
+		FVector2f GetAsVec2f() const
+		{
+			FVector2f res = {0.0f, 0.0f};
+			for (int32 c = 0; c < FMath::Min(m_components, 2); ++c)
+			{
+				ConvertData(c, &res[0], MBF_FLOAT32, ptr(), m_format);
+			}
+
+			return res;
+		}
+		
+        uint32_t GetAsUINT32() const
+        {
+            uint32_t res;
+            ConvertData( 0, &res, MBF_UINT32, ptr(), m_format );
+            return res;
+        }
+
+		void SetFromUINT32(uint32_t v)
+		{
+			ConvertData(0, ptr(), m_format, &v, MBF_UINT32);
+		}
+
+		void SetFromVec3f(const FVector3f& v)
+		{
+			for (int c = 0; c < FMath::Min(m_components, 3); ++c)
+			{
+				ConvertData(c, ptr(), m_format, &v, MBF_FLOAT32);
+			}
+		}
+
+		void SetFromVec3d(const FVector3d& v)
+		{
+			for (int c = 0; c < FMath::Min(m_components, 3); ++c)
+			{
+				ConvertData(c, ptr(), m_format, &v, MBF_FLOAT64);
+			}
+		}
+
+	protected:
+
+		int m_elementSize;
+        uint8_t* m_pBuf;
+		MESH_BUFFER_FORMAT m_format;
+		int m_components;
+	};
+
+
+	//---------------------------------------------------------------------------------------------
+	//! Class to iterate a specific buffer channel with known type
+	//---------------------------------------------------------------------------------------------
+	template<MESH_BUFFER_FORMAT FORMAT, class CTYPE, int COMPONENTS>
+	class MeshBufferIteratorConst;
+
+	template<MESH_BUFFER_FORMAT FORMAT, class CTYPE, int COMPONENTS>
+	class MUTABLERUNTIME_API MeshBufferIterator : public UntypedMeshBufferIterator
+	{
+	public:
+		inline MeshBufferIterator
+			(
+				FMeshBufferSet& bufferSet,
+				MESH_BUFFER_SEMANTIC semantic,
+				int semanticIndex = 0
+			)
+			: UntypedMeshBufferIterator( bufferSet, semantic, semanticIndex )
+		{
+			if (m_pBuf)
+			{
+				// Extra checks
+				int buffer = -1;
+				int channel = -1;
+				bufferSet.FindChannel(semantic, semanticIndex, &buffer, &channel);
+				
+				if (buffer >= 0 && channel >= 0)
+				{
+					MESH_BUFFER_SEMANTIC realSemantic = MBS_NONE;
+					int realSemanticIndex = 0;
+					MESH_BUFFER_FORMAT format = MBF_NONE;
+					int components = 0;
+					int offset = 0;
+					bufferSet.GetChannel
+					(
+						buffer, channel,
+						&realSemantic, &realSemanticIndex,
+						&format, &components,
+						&offset
+					);
+
+					if (format!=FORMAT || components!=COMPONENTS)
+					{
+						// Invalidate
+						m_format = MBF_NONE;
+						m_components = 0;
+						m_elementSize = 0;
+						m_pBuf = nullptr;
+					}
+				}
+				else
+				{
+					// Invalidate
+					m_format = MBF_NONE;
+					m_components = 0;
+					m_elementSize = 0;
+					m_pBuf = nullptr;
+				}
+			}
+		}
+
+		inline vec<CTYPE,COMPONENTS>& operator*()
+		{
+			return *reinterpret_cast< vec<CTYPE,COMPONENTS>* >( m_pBuf );
+		}
+
+		inline vec<CTYPE,COMPONENTS>& operator[]( size_t i )
+		{
+			return *((*this)+int(i));
+		}
+
+		inline MeshBufferIterator<FORMAT,CTYPE,COMPONENTS> operator+(int c)
+		{
+			MeshBufferIterator<FORMAT,CTYPE,COMPONENTS> res = *this;
+			res += c;
+			return res;
+		}
+
+	};
+
+
+	//---------------------------------------------------------------------------------------------
+	//! Class to iterate a specific buffer channel of unknown type
+	//---------------------------------------------------------------------------------------------
+	class MUTABLERUNTIME_API UntypedMeshBufferIteratorConst
+	{
+	public:
+
+		UntypedMeshBufferIteratorConst() = default;
+		UntypedMeshBufferIteratorConst(UntypedMeshBufferIteratorConst&&) = default;
+		UntypedMeshBufferIteratorConst(const UntypedMeshBufferIteratorConst&) = default;
+		UntypedMeshBufferIteratorConst& operator=(UntypedMeshBufferIteratorConst&&) = default;
+		UntypedMeshBufferIteratorConst& operator=(const UntypedMeshBufferIteratorConst&) = default;
+
+		inline UntypedMeshBufferIteratorConst
+			(
+				const FMeshBufferSet& bufferSet,
+				MESH_BUFFER_SEMANTIC semantic,
+				int semanticIndex = 0
+			)
+		{
+			int buffer = -1;
+			int channel = -1;
+			bufferSet.FindChannel( semantic, semanticIndex, &buffer, &channel );
+
+			if ( buffer>=0 && channel>=0 )
+			{
+				MESH_BUFFER_SEMANTIC realSemantic = MBS_NONE;
+				int realSemanticIndex = 0;
+				m_format = MBF_NONE;
+				m_components = 0;
+				int offset = 0;
+				bufferSet.GetChannel
+						(
+							buffer, channel,
+							&realSemantic, &realSemanticIndex,
+							&m_format, &m_components,
+							&offset
+						);
+				check( realSemantic == semantic );
+				check( realSemanticIndex == semanticIndex );
+
+				m_elementSize = bufferSet.GetElementSize( buffer );
+
+				m_pBuf = bufferSet.GetBufferData( buffer );
+				m_pBuf += offset;
+			}
+			else
+			{
+				m_format = MBF_NONE;
+				m_components = 0;
+				m_elementSize = 0;
+				m_pBuf = nullptr;
+			}
+		}
+
+        inline const uint8_t* ptr() const
+		{
+			return m_pBuf;
+		}
+
+		inline void operator++()
+		{
+			m_pBuf += m_elementSize;
+		}
+
+		inline void operator++(int)
+		{
+			m_pBuf += m_elementSize;
+		}
+
+		inline void operator+=( int c )
+		{
+			m_pBuf += c*m_elementSize;
+		}
+
+		inline std::size_t operator-( const UntypedMeshBufferIterator& other ) const
+		{
+			// Special degenerate case.
+			if (m_elementSize == 0)
+			{
+				return 0;
+			}
+			check( other.GetElementSize() == m_elementSize );
+			check( (ptr()-other.ptr()) % m_elementSize == 0 );
+			return (ptr()-other.ptr()) / m_elementSize;
+		}
+
+		inline std::size_t operator-( const UntypedMeshBufferIteratorConst& other ) const
+		{
+			// Special degenerate case.
+			if (m_elementSize == 0)
+			{
+				return 0;
+			}
+			check( other.GetElementSize() == m_elementSize );
+			check( (ptr()-other.ptr()) % m_elementSize == 0 );
+			return (ptr()-other.ptr()) / m_elementSize;
+		}
+
+		inline int GetElementSize() const
+		{
+			return m_elementSize;
+		}
+
+        inline MESH_BUFFER_FORMAT GetFormat() const
+        {
+            return m_format;
+        }
+
+        inline int GetComponents() const
+        {
+            return m_components;
+        }
+
+		FVector4f GetAsVec4f() const
+		{
+			FVector4f res(0.0f,0.0f,0.0f,0.0f);
+			for (int c = 0; c < FMath::Min(m_components, 4); ++c)
+			{
+				ConvertData(c, &res[0], MBF_FLOAT32, ptr(), m_format);
+			}
+			return res;
+		}
+
+		FVector3f GetAsVec3f() const
+		{
+			FVector3f res = { 0,0,0 };
+			for (int c = 0; c < FMath::Min(m_components, 3); ++c)
+			{
+				ConvertData(c, &res[0], MBF_FLOAT32, ptr(), m_format);
+			}
+			return res;
+		}
+
+		FVector3d GetAsVec3d() const
+		{
+			FVector3d res = { 0,0,0 };
+			for (int c = 0; c < FMath::Min(m_components, 3); ++c)
+			{
+				ConvertData(c, &res[0], MBF_FLOAT64, ptr(), m_format);
+			}
+			return res;
+		}
+
+		FVector2f GetAsVec2f() const
+		{
+			FVector2f res = {0.0f, 0.0f};
+			for (int32 c = 0; c < FMath::Min(m_components, 2); ++c)
+			{
+				ConvertData(c, &res[0], MBF_FLOAT32, ptr(), m_format);
+			}
+
+			return res;
+		}
+		
+        uint32_t GetAsUINT32() const
+        {
+            uint32_t res = 0;
+            ConvertData( 0, &res, MBF_UINT32, ptr(), m_format );
+            return res;
+        }
+        
+        vec<int32_t,8> GetAsVec8i() const
+        {
+            vec<int32_t, 8> res;
+            for ( int c = 0; c < FMath::Min( m_components, 8 ); ++c )
+            {
+                ConvertData( c, &res[0], MBF_INT32, ptr(), m_format );
+            }
+            return res;
+        }
+
+		vec<int32_t, 12> GetAsVec12i() const
+		{
+			vec<int32_t, 12> res;
+			for (int c = 0; c < FMath::Min(m_components, 12); ++c)
+			{
+				ConvertData(c, &res[0], MBF_INT32, ptr(), m_format);
+			}
+			return res;
+		}
+
+        inline UntypedMeshBufferIteratorConst operator+(int c) const
+        {
+            UntypedMeshBufferIteratorConst res = *this;
+            res += c;
+            return res;
+        }
+
+	protected:
+
+		int m_elementSize = 0;
+        const uint8_t* m_pBuf = nullptr;
+		MESH_BUFFER_FORMAT m_format = MBF_NONE;
+		int m_components = 0;
+	};
+
+
+
+
+	//---------------------------------------------------------------------------------------------
+	//! Class to iterate a specific buffer channel of a constant buffer set
+	//---------------------------------------------------------------------------------------------
+	template<MESH_BUFFER_FORMAT FORMAT, class CTYPE, int COMPONENTS>
+	class MUTABLERUNTIME_API MeshBufferIteratorConst : public UntypedMeshBufferIteratorConst
+	{
+	public:
+
+		inline MeshBufferIteratorConst()
+		{
+		}
+
+
+		inline MeshBufferIteratorConst
+			(
+				const FMeshBufferSet& bufferSet,
+				MESH_BUFFER_SEMANTIC semantic,
+				int semanticIndex = 0
+			)
+			: UntypedMeshBufferIteratorConst( bufferSet, semantic, semanticIndex )
+		{
+			if (m_pBuf)
+			{
+				// Extra checks
+				int buffer = -1;
+				int channel = -1;
+				bufferSet.FindChannel(semantic, semanticIndex, &buffer, &channel);
+
+				if (buffer >= 0 && channel >= 0)
+				{
+					MESH_BUFFER_SEMANTIC realSemantic = MBS_NONE;
+					int realSemanticIndex = 0;
+					MESH_BUFFER_FORMAT format = MBF_NONE;
+					int components = 0;
+					int offset = 0;
+					bufferSet.GetChannel
+					(
+						buffer, channel,
+						&realSemantic, &realSemanticIndex,
+						&format, &components,
+						&offset
+					);
+					
+					if (format!=FORMAT || components!=COMPONENTS)
+					{
+						// Invalidate
+						m_format = MBF_NONE;
+						m_components = 0;
+						m_elementSize = 0;
+						m_pBuf = nullptr;
+					}
+				}
+				else
+				{
+					// Invalidate
+					m_format = MBF_NONE;
+					m_components = 0;
+					m_elementSize = 0;
+					m_pBuf = nullptr;
+				}
+			}
+		}
+
+		inline const vec<CTYPE,COMPONENTS>& operator*() const
+		{
+			return *reinterpret_cast< const vec<CTYPE,COMPONENTS>* >( m_pBuf );
+		}
+
+		inline const vec<CTYPE,COMPONENTS>& operator[]( size_t i ) const
+		{
+			return *((*this)+int(i));
+		}
+
+		inline MeshBufferIteratorConst<FORMAT,CTYPE,COMPONENTS> operator+(int c) const
+		{
+			MeshBufferIteratorConst<FORMAT,CTYPE,COMPONENTS> res = *this;
+			res += c;
+			return res;
+		}
+
+	};
+
+
+	//---------------------------------------------------------------------------------------------
+	// Optimised mesh formats that are identified in some operations to chose a faster version
+	//---------------------------------------------------------------------------------------------
+	typedef enum
+	{
+		SMF_NONE,
+
+        SMF_PROJECT,
+
+        SMF_PROJECTWRAPPING,
+
+		SMF_COUNT
+	} STATIC_MESH_FORMATS;
+
+	//! Array of functions that can identify each format
+	typedef bool (*STATIC_MESH_FORMAT_ID_FUNC)( const Mesh* );
+	extern STATIC_MESH_FORMAT_ID_FUNC s_staticMeshFormatIdentify[SMF_COUNT];
+
+
+}
diff --git a/Engine/Plugins/Experimental/Mutable/Source/MutableRuntime/Internal/MuR/OpImageProject.h b/Engine/Plugins/Experimental/Mutable/Source/MutableRuntime/Internal/MuR/OpImageProject.h
new file mode 100644
--- /dev/null
+++ b/Engine/Plugins/Experimental/Mutable/Source/MutableRuntime/Internal/MuR/OpImageProject.h
@@ -0,0 +1,58 @@
+// Copyright Epic Games, Inc. All Rights Reserved.
+
+#pragma once
+
+
+#include "Math/IntVector.h"
+
+namespace mu
+{
+	class Image;
+	class Mesh;
+	struct FProjector;
+	enum class ESamplingMethod : uint8;
+	template <int NUM_INTERPOLATORS> class RasterVertex;
+
+	struct FScratchImageProject
+	{
+		TArray<RasterVertex<4>> Vertices;
+		TArray<uint8> CulledVertex;
+	};
+
+    extern void ImageRasterProjectedPlanar( const Mesh* pMesh, Image* pTargetImage,
+		const Image* pSource, const Image* pMask,
+		bool bIsRGBFadingEnabled, bool bIsAlphaFadingEnabled,
+		ESamplingMethod SamplingMethod,
+		float FadeStart, float FadeEnd, float MipInterpolationFactor,
+		int Layout, int32 Block,
+		UE::Math::TIntVector2<uint16> CropMin, UE::Math::TIntVector2<uint16> UncroppedSize,
+		FScratchImageProject* scratch, bool bUseVectorImplementation = false);
+
+    extern void ImageRasterProjectedCylindrical( const Mesh* pMesh, Image* pTargetImage,
+		const Image* pSource, const Image* pMask,
+		bool bIsRGBFadingEnabled, bool bIsAlphaFadingEnabled,
+		ESamplingMethod SamplingMethod,
+		float FadeStart, float FadeEnd, float MipInterpolationFactor,
+		int32 Layout,
+		float ProjectionAngle,
+		UE::Math::TIntVector2<uint16> CropMin, UE::Math::TIntVector2<uint16> UncroppedSize,
+		FScratchImageProject* Scratch, bool bUseVectorImplementation = false);
+
+    extern void ImageRasterProjectedWrapping( const Mesh* pMesh, Image* pTargetImage,
+		const Image* pSource, const Image* pMask,
+		bool bIsRGBFadingEnabled, bool bIsAlphaFadingEnabled,
+		ESamplingMethod SamplingMethod,
+		float FadeStart, float FadeEnd, float MipInterpolationFactor,
+		int32 Layout, int32 Block,
+		UE::Math::TIntVector2<uint16> CropMin, UE::Math::TIntVector2<uint16> UncroppedSize,
+		FScratchImageProject* Scratch, bool bUseVectorImplementation = false);
+
+	extern float ComputeProjectedFootprintBestMip(
+			const Mesh* pMesh, const FProjector& Projector, const FVector2f& TargetSize, const FVector2f& SourceSize);
+
+    extern void MeshProject(Mesh* Result, const Mesh* pMesh, const FProjector& Projector, bool& bOutSuccess);
+
+	MUTABLERUNTIME_API extern void CreateMeshOptimisedForProjection(Mesh* Result, int32 Layout);
+	MUTABLERUNTIME_API extern void CreateMeshOptimisedForWrappingProjection(Mesh* Result, int32 Layout);
+
+}
diff --git a/Engine/Plugins/Experimental/Mutable/Source/MutableRuntime/Internal/MuR/OpMeshMerge.h b/Engine/Plugins/Experimental/Mutable/Source/MutableRuntime/Internal/MuR/OpMeshMerge.h
new file mode 100644
--- /dev/null
+++ b/Engine/Plugins/Experimental/Mutable/Source/MutableRuntime/Internal/MuR/OpMeshMerge.h
@@ -0,0 +1,1246 @@
+// Copyright Epic Games, Inc. All Rights Reserved.
+
+#pragma once
+
+#include "MuR/MeshPrivate.h"
+#include "MuR/Platform.h"
+#include "MuR/MutableMath.h"
+#include "MuR/OpMeshFormat.h"
+#include "MuR/MutableTrace.h"
+
+
+namespace mu
+{
+
+	struct FMeshMergeScratchMeshes
+	{
+		Ptr<Mesh> FirstReformat;
+		Ptr<Mesh> SecondReformat;
+	};
+
+	//---------------------------------------------------------------------------------------------
+	//! Merge two meshes into one new mesh
+	//---------------------------------------------------------------------------------------------
+	inline void MeshMerge(Mesh* Result, const Mesh* pFirst, const Mesh* pSecond, bool bMergeSurfaces, 
+		FMeshMergeScratchMeshes& ScratchMeshes)
+	{
+		MUTABLE_CPUPROFILER_SCOPE(MeshMerge);
+
+		// Should never happen, but fixes static analysis warnings.
+		if (!(pFirst && pSecond))
+		{
+			return;
+		}
+
+		// Indices
+		//-----------------
+		if (pFirst->GetIndexBuffers().GetBufferCount() > 0)
+		{
+			MUTABLE_CPUPROFILER_SCOPE(Indices);
+
+			const int32 FirstCount = pFirst->GetIndexBuffers().GetElementCount();
+			const int32 SecondCount = pSecond->GetIndexBuffers().GetElementCount();
+			Result->GetIndexBuffers().SetElementCount(FirstCount + SecondCount);
+
+			check(pFirst->GetIndexBuffers().GetBufferCount() <= 1);
+			check(pSecond->GetIndexBuffers().GetBufferCount() <= 1);
+			Result->GetIndexBuffers().SetBufferCount(1);
+
+			MESH_BUFFER& ResultIndexBuffer = Result->GetIndexBuffers().m_buffers[0];
+
+			const MESH_BUFFER& FirstIndexBuffer = pFirst->GetIndexBuffers().m_buffers[0];
+			const MESH_BUFFER& SecondIndexBuffer = pSecond->GetIndexBuffers().m_buffers[0];
+
+			// Avoid unused variable warnings
+			(void)FirstIndexBuffer;
+			(void)SecondIndexBuffer;
+
+			// This will be changed below if need to change the format of the index buffers.
+			MESH_BUFFER_FORMAT IndexBufferFormat = MBF_NONE;
+
+			if (FirstCount && SecondCount)
+			{
+				check(!FirstIndexBuffer.m_channels.IsEmpty());
+				check(FirstIndexBuffer.m_channels == SecondIndexBuffer.m_channels);
+				check(FirstIndexBuffer.m_elementSize == SecondIndexBuffer.m_elementSize);
+
+				// We need to know the total number of vertices in case we need to adjust the index buffer format.
+				const uint64 totalVertexCount = pFirst->GetVertexBuffers().GetElementCount() + pSecond->GetVertexBuffers().GetElementCount();
+				const uint64 maxValueBits = GetMeshFormatData(pFirst->GetIndexBuffers().m_buffers[0].m_channels[0].m_format).m_maxValueBits;
+				const uint64 maxSupportedVertices = uint64(1) << maxValueBits;
+				
+				if (totalVertexCount > maxSupportedVertices)
+				{
+					IndexBufferFormat = totalVertexCount > MAX_uint16 ? MBF_UINT32 : MBF_UINT16;
+				}
+
+			}
+			
+			if (IndexBufferFormat != MBF_NONE)
+			{
+				// We only support vertex indices in case of having to change the format.
+				check(FirstIndexBuffer.m_channels.Num() == 1);
+
+				ResultIndexBuffer.m_channels.SetNum(1);
+				ResultIndexBuffer.m_channels[0].m_semantic = MBS_VERTEXINDEX;
+				ResultIndexBuffer.m_channels[0].m_format = IndexBufferFormat;
+				ResultIndexBuffer.m_channels[0].m_componentCount = 1;
+				ResultIndexBuffer.m_channels[0].m_semanticIndex = 0;
+				ResultIndexBuffer.m_channels[0].m_offset = 0;
+				ResultIndexBuffer.m_elementSize = GetMeshFormatData(IndexBufferFormat).m_size;
+			}
+			else if (FirstCount)
+			{
+				ResultIndexBuffer.m_channels = FirstIndexBuffer.m_channels;
+				ResultIndexBuffer.m_elementSize = FirstIndexBuffer.m_elementSize;
+			}
+			else if (SecondCount)
+			{
+				ResultIndexBuffer.m_channels = SecondIndexBuffer.m_channels;
+				ResultIndexBuffer.m_elementSize = SecondIndexBuffer.m_elementSize;
+			}
+
+			ResultIndexBuffer.m_data.SetNum(ResultIndexBuffer.m_elementSize * (FirstCount + SecondCount));
+
+			check(ResultIndexBuffer.m_channels.Num() == 1);
+			check(ResultIndexBuffer.m_channels[0].m_semantic == MBS_VERTEXINDEX);
+
+			if (!ResultIndexBuffer.m_data.IsEmpty())
+			{
+				if (FirstCount)
+				{
+					if (IndexBufferFormat == MBF_NONE
+						|| IndexBufferFormat == FirstIndexBuffer.m_channels[0].m_format)
+					{
+						FMemory::Memcpy(&ResultIndexBuffer.m_data[0],
+							&FirstIndexBuffer.m_data[0],
+							FirstIndexBuffer.m_elementSize * FirstCount);
+					}
+					else
+					{
+						// Conversion required
+						const uint8_t* pSource = &FirstIndexBuffer.m_data[0];
+						uint8_t* pDest = &ResultIndexBuffer.m_data[0];
+						switch (IndexBufferFormat)
+						{
+						case MBF_UINT32:
+						{
+							switch (FirstIndexBuffer.m_channels[0].m_format)
+							{
+							case MBF_UINT16:
+							{
+								for (int v = 0; v < FirstCount; ++v)
+								{
+									*(uint32_t*)pDest = *(const uint16*)pSource;
+									pSource += FirstIndexBuffer.m_elementSize;
+									pDest += ResultIndexBuffer.m_elementSize;
+								}
+								break;
+							}
+
+							case MBF_UINT8:
+							{
+								for (int v = 0; v < FirstCount; ++v)
+								{
+									*(uint32_t*)pDest = *(const uint8_t*)pSource;
+									pSource += FirstIndexBuffer.m_elementSize;
+									pDest += ResultIndexBuffer.m_elementSize;
+								}
+								break;
+							}
+
+							default:
+								checkf(false, TEXT("Format not supported."));
+								break;
+							}
+							break;
+						}
+
+						case MBF_UINT16:
+						{
+							switch (FirstIndexBuffer.m_channels[0].m_format)
+							{
+
+							case MBF_UINT8:
+							{
+								for (int v = 0; v < FirstCount; ++v)
+								{
+									*(uint16*)pDest = *(const uint8_t*)pSource;
+									pSource += FirstIndexBuffer.m_elementSize;
+									pDest += ResultIndexBuffer.m_elementSize;
+								}
+								break;
+							}
+
+							default:
+								checkf(false, TEXT("Format not supported."));
+								break;
+							}
+							break;
+						}
+
+						default:
+							checkf(false, TEXT("Format not supported."));
+							break;
+						}
+					}
+				}
+
+				if (SecondCount)
+				{
+					const uint8_t* pSource = &SecondIndexBuffer.m_data[0];
+					uint8_t* pDest = &ResultIndexBuffer.m_data[ResultIndexBuffer.m_elementSize * FirstCount];
+
+					uint32_t firstVertexCount = pFirst->GetVertexBuffers().GetElementCount();
+
+					if (IndexBufferFormat == MBF_NONE
+						|| IndexBufferFormat == SecondIndexBuffer.m_channels[0].m_format)
+					{
+						switch (SecondIndexBuffer.m_channels[0].m_format)
+						{
+						case MBF_INT32:
+						case MBF_UINT32:
+						case MBF_NINT32:
+						case MBF_NUINT32:
+						{
+							for (int v = 0; v < SecondCount; ++v)
+							{
+								*(uint32_t*)pDest = firstVertexCount + *(const uint32_t*)pSource;
+								pSource += SecondIndexBuffer.m_elementSize;
+								pDest += ResultIndexBuffer.m_elementSize;
+							}
+							break;
+						}
+
+						case MBF_INT16:
+						case MBF_UINT16:
+						case MBF_NINT16:
+						case MBF_NUINT16:
+						{
+							for (int v = 0; v < SecondCount; ++v)
+							{
+								*(uint16*)pDest = uint16(firstVertexCount) + *(const uint16*)pSource;
+								pSource += SecondIndexBuffer.m_elementSize;
+								pDest += ResultIndexBuffer.m_elementSize;
+							}
+							break;
+						}
+
+						case MBF_INT8:
+						case MBF_UINT8:
+						case MBF_NINT8:
+						case MBF_NUINT8:
+						{
+							for (int v = 0; v < SecondCount; ++v)
+							{
+								*(uint8_t*)pDest = uint8_t(firstVertexCount) + *(const uint8_t*)pSource;
+								pSource += SecondIndexBuffer.m_elementSize;
+								pDest += ResultIndexBuffer.m_elementSize;
+							}
+							break;
+						}
+
+						default:
+							checkf(false, TEXT("Format not supported."));
+							break;
+						}
+					}
+					else
+					{
+						// Format conversion required
+						switch (IndexBufferFormat)
+						{
+
+						case MBF_UINT32:
+						{
+							switch (SecondIndexBuffer.m_channels[0].m_format)
+							{
+							case MBF_INT16:
+							case MBF_UINT16:
+							case MBF_NINT16:
+							case MBF_NUINT16:
+							{
+								for (int v = 0; v < SecondCount; ++v)
+								{
+									*(uint32_t*)pDest = uint32_t(firstVertexCount) + *(const uint16*)pSource;
+									pSource += SecondIndexBuffer.m_elementSize;
+									pDest += ResultIndexBuffer.m_elementSize;
+								}
+								break;
+							}
+
+							case MBF_INT8:
+							case MBF_UINT8:
+							case MBF_NINT8:
+							case MBF_NUINT8:
+							{
+								for (int v = 0; v < SecondCount; ++v)
+								{
+									*(uint32_t*)pDest = uint32_t(firstVertexCount) + *(const uint8_t*)pSource;
+									pSource += SecondIndexBuffer.m_elementSize;
+									pDest += ResultIndexBuffer.m_elementSize;
+								}
+								break;
+							}
+
+							default:
+								checkf(false, TEXT("Format not supported."));
+								break;
+							}
+
+							break;
+						}
+
+						case MBF_UINT16:
+						{
+							switch (SecondIndexBuffer.m_channels[0].m_format)
+							{
+							case MBF_INT8:
+							case MBF_UINT8:
+							case MBF_NINT8:
+							case MBF_NUINT8:
+							{
+								for (int v = 0; v < SecondCount; ++v)
+								{
+									*(uint16*)pDest = uint16(firstVertexCount) + *(const uint8_t*)pSource;
+									pSource += SecondIndexBuffer.m_elementSize;
+									pDest += ResultIndexBuffer.m_elementSize;
+								}
+								break;
+							}
+
+							default:
+								checkf(false, TEXT("Format not supported."));
+								break;
+							}
+
+							break;
+						}
+
+						default:
+							checkf(false, TEXT("Format not supported."));
+							break;
+
+						}
+					}
+				}
+			}
+		}
+
+		// Faces
+		//-----------------
+		{
+			MUTABLE_CPUPROFILER_SCOPE(Faces);
+
+			const int32 FirstCount = pFirst->GetFaceBuffers().GetElementCount();
+			const int32 SecondCount = pSecond->GetFaceBuffers().GetElementCount();
+			Result->GetFaceBuffers().SetElementCount(FirstCount + SecondCount);
+			Result->GetFaceBuffers().SetBufferCount(pFirst->GetFaceBuffers().GetBufferCount());
+
+			// Merge only the buffers present in the first mesh
+			for (int32 b = 0; b < Result->GetFaceBuffers().GetBufferCount(); ++b)
+			{
+				const MESH_BUFFER& first = pFirst->GetFaceBuffers().m_buffers[b];
+
+				MESH_BUFFER& result = Result->GetFaceBuffers().m_buffers[b];
+				result.m_channels = first.m_channels;
+				result.m_elementSize = first.m_elementSize;
+				result.m_data.SetNum(result.m_elementSize * (FirstCount + SecondCount));
+
+				MESH_BUFFER_SEMANTIC semantic = MBS_NONE;
+				int semanticIndex = 0;
+				MESH_BUFFER_FORMAT type = MBF_NONE;
+				int components = 0;
+				int offset = 0;
+				pFirst->GetFaceBuffers().GetChannel
+				(b, 0, &semantic, &semanticIndex, &type, &components, &offset);
+
+				if (FirstCount)
+				{
+					FMemory::Memcpy(&result.m_data[0],
+						&first.m_data[0],
+						result.m_elementSize * FirstCount
+					);
+				}
+
+				if (SecondCount)
+				{
+					// Find in the second mesh
+					int otherBuffer = -1;
+					int otherChannel = -1;
+					pSecond->GetFaceBuffers().FindChannel
+					(semantic, 0, &otherBuffer, &otherChannel);
+
+					if (otherBuffer >= 0)
+					{
+						const MESH_BUFFER& second =
+							pSecond->GetFaceBuffers().m_buffers[otherBuffer];
+						check(first.m_channels == second.m_channels);
+
+						// Raw copy
+						FMemory::Memcpy(&result.m_data[result.m_elementSize * FirstCount],
+							&second.m_data[0],
+							result.m_elementSize * SecondCount);
+					}
+					else
+					{
+						// Fill with zeroes
+						FMemory::Memzero(&result.m_data[result.m_elementSize * FirstCount],
+							result.m_elementSize * SecondCount);
+					}
+				}
+			}
+		}
+
+
+		// Layouts
+		//-----------------
+		{
+			MUTABLE_CPUPROFILER_SCOPE(Layouts);
+
+			Result->m_layouts.SetNum(pFirst->m_layouts.Num());
+			for (int i = 0; i < pFirst->m_layouts.Num(); ++i)
+			{
+				const Layout* pF = pFirst->m_layouts[i].get();
+				LayoutPtr pR = pF->Clone();
+
+				if (i < pSecond->m_layouts.Num())
+				{
+					const Layout* pS = pSecond->m_layouts[i].get();
+
+					pR->m_blocks.Append(pS->m_blocks);
+				}
+
+				Result->m_layouts[i] = pR;
+			}
+		}
+
+
+		// Skeleton
+		//---------------------------
+
+		// Add SkeletonIDs
+		Result->SkeletonIDs = pFirst->SkeletonIDs;
+
+		for (const int32 SkeletonID : pSecond->SkeletonIDs)
+		{
+			Result->SkeletonIDs.AddUnique(SkeletonID);
+		}
+
+		// Do they have the same skeleton?
+		bool bMergeSkeletons = pFirst->GetSkeleton() != pSecond->GetSkeleton();
+
+		// Are they different skeletons but with the same data?
+		if (bMergeSkeletons && pFirst->GetSkeleton() && pSecond->GetSkeleton())
+		{
+			bMergeSkeletons = !(*pFirst->GetSkeleton() == *pSecond->GetSkeleton());
+		}
+
+
+		if (bMergeSkeletons)
+		{
+			MUTABLE_CPUPROFILER_SCOPE(MergeSkeleton);
+			Ptr<Skeleton> pResultSkeleton;
+
+			mu::SkeletonPtrConst pFirstSkeleton = pFirst->GetSkeleton();
+			mu::SkeletonPtrConst pSecondSkeleton = pSecond->GetSkeleton();
+
+			const int32 NumBonesFirst = pFirstSkeleton ? pFirstSkeleton->GetBoneCount() : 0;
+			const int32 NumBonesSecond = pSecondSkeleton ? pSecondSkeleton->GetBoneCount() : 0;
+
+			pResultSkeleton = pFirstSkeleton ? pFirstSkeleton->Clone() : new Skeleton;
+			Result->SetSkeleton(pResultSkeleton);
+
+			TArray<uint16> SecondToResultBoneIndices;
+			SecondToResultBoneIndices.SetNumUninitialized(NumBonesSecond);
+
+			// Merge pSecond and build the remap array 
+			for (int32 SecondBoneIndex = 0; SecondBoneIndex < NumBonesSecond; ++SecondBoneIndex)
+			{
+				const uint16 BoneNameId = pSecondSkeleton->BoneIds[SecondBoneIndex];
+				int32 Index = pResultSkeleton->FindBone(BoneNameId);
+
+				// Add a new bone
+				if (Index == INDEX_NONE)
+				{
+					Index = pResultSkeleton->BoneIds.Add(BoneNameId);
+
+					// Add an incorrect index, to be fixed below in case the parent index is later in the bone array.
+					pResultSkeleton->BoneParents.Add(pSecondSkeleton->BoneParents[SecondBoneIndex]);
+
+					if (pSecondSkeleton->BoneNames.IsValidIndex(SecondBoneIndex))
+					{
+						pResultSkeleton->BoneNames.Add(pSecondSkeleton->BoneNames[SecondBoneIndex]);
+					}
+				}
+
+				SecondToResultBoneIndices[SecondBoneIndex] = (uint16)Index;
+			}
+
+			// Fix second mesh bone parents
+			for (int32 ob = NumBonesFirst; ob < pResultSkeleton->BoneParents.Num(); ++ob)
+			{
+				int16 secondMeshIndex = pResultSkeleton->BoneParents[ob];
+				if (secondMeshIndex != INDEX_NONE)
+				{
+					pResultSkeleton->BoneParents[ob] = SecondToResultBoneIndices[secondMeshIndex];
+				}
+			}
+		}
+		else
+		{
+			Result->SetSkeleton(pFirst->GetSkeleton());
+		}
+
+
+		// Surfaces
+		//---------------------------
+		
+		// Remap bone indices if we merge surfaces since bonemaps will be merged too.
+		bool bRemapBoneIndices = false;
+		TArray<uint16> RemappedBoneMapIndices;
+
+		// Used to know the format of the bone index buffer
+		uint32 MaxNumBonesInBoneMaps = 0;
+		const int32 NumSecondBonesInBoneMap = pSecond->BoneMap.Num();
+
+		{
+			MUTABLE_CPUPROFILER_SCOPE(Surfaces);
+			
+			const int32 NumFirstBonesInBoneMap = pFirst->BoneMap.Num();
+			Result->BoneMap = pFirst->BoneMap;
+
+			if (bMergeSurfaces)
+			{
+				// Merge BoneMaps
+				RemappedBoneMapIndices.SetNumUninitialized(NumSecondBonesInBoneMap);
+
+				for (uint16 SecondBoneMapIndex = 0; SecondBoneMapIndex < NumSecondBonesInBoneMap; ++SecondBoneMapIndex)
+				{
+					const int32 BoneMapIndex = Result->BoneMap.AddUnique(pSecond->BoneMap[SecondBoneMapIndex]);
+					RemappedBoneMapIndices[SecondBoneMapIndex] = BoneMapIndex;
+
+					bRemapBoneIndices = bRemapBoneIndices || BoneMapIndex != SecondBoneMapIndex;
+				}
+
+				MESH_SURFACE& NewSurface = Result->m_surfaces.AddDefaulted_GetRef();
+				NewSurface.m_vertexCount = pFirst->GetVertexCount() + pSecond->GetVertexCount();
+				NewSurface.m_indexCount = pFirst->GetIndexCount() + pSecond->GetIndexCount();
+				NewSurface.BoneMapCount = Result->BoneMap.Num();
+				
+				//All merged surfaces will have the same bCastShadow value. Decided by the first merged mesh
+				NewSurface.bCastShadow = pFirst->m_surfaces.Last().bCastShadow;
+			}
+			else
+			{
+				// Add the BoneMap of the second mesh
+				Result->BoneMap.Append(pSecond->BoneMap);
+
+				// Add pFirst surfaces
+				Result->m_surfaces = pFirst->m_surfaces;
+
+				const int32 FirstVertexIndex = pFirst->GetVertexCount();
+				const int32 FirstIndexIndex = pFirst->GetIndexCount();
+
+				check(pSecond->m_surfaces.Num() == 1);
+				MESH_SURFACE& NewSurface = Result->m_surfaces.Add_GetRef(pSecond->m_surfaces[0]);
+				NewSurface.m_firstVertex += FirstVertexIndex;
+				NewSurface.m_firstIndex += FirstIndexIndex;
+				NewSurface.BoneMapIndex += NumFirstBonesInBoneMap;
+			}
+
+			for (const MESH_SURFACE& Surface : Result->m_surfaces)
+			{
+				MaxNumBonesInBoneMaps = FMath::Max(MaxNumBonesInBoneMaps, Surface.BoneMapCount);
+			}
+
+			Result->BoneMap.Shrink();
+		}
+
+
+		// Pose
+		//---------------------------
+		{
+			MUTABLE_CPUPROFILER_SCOPE(Pose);
+
+			Result->BonePoses.Reserve(Result->GetSkeleton()->GetBoneCount());
+
+			// Copy poses from the first mesh
+			Result->BonePoses = pFirst->BonePoses;
+
+			// Add or override bone poses
+			for (const Mesh::FBonePose& SecondBonePose : pSecond->BonePoses)
+			{
+				const int32 ResultBoneIndex = Result->FindBonePose(SecondBonePose.BoneId);
+
+				if (ResultBoneIndex != INDEX_NONE)
+				{
+					Mesh::FBonePose& ResultBonePose = Result->BonePoses[ResultBoneIndex];
+
+					// TODO: Not sure how to tune this priority, review it.
+					// For now use a similar strategy as before. 
+					auto ComputeBoneMergePriority = [](const Mesh::FBonePose& BonePose)
+					{
+						return (EnumHasAnyFlags(BonePose.BoneUsageFlags, EBoneUsageFlags::Skinning) ? 1 : 0) +
+							(EnumHasAnyFlags(BonePose.BoneUsageFlags, EBoneUsageFlags::Reshaped) ? 1 : 0);
+					};
+
+					if (ComputeBoneMergePriority(ResultBonePose) < ComputeBoneMergePriority(SecondBonePose))
+					{
+						//ResultBonePose.BoneName = SecondBonePose.BoneName;
+						ResultBonePose.BoneTransform = SecondBonePose.BoneTransform;
+						// Merge usage flags
+						EnumAddFlags(ResultBonePose.BoneUsageFlags, SecondBonePose.BoneUsageFlags);
+					}
+				}
+				else
+				{
+					Result->BonePoses.Add(SecondBonePose);
+				}
+			}
+
+			Result->BonePoses.Shrink();
+		}
+
+
+		// PhysicsBodies
+		//---------------------------
+		{
+			MUTABLE_CPUPROFILER_SCOPE(PhysicsBodies);
+
+			// Appends InPhysicsBody to OutPhysicsBody removing Bodies that are equal, have same bone and customId and its properies are identical.	
+			auto AppendPhysicsBodiesUnique = [](PhysicsBody& OutPhysicsBody, const PhysicsBody& InPhysicsBody) -> bool
+			{
+				TArray<uint16>& OutBones = OutPhysicsBody.BoneIds;
+				TArray<int32>& OutCustomIds = OutPhysicsBody.BodiesCustomIds;
+				TArray<FPhysicsBodyAggregate>& OutBodies = OutPhysicsBody.Bodies;
+
+				const TArray<uint16>& InBones = InPhysicsBody.BoneIds;
+				const TArray<int32>& InCustomIds = InPhysicsBody.BodiesCustomIds;
+				const TArray<FPhysicsBodyAggregate>& InBodies = InPhysicsBody.Bodies;
+
+				const int32 InBodyCount = InPhysicsBody.GetBodyCount();
+				const int32 OutBodyCount = OutPhysicsBody.GetBodyCount();
+
+				bool bModified = false;
+
+				for (int32 InBodyIndex = 0; InBodyIndex < InBodyCount; ++InBodyIndex)
+				{
+					int32 FoundIndex = INDEX_NONE;
+					for (int32 OutBodyIndex = 0; OutBodyIndex < OutBodyCount; ++OutBodyIndex)
+					{
+						if (InCustomIds[InBodyIndex] == OutCustomIds[OutBodyIndex] && InBones[InBodyIndex] == OutBones[OutBodyIndex])
+						{
+							FoundIndex = OutBodyIndex;
+							break;
+						}
+					}
+
+					if (FoundIndex == INDEX_NONE)
+					{
+						OutBones.Add(InBones[InBodyIndex]);
+						OutCustomIds.Add(InCustomIds[InBodyIndex]);
+						OutBodies.Add(InBodies[InBodyIndex]);
+
+						bModified |= true;
+
+						continue;
+					}
+
+					for (const FSphereBody& Body : InBodies[InBodyIndex].Spheres)
+					{
+						const int32 NumBeforeAddition = OutBodies[FoundIndex].Spheres.Num();
+						bModified |= NumBeforeAddition == OutBodies[FoundIndex].Spheres.AddUnique(Body);
+					}
+
+					for (const FBoxBody& Body : InBodies[InBodyIndex].Boxes)
+					{
+						const int32 NumBeforeAddition = OutBodies[FoundIndex].Boxes.Num();
+						bModified |= NumBeforeAddition == OutBodies[FoundIndex].Boxes.AddUnique(Body);
+					}
+
+					for (const FSphylBody& Body : InBodies[InBodyIndex].Sphyls)
+					{
+						const int32 NumBeforeAddition = OutBodies[FoundIndex].Sphyls.Num();
+						bModified |= NumBeforeAddition == OutBodies[FoundIndex].Sphyls.AddUnique(Body);
+					}
+
+					for (const FTaperedCapsuleBody& Body : InBodies[InBodyIndex].TaperedCapsules)
+					{
+						const int32 NumBeforeAddition = OutBodies[FoundIndex].TaperedCapsules.Num();
+						bModified |= NumBeforeAddition == OutBodies[FoundIndex].TaperedCapsules.AddUnique(Body);
+					}
+
+					for (const FConvexBody& Body : InBodies[InBodyIndex].Convex)
+					{
+						const int32 NumBeforeAddition = OutBodies[FoundIndex].Convex.Num();
+						bModified |= NumBeforeAddition == OutBodies[FoundIndex].Convex.AddUnique(Body);
+					}
+				}
+
+				return bModified;
+			};
+
+			TTuple<const PhysicsBody*, bool> SharedResultPhysicsBody = Invoke([&]()
+				-> TTuple<const PhysicsBody*, bool>
+			{
+				if (pFirst->GetPhysicsBody() == pSecond->GetPhysicsBody())
+				{
+					return MakeTuple(pFirst->GetPhysicsBody().get(), true);
+				}
+
+				if (pFirst->GetPhysicsBody() && !pSecond->GetPhysicsBody())
+				{
+					return MakeTuple(pFirst->GetPhysicsBody().get(), true);
+				}
+
+				if (!pFirst->GetPhysicsBody() && pSecond->GetPhysicsBody())
+				{
+					return MakeTuple(pSecond->GetPhysicsBody().get(), true);
+				}
+
+				return MakeTuple(nullptr, false);
+			});
+
+			if (SharedResultPhysicsBody.Get<1>())
+			{
+				// Only one or non of the meshes has physics, share the result.
+				Result->SetPhysicsBody(SharedResultPhysicsBody.Get<0>());
+			}
+			else
+			{
+				check(pFirst->GetPhysicsBody() && pSecond->GetPhysicsBody());
+
+				Ptr<PhysicsBody> MergedResultPhysicsBody = pFirst->GetPhysicsBody()->Clone();
+
+				MergedResultPhysicsBody->bBodiesModified =
+					AppendPhysicsBodiesUnique(*MergedResultPhysicsBody, *pSecond->GetPhysicsBody()) ||
+					pFirst->GetPhysicsBody()->bBodiesModified || pSecond->GetPhysicsBody()->bBodiesModified;
+
+				Result->SetPhysicsBody(MergedResultPhysicsBody);
+			}
+
+			// Additional physics bodies.
+			const int32 MaxAdditionalPhysicsResultNum = pFirst->AdditionalPhysicsBodies.Num() + pSecond->AdditionalPhysicsBodies.Num();
+
+			Result->AdditionalPhysicsBodies.Reserve(MaxAdditionalPhysicsResultNum);
+			Result->AdditionalPhysicsBodies.Append(pFirst->AdditionalPhysicsBodies);
+			
+			// Not very many additional bodies expected, do a quadratic search to have unique bodies based on external id.
+			const int32 NumSecondAdditionalBodies = pSecond->AdditionalPhysicsBodies.Num();
+			for (int32 Index = 0; Index < NumSecondAdditionalBodies; ++Index)
+			{
+				const int32 CustomIdToMerge = pSecond->AdditionalPhysicsBodies[Index]->CustomId;
+
+				const mu::Ptr<const PhysicsBody>* Found = pFirst->AdditionalPhysicsBodies.FindByPredicate(
+					[CustomIdToMerge](const Ptr<const mu::PhysicsBody>& Body) { return CustomIdToMerge == Body->CustomId; });
+
+				// TODO: current usages do not expect collisions, but same Id collision with bodies modified in differnet ways
+				// may need to be contemplated at some point.
+				if (!Found)
+				{
+					Result->AdditionalPhysicsBodies.Add(pSecond->AdditionalPhysicsBodies[Index]);
+				}
+			}
+		}
+
+		// Vertices
+		//-----------------
+		{
+            MUTABLE_CPUPROFILER_SCOPE(Vertices);
+
+            const int32 FirstCount = pFirst->GetVertexBuffers().GetElementCount();
+			const int32 SecondCount = pSecond->GetVertexBuffers().GetElementCount();
+
+			// TODO: when formats match, which at runtime should be always.
+			bool bFastPath = pFirst->GetVertexBuffers().HasSameFormat( pSecond->GetVertexBuffers() );
+
+			// Check if the format of the BoneIndex buffer has to change
+			bool bChangeBoneIndicesFormat = false;
+			MESH_BUFFER_FORMAT BoneIndexFormat = MaxNumBonesInBoneMaps > MAX_uint8 ? MBF_UINT16 : MBF_UINT8;
+
+			// Iterate all vertex buffers to check if we need to format bone indices
+			{
+				{
+					// Check if pFirst requires a reformat of the bone index buffers
+					const FMeshBufferSet& VertexBuffers = pFirst->GetVertexBuffers();
+					for (int32 VertexBufferIndex = 0; !bChangeBoneIndicesFormat && VertexBufferIndex < VertexBuffers.m_buffers.Num(); ++VertexBufferIndex)
+					{
+						const MESH_BUFFER& Buffer = VertexBuffers.m_buffers[VertexBufferIndex];
+
+						const int32 elemSize = VertexBuffers.GetElementSize(VertexBufferIndex);
+						const int32 firstSize = FirstCount * elemSize;
+
+						const int32 ChannelsCount = VertexBuffers.GetBufferChannelCount(VertexBufferIndex);
+						for (int32 ChannelIndex = 0; ChannelIndex < ChannelsCount; ++ChannelIndex)
+						{
+							if (Buffer.m_channels[ChannelIndex].m_semantic == MBS_BONEINDICES)
+							{
+								bChangeBoneIndicesFormat = Buffer.m_channels[ChannelIndex].m_format != BoneIndexFormat;
+								break;
+							}
+						}
+					}
+				}
+
+				{
+					// Check if pSecond requires a reformat of the bone index buffers
+					const FMeshBufferSet& VertexBuffers = pSecond->GetVertexBuffers();
+					for (int32 VertexBufferIndex = 0; !bChangeBoneIndicesFormat && VertexBufferIndex < VertexBuffers.m_buffers.Num(); ++VertexBufferIndex)
+					{
+						const MESH_BUFFER& Buffer = VertexBuffers.m_buffers[VertexBufferIndex];
+
+						const int32 elemSize = VertexBuffers.GetElementSize(VertexBufferIndex);
+						const int32 firstSize = FirstCount * elemSize;
+
+						const int32 ChannelsCount = VertexBuffers.GetBufferChannelCount(VertexBufferIndex);
+						for (int32 ChannelIndex = 0; ChannelIndex < ChannelsCount; ++ChannelIndex)
+						{
+							if (Buffer.m_channels[ChannelIndex].m_semantic == MBS_BONEINDICES)
+							{
+								bChangeBoneIndicesFormat = Buffer.m_channels[ChannelIndex].m_format != BoneIndexFormat;
+								break;
+							}
+						}
+					}
+				}
+
+				bFastPath = bFastPath && !bChangeBoneIndicesFormat;
+			}
+
+			Ptr<const Mesh> pVFirst;
+			Ptr<const Mesh> pVSecond;
+			
+			if (!bFastPath)
+			{
+                MUTABLE_CPUPROFILER_SCOPE(SlowPath);
+
+                // Very slow! Should we warn somehow?
+
+				// Expand component counts in vertex channels of the format mesh
+				int32 vbcount = pFirst->GetVertexBuffers().m_buffers.Num();
+				Result->GetVertexBuffers().SetBufferCount( (int)vbcount );
+
+				for ( int32 vb = 0; vb<vbcount; ++vb )
+				{
+					MESH_BUFFER& result = Result->GetVertexBuffers().m_buffers[vb];
+					const MESH_BUFFER& first = pFirst->GetVertexBuffers().m_buffers[vb];
+
+					result.m_channels = first.m_channels;
+					result.m_elementSize = first.m_elementSize;
+
+					// See if we need to enlarge the components of any of the result channels
+					bool resetOffsets = false;
+					for ( int32 c=0; c<result.m_channels.Num(); ++c )
+					{
+						int sb = -1;
+						int sc = -1;
+						pSecond->GetVertexBuffers().FindChannel
+								(
+									result.m_channels[c].m_semantic,
+									result.m_channels[c].m_semanticIndex,
+									&sb, &sc
+								);
+						if ( sb>=0 )
+						{
+							const MESH_BUFFER& second = pSecond->GetVertexBuffers().m_buffers[sb];
+
+							if ( second.m_channels[sc].m_componentCount
+								 >
+								 result.m_channels[c].m_componentCount )
+							{
+								result.m_channels[c].m_componentCount =
+										second.m_channels[sc].m_componentCount;
+								resetOffsets = true;
+							}
+						}
+					}
+
+					// Reset the channel offsets if necessary
+					if (resetOffsets)
+					{
+						int offset = 0;
+						for ( int32 c=0; c<result.m_channels.Num(); ++c )
+						{
+                            result.m_channels[c].m_offset = (uint8_t)offset;
+							offset += result.m_channels[c].m_componentCount
+									*
+									GetMeshFormatData(result.m_channels[c].m_format).m_size;
+						}
+						result.m_elementSize = offset;
+					}
+				}
+
+                // See if we need to add additional buffers from the second mesh (like vertex colours or additional UV Channels)
+                // This is a bit ad-hoc: we only add buffers containing all new channels
+                for ( const MESH_BUFFER& buf : pSecond->GetVertexBuffers().m_buffers )
+                {
+                    bool someChannel = false;
+                    bool allNewChannels = true;
+					for (const MESH_BUFFER_CHANNEL& chan : buf.m_channels)
+					{
+						// Skip system buffers
+						if (chan.m_semantic == MBS_VERTEXINDEX
+							||
+							chan.m_semantic == MBS_LAYOUTBLOCK)
+						{
+							continue;
+						}
+
+						someChannel = true;
+
+						int foundBuffer = -1;
+						int foundChannel = -1;
+						pFirst->GetVertexBuffers().FindChannel(chan.m_semantic, chan.m_semanticIndex, &foundBuffer, &foundChannel);
+						if (foundBuffer >= 0)
+						{
+							// There's at least one channel that already exists in the first mesh. Don't add the buffer.
+							allNewChannels = false;
+							continue;
+						}
+						
+						// If there are additional UV channels try to add them
+						if(!allNewChannels && chan.m_semantic == MBS_TEXCOORDS
+							&&
+							chan.m_semanticIndex > 0)
+						{
+							// Add additional UV channels if the previous one is found.
+							FMeshBufferSet& VertexBuffers = Result->GetVertexBuffers();
+							VertexBuffers.FindChannel(MBS_TEXCOORDS, chan.m_semanticIndex - 1, &foundBuffer, &foundChannel);
+
+							if (foundBuffer >= 0)
+							{
+								MESH_BUFFER& Buffer = VertexBuffers.m_buffers[foundBuffer];
+								Buffer.m_channels.Insert(chan, foundChannel + 1);
+
+								// Update offsets
+								int32 Offset = Buffer.m_channels[foundChannel].m_offset;
+								for (int32 c = foundChannel; c < Buffer.m_channels.Num(); ++c)
+								{
+									Buffer.m_channels[c].m_offset = (uint8)Offset;
+									Offset += Buffer.m_channels[c].m_componentCount
+										*
+										GetMeshFormatData(Buffer.m_channels[c].m_format).m_size;
+								}
+								Buffer.m_elementSize = Offset;
+							}
+						}
+					}
+
+                    if (someChannel && allNewChannels)
+                    {
+                        Result->GetVertexBuffers().m_buffers.Add(buf);
+                    }
+                }
+
+				// Change the format of the bone indices buffer
+				if(bChangeBoneIndicesFormat)
+				{
+					// Iterate all vertex buffers and update the format
+					FMeshBufferSet& VertexBuffers = Result->GetVertexBuffers();
+					for (int32 VertexBufferIndex = 0; VertexBufferIndex < VertexBuffers.m_buffers.Num(); ++VertexBufferIndex)
+					{
+						MESH_BUFFER& result = VertexBuffers.m_buffers[VertexBufferIndex];
+
+						const int32 ChannelsCount = VertexBuffers.GetBufferChannelCount(VertexBufferIndex);
+						for (int32 ChannelIndex = 0; ChannelIndex < ChannelsCount; ++ChannelIndex)
+						{
+							if (result.m_channels[ChannelIndex].m_semantic == MBS_BONEINDICES)
+							{
+								result.m_channels[ChannelIndex].m_format = BoneIndexFormat;
+
+								// Reset offsets
+								int32 offset = 0;
+								for (int32 AuxChannelIndex = 0; AuxChannelIndex < ChannelsCount; ++AuxChannelIndex)
+								{
+									result.m_channels[AuxChannelIndex].m_offset = (uint8_t)offset;
+									offset += result.m_channels[AuxChannelIndex].m_componentCount
+										*
+										GetMeshFormatData(result.m_channels[AuxChannelIndex].m_format).m_size;
+								}
+								result.m_elementSize = offset;
+							}
+						}
+					}
+				}
+
+                // Allocate vertices
+                Result->GetVertexBuffers().SetElementCount( FirstCount + SecondCount );
+				
+				// Convert the source meshes to the new format
+                if (pFirst->GetVertexBuffers().HasSameFormat(Result->GetVertexBuffers()))
+                {
+                    pVFirst = pFirst;
+                }
+                else
+                {
+					check(ScratchMeshes.FirstReformat);
+
+					bool bOutSuccess = false;
+					MeshFormat(ScratchMeshes.FirstReformat.get(), pFirst, Result, false, true, false, false, false, bOutSuccess);
+					
+					if (bOutSuccess)
+					{
+						pVFirst = ScratchMeshes.FirstReformat;
+					}
+                }
+
+
+                if (pSecond->GetVertexBuffers().HasSameFormat(Result->GetVertexBuffers()))
+                {
+                    pVSecond = pSecond;
+                }
+                else
+                {
+					check(ScratchMeshes.SecondReformat);
+
+					bool bOutSuccess = false;
+                    MeshFormat(ScratchMeshes.SecondReformat.get(), pSecond, Result, false, true, false, false, false, bOutSuccess);
+
+					if (bOutSuccess)
+					{
+						pVSecond = ScratchMeshes.SecondReformat;
+					}
+                }
+
+				check(pVFirst->GetVertexBuffers().HasSameFormat(pVSecond->GetVertexBuffers()));
+			}
+			else
+			{
+                MUTABLE_CPUPROFILER_SCOPE(FastPath);
+
+                pVFirst = pFirst;
+				pVSecond = pSecond;
+			}
+
+			Result->m_VertexBuffers = pVFirst->m_VertexBuffers;
+			Result->GetVertexBuffers().SetElementCount( FirstCount + SecondCount );
+
+            // first copy all the vertex data
+			{
+				MUTABLE_CPUPROFILER_SCOPE(CopyVertexData);
+				for (int32 vb = 0; vb < Result->GetVertexBuffers().m_buffers.Num(); ++vb)
+				{
+					MESH_BUFFER& result = Result->GetVertexBuffers().m_buffers[vb];
+					const MESH_BUFFER& second = pVSecond->GetVertexBuffers().m_buffers[vb];
+
+					if (SecondCount)
+					{
+						int elemSize = Result->GetVertexBuffers().GetElementSize((int)vb);
+						int firstSize = FirstCount * elemSize;
+						int secondSize = SecondCount * elemSize;
+
+						FMemory::Memcpy( &result.m_data[firstSize], &second.m_data[0], secondSize );
+					}
+				}
+			}
+
+
+            if (bRemapBoneIndices)
+            {
+                MUTABLE_CPUPROFILER_SCOPE(Remap);
+
+				// We need to remap the bone indices of the second mesh vertices that we already copied
+				// to result
+				check(!RemappedBoneMapIndices.IsEmpty())
+
+				// Iterate all vertex buffers and update the format
+				FMeshBufferSet& VertexBuffers = Result->GetVertexBuffers();
+				for (int32 VertexBufferIndex = 0; VertexBufferIndex < VertexBuffers.m_buffers.Num(); ++VertexBufferIndex)
+				{
+					MESH_BUFFER& ResultBuffer = VertexBuffers.m_buffers[VertexBufferIndex];
+
+					const int32 ElemSize = VertexBuffers.GetElementSize(VertexBufferIndex);
+					const int32 FirstSize = FirstCount * ElemSize;
+
+					const int32 ChannelsCount = VertexBuffers.GetBufferChannelCount(VertexBufferIndex);
+					for (int32 ChannelIndex = 0; ChannelIndex < ChannelsCount; ++ChannelIndex)
+					{
+						if (ResultBuffer.m_channels[ChannelIndex].m_semantic != MBS_BONEINDICES)
+						{
+							continue;
+						}
+						
+						int32 ResultOffset = FirstSize + ResultBuffer.m_channels[ChannelIndex].m_offset;
+
+						const int32 NumComponents = ResultBuffer.m_channels[ChannelIndex].m_componentCount;
+
+						// Bone indices may need remapping
+						for (int32 VertexIndex = 0; VertexIndex < pVSecond->GetVertexCount(); ++VertexIndex)
+						{
+							switch (BoneIndexFormat)
+							{
+							case MBF_INT8:
+							case MBF_UINT8:
+							{
+								uint8* pD = reinterpret_cast<uint8*>
+									(&ResultBuffer.m_data[ResultOffset]);
+
+								for (int32 ComponentIndex = 0; ComponentIndex < NumComponents; ++ComponentIndex)
+								{
+									uint8 BoneMapIndex = pD[ComponentIndex];
+
+									// be defensive
+									if (BoneMapIndex < NumSecondBonesInBoneMap)
+									{
+										pD[ComponentIndex] = (uint8)RemappedBoneMapIndices[BoneMapIndex];
+									}
+									else
+									{
+										pD[ComponentIndex] = 0;
+									}
+								}
+
+								ResultOffset += ElemSize;
+								break;
+							}
+
+							case MBF_INT16:
+							case MBF_UINT16:
+							{
+								uint16* pD = reinterpret_cast<uint16*>
+									(&ResultBuffer.m_data[ResultOffset]);
+
+								for (int32 ComponentIndex = 0; ComponentIndex < NumComponents; ++ComponentIndex)
+								{
+									uint16 BoneMapIndex = pD[ComponentIndex];
+
+									// be defensive
+									if (BoneMapIndex < NumSecondBonesInBoneMap)
+									{
+										pD[ComponentIndex] = (uint16)RemappedBoneMapIndices[BoneMapIndex];
+									}
+									else
+									{
+										pD[ComponentIndex] = 0;
+									}
+								}
+
+								ResultOffset += ElemSize;
+								break;
+							}
+
+							case MBF_INT32:
+							case MBF_UINT32:
+							{
+								// Unreal does not support 32 bit bone indices
+								checkf(false, TEXT("Format not supported."));
+								break;
+							}
+
+							default:
+								checkf(false, TEXT("Format not supported."));
+							}
+						}
+					}
+				}
+			}
+		}
+
+		// Tags
+		Result->m_tags = pFirst->m_tags;
+
+		for (const FString& SecondTag : pSecond->m_tags)
+		{
+			Result->m_tags.AddUnique(SecondTag);
+		}
+
+
+		// Streamed Resources
+		Result->StreamedResources = pFirst->StreamedResources;
+
+		const int32 NumStreamedResources = pSecond->StreamedResources.Num();
+		for (int32 Index = 0; Index < NumStreamedResources; ++Index)
+		{
+			Result->StreamedResources.AddUnique(pSecond->StreamedResources[Index]);
+		}
+	}
+	
+
+    //---------------------------------------------------------------------------------------------
+    //!
+    //---------------------------------------------------------------------------------------------
+    inline void ExtendSkeleton( Skeleton* pBase, const Skeleton* pOther )
+    {
+        TMap<int,int> otherToResult;
+
+        int initialBones = pBase->GetBoneCount();
+        for ( int b=0; pOther && b<pOther->GetBoneCount(); ++b)
+        {
+            int32 resultBoneIndex = pBase->FindBone( pOther->GetBoneId(b) );
+            if ( resultBoneIndex<0 )
+            {
+                int32 newIndex = pBase->BoneIds.Num();
+                otherToResult.Add(b,newIndex);
+                pBase->BoneIds.Add( pOther->BoneIds[b] );
+
+                // Will be remapped below
+                pBase->BoneParents.Add(pOther->BoneParents[b] );
+
+				if (pOther->BoneNames.IsValidIndex(b))
+				{
+					pBase->BoneNames.Add(pOther->BoneNames[b]);
+				}
+            }
+            else
+            {
+                otherToResult.Add(b,resultBoneIndex);
+            }
+        }
+
+        // Fix bone parent indices of the bones added from pOther
+        for ( int b=initialBones;b<pBase->GetBoneCount(); ++b)
+        {
+            int16_t sourceIndex = pBase->BoneParents[b];
+            if (sourceIndex>=0)
+            {
+                pBase->BoneParents[b] = (int16_t)otherToResult[ sourceIndex ];
+            }
+        }
+    }
+
+
+    //---------------------------------------------------------------------------------------------
+    //! Return null if there is no need to remap the mesh
+    //---------------------------------------------------------------------------------------------
+    inline void MeshRemapSkeleton(Mesh* Result, const Mesh* SourceMesh, const Skeleton* Skeleton, bool& bOutSuccess)
+    {
+		bOutSuccess = true;
+
+        if (SourceMesh->GetVertexCount() == 0 || !SourceMesh->GetSkeleton() || SourceMesh->GetSkeleton()->GetBoneCount() == 0)
+        {
+			bOutSuccess = false;
+            return;
+        }
+
+  //      mu::SkeletonPtrConst SourceSkeleton = SourceMesh->GetSkeleton();
+		//const TArray<uint16>& SourceBoneMaps = SourceMesh->GetBoneMap();
+
+  //      // Remap the indices of the bonemap to those of the new skeleton
+		//TArray<uint16> RemappedBoneMap;
+
+		//const int32 NumBonesBoneMap = SourceBoneMaps.Num();
+		//RemappedBoneMap.Reserve(NumBonesBoneMap);
+
+		//bool bBonesRemapped = false;
+
+		//for (const uint16& SourceBoneIndex : SourceBoneMaps)
+		//{
+		//	const uint16 BoneIndex = Skeleton->FindBone(SourceSkeleton->GetBoneId(SourceBoneIndex));
+
+		//	bBonesRemapped = bBonesRemapped || BoneIndex != SourceBoneIndex;
+		//	RemappedBoneMap.Add(BoneIndex);
+		//}
+
+  //      if (!bBonesRemapped)
+  //      {
+		//	bOutSuccess = false;
+  //          return;
+  //      }
+
+		Result->CopyFrom(*SourceMesh);
+		Result->SetSkeleton(Skeleton);
+		//Result->SetBoneMap(RemappedBoneMap);
+    }
+	
+}
diff --git a/Engine/Plugins/Experimental/Mutable/Source/MutableRuntime/Internal/MuR/Platform.h b/Engine/Plugins/Experimental/Mutable/Source/MutableRuntime/Internal/MuR/Platform.h
new file mode 100644
--- /dev/null
+++ b/Engine/Plugins/Experimental/Mutable/Source/MutableRuntime/Internal/MuR/Platform.h
@@ -0,0 +1,18 @@
+// Copyright Epic Games, Inc. All Rights Reserved.
+
+#pragma once
+
+// This file contains some hacks to solve differences between platforms
+
+#include "HAL/PlatformTime.h"
+#include "MuR/Types.h"
+#include "HAL/UnrealMemory.h"
+
+
+
+//! Unify debug defines
+#if !defined(MUTABLE_DEBUG)
+    #if !defined(NDEBUG) || defined(_DEBUG)
+        #define MUTABLE_DEBUG
+    #endif
+#endif
diff --git a/Engine/Plugins/Experimental/Mutable/Source/MutableRuntime/Private/MuR/BlockCompression/Miro/Miro.h b/Engine/Plugins/Experimental/Mutable/Source/MutableRuntime/Private/MuR/BlockCompression/Miro/Miro.h
deleted file mode 100644
diff --git a/Engine/Plugins/Experimental/Mutable/Source/MutableRuntime/Private/MuR/CodeVisitor.h b/Engine/Plugins/Experimental/Mutable/Source/MutableRuntime/Private/MuR/CodeVisitor.h
deleted file mode 100644
diff --git a/Engine/Plugins/Experimental/Mutable/Source/MutableRuntime/Private/MuR/ConvertData.h b/Engine/Plugins/Experimental/Mutable/Source/MutableRuntime/Private/MuR/ConvertData.h
deleted file mode 100644
diff --git a/Engine/Plugins/Experimental/Mutable/Source/MutableRuntime/Private/MuR/ImagePrivate.h b/Engine/Plugins/Experimental/Mutable/Source/MutableRuntime/Private/MuR/ImagePrivate.h
deleted file mode 100644
diff --git a/Engine/Plugins/Experimental/Mutable/Source/MutableRuntime/Private/MuR/MeshPrivate.h b/Engine/Plugins/Experimental/Mutable/Source/MutableRuntime/Private/MuR/MeshPrivate.h
deleted file mode 100644
diff --git a/Engine/Plugins/Experimental/Mutable/Source/MutableRuntime/Private/MuR/OpImageProject.h b/Engine/Plugins/Experimental/Mutable/Source/MutableRuntime/Private/MuR/OpImageProject.h
deleted file mode 100644
diff --git a/Engine/Plugins/Experimental/Mutable/Source/MutableRuntime/Private/MuR/OpMeshMerge.h b/Engine/Plugins/Experimental/Mutable/Source/MutableRuntime/Private/MuR/OpMeshMerge.h
deleted file mode 100644
diff --git a/Engine/Plugins/Experimental/Mutable/Source/MutableRuntime/Private/MuR/Platform.h b/Engine/Plugins/Experimental/Mutable/Source/MutableRuntime/Private/MuR/Platform.h
deleted file mode 100644
diff --git a/Engine/Plugins/Experimental/Mutable/Source/MutableTools/MutableTools.Build.cs b/Engine/Plugins/Experimental/Mutable/Source/MutableTools/MutableTools.Build.cs
--- a/Engine/Plugins/Experimental/Mutable/Source/MutableTools/MutableTools.Build.cs
+++ b/Engine/Plugins/Experimental/Mutable/Source/MutableTools/MutableTools.Build.cs
@@ -12,11 +12,6 @@ namespace UnrealBuildTool.Rules
 
 			DefaultBuildSettings = BuildSettingsVersion.V2;
 			IWYUSupport = IWYUSupport.KeepAsIsForNow;
-			//bUseUnity = false;
-
-			PrivateIncludePaths.AddRange(new string[] {
-				System.IO.Path.Combine(GetModuleDirectory("MutableRuntime"), "Private"),
-			});
 
 			PublicDependencyModuleNames.AddRange(
                 new string[] {
diff --git a/Engine/Plugins/Experimental/Mutable/Source/MutableTools/Private/MuT/CodeGenerator_Mesh.cpp b/Engine/Plugins/Experimental/Mutable/Source/MutableTools/Private/MuT/CodeGenerator_Mesh.cpp
--- a/Engine/Plugins/Experimental/Mutable/Source/MutableTools/Private/MuT/CodeGenerator_Mesh.cpp
+++ b/Engine/Plugins/Experimental/Mutable/Source/MutableTools/Private/MuT/CodeGenerator_Mesh.cpp
@@ -6,6 +6,7 @@
 #include "Misc/AssertionMacros.h"
 #include "MuR/Layout.h"
 #include "MuR/Mesh.h"
+#include "MuR/MeshPrivate.h"
 #include "MuR/MeshBufferSet.h"
 #include "MuR/MutableMath.h"
 #include "MuR/Operations.h"
diff --git a/Engine/Plugins/Experimental/Mutable/Source/MutableTools/Private/MuT/CodeOptimiser.h b/Engine/Plugins/Experimental/Mutable/Source/MutableTools/Private/MuT/CodeOptimiser.h
--- a/Engine/Plugins/Experimental/Mutable/Source/MutableTools/Private/MuT/CodeOptimiser.h
+++ b/Engine/Plugins/Experimental/Mutable/Source/MutableTools/Private/MuT/CodeOptimiser.h
@@ -3,8 +3,6 @@
 #pragma once
 
 #include "MuR/Operations.h"
-#include "MuR/ModelPrivate.h"
-#include "MuR/ImagePrivate.h"
 #include "MuR/CodeVisitor.h"
 #include "MuT/CompilerPrivate.h"
 #include "MuT/AST.h"
diff --git a/Engine/Plugins/Experimental/Mutable/Source/MutableTools/Private/MuT/CodeOptimiser_Semantic.cpp b/Engine/Plugins/Experimental/Mutable/Source/MutableTools/Private/MuT/CodeOptimiser_Semantic.cpp
--- a/Engine/Plugins/Experimental/Mutable/Source/MutableTools/Private/MuT/CodeOptimiser_Semantic.cpp
+++ b/Engine/Plugins/Experimental/Mutable/Source/MutableTools/Private/MuT/CodeOptimiser_Semantic.cpp
@@ -1,10 +1,9 @@
 // Copyright Epic Games, Inc. All Rights Reserved.
 
+#include "NodeImageFormatPrivate.h"
 #include "HAL/PlatformMath.h"
 #include "Misc/AssertionMacros.h"
 #include "MuR/Image.h"
-#include "MuR/ImagePrivate.h"
-#include "MuR/ImageRLE.h"
 #include "MuR/MutableMath.h"
 #include "MuR/MutableTrace.h"
 #include "MuR/Operations.h"
diff --git a/Engine/Plugins/Experimental/Mutable/Source/MutableTools/Private/MuT/NodeMeshInterpolatePrivate.h b/Engine/Plugins/Experimental/Mutable/Source/MutableTools/Private/MuT/NodeMeshInterpolatePrivate.h
--- a/Engine/Plugins/Experimental/Mutable/Source/MutableTools/Private/MuT/NodeMeshInterpolatePrivate.h
+++ b/Engine/Plugins/Experimental/Mutable/Source/MutableTools/Private/MuT/NodeMeshInterpolatePrivate.h
@@ -7,8 +7,6 @@
 #include "MuT/NodeMeshPrivate.h"
 #include "MuT/AST.h"
 
-#include "MuR/MeshPrivate.h"
-
 
 namespace mu
 {
