[Mutable] Fix instance generation crash in the editor if the FirstLODAvailable is higher than 0 and mesh cache is enabled.

#rnx

diff --git a/Engine/Plugins/Experimental/Mutable/Source/CustomizableObject/Internal/MuCO/CustomizableObjectSystemPrivate.h b/Engine/Plugins/Experimental/Mutable/Source/CustomizableObject/Internal/MuCO/CustomizableObjectSystemPrivate.h
--- a/Engine/Plugins/Experimental/Mutable/Source/CustomizableObject/Internal/MuCO/CustomizableObjectSystemPrivate.h
+++ b/Engine/Plugins/Experimental/Mutable/Source/CustomizableObject/Internal/MuCO/CustomizableObjectSystemPrivate.h
@@ -544,6 +544,7 @@ public:
 
 	int32 NumComponents = 0;
 	int32 NumLODsAvailable = 0;
+	int32 FirstLODAvailable = 0;
 
 	TMap<uint32, FTexturePlatformData*> ImageToPlatformDataMap;
 
diff --git a/Engine/Plugins/Experimental/Mutable/Source/CustomizableObject/Private/MuCO/CustomizableObjectSystem.cpp b/Engine/Plugins/Experimental/Mutable/Source/CustomizableObject/Private/MuCO/CustomizableObjectSystem.cpp
--- a/Engine/Plugins/Experimental/Mutable/Source/CustomizableObject/Private/MuCO/CustomizableObjectSystem.cpp
+++ b/Engine/Plugins/Experimental/Mutable/Source/CustomizableObject/Private/MuCO/CustomizableObjectSystem.cpp
@@ -135,6 +135,7 @@ FUpdateContextPrivate::FUpdateContextPrivate(UCustomizableObjectInstance& InInst
 	CapturedDescriptorHash = FDescriptorHash(Descriptor);
 	Parameters = Descriptor.GetParameters();
 	NumComponents = InInstance.GetCustomizableObject()->GetComponentCount();
+	FirstLODAvailable = InInstance.GetCustomizableObject()->GetPrivate()->GetMinLODIndex();
 	
 	InInstance.GetCustomizableObject()->GetPrivate()->ApplyStateForcedValuesToParameters(CapturedDescriptor.GetState(), Parameters.get());
 
@@ -1607,7 +1608,7 @@ namespace impl
 	{
 		Operation->NumLODsAvailable = Operation->MutableInstance->GetLODCount();
 
-		int32 CurrentMinLOD = Operation->GetMinLOD();
+		int32 CurrentMinLOD = FMath::Max(Operation->GetMinLOD(), Operation->FirstLODAvailable);
 		int32 CurrentMaxLOD = Operation->GetMaxLOD();
 
 		if (CurrentMinLOD >= Operation->NumLODsAvailable)
diff --git a/Engine/Plugins/Experimental/Mutable/Source/CustomizableObjectEditor/Private/MuCOE/GenerateMutableSource/GenerateMutableSource.cpp b/Engine/Plugins/Experimental/Mutable/Source/CustomizableObjectEditor/Private/MuCOE/GenerateMutableSource/GenerateMutableSource.cpp
--- a/Engine/Plugins/Experimental/Mutable/Source/CustomizableObjectEditor/Private/MuCOE/GenerateMutableSource/GenerateMutableSource.cpp
+++ b/Engine/Plugins/Experimental/Mutable/Source/CustomizableObjectEditor/Private/MuCOE/GenerateMutableSource/GenerateMutableSource.cpp
@@ -798,7 +798,7 @@ mu::NodeObjectPtr GenerateMutableSource(const UEdGraphPin * Pin, FMutableGraphGe
 			{
 				FSupportedQualityLevelArray SupportedQualityLevels = LODSettings.MinQualityLevelLOD.GetSupportedQualityLevels(*GenerationContext.Options.TargetPlatform->GetPlatformInfo().IniPlatformName.ToString());
 				
-				int32 MinValue = MAX_int32;
+				int32 MinValue = GenerationContext.NumLODsInRoot - 1;
 				for (int32& QL : SupportedQualityLevels)
 				{
 					// check if have data for the supported quality level or set to default.
@@ -813,7 +813,7 @@ mu::NodeObjectPtr GenerateMutableSource(const UEdGraphPin * Pin, FMutableGraphGe
 					}
 				}
 
-				GenerationContext.FirstLODAvailable = LODSettings.MinQualityLevelLOD.GetValueForPlatform(GenerationContext.Options.TargetPlatform);
+				GenerationContext.FirstLODAvailable = FMath::Max(0, MinValue);
 			}
 			else
 			{
