[Mutable] Fix CI crash due to missing backward compatibility for BoneParents.

#rnx

diff --git a/Engine/Plugins/Experimental/Mutable/Source/MutableRuntime/Private/MuR/Skeleton.cpp b/Engine/Plugins/Experimental/Mutable/Source/MutableRuntime/Private/MuR/Skeleton.cpp
--- a/Engine/Plugins/Experimental/Mutable/Source/MutableRuntime/Private/MuR/Skeleton.cpp
+++ b/Engine/Plugins/Experimental/Mutable/Source/MutableRuntime/Private/MuR/Skeleton.cpp
@@ -173,6 +173,16 @@ namespace mu
 
 		arch >> BoneParents;
 
+		if (ver < 6) // before BoneIds 
+		{
+			int16 ParentIndex = INDEX_NONE;
+			for (int16& BoneParent : BoneParents)
+			{
+				BoneParent = ParentIndex;
+				ParentIndex++;
+			}
+		}
+
 		if (ver <= 4)
 		{
 			TArray<int32> BoneIds_DEPRECATED;
