[mutable] Fixed crash of multilayer projectors with only one image parameter option.
+ Fixed projector's selector button.
+ Improved multilayer projector's widget.

#rnx

#virtualized

diff --git a/Engine/Plugins/Experimental/Mutable/Source/CustomizableObject/Private/MuCO/CustomizableObjectInstanceDescriptor.cpp b/Engine/Plugins/Experimental/Mutable/Source/CustomizableObject/Private/MuCO/CustomizableObjectInstanceDescriptor.cpp
--- a/Engine/Plugins/Experimental/Mutable/Source/CustomizableObject/Private/MuCO/CustomizableObjectInstanceDescriptor.cpp
+++ b/Engine/Plugins/Experimental/Mutable/Source/CustomizableObject/Private/MuCO/CustomizableObjectInstanceDescriptor.cpp
@@ -872,6 +872,15 @@ void FCustomizableObjectInstanceDescriptor::ReloadParameters()
 						const int32 ParamValue = MutableParameters->GetIntValue(ParamIndex);
 						Param.ParameterValueName = CustomizableObject->FindIntParameterValueName(ParamIndex, ParamValue);
 					}
+
+					// Multilayer ints with one option are not multidimensional parameters. However, we need to preserve the layer 
+					// information in case that we add a new option to the parameter, and it is converted to multidimensional.
+					for (int32 RangeIndex = 0; RangeIndex < Result->ParameterRangeValueNames.Num(); ++RangeIndex)
+					{
+						const int32 Value = MutableParameters->GetIntValue(ParamIndex);
+						const FString AuxParameterValueName = CustomizableObject->FindIntParameterValueName(ParamIndex, Value);
+						Param.ParameterRangeValueNames.Add(AuxParameterValueName);
+					}
 				}
 			} 
 			else // Not found in Instance Parameters. Use Mutable Parameters.
diff --git a/Engine/Plugins/Experimental/Mutable/Source/CustomizableObjectEditor/CustomizableObjectEditor.Build.cs b/Engine/Plugins/Experimental/Mutable/Source/CustomizableObjectEditor/CustomizableObjectEditor.Build.cs
--- a/Engine/Plugins/Experimental/Mutable/Source/CustomizableObjectEditor/CustomizableObjectEditor.Build.cs
+++ b/Engine/Plugins/Experimental/Mutable/Source/CustomizableObjectEditor/CustomizableObjectEditor.Build.cs
@@ -51,6 +51,7 @@ public class CustomizableObjectEditor : ModuleRules
 				"ClothingSystemRuntimeInterface",
 				"DeveloperToolSettings",
 				"ToolMenus",
+				"ToolWidgets",
 				"EditorFramework",
 				"UMG",
 				"Persona",
diff --git a/Engine/Plugins/Experimental/Mutable/Source/CustomizableObjectEditor/Private/MuCOE/CustomizableInstanceDetails.cpp b/Engine/Plugins/Experimental/Mutable/Source/CustomizableObjectEditor/Private/MuCOE/CustomizableInstanceDetails.cpp
--- a/Engine/Plugins/Experimental/Mutable/Source/CustomizableObjectEditor/Private/MuCOE/CustomizableInstanceDetails.cpp
+++ b/Engine/Plugins/Experimental/Mutable/Source/CustomizableObjectEditor/Private/MuCOE/CustomizableInstanceDetails.cpp
@@ -22,6 +22,7 @@
 #include "MuCOE/SMutableTextSearchBox.h"
 #include "MuCOE/UnrealEditorPortabilityHelpers.h"
 
+#include "SSearchableComboBox.h"
 #include "Widgets/Colors/SColorBlock.h"
 #include "Widgets/Colors/SColorPicker.h"
 #include "Widgets/Input/SButton.h"
@@ -897,6 +898,7 @@ TSharedRef<SWidget> FCustomizableInstanceDetails::GenerateIntWidget(const int32
 			OptionNamesAttribute.Add(CustomizableObject->GetIntParameterAvailableOption(ParamIndexInObject, i));
 		}
 
+		//TODO(Max): UE-207137
 		return SNew(SMutableTextSearchBox)
 			.ToolTipText(FText::FromString(ToolTipText))
 			.PossibleSuggestions(OptionNamesAttribute)
@@ -1405,6 +1407,7 @@ TSharedRef<SWidget> FCustomizableInstanceDetails::GenerateMultidimensionalProjec
 			.VAlign(VAlign_Center)
 			.FillWidth(0.45f)
 			[
+				//TODO(Max): UE-207137
 				SNew(SMutableTextSearchBox)
 				.PossibleSuggestions(PoseOptionNamesAttribute)
 				.InitialText(FText::FromString(PoseOptionNamesAttribute[PoseValueIndex]))
@@ -1444,170 +1447,175 @@ TSharedRef<SWidget> FCustomizableInstanceDetails::GenerateMultidimensionalProjec
 	}
 
 	const FString TextureSwitchEnumParamName = ParamName + FMultilayerProjector::IMAGE_PARAMETER_POSTFIX;
+	const FString OpacitySliderParamName = ParamName + FMultilayerProjector::OPACITY_PARAMETER_POSTFIX;
 
 	IDetailGroup* ProjectorGroup = *ParentsGroups.Find(ParamName);
 
 	for (int32 RangeIndex = 0; RangeIndex < ProjectorParameters[ProjectorParamIndex].RangeValues.Num(); ++RangeIndex)
 	{
 		const int32 TextureSwitchEnumParamIndexInObject = CustomizableObject->FindParameter(TextureSwitchEnumParamName);
-		check(TextureSwitchEnumParamIndexInObject >= 0);
-		int32 NumValues = CustomizableObject->GetIntParameterNumOptions(TextureSwitchEnumParamIndexInObject);
-
-		FString OpacitySliderParamName = ParamName + FMultilayerProjector::OPACITY_PARAMETER_POSTFIX;
-
-		TArray<FString> OptionNamesAttribute;
-		FString Value = CustomInstance->GetIntParameterSelectedOption(TextureSwitchEnumParamName, RangeIndex);
-		int32 ValueIndex = 0;
+		check(TextureSwitchEnumParamIndexInObject >= 0); TSharedPtr<FString> CurrentStateName = nullptr;
 
 		const UProjectorParameter* ProjectorParameter = Editor->GetProjectorParameter();
 		const bool bSelectedProjector = ProjectorParameter->IsProjectorSelected(ParamName, RangeIndex);
 
-		for (int32 CandidateIndex = 0; CandidateIndex < NumValues; ++CandidateIndex)
-		{
-			FString PossibleValue = CustomizableObject->GetIntParameterAvailableOption(TextureSwitchEnumParamIndexInObject, CandidateIndex);
-			if (PossibleValue == Value)
-			{
-				ValueIndex = CandidateIndex;
-			}
-			OptionNamesAttribute.Add(CustomizableObject->GetIntParameterAvailableOption(TextureSwitchEnumParamIndexInObject, CandidateIndex));
-		}
-
-		TSharedPtr<SHorizontalBox> SliderBox;
-		TSharedPtr<SSpinBox<float>> Slider;
-		TSharedPtr<SHorizontalBox> PropertiesHB;
+		//Vertical box that owns all the layer properties
+		TSharedPtr<SVerticalBox> LayerProperties = SNew(SVerticalBox);
+		//Horizontal box that owns all the projector properties
+		TSharedPtr<SHorizontalBox> ProjectorProperties;
+		// Widget to set the opacity and remove a layer
+		TSharedPtr<SHorizontalBox> OpacityRemoveWidget = SNew(SHorizontalBox);
+		// Button Ptr needed to edit its style
 		TSharedPtr<SButton> Button;
-		TSharedPtr<STextBlock> TextBlock;
 
-		//Horizontal box that owns all the hidden layer properties
-		SAssignNew(PropertiesHB, SHorizontalBox)
+		SAssignNew(ProjectorProperties, SHorizontalBox)
 		+ SHorizontalBox::Slot()
 		.Padding(1, 0)
-		.AutoWidth()
-		[
-			SNew(SHorizontalBox)
-			+ SHorizontalBox::Slot()
-			.Padding(1, 0)
-			[
-				SNew(SBox)
-				.MinDesiredWidth(115.f)
-				.MaxDesiredWidth(115.f)
-				[
-					SAssignNew(Button, SButton)
-					.OnClicked(this, &FCustomizableInstanceDetails::OnProjectorSelectChanged, ParamName, RangeIndex)
-					.VAlign(VAlign_Center)
-					.HAlign(HAlign_Center)
-					.Content()
-					[
-						SNew(STextBlock)
-						.Text(bSelectedProjector ? LOCTEXT("Unselect Projector", "Unselect Projector") : LOCTEXT("Select Projector", "Select Projector"))
-						.ToolTipText(bSelectedProjector ? LOCTEXT("Unselect Projector", "Unselect Projector") : LOCTEXT("Select Projector", "Select Projector"))
-						.Justification(ETextJustify::Center)
-						.Font(LayoutBuilder.Pin()->GetDetailFont())
-						//.AutoWrapText(true)
-					]
-				]
-			]
-
-			+ SHorizontalBox::Slot()
-			.Padding(1, 0)
-			[
-				SNew(SBox)
-				.MinDesiredWidth(120.f)
-				[
-					SNew(SMutableTextSearchBox)
-					.PossibleSuggestions(OptionNamesAttribute)
-					.InitialText(FText::FromString(OptionNamesAttribute[ValueIndex]))
-					.MustMatchPossibleSuggestions(TAttribute<bool>(true))
-					.SuggestionListPlacement(EMenuPlacement::MenuPlacement_ComboBox)
-					.OnTextCommitted(this, &FCustomizableInstanceDetails::OnProjectorTextureParameterComboBoxChanged, TextureSwitchEnumParamName, RangeIndex)
-				]
-			]
-		];
-
-		Button->SetBorderBackgroundColor(bSelectedProjector ? FLinearColor::Green : FLinearColor::White);
-
-		PropertiesHB->AddSlot()
-		.HAlign(HAlign_Fill)
 		[
-			SNew(SHorizontalBox)
-			+ SHorizontalBox::Slot()
-			.Padding(1, 0)
-			.HAlign(HAlign_Fill)
+			SNew(SBox)
+			.MinDesiredWidth(115.f)
+			.MaxDesiredWidth(115.f)
 			[
-				SNew(SHorizontalBox)
-				+ SHorizontalBox::Slot()
+				SAssignNew(Button, SButton)
+				.OnClicked(this, &FCustomizableInstanceDetails::OnProjectorSelectChanged, ParamName, RangeIndex)
+				.VAlign(VAlign_Center)
+				.HAlign(HAlign_Center)
+				.Content()
 				[
-					SNew(SSpinBox<float>)
-					.MinValue(0.0f)
-					.MaxValue(1.0f)
-					.Value(this, &FCustomizableInstanceDetails::GetFloatParameterValue, OpacitySliderParamName, RangeIndex)
-					.OnValueChanged(this, &FCustomizableInstanceDetails::OnFloatParameterChanged, OpacitySliderParamName, RangeIndex)
-					.OnEndSliderMovement(this, &FCustomizableInstanceDetails::OnFloatParameterSliderEnd, OpacitySliderParamName, RangeIndex)
+					SNew(STextBlock)
+					.Text(bSelectedProjector ? LOCTEXT("Unselect Projector", "Unselect Projector") : LOCTEXT("Select Projector", "Select Projector"))
+					.ToolTipText(bSelectedProjector ? LOCTEXT("Unselect Projector", "Unselect Projector") : LOCTEXT("Select Projector", "Select Projector"))
+					.Justification(ETextJustify::Center)
+					.Font(LayoutBuilder.Pin()->GetDetailFont())
 				]
 			]
-		];
+		]
 
-		PropertiesHB->AddSlot()
-		.HAlign(HAlign_Right)
-		.AutoWidth()
+		+ SHorizontalBox::Slot()
+		.Padding(1, 0)
 		[
-			SNew(SHorizontalBox)
-			+ SHorizontalBox::Slot()
-			.Padding(1, 0)
-			[
-				SNew(SButton)
+			SNew(SButton)
 				.OnClicked(this, &FCustomizableInstanceDetails::OnProjectorCopyTransform, ParamName, RangeIndex)
 				.VAlign(VAlign_Center)
 				.HAlign(HAlign_Center)
 				.Content()
 				[
 					SNew(STextBlock)
-					.ToolTipText(LOCTEXT("Copy Transform", "Copy Transform"))
-					.Text(LOCTEXT("Copy Transform", "Copy Transform"))
-					.AutoWrapText(true)
-					.Justification(ETextJustify::Center)
-					.Font(LayoutBuilder.Pin()->GetDetailFont())
+						.ToolTipText(LOCTEXT("Copy Transform", "Copy Transform"))
+						.Text(LOCTEXT("Copy Transform", "Copy Transform"))
+						.AutoWrapText(true)
+						.Justification(ETextJustify::Center)
+						.Font(LayoutBuilder.Pin()->GetDetailFont())
 				]
-			]
+		]
 
-			+ SHorizontalBox::Slot()
-			.Padding(1, 0)
-			[
-				SNew(SButton)
+		+ SHorizontalBox::Slot()
+		.Padding(1, 0)
+		[
+			SNew(SButton)
 				.OnClicked(this, &FCustomizableInstanceDetails::OnProjectorPasteTransform, ParamName, RangeIndex)
 				.VAlign(VAlign_Center)
 				.HAlign(HAlign_Center)
 				.Content()
 				[
 					SNew(STextBlock)
-					.ToolTipText(LOCTEXT("Paste Transform", "Paste Transform"))
-					.Text(LOCTEXT("Paste Transform", "Paste Transform"))
-					.Justification(ETextJustify::Center)
-					.AutoWrapText(true)
-					.Font(LayoutBuilder.Pin()->GetDetailFont())
+						.ToolTipText(LOCTEXT("Paste Transform", "Paste Transform"))
+						.Text(LOCTEXT("Paste Transform", "Paste Transform"))
+						.Justification(ETextJustify::Center)
+						.AutoWrapText(true)
+						.Font(LayoutBuilder.Pin()->GetDetailFont())
 				]
-			]
+		]
 
-			+ SHorizontalBox::Slot()
-			.Padding(1, 0)
-			[
-				SNew(SButton)
+		+ SHorizontalBox::Slot()
+		.Padding(1, 0)
+		[
+			SNew(SButton)
 				.OnClicked(this, &FCustomizableInstanceDetails::OnProjectorResetTransform, ParamName, RangeIndex)
 				.VAlign(VAlign_Center)
 				.HAlign(HAlign_Center)
 				.Content()
 				[
 					SNew(STextBlock)
-					.ToolTipText(LOCTEXT("Reset Transform", "Reset Transform"))
-					.Text(LOCTEXT("Reset Transform", "Reset Transform"))
-					.Justification(ETextJustify::Center)
-					.AutoWrapText(true)
-					.Font(LayoutBuilder.Pin()->GetDetailFont())
+						.ToolTipText(LOCTEXT("Reset Transform", "Reset Transform"))
+						.Text(LOCTEXT("Reset Transform", "Reset Transform"))
+						.Justification(ETextJustify::Center)
+						.AutoWrapText(true)
+						.Font(LayoutBuilder.Pin()->GetDetailFont())
+				]
+		];
+
+		Button->SetBorderBackgroundColor(bSelectedProjector ? FLinearColor::Green : FLinearColor::White);
+
+		// If number of options is equal to 1, Mutable does not consider it multidimensional parameters
+		int32 NumValues = CustomizableObject->GetIntParameterNumOptions(TextureSwitchEnumParamIndexInObject);
+		FString Value = CustomizableObject->IsParameterMultidimensional(TextureSwitchEnumParamName) ? 
+			CustomInstance->GetIntParameterSelectedOption(TextureSwitchEnumParamName, RangeIndex) : CustomInstance->GetIntParameterSelectedOption(TextureSwitchEnumParamName);
+
+		TArray<TSharedPtr<FString>> OptionNamesAttribute;
+		int32 ValueIndex = 0;
+
+		for (int32 CandidateIndex = 0; CandidateIndex < NumValues; ++CandidateIndex)
+		{
+			FString PossibleValue = CustomizableObject->GetIntParameterAvailableOption(TextureSwitchEnumParamIndexInObject, CandidateIndex);
+			if (PossibleValue == Value)
+			{
+				ValueIndex = CandidateIndex;
+			}
+
+			OptionNamesAttribute.Add(MakeShared<FString>(CustomizableObject->GetIntParameterAvailableOption(TextureSwitchEnumParamIndexInObject, CandidateIndex)));
+		}
+		
+		// Avoid filling this arraw with repeated array  options
+		if (RangeIndex == 0)
+		{
+			ProjectorTextureOptions.Add(OptionNamesAttribute);
+		}
+
+		if (NumValues > 1)
+		{
+			OpacityRemoveWidget->AddSlot()
+			.Padding(1, 0)
+			.FillWidth(0.3f)
+			[
+				SNew(SBox)
+				[
+					SNew(SSearchableComboBox)
+					.OptionsSource(&ProjectorTextureOptions.Last())
+					.InitiallySelectedItem(ProjectorTextureOptions.Last()[ValueIndex])
+					.OnGenerateWidget(this, &FCustomizableInstanceDetails::MakeTextureComboEntryWidget)
+					.OnSelectionChanged(this, &FCustomizableInstanceDetails::OnProjectorTextureParameterComboBoxChanged, TextureSwitchEnumParamName, RangeIndex)
+					.Content()
+					[
+						SNew(STextBlock)
+						.Text(FText::FromString(*ProjectorTextureOptions.Last()[ValueIndex]))
+					]
 				]
+			];
+		}
+		
+		OpacityRemoveWidget->AddSlot()
+		.HAlign(HAlign_Fill)
+		.Padding(1, 0)
+		.FillWidth(0.7f)
+		[
+			SNew(SHorizontalBox)
+			+SHorizontalBox::Slot()
+			.Padding(1, 0)
+			.HAlign(EHorizontalAlignment::HAlign_Fill)
+			[
+				SNew(SSpinBox<float>)
+				.MinValue(0.0f)
+				.MaxValue(1.0f)
+				.Value(this, &FCustomizableInstanceDetails::GetFloatParameterValue, OpacitySliderParamName, RangeIndex)
+				.OnValueChanged(this, &FCustomizableInstanceDetails::OnFloatParameterChanged, OpacitySliderParamName, RangeIndex)
+				.OnEndSliderMovement(this, &FCustomizableInstanceDetails::OnFloatParameterSliderEnd, OpacitySliderParamName, RangeIndex)
+				.Font(LayoutBuilder.Pin()->GetDetailFont())
 			]
 
 			+ SHorizontalBox::Slot()
 			.Padding(1, 0)
+			.AutoWidth()
 			[
 				SNew(SButton)
 				.OnClicked(this, &FCustomizableInstanceDetails::OnProjectorLayerRemoved, ParamName, RangeIndex)
@@ -1616,22 +1624,35 @@ TSharedRef<SWidget> FCustomizableInstanceDetails::GenerateMultidimensionalProjec
 				.Content()
 				[
 					SNew(STextBlock)
-					.ToolTipText(LOCTEXT("Remove Layer", "Remove Layer"))
-					.Text(LOCTEXT("Remove Layer", "Remove Layer"))
-					.Justification(ETextJustify::Center)
-					.AutoWrapText(true)
-					.Font(LayoutBuilder.Pin()->GetDetailFont())
+						.ToolTipText(LOCTEXT("LayerProjectorRemoveLayer_ToolTip", "Remove Layer"))
+						.Text(LOCTEXT("LayerProjectorRemoveLayer_Text", "X"))
+						.Justification(ETextJustify::Center)
+						.AutoWrapText(true)
+						.Font(LayoutBuilder.Pin()->GetDetailFont())
 				]
 			]
 		];
 
-		FString LayerName = "Layer " + FString::FromInt(RangeIndex);
+		LayerProperties->AddSlot()
+		.AutoHeight()
+		[
+			OpacityRemoveWidget.ToSharedRef()
+		];
+		LayerProperties->AddSlot()
+		.AutoHeight()
+		.Padding(0.0f,5.0f,0.0f,0.0f)
+		[
+			ProjectorProperties.ToSharedRef()
+		];
 
+		// Final composed widget
 		ProjectorGroup->AddWidgetRow()
 		.NameContent()
+		.VAlign(EVerticalAlignment::VAlign_Center)
+		.HAlign(EHorizontalAlignment::HAlign_Left)
 		[
 			SNew(STextBlock)
-			.Text(FText::FromString(LayerName))
+			.Text(FText::FromString("Layer " + FString::FromInt(RangeIndex)))
 		]
 		.ValueContent()
 		.HAlign(EHorizontalAlignment::HAlign_Fill)
@@ -1642,7 +1663,7 @@ TSharedRef<SWidget> FCustomizableInstanceDetails::GenerateMultidimensionalProjec
 			.HAlign(HAlign_Fill)
 			.Padding(0.0f, 5.0f, 0.0f, 5.0f)
 			[
-				PropertiesHB.ToSharedRef()
+				LayerProperties.ToSharedRef()
 			]
 		];
 	}
@@ -1779,6 +1800,12 @@ TSharedPtr<ICustomizableObjectInstanceEditor> FCustomizableInstanceDetails::GetE
 }
 
 
+TSharedRef<SWidget> FCustomizableInstanceDetails::MakeTextureComboEntryWidget(TSharedPtr<FString> InItem) const
+{
+	return SNew(STextBlock).Text(FText::FromString(*InItem.Get()));
+}
+
+
 // BOOL PARAMETERS -----------------------------------------------------------------------------------------------------------------
 
 TSharedRef<SWidget> FCustomizableInstanceDetails::GenerateBoolWidget(const int32 ParamIndexInObject)
diff --git a/Engine/Plugins/Experimental/Mutable/Source/CustomizableObjectEditor/Private/MuCOE/CustomizableInstanceDetails.h b/Engine/Plugins/Experimental/Mutable/Source/CustomizableObjectEditor/Private/MuCOE/CustomizableInstanceDetails.h
--- a/Engine/Plugins/Experimental/Mutable/Source/CustomizableObjectEditor/Private/MuCOE/CustomizableInstanceDetails.h
+++ b/Engine/Plugins/Experimental/Mutable/Source/CustomizableObjectEditor/Private/MuCOE/CustomizableInstanceDetails.h
@@ -110,6 +110,7 @@ private:
 	FReply OnProjectorLayerAdded(const FString ParamName) const;
 	FReply OnProjectorLayerRemoved(const FString ParamName, const int32 RangeIndex) const;
 	void OnProjectorTextureParameterComboBoxChanged(TSharedPtr<FString> Selection, ESelectInfo::Type SelectInfo, const FString ParamName, int32 RangeIndex) const;
+	TSharedRef<SWidget> MakeTextureComboEntryWidget(TSharedPtr<FString> InItem) const;
 
 	// Parameter Functions
 	FReply OnCopyAllParameters();
@@ -152,5 +153,8 @@ private:
 	/** Maps param index to bool telling if it has parent, same use as previous line */
 	TMap<int32, bool> ParamHasParent;
 
+	/** Array with all the possible multilayer projector texture options */
+	TArray<TArray<TSharedPtr<FString>>> ProjectorTextureOptions;
+
 };
 
diff --git a/Engine/Plugins/Experimental/Mutable/Source/CustomizableObjectEditor/Private/MuCOE/CustomizableObjectInstanceEditor.cpp b/Engine/Plugins/Experimental/Mutable/Source/CustomizableObjectEditor/Private/MuCOE/CustomizableObjectInstanceEditor.cpp
--- a/Engine/Plugins/Experimental/Mutable/Source/CustomizableObjectEditor/Private/MuCOE/CustomizableObjectInstanceEditor.cpp
+++ b/Engine/Plugins/Experimental/Mutable/Source/CustomizableObjectEditor/Private/MuCOE/CustomizableObjectInstanceEditor.cpp
@@ -631,6 +631,8 @@ void FCustomizableObjectInstanceEditor::HideGizmo()
 
 void FCustomizableObjectInstanceEditor::ShowGizmoProjectorParameter(const FString& ParamName, int32 RangeIndex)
 {
+	HideGizmo();
+
 	ShowGizmoProjectorParameter(ParamName, RangeIndex, SharedThis(this), Viewport, CustomizableInstanceDetailsView, ProjectorParameter, CustomizableObjectInstance);
 }
 
@@ -776,9 +778,7 @@ void FCustomizableObjectInstanceEditor::ShowGizmoProjectorParameter(const FStrin
 	const TSharedPtr<ICustomizableObjectInstanceEditor>& Editor, const TSharedPtr<SCustomizableObjectEditorViewportTabBody>& Viewport,
 	const TSharedPtr<IDetailsView>& InstanceDetailsView, UProjectorParameter* ProjectorParameter,
 	UCustomizableObjectInstance* Instance)
-{
-	Editor->HideGizmo();
-	
+{	
 	ProjectorParameter->SelectProjector(ParamName, RangeIndex);
 	
 	ProjectorParameter->SetPosition(Instance->GetProjectorPosition(ParamName, RangeIndex));
