[mutable] Fixed focus/rotation issues when previewing the second component of a CO object.
- Opening the CO representing the 2nd component did not allow for focusing and orbital rotation was a bit odd.

diff --git a/Engine/Plugins/Experimental/Mutable/Source/CustomizableObjectEditor/Private/MuCOE/CustomizableObjectEditorViewportClient.cpp b/Engine/Plugins/Experimental/Mutable/Source/CustomizableObjectEditor/Private/MuCOE/CustomizableObjectEditorViewportClient.cpp
--- a/Engine/Plugins/Experimental/Mutable/Source/CustomizableObjectEditor/Private/MuCOE/CustomizableObjectEditorViewportClient.cpp
+++ b/Engine/Plugins/Experimental/Mutable/Source/CustomizableObjectEditor/Private/MuCOE/CustomizableObjectEditorViewportClient.cpp
@@ -165,8 +165,19 @@ FCustomizableObjectEditorViewportClient::FCustomizableObjectEditorViewportClient
 
 void FCustomizableObjectEditorViewportClient::UpdateCameraSetup()
 {
+	// Look for any Skeletal Mesh Component that we can focus to.
+	bool bWasValidComponentFound = false;
+	for	(TWeakObjectPtr<UDebugSkelMeshComponent> SkeletalMeshComponent : SkeletalMeshComponents)
+	{
+		if (SkeletalMeshComponent.IsValid() && UE_MUTABLE_GETSKINNEDASSET(SkeletalMeshComponent))
+		{
+			bWasValidComponentFound = true;
+			break;
+		}
+	}
+	
 	static FRotator CustomOrbitRotation(-33.75, -135, 0);
-	if ( (SkeletalMeshComponents.Num() && SkeletalMeshComponents[0].IsValid() && UE_MUTABLE_GETSKINNEDASSET(SkeletalMeshComponents[0]))
+	if ( bWasValidComponentFound
 		||
 		(StaticMeshComponent.IsValid() && StaticMeshComponent->GetStaticMesh()) )
 	{
@@ -192,7 +203,7 @@ void FCustomizableObjectEditorViewportClient::UpdateFloor()
 
 	for (TWeakObjectPtr<USkeletalMeshComponent> SkeletalMeshComponent : SkeletalMeshComponents)
 	{
-		if (SkeletalMeshComponent.IsValid())
+		if (SkeletalMeshComponent.IsValid() )
 		{
 			SkeletalMeshComponent->bComponentUseFixedSkelBounds = true;
 			SkeletalMeshComponent->UpdateBounds();
@@ -815,9 +826,6 @@ FSphere FCustomizableObjectEditorViewportClient::GetCameraTarget()
 	{
 		if (SkeletalMeshComponent.IsValid())
 		{
-			FTransform ComponentToWorld = SkeletalMeshComponent->GetComponentTransform();
-			SkeletalMeshComponent.Get()->CalcBounds(ComponentToWorld);
-
 			FBoxSphereBounds Bounds = SkeletalMeshComponent.Get()->CalcBounds(FTransform::Identity);
 
 			if (!bFoundTarget)
@@ -835,9 +843,6 @@ FSphere FCustomizableObjectEditorViewportClient::GetCameraTarget()
 
 	if(!bFoundTarget && StaticMeshComponent.IsValid())
 	{
-		FTransform ComponentToWorld = StaticMeshComponent->GetComponentTransform();
-		StaticMeshComponent.Get()->CalcBounds(ComponentToWorld);
-
 		if( !bFoundTarget )
 		{
 			FBoxSphereBounds Bounds = StaticMeshComponent.Get()->CalcBounds(FTransform::Identity);
@@ -877,13 +882,8 @@ void FCustomizableObjectEditorViewportClient::SetPreviewComponent(UStaticMeshCom
 void FCustomizableObjectEditorViewportClient::SetPreviewComponents(const TArray<UDebugSkelMeshComponent*>& InSkeletalMeshComponents)
 {
 	SkeletalMeshComponents.Reset(InSkeletalMeshComponents.Num());
+	SkeletalMeshComponents.Append(InSkeletalMeshComponents);
 	
-	for (UDebugSkelMeshComponent* SkeletalMeshComponent : InSkeletalMeshComponents)
-	{
-		SkeletalMeshComponents.Add(SkeletalMeshComponent);
-
-	}
-
 	StaticMeshComponent = nullptr;
 }
 
@@ -2243,3 +2243,4 @@ bool SMutableSelectFolderDlg::GetGenerateConstantMaterialInstances() const
 }
 
 #undef LOCTEXT_NAMESPACE 
+
