[mutable] Changed the way we show the size of the constant resources of a CO
- Previous version was precise but a bit odd to read

diff --git a/Engine/Plugins/Experimental/Mutable/Source/CustomizableObjectEditor/Private/MuCOE/SMutableConstantsWidget.cpp b/Engine/Plugins/Experimental/Mutable/Source/CustomizableObjectEditor/Private/MuCOE/SMutableConstantsWidget.cpp
--- a/Engine/Plugins/Experimental/Mutable/Source/CustomizableObjectEditor/Private/MuCOE/SMutableConstantsWidget.cpp
+++ b/Engine/Plugins/Experimental/Mutable/Source/CustomizableObjectEditor/Private/MuCOE/SMutableConstantsWidget.cpp
@@ -21,61 +21,39 @@ class SWidget;
 // Private namespace with utility functionality used by this slate object
 namespace
 {
-	/** Provided a byte count this function proceeds to output that byte value as bytes, kilobytes, megabytes and gigabytes.
+	/** Provided a byte count this function proceeds to output that byte value as text alongside with it's unit of mesure (Bytes, KB...).
 	 * @param SizeInBytes The amount of bytes to convert to a formatted text that represents it.
-	 * @return A text representing the value provided as Bytes, Kilobytes, Megabytes and Gigabytes. It will not return
-	 * size types with value 0.
+	 * @return A text representing the value provided as Bytes, Kilobytes, Megabytes or Gigabytes. 
 	 */
-	FText GetSizeInBytesAsBiggerUnits(const uint64 SizeInBytes)
+	FText GenerateTextForSize(const uint64 SizeInBytes)
 	{
-		FString OutputString = "";
-		
-		uint64 Bytes = SizeInBytes;
-		uint64 KiloBytes = 0;
-		uint64 MegaBytes = 0;
-		uint64 GigaBytes = 0;
+		FString Unit = TEXT("Bytes");
+		double Value = SizeInBytes;
 		
 		// B to KB
-		if (Bytes >= 1024)
+		if (SizeInBytes >= 1024)
 		{
-			KiloBytes = Bytes / 1024;
-			Bytes = Bytes % 1024;
+			Unit = TEXT("KB");
+			Value = SizeInBytes / 1024.0;
 			
 			// KB to MB
-			if (KiloBytes >= 1024)
+			if (Value >= 1024.0)
 			{
-				MegaBytes = KiloBytes / 1024;
-				KiloBytes = KiloBytes % 1024;
+				Unit = TEXT("MB");
+				Value = Value / 1024.0;
 				
 				// MB to GB
-				if (MegaBytes >= 1024)
+				if (Value >= 1024.0)
 				{
-					GigaBytes = MegaBytes / 1024;
-					MegaBytes = MegaBytes % 1024;
+					Unit = TEXT("GB");
+					Value = Value / 1024.0;
 				}
 			}
 		}
 
-		// Compose the string that will be the output of this function 
-		if (GigaBytes)
-		{
-			OutputString.Append(FString::FromInt(GigaBytes) + " GB, ");
-		}
-
-		if (MegaBytes)
-		{
-			OutputString.Append(FString::FromInt(MegaBytes) + " MB, ");
-		}
-
-		if (KiloBytes)
-		{
-			OutputString.Append(FString::FromInt(KiloBytes) + " KB, ");
-		}
-
-		if (Bytes)
-		{
-			OutputString.Append( FString::FromInt(Bytes) + " B ");
-		}
+		FString OutputString = FString::Printf(TEXT("%.2f"), Value);
+		OutputString.Append(" ");
+		OutputString.Append(Unit);
 
 		return FText::FromString(OutputString);
 	}
@@ -1135,7 +1113,7 @@ void SMutableConstantsWidget::LoadConstantStrings()
 	}
 
 	// Cache the size in memory of the constants as a formatted text so it is able to be be later used by the UI
-	ConstantStringsFormattedSize = GetSizeInBytesAsBiggerUnits(ConstantStringsAccumulatedSize);
+	ConstantStringsFormattedSize = GenerateTextForSize(ConstantStringsAccumulatedSize);
 }
 
 void SMutableConstantsWidget::LoadConstantImages()
@@ -1162,7 +1140,7 @@ void SMutableConstantsWidget::LoadConstantImages()
 	}
 
 	// Cache the size in memory of the constants as a formatted text so it is able to be be later used by the UI
-	ConstantImagesFormattedSize = GetSizeInBytesAsBiggerUnits(ConstantImagesAccumulatedSize);
+	ConstantImagesFormattedSize = GenerateTextForSize(ConstantImagesAccumulatedSize);
 	
 	// Regenerate the array with only a part of the the full images array
 	RegenerateProxyImageArray();
@@ -1199,7 +1177,7 @@ void SMutableConstantsWidget::LoadConstantMeshes()
 	}
 
 	// Cache the size in memory of the constants as a formatted text so it is able to be be later used by the UI
-	ConstantMeshesFormattedSize = GetSizeInBytesAsBiggerUnits(ConstantMeshesAccumulatedSize);
+	ConstantMeshesFormattedSize = GenerateTextForSize(ConstantMeshesAccumulatedSize);
 }
 
 void SMutableConstantsWidget::LoadConstantLayouts()
@@ -1224,8 +1202,7 @@ void SMutableConstantsWidget::LoadConstantLayouts()
 	}
 
 	// Cache the size in memory of the constants as a formatted text so it is able to be be later used by the UI
-	ConstantLayoutsFormattedSize = GetSizeInBytesAsBiggerUnits(Stream.GetBufferSize());
-
+	ConstantLayoutsFormattedSize = GenerateTextForSize(Stream.GetBufferSize());
 }
 
 void SMutableConstantsWidget::LoadConstantSkeletons()
@@ -1249,7 +1226,7 @@ void SMutableConstantsWidget::LoadConstantSkeletons()
 	}
 
 	// Cache the size in memory of the constants as a formatted text so it is able to be be later used by the UI
-	ConstantSkeletonsFormattedSize = GetSizeInBytesAsBiggerUnits(Stream.GetBufferSize());
+	ConstantSkeletonsFormattedSize = GenerateTextForSize(Stream.GetBufferSize());
 }
 
 void SMutableConstantsWidget::LoadConstantProjectors()
@@ -1273,7 +1250,7 @@ void SMutableConstantsWidget::LoadConstantProjectors()
 	}
 
 	// Cache the size in memory of the constants as a formatted text so it is able to be be later used by the UI
-	ConstantProjectorsFormattedSize = GetSizeInBytesAsBiggerUnits(Stream.GetBufferSize());
+	ConstantProjectorsFormattedSize = GenerateTextForSize(Stream.GetBufferSize());
 }
 
 void SMutableConstantsWidget::LoadConstantMatrices()
@@ -1297,7 +1274,7 @@ void SMutableConstantsWidget::LoadConstantMatrices()
 	}
 
 	// Cache the size in memory of the constants as a formatted text so it is able to be be later used by the UI
-	ConstantMatricesFormattedSize = GetSizeInBytesAsBiggerUnits(Stream.GetBufferSize());
+	ConstantMatricesFormattedSize = GenerateTextForSize(Stream.GetBufferSize());
 }
 
 void SMutableConstantsWidget::LoadConstantShapes()
@@ -1321,7 +1298,7 @@ void SMutableConstantsWidget::LoadConstantShapes()
 	}
 	
 	// Cache the size in memory of the constants as a formatted text so it is able to be be later used by the UI
-	ConstantShapesFormattedSize = GetSizeInBytesAsBiggerUnits(Stream.GetBufferSize());
+	ConstantShapesFormattedSize = GenerateTextForSize(Stream.GetBufferSize());
 }
 
 void SMutableConstantsWidget::LoadConstantCurves()
@@ -1345,7 +1322,7 @@ void SMutableConstantsWidget::LoadConstantCurves()
 	}
 
 	// Cache the size in memory of the constants as a formatted text so it is able to be be later used by the UI
-	ConstantCurvesFormattedSize = GetSizeInBytesAsBiggerUnits(Stream.GetBufferSize());
+	ConstantCurvesFormattedSize = GenerateTextForSize(Stream.GetBufferSize());
 }
 
 #pragma  endregion
