[Mutable] Forbid Node Material Variation to have multiple outputs.

#rnx

diff --git a/Engine/Plugins/Experimental/Mutable/Source/CustomizableObjectEditor/Classes/MuCOE/Nodes/CustomizableObjectNodeMaterialVariation.h b/Engine/Plugins/Experimental/Mutable/Source/CustomizableObjectEditor/Classes/MuCOE/Nodes/CustomizableObjectNodeMaterialVariation.h
--- a/Engine/Plugins/Experimental/Mutable/Source/CustomizableObjectEditor/Classes/MuCOE/Nodes/CustomizableObjectNodeMaterialVariation.h
+++ b/Engine/Plugins/Experimental/Mutable/Source/CustomizableObjectEditor/Classes/MuCOE/Nodes/CustomizableObjectNodeMaterialVariation.h
@@ -25,17 +25,19 @@ public:
 
 	UPROPERTY(EditAnywhere, Category = CustomizableObject)
 	ECustomizableObjectNodeMaterialVariationType Type = ECustomizableObjectNodeMaterialVariationType::Tag;
-	
+
+private:
 	// Deprecated properties
 	UPROPERTY()
 	TArray<FCustomizableObjectMaterialVariation> Variations_DEPRECATED;
 
+public:
 	// UCustomizableObjectNode interface
 	virtual void BackwardsCompatibleFixup() override;
+	virtual bool IsSingleOutputNode() const override;
 	
 	// UCustomizableObjectNodeVariation interface
 	virtual FName GetCategory() const override;
-
-	bool IsInputPinArray() const;
+	virtual bool IsInputPinArray() const override;
 };
 
diff --git a/Engine/Plugins/Experimental/Mutable/Source/CustomizableObjectEditor/Private/MuCOE/Nodes/CustomizableObjectNodeMaterialVariation.cpp b/Engine/Plugins/Experimental/Mutable/Source/CustomizableObjectEditor/Private/MuCOE/Nodes/CustomizableObjectNodeMaterialVariation.cpp
--- a/Engine/Plugins/Experimental/Mutable/Source/CustomizableObjectEditor/Private/MuCOE/Nodes/CustomizableObjectNodeMaterialVariation.cpp
+++ b/Engine/Plugins/Experimental/Mutable/Source/CustomizableObjectEditor/Private/MuCOE/Nodes/CustomizableObjectNodeMaterialVariation.cpp
@@ -35,3 +35,9 @@ bool UCustomizableObjectNodeMaterialVariation::IsInputPinArray() const
 {
 	return true;
 }
+
+
+bool UCustomizableObjectNodeMaterialVariation::IsSingleOutputNode() const
+{
+	return true;
+}
