[mutable] Ensure that tasks are resolved in a memory-optimizing order when not using concurrency, to reduce the memory footprint.
Move insights scope markers so that they are also correct for cooking processes.
Don't do the derived data save on a separeate thread when the compilation is synchronous.
#rnx

diff --git a/Engine/Plugins/Experimental/Mutable/Source/CustomizableObject/Private/MuCO/CustomizableObject.cpp b/Engine/Plugins/Experimental/Mutable/Source/CustomizableObject/Private/MuCO/CustomizableObject.cpp
--- a/Engine/Plugins/Experimental/Mutable/Source/CustomizableObject/Private/MuCO/CustomizableObject.cpp
+++ b/Engine/Plugins/Experimental/Mutable/Source/CustomizableObject/Private/MuCO/CustomizableObject.cpp
@@ -809,6 +809,8 @@ void UCustomizableObjectPrivate::LoadCompiledDataFromDisk()
 
 void UCustomizableObjectPrivate::CachePlatformData(const ITargetPlatform* InTargetPlatform, const TArray64<uint8>& InModelBytes, const TArray64< uint8>& InBulkBytes)
 {
+	MUTABLE_CPUPROFILER_SCOPE(CachePlatformData)
+
 	FString PlatformName = InTargetPlatform ? InTargetPlatform->PlatformName() : FPlatformProperties::PlatformName();
 
 	check(!CachedPlatformsData.Find(PlatformName));
@@ -821,8 +823,8 @@ void UCustomizableObjectPrivate::CachePlatformData(const ITargetPlatform* InTarg
 	Data.ModelData.Append(InModelBytes);
 
 	// Cache streamable bulk data
-	Data.StreamableData.Reset(InBulkBytes.Num());
-	Algo::Copy(InBulkBytes, Data.StreamableData);
+	Data.StreamableData.SetNumUninitialized(InBulkBytes.Num(),EAllowShrinking::No);
+	FMemory::Memcpy(Data.StreamableData.GetData(), InBulkBytes.GetData(), InBulkBytes.Num() );
 }
 
 
diff --git a/Engine/Plugins/Experimental/Mutable/Source/CustomizableObjectEditor/Private/MuCOE/CustomizableObjectCompileRunnable.cpp b/Engine/Plugins/Experimental/Mutable/Source/CustomizableObjectEditor/Private/MuCOE/CustomizableObjectCompileRunnable.cpp
--- a/Engine/Plugins/Experimental/Mutable/Source/CustomizableObjectEditor/Private/MuCOE/CustomizableObjectCompileRunnable.cpp
+++ b/Engine/Plugins/Experimental/Mutable/Source/CustomizableObjectEditor/Private/MuCOE/CustomizableObjectCompileRunnable.cpp
@@ -244,6 +244,8 @@ const TArray<FCustomizableObjectCompileRunnable::FError>& FCustomizableObjectCom
 
 FCustomizableObjectSaveDDRunnable::FCustomizableObjectSaveDDRunnable(UCustomizableObject* CustomizableObject, const FCompilationOptions& InOptions)
 {
+	MUTABLE_CPUPROFILER_SCOPE(FCustomizableObjectSaveDDRunnable::FCustomizableObjectSaveDDRunnable)
+		
 	Model = CustomizableObject->GetPrivate()->GetModel();
 	Options = InOptions;
 	
@@ -266,6 +268,8 @@ FCustomizableObjectSaveDDRunnable::FCustomizableObjectSaveDDRunnable(UCustomizab
 
 uint32 FCustomizableObjectSaveDDRunnable::Run()
 {
+	MUTABLE_CPUPROFILER_SCOPE(FCustomizableObjectSaveDDRunnable::Run)
+
 	bool bModelSerialized = Model.Get() != nullptr;
 
 	if (Options.bIsCooking)
diff --git a/Engine/Plugins/Experimental/Mutable/Source/CustomizableObjectEditor/Private/MuCOE/CustomizableObjectCompileRunnable.h b/Engine/Plugins/Experimental/Mutable/Source/CustomizableObjectEditor/Private/MuCOE/CustomizableObjectCompileRunnable.h
--- a/Engine/Plugins/Experimental/Mutable/Source/CustomizableObjectEditor/Private/MuCOE/CustomizableObjectCompileRunnable.h
+++ b/Engine/Plugins/Experimental/Mutable/Source/CustomizableObjectEditor/Private/MuCOE/CustomizableObjectCompileRunnable.h
@@ -9,6 +9,8 @@
 #include "MuT/Node.h"
 #include "MuCOE/CustomizableObjectEditorLogger.h"
 
+#include <atomic>
+
 class ITargetPlatform;
 
 
@@ -77,7 +79,7 @@ public:
 	FString ErrorMsg;
 
 	// Whether the thread has finished running
-	bool bThreadCompleted;
+	std::atomic<bool> bThreadCompleted;
 };
 
 
@@ -125,5 +127,5 @@ private:
 	bool bIsCooking = false;
 
 	// Whether the thread has finished running
-	bool bThreadCompleted = false;
+	std::atomic<bool> bThreadCompleted = false;
 };
diff --git a/Engine/Plugins/Experimental/Mutable/Source/CustomizableObjectEditor/Private/MuCOE/CustomizableObjectCompiler.cpp b/Engine/Plugins/Experimental/Mutable/Source/CustomizableObjectEditor/Private/MuCOE/CustomizableObjectCompiler.cpp
--- a/Engine/Plugins/Experimental/Mutable/Source/CustomizableObjectEditor/Private/MuCOE/CustomizableObjectCompiler.cpp
+++ b/Engine/Plugins/Experimental/Mutable/Source/CustomizableObjectEditor/Private/MuCOE/CustomizableObjectCompiler.cpp
@@ -84,9 +84,6 @@ bool FCustomizableObjectCompiler::Tick()
 		UpdateArrayGCProtect();
 
 		CompilationLogsContainer.ClearMessagesArray();
-
-		TRACE_END_REGION(UE_MUTABLE_COMPILE_REGION);
-
 	}
 
 	if (SaveDDTask.IsValid() && SaveDDTask->IsCompleted())
@@ -101,8 +98,6 @@ bool FCustomizableObjectCompiler::Tick()
 
 		UE_LOG(LogMutable, Verbose, TEXT("PROFILE: [ %16.8f ] Finished Saving Derived Data task."), FPlatformTime::Seconds());
 		UE_LOG(LogMutable, Verbose, TEXT("PROFILE: -----------------------------------------------------------"));
-
-		TRACE_END_REGION(UE_MUTABLE_SAVEDD_REGION);
 	}
 
 	if (CompilationLaunchPending)
@@ -1090,8 +1085,8 @@ void FCustomizableObjectCompiler::CompileInternal(UCustomizableObject* Object, c
 
 			if (SaveDDTask.IsValid())
 			{
-				SaveCODerivedData();
-				SaveDDThread->WaitForCompletion();
+				SaveDDTask->Init();
+				SaveDDTask->Run();
 				FinishSavingDerivedData();
 			}
 
@@ -1239,6 +1234,8 @@ void FCustomizableObjectCompiler::FinishCompilation()
 	CompileThread.Reset();
 	CompileTask.Reset();
 
+	TRACE_END_REGION(UE_MUTABLE_COMPILE_REGION);
+
 	if (!Options.bDontUpdateStreamedDataAndCache)
 	{
 		TRACE_BEGIN_REGION(UE_MUTABLE_SAVEDD_REGION);
@@ -1258,6 +1255,8 @@ void FCustomizableObjectCompiler::FinishCompilation()
 
 void FCustomizableObjectCompiler::FinishSavingDerivedData()
 {
+	MUTABLE_CPUPROFILER_SCOPE(FinishSavingDerivedData)
+
 	check(SaveDDTask.IsValid());
 
 	if (Options.bIsCooking)
@@ -1277,6 +1276,9 @@ void FCustomizableObjectCompiler::FinishSavingDerivedData()
 	RemoveCompileNotification();
 
 	NotifyCompilationErrors();
+
+	TRACE_END_REGION(UE_MUTABLE_SAVEDD_REGION);
+
 }
 
 
diff --git a/Engine/Plugins/Experimental/Mutable/Source/MutableTools/Private/MuT/CodeOptimiser.cpp b/Engine/Plugins/Experimental/Mutable/Source/MutableTools/Private/MuT/CodeOptimiser.cpp
--- a/Engine/Plugins/Experimental/Mutable/Source/MutableTools/Private/MuT/CodeOptimiser.cpp
+++ b/Engine/Plugins/Experimental/Mutable/Source/MutableTools/Private/MuT/CodeOptimiser.cpp
@@ -449,9 +449,9 @@ namespace mu
 	//---------------------------------------------------------------------------------------------
 	//---------------------------------------------------------------------------------------------
 	//---------------------------------------------------------------------------------------------
-	class ConstantTask 
+	class FConstantTask 
 	{
-	private:
+	public:
 
 		// input
 		Ptr<ASTOp> Source;
@@ -468,7 +468,7 @@ namespace mu
 
 	public:
 
-		ConstantTask( const Ptr<ASTOp>& InSource, const CompilerOptions::Private* InOptions, int32 InOptimizationPass )
+		FConstantTask( const Ptr<ASTOp>& InSource, const CompilerOptions::Private* InOptions, int32 InOptimizationPass )
 		{
 			OptimizationPass = InOptimizationPass;
 			Source = InSource;
@@ -477,14 +477,6 @@ namespace mu
 			ReferencedResourceProvider = InOptions->OptimisationOptions.ReferencedResourceProvider;
 		}
 
-		void Prepare()
-		{
-			// This runs in the AST managing thread.
-
-			// We need the clone because linking modifies ASTOp state and also to be safe for concurrency.
-			SourceCloned = ASTOp::DeepClone(Source);
-		}
-
 		void Run(FImageOperator ImOp)
 		{
 			MUTABLE_CPUPROFILER_SCOPE(ConstantTask_Run);
@@ -617,13 +609,6 @@ namespace mu
 			SourceCloned = nullptr;
 		}
 
-		void Complete()
-		{
-			// This runs in the AST managing thread.
-			ASTOp::Replace(Source, Result);
-			Source = nullptr;
-			Result = nullptr;
-		}
 	};
 
 
@@ -645,6 +630,11 @@ namespace mu
 		{
 			Ptr<ASTOp> Root;
 			UE::Tasks::FTaskEvent CompletedEvent;
+
+			// This is only necessary for non-concurrent execution
+			TArray< UE::Tasks::FTask, TInlineAllocator<8> > Requisites;
+			TFunction<void()> NonConcurrentTask;
+			TUniquePtr<FConstantTask> TaskData;
 		};
 		TArray< FConstantSubgraph > ConstantSubgraphs;
 		ConstantSubgraphs.Reserve(256);
@@ -763,101 +753,95 @@ namespace mu
 				});
 		}
 
-		bool bUseConcurrency = InOptions->bUseConcurrency;
-
-		// Launch the tasks. Dot it from a task in the ASTPipe to make sure we "lock" modifications of the AST.
-		UE::Tasks::FTaskEvent EverythingComplete(TEXT("ConstantGenerator_EverythingComplete"));
-		auto LaunchFunc = [&ConstantSubgraphs, &ASTPipe, &EverythingComplete, Pass, InOptions, bUseConcurrency]()
-			{
-				MUTABLE_CPUPROFILER_SCOPE(ConstantGenerator_LaunchTasks);
+		auto GetRequisites = [&ConstantSubgraphs](const Ptr<ASTOp>& SubgraphRoot, TArray< UE::Tasks::FTask, TInlineAllocator<8> >& OutRequisites)
+		{
+			TArray< Ptr<ASTOp> > ScanRoots;
+			ScanRoots.Add(SubgraphRoot);
+			ASTOp::Traverse_TopDown_Unique_Imprecise(ScanRoots, [&SubgraphRoot, &OutRequisites, &ConstantSubgraphs](Ptr<ASTOp>& ChildNode)
+				{
+					bool bRecurse = true;
 
-				FImageOperator ImOp = FImageOperator::GetDefault(InOptions->ImageFormatFunc);
+					// Subgraph root?
+					if (SubgraphRoot == ChildNode)
+					{
+						return bRecurse;
+					}
 
-				// Traverse list of constants to generate. It is ordered in a bottom-up way.
-				int32 SubgraphCount = ConstantSubgraphs.Num();
-				for (int32 OrderIndex = 0; OrderIndex < SubgraphCount; ++OrderIndex)
-				{
-					int32 Index = OrderIndex;
-					if (bUseConcurrency)
+					FConstantSubgraph* DependencyFound = ConstantSubgraphs.FindByPredicate([&ChildNode](const FConstantSubgraph& Candidate) { return Candidate.Root == ChildNode; });
+					if (DependencyFound)
 					{
-						// If we are using concurrency, we prefer to resolve tasks top-to-bottom
-						Index = SubgraphCount - 1 - OrderIndex;
+						bRecurse = false;
+						OutRequisites.Add(DependencyFound->CompletedEvent);
 					}
 
-					Ptr<ASTOp> SubgraphRoot = ConstantSubgraphs[Index].Root;
-					UE::Tasks::FTaskEvent& SubgraphCompletionEvent = ConstantSubgraphs[Index].CompletedEvent;
+					return bRecurse;
+				});
+		};
 
-					// Referenced images are resolved in its own task to prevent requesting them twice (and loading them twice)
-					bool bIsCompileTimeReferenceImage = (SubgraphRoot->GetOpType() == OP_TYPE::IM_REFERENCE);
+		bool bUseConcurrency = InOptions->bUseConcurrency;
 
-					// Launch the task with its dependencies
-					if (bIsCompileTimeReferenceImage)
+		if (bUseConcurrency)
+		{
+			// Launch the tasks. Do it from a task in the ASTPipe to make sure to avoid race conditions in the main AST manipulation.
+			UE::Tasks::FTaskEvent EverythingComplete(TEXT("ConstantGenerator_EverythingComplete"));
+			UE::Tasks::FTask LaunchTask = ASTPipe.Launch(TEXT("ConstantGeneratorLaunchTasks"), 
+				[&ConstantSubgraphs, &ASTPipe, &EverythingComplete, &GetRequisites, Pass, InOptions]()
+				{
+					MUTABLE_CPUPROFILER_SCOPE(ConstantGenerator_LaunchTasks);
+
+					FImageOperator ImOp = FImageOperator::GetDefault(InOptions->ImageFormatFunc);
+
+					// Traverse list of constants to generate. It is ordered in a bottom-up way.
+					int32 SubgraphCount = ConstantSubgraphs.Num();
+					for (int32 OrderIndex = 0; OrderIndex < SubgraphCount; ++OrderIndex)
 					{
-						// Instead of generating the constant we resolve the reference, which also replaces the ASTOp.
-						TSharedPtr< Ptr<Image> > ResolveImage = MakeShared<Ptr<Image>>();
+						int32 Index = SubgraphCount - 1 - OrderIndex;
 
-						const ASTOpReferenceResource* Typed = static_cast<const ASTOpReferenceResource*>(SubgraphRoot.get());
-						UE::Tasks::FTaskEvent ReferenceCompletionEvent = InOptions->OptimisationOptions.ReferencedResourceProvider(Typed->ID, ResolveImage);
+						Ptr<ASTOp> SubgraphRoot = ConstantSubgraphs[Index].Root;
+						UE::Tasks::FTaskEvent& SubgraphCompletionEvent = ConstantSubgraphs[Index].CompletedEvent;
 
-						auto CompleteFunc = [SubgraphRoot, InOptions, ResolveImage]()
-							{
-								Ptr<ASTOpConstantResource> ConstantOp = new ASTOpConstantResource;
-								ConstantOp->type = OP_TYPE::IM_CONSTANT;
-								ConstantOp->SetValue(ResolveImage->get(), InOptions->OptimisationOptions.bUseDiskCache);
-								ASTOp::Replace(SubgraphRoot, ConstantOp);
-							};
+						// Referenced images are resolved in its own task to prevent requesting them twice (and loading them twice)
+						bool bIsCompileTimeReferenceImage = (SubgraphRoot->GetOpType() == OP_TYPE::IM_REFERENCE);
 
-						if (bUseConcurrency)
+						// Launch the task with its dependencies
+						if (bIsCompileTimeReferenceImage)
 						{
+							// Instead of generating the constant we resolve the reference, which also replaces the ASTOp.
+							const ASTOpReferenceResource* Typed = static_cast<const ASTOpReferenceResource*>(SubgraphRoot.get());
+							uint32 ImageID = Typed->ID;
+
+							TSharedPtr< Ptr<Image> > ResolveImage = MakeShared<Ptr<Image>>();
+
+							UE::Tasks::FTaskEvent ReferenceCompletionEvent = InOptions->OptimisationOptions.ReferencedResourceProvider(ImageID, ResolveImage);
+
 							UE::Tasks::FTask CompleteTask = ASTPipe.Launch(TEXT("MutableResolveComplete"),
-								MoveTemp(CompleteFunc),
+								[SubgraphRoot, InOptions, ResolveImage]()
+								{
+									Ptr<ASTOpConstantResource> ConstantOp = new ASTOpConstantResource;
+									ConstantOp->type = OP_TYPE::IM_CONSTANT;
+									ConstantOp->SetValue(ResolveImage->get(), InOptions->OptimisationOptions.bUseDiskCache);
+									ASTOp::Replace(SubgraphRoot, ConstantOp);
+								},
 								ReferenceCompletionEvent,
 								LowLevelTasks::ETaskPriority::BackgroundNormal);
 
 							SubgraphCompletionEvent.AddPrerequisites(CompleteTask);
 						}
+
 						else
 						{
-							ReferenceCompletionEvent.Wait();
-							CompleteFunc();
-						}
-					}
-
-					else
-					{
-						// Scan for requisites
-						TArray< UE::Tasks::FTask, TInlineAllocator<8> > Requisites;
-						TArray< Ptr<ASTOp> > ScanRoots;
-						ScanRoots.Add(SubgraphRoot);
-						ASTOp::Traverse_TopDown_Unique_Imprecise(ScanRoots, [&SubgraphRoot, &Requisites, &ConstantSubgraphs](Ptr<ASTOp>& ChildNode)
-							{
-								bool bRecurse = true;
-
-								// Subgraph root?
-								if (SubgraphRoot == ChildNode)
-								{
-									return bRecurse;
-								}
-
-								FConstantSubgraph* DependencyFound = ConstantSubgraphs.FindByPredicate([&ChildNode](const FConstantSubgraph& Candidate) { return Candidate.Root == ChildNode; });
-								if (DependencyFound)
-								{
-									bRecurse = false;
-									Requisites.Add(DependencyFound->CompletedEvent);
-								}
-
-								return bRecurse;
-							});
+							// Scan for requisites
+							TArray< UE::Tasks::FTask, TInlineAllocator<8> > Requisites;
+							GetRequisites(SubgraphRoot,Requisites);
 
-						TUniquePtr<ConstantTask> Task(new ConstantTask(SubgraphRoot, InOptions, Pass));
-						ConstantTask* TaskPtr = Task.Get();
+							TUniquePtr<FConstantTask> Task(new FConstantTask(SubgraphRoot, InOptions, Pass));
+							FConstantTask* TaskPtr = Task.Get();
 
-						if (bUseConcurrency)
-						{
 							// Launch the preparation on the AST-modification pipe
 							UE::Tasks::FTask PrepareTask = ASTPipe.Launch(TEXT("MutableConstantPrepare"), [TaskPtr]()
 								{
-									TaskPtr->Prepare();
+									// We need the clone because linking modifies ASTOp state and also to be safe for concurrency.
+									TaskPtr->SourceCloned = ASTOp::DeepClone(TaskPtr->Source);
 								},
 								Requisites,
 								LowLevelTasks::ETaskPriority::BackgroundHigh);
@@ -873,31 +857,22 @@ namespace mu
 							// Launch the completion on the AST-modification pipe
 							UE::Tasks::FTask CompleteTask = ASTPipe.Launch(TEXT("MutableConstantComplete"), [TaskPtr = MoveTemp(Task)]()
 								{
-									TaskPtr->Complete();
+									ASTOp::Replace(TaskPtr->Source, TaskPtr->Result);
+									TaskPtr->Source = nullptr;
+									TaskPtr->Result = nullptr;
 								},
 								RunTask,
 								LowLevelTasks::ETaskPriority::BackgroundHigh);
 
 							SubgraphCompletionEvent.AddPrerequisites(CompleteTask);
 						}
-						else
-						{
-							TaskPtr->Prepare();
-							TaskPtr->Run(ImOp);
-							TaskPtr->Complete();
-						}
-					}
 
-					ConstantSubgraphs[Index].Root = nullptr;
-					SubgraphCompletionEvent.Trigger();
-					EverythingComplete.AddPrerequisites(SubgraphCompletionEvent);
-				}
-
-			};
+						ConstantSubgraphs[Index].Root = nullptr;
+						SubgraphCompletionEvent.Trigger();
+						EverythingComplete.AddPrerequisites(SubgraphCompletionEvent);
+					}
 
-		if (bUseConcurrency)
-		{ 
-			UE::Tasks::FTask LaunchTask = ASTPipe.Launch(TEXT("ConstantGeneratorLaunchTasks"), MoveTemp(LaunchFunc));
+				});
 
 			// Wait for pending tasks
 			{
@@ -907,10 +882,134 @@ namespace mu
 				EverythingComplete.Wait();
 			}
 		}
+
 		else
 		{
-			LaunchFunc();
-			EverythingComplete.Trigger();
+			// Non-concurrent version: Prepare the tasks.
+			{
+				MUTABLE_CPUPROFILER_SCOPE(ConstantGenerator_LaunchTasks);
+
+				FImageOperator ImOp = FImageOperator::GetDefault(InOptions->ImageFormatFunc);
+
+				// Traverse list of constants to generate. It is ordered in a bottom-up way.
+				int32 SubgraphCount = ConstantSubgraphs.Num();
+				for (int32 OrderIndex = 0; OrderIndex < SubgraphCount; ++OrderIndex)
+				{
+					int32 Index = SubgraphCount - 1 - OrderIndex;
+
+					Ptr<ASTOp> SubgraphRoot = ConstantSubgraphs[Index].Root;
+					UE::Tasks::FTaskEvent& SubgraphCompletionEvent = ConstantSubgraphs[Index].CompletedEvent;
+
+					// Referenced images are resolved in its own task to prevent requesting them twice (and loading them twice)
+					bool bIsCompileTimeReferenceImage = (SubgraphRoot->GetOpType() == OP_TYPE::IM_REFERENCE);
+
+					// Launch the task with its dependencies
+					if (bIsCompileTimeReferenceImage)
+					{
+						// Instead of generating the constant we resolve the reference, which also replaces the ASTOp.
+						const ASTOpReferenceResource* Typed = static_cast<const ASTOpReferenceResource*>(SubgraphRoot.get());
+						uint32 ImageID = Typed->ID;
+
+						auto ImmediateCompleteFunc = [SubgraphRoot, InOptions, ImageID]()
+							{
+								TSharedPtr< Ptr<Image> > ResolveImage = MakeShared<Ptr<Image>>();
+								UE::Tasks::FTaskEvent ReferenceCompletionEvent = InOptions->OptimisationOptions.ReferencedResourceProvider(ImageID, ResolveImage);
+								ReferenceCompletionEvent.Wait();
+								Ptr<ASTOpConstantResource> ConstantOp = new ASTOpConstantResource;
+								ConstantOp->type = OP_TYPE::IM_CONSTANT;
+								ConstantOp->SetValue(ResolveImage->get(), InOptions->OptimisationOptions.bUseDiskCache);
+								ASTOp::Replace(SubgraphRoot, ConstantOp);
+							};
+						ConstantSubgraphs[Index].NonConcurrentTask = ImmediateCompleteFunc;
+					}
+
+					else
+					{
+						// Scan for requisites
+						TArray< UE::Tasks::FTask, TInlineAllocator<8> > Requisites;
+						GetRequisites(SubgraphRoot, Requisites);
+
+						TUniquePtr<FConstantTask> Task(new FConstantTask(SubgraphRoot, InOptions, Pass));
+						FConstantTask* TaskPtr = Task.Get();
+
+						ConstantSubgraphs[Index].Requisites = Requisites;
+						ConstantSubgraphs[Index].TaskData = MoveTemp(Task);
+						ConstantSubgraphs[Index].NonConcurrentTask = [TaskPtr, ImOp]()
+							{
+								TaskPtr->SourceCloned = ASTOp::DeepClone(TaskPtr->Source);
+								TaskPtr->Run(ImOp);
+								ASTOp::Replace(TaskPtr->Source, TaskPtr->Result);
+							};
+					}
+
+					ConstantSubgraphs[Index].Root = nullptr;
+				}
+			}
+
+			auto TryRunFunc = [&ConstantSubgraphs](int32 Index) -> bool
+				{
+					bool bCanRun = true;
+					for (const UE::Tasks::FTask& Req : ConstantSubgraphs[Index].Requisites)
+					{
+						if (!Req.IsCompleted())
+						{
+							bCanRun = false;
+							break;
+						}
+					}
+
+					if (bCanRun)
+					{
+						ConstantSubgraphs[Index].NonConcurrentTask();
+						ConstantSubgraphs[Index].NonConcurrentTask.Reset();
+						ConstantSubgraphs[Index].TaskData.Reset();
+						ConstantSubgraphs[Index].CompletedEvent.Trigger();
+					}
+
+					return bCanRun;
+				};
+
+			// Resolve tasks top-down while as long as they have the prerequisites completed.
+			int32 SubgraphCount = ConstantSubgraphs.Num();
+			TArray<int32> Pending;
+			for (int32 OrderIndex = 0; OrderIndex < SubgraphCount; ++OrderIndex)
+			{
+				int32 Index = SubgraphCount - 1 - OrderIndex;
+
+				bool bRan = TryRunFunc(Index);
+
+				if (bRan)
+				{
+					// We may have unlocked one of the pending tasks. Review them all.
+					bool bModified = true;
+					while (bModified)
+					{
+						bModified = false;
+						for (int32 PendingIndexIndex = 0; PendingIndexIndex < Pending.Num(); )
+						{
+							int32 PendingIndex = Pending[PendingIndexIndex];
+							bRan = TryRunFunc(PendingIndex);
+							if (bRan)
+							{
+								bModified = true;
+								Pending.RemoveAtSwap(PendingIndexIndex);
+								break;
+							}
+							else
+							{
+								++PendingIndexIndex;
+							}
+						}
+					}
+				}
+				else
+				{
+					// We couldn't run the task yet, remember it for later.
+					Pending.Add(Index);
+				}
+			}
+
+			check(Pending.IsEmpty());
 		}
 
 		bool bSomethingModified = ConstantSubgraphs.Num() > 0;
