[mutable] Fixed crash when there is an empty texture node.
#rnx
#tests Editor

diff --git a/Engine/Plugins/Experimental/Mutable/Source/CustomizableObjectEditor/Private/MuCOE/GenerateMutableSource/GenerateMutableSourceImage.cpp b/Engine/Plugins/Experimental/Mutable/Source/CustomizableObjectEditor/Private/MuCOE/GenerateMutableSource/GenerateMutableSourceImage.cpp
--- a/Engine/Plugins/Experimental/Mutable/Source/CustomizableObjectEditor/Private/MuCOE/GenerateMutableSource/GenerateMutableSourceImage.cpp
+++ b/Engine/Plugins/Experimental/Mutable/Source/CustomizableObjectEditor/Private/MuCOE/GenerateMutableSource/GenerateMutableSourceImage.cpp
@@ -176,7 +176,7 @@ mu::NodeImagePtr GenerateMutableSourceImage(const UEdGraphPin* Pin, FMutableGrap
 		}
 		else
 		{
-			GenerationContext.Compiler->CompilerLog(LOCTEXT("MissingImage", "Missing image in texture node."), Node);
+			GenerationContext.Compiler->CompilerLog(LOCTEXT("MissingImage", "Missing image in texture node."), Node, EMessageSeverity::Warning);
 		}
 	}
 
diff --git a/Engine/Plugins/Experimental/Mutable/Source/MutableTools/Private/MuT/CodeGenerator.cpp b/Engine/Plugins/Experimental/Mutable/Source/MutableTools/Private/MuT/CodeGenerator.cpp
--- a/Engine/Plugins/Experimental/Mutable/Source/MutableTools/Private/MuT/CodeGenerator.cpp
+++ b/Engine/Plugins/Experimental/Mutable/Source/MutableTools/Private/MuT/CodeGenerator.cpp
@@ -1089,46 +1089,50 @@ namespace mu
 			{
 				MUTABLE_CPUPROFILER_SCOPE(SurfaceTexture);
 
-				if (NodeImagePtr pImageNode = node.m_images[t].m_pImage)
+				// Any image-specific format or mipmapping needs to be applied at the end
+				Ptr<NodeImageMipmap> mipmapNode;
+				Ptr<NodeImageFormat> formatNode;
+				Ptr<NodeImageSwizzle> swizzleNode;
+
+				bool bFound = false;
+				NodeImagePtr pImageNode = node.m_images[t].m_pImage;
+
+				while (!bFound && pImageNode)
 				{
-					// Any image-specific format or mipmapping needs to be applied at the end
-					Ptr<NodeImageMipmap> mipmapNode;
-					Ptr<NodeImageFormat> formatNode;
-					Ptr<NodeImageSwizzle> swizzleNode;
-					bool found = false;
-					while (!found)
+					if (pImageNode->GetType()==NodeImageMipmap::GetStaticType())
 					{
-						if (pImageNode->GetType()==NodeImageMipmap::GetStaticType())
-						{
-							NodeImageMipmap* tm = static_cast<NodeImageMipmap*>(pImageNode.get());
-							if (!mipmapNode) mipmapNode = tm;
-							pImageNode = tm->GetSource();
-						}
-						else if (pImageNode->GetType() == NodeImageFormat::GetStaticType())
+						NodeImageMipmap* tm = static_cast<NodeImageMipmap*>(pImageNode.get());
+						if (!mipmapNode) mipmapNode = tm;
+						pImageNode = tm->GetSource();
+					}
+					else if (pImageNode->GetType() == NodeImageFormat::GetStaticType())
+					{
+						NodeImageFormat* tf = static_cast<NodeImageFormat*>(pImageNode.get());
+						if (!formatNode) formatNode = tf;
+						pImageNode = tf->GetSource();
+					}
+					else if (pImageNode->GetType() == NodeImageSwizzle::GetStaticType())
+					{
+						NodeImageSwizzle* ts = static_cast<NodeImageSwizzle*>(pImageNode.get());
+						NodeImage* Source = ts->GetSource(0).get();
+						if (!swizzleNode && Source ==ts->GetSource(1) && Source==ts->GetSource(2) && Source==ts->GetSource(3))
 						{
-							NodeImageFormat* tf = static_cast<NodeImageFormat*>(pImageNode.get());
-							if (!formatNode) formatNode = tf;
-							pImageNode = tf->GetSource();
-						}
-						else if (pImageNode->GetType() == NodeImageSwizzle::GetStaticType())
-						{
-							NodeImageSwizzle* ts = static_cast<NodeImageSwizzle*>(pImageNode.get());
-							NodeImage* Source = ts->GetSource(0).get();
-							if (!swizzleNode && Source ==ts->GetSource(1) && Source==ts->GetSource(2) && Source==ts->GetSource(3))
-							{
-								swizzleNode = ts;
-								pImageNode = Source;
-							}
-							else
-							{
-								found = true;
-							}
+							swizzleNode = ts;
+							pImageNode = Source;
 						}
 						else
 						{
-							found = true;
+							bFound = true;
 						}
 					}
+					else
+					{
+						bFound = true;
+					}
+				}
+
+				if (bFound)
+				{
 
 					const int LayoutIndex = node.m_images[t].m_layoutIndex;
 
