[Mutable] Fix CustomizableObject cooked in multiple cook workers.

#rnx

diff --git a/Engine/Plugins/Experimental/Mutable/Source/CustomizableObject/Classes/MuCO/CustomizableObjectInstance.h b/Engine/Plugins/Experimental/Mutable/Source/CustomizableObject/Classes/MuCO/CustomizableObjectInstance.h
--- a/Engine/Plugins/Experimental/Mutable/Source/CustomizableObject/Classes/MuCO/CustomizableObjectInstance.h
+++ b/Engine/Plugins/Experimental/Mutable/Source/CustomizableObject/Classes/MuCO/CustomizableObjectInstance.h
@@ -132,9 +132,6 @@ public:
 #if WITH_EDITOR
 	virtual void PostEditChangeProperty(FPropertyChangedEvent& PropertyChangedEvent) override;
 	virtual bool CanEditChange( const FProperty* InProperty ) const override;
-	virtual void BeginCacheForCookedPlatformData(const ITargetPlatform* TargetPlatform) override;
-	virtual bool IsCachedCookedPlatformDataLoaded(const ITargetPlatform* TargetPlatform) override;
-	virtual void PreSave(FObjectPreSaveContext ObjectSaveContext) override;
 #endif //WITH_EDITOR
 	virtual void Serialize(FArchive& Ar) override;
 	virtual void PostLoad() override;
diff --git a/Engine/Plugins/Experimental/Mutable/Source/CustomizableObject/Private/MuCO/CustomizableObjectInstance.cpp b/Engine/Plugins/Experimental/Mutable/Source/CustomizableObject/Private/MuCO/CustomizableObjectInstance.cpp
--- a/Engine/Plugins/Experimental/Mutable/Source/CustomizableObject/Private/MuCO/CustomizableObjectInstance.cpp
+++ b/Engine/Plugins/Experimental/Mutable/Source/CustomizableObject/Private/MuCO/CustomizableObjectInstance.cpp
@@ -564,61 +564,10 @@ void UCustomizableObjectInstance::PostLoad()
 
 	// Skip the cost of ReloadParameters in the cook commandlet; it will be reloaded during PreSave. For cooked runtime
 	// and editor UI, reload on load because it will not otherwise reload unless the CustomizableObject recompiles.
-	if (!IsRunningCookCommandlet())
-	{
-		Descriptor.ReloadParameters();
-	}
+	Descriptor.ReloadParameters();
 	PrivateData->InitCustomizableObjectData(GetCustomizableObject());
 }
 
-#if WITH_EDITOR
-void UCustomizableObjectInstance::BeginCacheForCookedPlatformData(const ITargetPlatform* TargetPlatform)
-{
-	UCustomizableObject* CustomizableObject = GetCustomizableObject();
-	if (!CustomizableObject)
-	{
-		return;
-	}
-	CustomizableObject->BeginCacheForCookedPlatformData(TargetPlatform);
-}
-
-bool UCustomizableObjectInstance::IsCachedCookedPlatformDataLoaded(const ITargetPlatform* TargetPlatform)
-{
-	UCustomizableObject* CustomizableObject = GetCustomizableObject();
-	if (!CustomizableObject)
-	{
-		return true;
-	}
-
-	return CustomizableObject->IsCachedCookedPlatformDataLoaded(TargetPlatform);
-}
-
-void UCustomizableObjectInstance::PreSave(FObjectPreSaveContext ObjectSaveContext)
-{
-	Super::PreSave(ObjectSaveContext);
-
-	const ITargetPlatform* TargetPlatform = ObjectSaveContext.GetTargetPlatform();
-	if (!TargetPlatform)
-	{
-		return;
-	}
-
-	UCustomizableObject* CustomizableObject = GetCustomizableObject();
-	if (!CustomizableObject)
-	{
-		return;
-	}
-	if (CustomizableObject->GetPrivate()->TryUpdateIsChildObject() && CustomizableObject->IsChildObject())
-	{
-		UE_LOG(LogMutable, Error,
-			TEXT("CO Instance [%s] has an invalid dependency on CO that is a child object: [%s]. The instance will be unusable at runtime and may crash."),
-			*GetPathName(), *CustomizableObject->GetPathName());
-	}
-
-	CustomizableObject->GetPrivate()->TryLoadCompiledCookDataForPlatform(TargetPlatform);
-	Descriptor.ReloadParameters();
-}
-#endif
 
 FString UCustomizableObjectInstance::GetDesc()
 {
diff --git a/Engine/Plugins/Experimental/Mutable/Source/CustomizableObject/Private/MuCO/CustomizableObjectInstanceDescriptor.cpp b/Engine/Plugins/Experimental/Mutable/Source/CustomizableObject/Private/MuCO/CustomizableObjectInstanceDescriptor.cpp
--- a/Engine/Plugins/Experimental/Mutable/Source/CustomizableObject/Private/MuCO/CustomizableObjectInstanceDescriptor.cpp
+++ b/Engine/Plugins/Experimental/Mutable/Source/CustomizableObject/Private/MuCO/CustomizableObjectInstanceDescriptor.cpp
@@ -742,13 +742,13 @@ FString FCustomizableObjectInstanceDescriptor::ToString() const
 
 void FCustomizableObjectInstanceDescriptor::ReloadParameters()
 {
-	if (!CustomizableObject)
+	if (IsRunningCookCommandlet())
 	{
 		return;
 	}
 
-	if (!CustomizableObject->IsCompiled())
-	{	
+	if (!CustomizableObject || !CustomizableObject->IsCompiled())
+	{
 		return;
 	}
 
