[Mutable] Fix Pose Mesh node not using the BoneMap to find bones affected by the pose.

#rnx

diff --git a/Engine/Plugins/Experimental/Mutable/Source/MutableRuntime/Private/MuR/OpMeshApplyPose.h b/Engine/Plugins/Experimental/Mutable/Source/MutableRuntime/Private/MuR/OpMeshApplyPose.h
--- a/Engine/Plugins/Experimental/Mutable/Source/MutableRuntime/Private/MuR/OpMeshApplyPose.h
+++ b/Engine/Plugins/Experimental/Mutable/Source/MutableRuntime/Private/MuR/OpMeshApplyPose.h
@@ -31,26 +31,59 @@ namespace mu
         // For now we only convert the vertex positions.
         // \TODO: normals and tangents
 
+		// Find closest bone affected by the pose for each bone in the skeleton.
+		const int32 NumBones = pSkeleton->GetBoneCount();
+		TArray<int32> BoneToPoseIndex;
+		BoneToPoseIndex.Init(INDEX_NONE, NumBones);
+
+		for (int32 BoneIndex = 0; BoneIndex < NumBones; ++BoneIndex)
+		{
+			const int32 BonePoseIndex = pPose->FindBonePose(pSkeleton->GetBoneId(BoneIndex));
+			if (BonePoseIndex != INDEX_NONE)
+			{
+				BoneToPoseIndex[BoneIndex] = BonePoseIndex;
+				continue;
+			}
+
+			// Parent bones are in a strictly increassing order. Set the pose index from the parent bone.
+			const int32 ParentIndex = pSkeleton->GetBoneParent(BoneIndex);
+			if (ParentIndex != INDEX_NONE)
+			{
+				BoneToPoseIndex[BoneIndex] = BoneToPoseIndex[ParentIndex];
+			}
+		}
+
+
+		TArray<uint16> BoneMap = pBase->GetBoneMap();
+		const int32 NumBonesBoneMap = BoneMap.Num();
+
         // Prepare the skin matrices. They may be in different order, and we only need the ones
         // relevant for the base mesh deformation.
-		TArray<FTransform3f> skinTransforms;
-		skinTransforms.SetNum(pSkeleton->GetBoneCount());
-        for ( int32 b=0; b< pSkeleton->GetBoneCount(); ++b )
-        {
-            int32 poseBoneIndex = pPose->FindBonePose(pSkeleton->GetBoneId(b));
+		TArray<FTransform3f> SkinTransforms;
+		SkinTransforms.Reserve(NumBonesBoneMap);
 
-            if( poseBoneIndex != INDEX_NONE )
-            {
-				skinTransforms[b] = pPose->BonePoses[poseBoneIndex].BoneTransform;
-            }
-            else
-            {
-                // This could happen if the model is compiled with maximum optimization and the skeleton
-				// has bones used by other meshes.
-				skinTransforms[b] = FTransform3f::Identity;
-            }
+		bool bBonesAffected = false;
+		for (int32 Index = 0; Index < NumBonesBoneMap; ++Index)
+		{
+			const int32 BoneIndex = pSkeleton->FindBone(BoneMap[Index]);
+			if (BoneToPoseIndex[BoneIndex] != INDEX_NONE)
+			{
+				SkinTransforms.Add(pPose->BonePoses[BoneToPoseIndex[BoneIndex]].BoneTransform);
+				bBonesAffected = true;
+			}
+			else
+			{
+				// Bone not affected by the pose. Set identity.
+				SkinTransforms.Add(FTransform3f::Identity);
+			}
+		}
 
-        }
+		if (!bBonesAffected)
+		{
+			// The pose does not affect any vertex in the mesh. 
+			bOutSuccess = false;
+			return;
+		}
 
         // Get pointers to vertex position data
 		MeshBufferIteratorConst<MBF_FLOAT32, float, 3> itSource(pBase->m_VertexBuffers, MBS_POSITION, 0);
@@ -106,7 +139,7 @@ namespace mu
                 //vec4f p = skinMatrices[ boneIndex[w] ]
                 //        * vec4f( sourcePosition, 1.0f )
                 //        * weight[w];
-				FVector3f p = skinTransforms[boneIndex[w]].TransformPosition(sourcePosition);
+				FVector3f p = SkinTransforms[boneIndex[w]].TransformPosition(sourcePosition);
 				position += p * weight[w];
             }
 
