[mutable] Fix bug with mesh switching node and layouts.

The mesh switching was not overriding the option layouts with the layouts in the first connection.

#rnx

diff --git a/Engine/Plugins/Experimental/Mutable/Source/MutableTools/Private/MuT/CodeGenerator_Mesh.cpp b/Engine/Plugins/Experimental/Mutable/Source/MutableTools/Private/MuT/CodeGenerator_Mesh.cpp
--- a/Engine/Plugins/Experimental/Mutable/Source/MutableTools/Private/MuT/CodeGenerator_Mesh.cpp
+++ b/Engine/Plugins/Experimental/Mutable/Source/MutableTools/Private/MuT/CodeGenerator_Mesh.cpp
@@ -1087,7 +1087,7 @@ namespace mu
             {
 				// Take the layouts from the first non-null connection.
 				// \TODO: Take them from the first connection that actually returns layouts?
-				if (!bFirstValidConnectionFound)
+				if (bFirstValidConnectionFound)
 				{
 					TargetOptions.OverrideLayouts = OutResult.GeneratedLayouts;
 				}
