[Mutable] Fix ApplyPose pose generation when root is not the identity.

#rnx

diff --git a/Engine/Plugins/Experimental/Mutable/Source/CustomizableObjectEditor/Private/MuCOE/Nodes/CustomizableObjectNodeAnimationPose.cpp b/Engine/Plugins/Experimental/Mutable/Source/CustomizableObjectEditor/Private/MuCOE/Nodes/CustomizableObjectNodeAnimationPose.cpp
--- a/Engine/Plugins/Experimental/Mutable/Source/CustomizableObjectEditor/Private/MuCOE/Nodes/CustomizableObjectNodeAnimationPose.cpp
+++ b/Engine/Plugins/Experimental/Mutable/Source/CustomizableObjectEditor/Private/MuCOE/Nodes/CustomizableObjectNodeAnimationPose.cpp
@@ -140,7 +140,7 @@ void UCustomizableObjectNodeAnimationPose::StaticRetrievePoseInformation(UPoseAs
 			FTransform CumulativePoseTransform;
 
 			int32 ParentIndex = BoneIndex;
-			while (ParentIndex > 0)
+			while (ParentIndex >= 0)
 			{
 				FName ParentName = RefSkeleton.GetBoneName(ParentIndex);
 				if (int32 PoseBoneIndex = PoseRefSkeleton.FindBoneIndex(ParentName); PoseBoneIndex >= 0)
