[mutable] Removed editor for MutableDebugger.
- The debugger will no longer be an editor for the CO. That was causing issues where the default editor for the type CO was being reported as the debugger editor. That seemed to happen because we had more than one edtor for that type of asset.
- Removed duplicated/not maintaned code part of the different debugger slates

diff --git a/Engine/Plugins/Experimental/Mutable/Source/CustomizableObjectEditor/Private/MuCOE/AssetDefinition_CustomizableObject.cpp b/Engine/Plugins/Experimental/Mutable/Source/CustomizableObjectEditor/Private/MuCOE/AssetDefinition_CustomizableObject.cpp
--- a/Engine/Plugins/Experimental/Mutable/Source/CustomizableObjectEditor/Private/MuCOE/AssetDefinition_CustomizableObject.cpp
+++ b/Engine/Plugins/Experimental/Mutable/Source/CustomizableObjectEditor/Private/MuCOE/AssetDefinition_CustomizableObject.cpp
@@ -4,8 +4,10 @@
 
 #include "AssetRegistry/AssetRegistryModule.h"
 #include "AssetToolsModule.h"
+#include "AssetViewUtils.h"
 #include "ContentBrowserMenuContexts.h"
 #include "ContentBrowserModule.h"
+#include "CustomizableObjectEditor.h"
 #include "Editor.h"
 #include "IContentBrowserSingleton.h"
 #include "Misc/MessageDialog.h"
@@ -148,8 +150,18 @@ namespace MenuExtension_CustomizableObject
 		
 		for (UCustomizableObject* Object : Context->LoadSelectedObjects<UCustomizableObject>())
 		{
-			ICustomizableObjectEditorModule& CustomizableObjectEditorModule = FModuleManager::LoadModuleChecked<ICustomizableObjectEditorModule>("CustomizableObjectEditor");
-			CustomizableObjectEditorModule.CreateCustomizableObjectDebugger(EToolkitMode::Standalone, nullptr , Object);
+			if (!AssetViewUtils::OpenEditorForAsset(Object))
+			{
+				continue;
+			}
+
+			if (IAssetEditorInstance* AssetEditor = GEditor->GetEditorSubsystem<UAssetEditorSubsystem>()->FindEditorForAsset(Object, true))
+			{
+				if (const FCustomizableObjectEditor* CustomizableObjectEditor = StaticCast<FCustomizableObjectEditor*>(AssetEditor))
+				{
+					CustomizableObjectEditor->DebugObject();
+				}
+			}
 		}
 	}
 
diff --git a/Engine/Plugins/Experimental/Mutable/Source/CustomizableObjectEditor/Private/MuCOE/CustomizableObjectDebugger.cpp b/Engine/Plugins/Experimental/Mutable/Source/CustomizableObjectEditor/Private/MuCOE/CustomizableObjectDebugger.cpp
deleted file mode 100644
diff --git a/Engine/Plugins/Experimental/Mutable/Source/CustomizableObjectEditor/Private/MuCOE/CustomizableObjectDebugger.h b/Engine/Plugins/Experimental/Mutable/Source/CustomizableObjectEditor/Private/MuCOE/CustomizableObjectDebugger.h
deleted file mode 100644
diff --git a/Engine/Plugins/Experimental/Mutable/Source/CustomizableObjectEditor/Private/MuCOE/CustomizableObjectEditor.cpp b/Engine/Plugins/Experimental/Mutable/Source/CustomizableObjectEditor/Private/MuCOE/CustomizableObjectEditor.cpp
--- a/Engine/Plugins/Experimental/Mutable/Source/CustomizableObjectEditor/Private/MuCOE/CustomizableObjectEditor.cpp
+++ b/Engine/Plugins/Experimental/Mutable/Source/CustomizableObjectEditor/Private/MuCOE/CustomizableObjectEditor.cpp
@@ -60,6 +60,7 @@
 #include "Widgets/Input/STextComboBox.h"
 #include "Widgets/Layout/SScrollBox.h"
 #include "Widgets/Notifications/SNotificationList.h"
+#include "SMutableObjectViewer.h"
 
 class FAdvancedPreviewScene;
 class FWorkspaceItem;
@@ -1577,18 +1578,16 @@ void FCustomizableObjectEditor::CompileObject()
 }
 
 
-void FCustomizableObjectEditor::DebugObject()
+void FCustomizableObjectEditor::DebugObject() const
 {
-	if (!CustomizableObject)
-	{
-		return;
-	}
-
-	// Open a mutable debugger as a standalone window
-	ICustomizableObjectEditorModule* CustomizableObjectEditorModule = &FModuleManager::LoadModuleChecked<ICustomizableObjectEditorModule>("CustomizableObjectEditor");
-	if (!CustomizableObjectEditorModule) return;
-
-	CustomizableObjectEditorModule->CreateCustomizableObjectDebugger(EToolkitMode::Standalone, nullptr, CustomizableObject);
+	const TSharedPtr<SDockTab> NewMutableObjectTab = SNew(SDockTab)
+	.Label(FText::FromString(TEXT("Debugger")))
+	[
+		SNew(SMutableObjectViewer, CustomizableObject)
+	];
+	
+	// Spawn the debugger tab alongside the Graph Tab 
+	TabManager->InsertNewDocumentTab(GraphTabId, FTabManager::ESearchPreference::PreferLiveTab, NewMutableObjectTab.ToSharedRef());
 }
 
 
diff --git a/Engine/Plugins/Experimental/Mutable/Source/CustomizableObjectEditor/Private/MuCOE/CustomizableObjectEditor.h b/Engine/Plugins/Experimental/Mutable/Source/CustomizableObjectEditor/Private/MuCOE/CustomizableObjectEditor.h
--- a/Engine/Plugins/Experimental/Mutable/Source/CustomizableObjectEditor/Private/MuCOE/CustomizableObjectEditor.h
+++ b/Engine/Plugins/Experimental/Mutable/Source/CustomizableObjectEditor/Private/MuCOE/CustomizableObjectEditor.h
@@ -276,6 +276,9 @@ public:
 
 	TSharedPtr<class SCustomizableObjectNodeLayoutBlocksEditor> GetLayoutBlocksEditor() { return LayoutBlocksEditor; }
 
+	/** Debug the object as a raw mutable data in the internal tools. */
+	void DebugObject() const;
+	
 private:
 	explicit FCustomizableObjectEditor(UCustomizableObject& ObjectToEdit);
 
@@ -305,9 +308,6 @@ private:
 	void CompileObjectUserPressedButton();
 	void CompileOnlySelectedObjectUserPressedButton();
 	
-	/** Debug the object as a raw mutable data in the internal tools. */
-	void DebugObject();
-
 	// Compile options menu callbacks
 	TSharedRef<SWidget> GenerateCompileOptionsMenuContent(TSharedRef<FUICommandList> InCommandList);
 	void ResetCompileOptions();
diff --git a/Engine/Plugins/Experimental/Mutable/Source/CustomizableObjectEditor/Private/MuCOE/CustomizableObjectEditorModule.cpp b/Engine/Plugins/Experimental/Mutable/Source/CustomizableObjectEditor/Private/MuCOE/CustomizableObjectEditorModule.cpp
--- a/Engine/Plugins/Experimental/Mutable/Source/CustomizableObjectEditor/Private/MuCOE/CustomizableObjectEditorModule.cpp
+++ b/Engine/Plugins/Experimental/Mutable/Source/CustomizableObjectEditor/Private/MuCOE/CustomizableObjectEditorModule.cpp
@@ -14,7 +14,6 @@
 #include "MuCOE/CustomizableInstanceDetails.h"
 #include "MuCOE/CustomizableObjectCustomSettings.h"
 #include "MuCOE/CustomizableObjectCustomSettingsDetails.h"
-#include "MuCOE/CustomizableObjectDebugger.h"
 #include "MuCOE/CustomizableObjectDetails.h"
 #include "MuCOE/CustomizableObjectEditor.h"
 #include "MuCOE/CustomizableObjectEditorLogger.h"
@@ -92,7 +91,6 @@ public:
 	// ICustomizableObjectEditorModule interface
 	TSharedRef<ICustomizableObjectEditor> CreateCustomizableObjectEditor( const EToolkitMode::Type Mode, const TSharedPtr< IToolkitHost >& InitToolkitHost, UCustomizableObject* CustomizableObject ) override;
 	TSharedRef<ICustomizableObjectInstanceEditor> CreateCustomizableObjectInstanceEditor( const EToolkitMode::Type Mode, const TSharedPtr< IToolkitHost >& InitToolkitHost, UCustomizableObjectInstance* CustomizableObjectInstance ) override;
-	TSharedRef<ICustomizableObjectDebugger> CreateCustomizableObjectDebugger(const EToolkitMode::Type Mode, const TSharedPtr< IToolkitHost >& InitToolkitHost, UCustomizableObject* CustomizableObject) override;
 	virtual EAssetTypeCategories::Type GetAssetCategory() const override;
 	virtual FCustomizableObjectEditorLogger& GetLogger() override;
 
@@ -246,12 +244,6 @@ TSharedRef<ICustomizableObjectInstanceEditor> FCustomizableObjectEditorModule::C
 }
 
 
-TSharedRef<ICustomizableObjectDebugger> FCustomizableObjectEditorModule::CreateCustomizableObjectDebugger(const EToolkitMode::Type Mode, const TSharedPtr< IToolkitHost >& InitToolkitHost, UCustomizableObject* CustomizableObject)
-{
-	TSharedRef<FCustomizableObjectDebugger> NewCustomizableObjectDebugger(new FCustomizableObjectDebugger());
-	NewCustomizableObjectDebugger->InitCustomizableObjectDebugger(Mode, InitToolkitHost, CustomizableObject);
-	return NewCustomizableObjectDebugger;
-}
 
 
 FCustomizableObjectEditorLogger& FCustomizableObjectEditorModule::GetLogger()
diff --git a/Engine/Plugins/Experimental/Mutable/Source/CustomizableObjectEditor/Private/MuCOE/SMutableGraphViewer.cpp b/Engine/Plugins/Experimental/Mutable/Source/CustomizableObjectEditor/Private/MuCOE/SMutableGraphViewer.cpp
--- a/Engine/Plugins/Experimental/Mutable/Source/CustomizableObjectEditor/Private/MuCOE/SMutableGraphViewer.cpp
+++ b/Engine/Plugins/Experimental/Mutable/Source/CustomizableObjectEditor/Private/MuCOE/SMutableGraphViewer.cpp
@@ -119,17 +119,12 @@ FString SMutableGraphViewer::GetReferencerName() const
 }
 
 
-void SMutableGraphViewer::Construct(const FArguments& InArgs, const mu::NodePtr& InRootNode,
-	const FCompilationOptions& InCompileOptions,
-	TWeakPtr<FTabManager> InParentTabManager, const FName& InParentNewTabId)
+void SMutableGraphViewer::Construct(const FArguments& InArgs, const mu::NodePtr& InRootNode)
 {
 	DataTag = InArgs._DataTag;
 	ReferencedRuntimeTextures = InArgs._ReferencedRuntimeTextures;
 	ReferencedCompileTextures = InArgs._ReferencedCompileTextures;
 	RootNode = InRootNode;
-	CompileOptions = InCompileOptions;
-	ParentTabManager = InParentTabManager;
-	ParentNewTabId = InParentNewTabId;
 
 	FToolBarBuilder ToolbarBuilder(TSharedPtr<const FUICommandList>(), FMultiBoxCustomization::None, TSharedPtr<FExtender>(), true);
 	ToolbarBuilder.SetLabelVisibility(EVisibility::Visible);
@@ -179,28 +174,6 @@ void SMutableGraphViewer::Construct(const FArguments& InArgs, const mu::NodePtr&
 		
 	ToolbarBuilder.EndSection();
 
-
-	ToolbarBuilder.BeginSection("Compilation");
-
-	ToolbarBuilder.AddToolBarButton(
-		FUIAction(FExecuteAction::CreateSP(this, &SMutableGraphViewer::CompileMutableCodePressed)),
-		NAME_None,
-		LOCTEXT("GenerateMutableCode", "Unreal to Mutable Code"),
-		LOCTEXT("GenerateMutableCodeTooltip", "Generate a mutable code from the mutable graph."),
-		FSlateIcon(FCustomizableObjectEditorStyle::Get().GetStyleSetName(), "CustomizableObjectDebugger.CompileMutableCode", "CustomizableObjectDebugger.CompileMutableCode.Small"),
-		EUserInterfaceActionType::Button
-	);
-
-	ToolbarBuilder.AddComboButton(
-		FUIAction(),
-		FOnGetContent::CreateSP(this, &SMutableGraphViewer::GenerateCompileOptionsMenuContent),
-		LOCTEXT("Compile_Options_Label", "Compile Options"),
-		LOCTEXT("Compile_Options_Tooltip", "Change Compile Options"),
-		TAttribute<FSlateIcon>(),
-		true);
-
-	ToolbarBuilder.EndSection();
-
 	ToolbarBuilder.AddWidget(SNew(STextBlock).Text(MakeAttributeLambda([this]() { return FText::FromString(DataTag); })));
 
 	ChildSlot
@@ -257,122 +230,6 @@ void SMutableGraphViewer::Construct(const FArguments& InArgs, const mu::NodePtr&
 }
 
 
-void SMutableGraphViewer::CompileMutableCodePressed()
-{
-	// Do the compilation to Mutable Code synchronously.
-	TSharedPtr<FCustomizableObjectCompileRunnable> CompileTask = MakeShareable(new FCustomizableObjectCompileRunnable(RootNode));
-	CompileTask->Options = CompileOptions;
-	CompileTask->ReferencedTextures = ReferencedCompileTextures;
-	CompileTask->Init();
-	CompileTask->Run();
-
-	FString NewDataTag = FString::Printf(TEXT("%s graph, opt %d "), *DataTag, CompileOptions.OptimizationLevel);
-
-	TSharedPtr<SDockTab> NewMutableCodeTab = SNew(SDockTab)
-		.Label(LOCTEXT("MutableCode", "Mutable Code"))
-		[
-			SNew(SMutableCodeViewer, CompileTask->Model, ReferencedRuntimeTextures)
-			.DataTag(NewDataTag)
-		];
-
-	TSharedPtr<FTabManager> TabManager = ParentTabManager.Pin();
-	check(TabManager);
-	TabManager->InsertNewDocumentTab(ParentNewTabId, FTabManager::ESearchPreference::PreferLiveTab, NewMutableCodeTab.ToSharedRef());
-}
-
-
-TSharedRef<SWidget> SMutableGraphViewer::GenerateCompileOptionsMenuContent()
-{
-	const bool bShouldCloseWindowAfterMenuSelection = false;
-	FMenuBuilder MenuBuilder(bShouldCloseWindowAfterMenuSelection, nullptr);
-
-	// settings
-	MenuBuilder.BeginSection("Optimization", LOCTEXT("MutableCompileOptimizationHeading", "Optimization"));
-	{
-		// Compilation options
-		//-----------------------------------
-
-		// Optimisation level
-		CompileOptimizationStrings.Empty();
-		CompileOptimizationStrings.Add(MakeShareable(new FString(NSLOCTEXT("CustomizableObjectEditor", "Debugger_OptimizationNone", "None").ToString())));
-		CompileOptimizationStrings.Add(MakeShareable(new FString(NSLOCTEXT("CustomizableObjectEditor", "Debugger_OptimizationMin", "Minimal").ToString())));
-		CompileOptimizationStrings.Add(MakeShareable(new FString(NSLOCTEXT("CustomizableObjectEditor", "Debugger_OptimizationMax", "Maximum").ToString())));
-
-		CompileOptions.OptimizationLevel = FMath::Min(CompileOptions.OptimizationLevel, CompileOptimizationStrings.Num() - 1);
-
-		CompileOptimizationCombo =
-			SNew(STextComboBox)
-			.OptionsSource(&CompileOptimizationStrings)
-			.InitiallySelectedItem(CompileOptimizationStrings[CompileOptions.OptimizationLevel])
-			.OnSelectionChanged(this, &SMutableGraphViewer::OnChangeCompileOptimizationLevel)
-			;
-		MenuBuilder.AddWidget(CompileOptimizationCombo.ToSharedRef(), LOCTEXT("MutableCompileOptimizationLevel", "Optimization Level"));
-
-		{
-			CompileTextureCompressionStrings.Empty();
-			CompileTextureCompressionStrings.Add(MakeShareable(new FString(LOCTEXT("MutableTextureCompressionNone", "None").ToString())));
-			CompileTextureCompressionStrings.Add(MakeShareable(new FString(LOCTEXT("MutableTextureCompressionFast", "Fast").ToString())));
-			CompileTextureCompressionStrings.Add(MakeShareable(new FString(LOCTEXT("MutableTextureCompressionHighQuality", "High Quality").ToString())));
-
-			int32 SelectedCompression = FMath::Clamp(int32(CompileOptions.TextureCompression), 0, CompileTextureCompressionStrings.Num() - 1);
-			CompileTextureCompressionCombo =
-				SNew(STextComboBox)
-				.OptionsSource(&CompileTextureCompressionStrings)
-				.InitiallySelectedItem(CompileTextureCompressionStrings[SelectedCompression])
-				.OnSelectionChanged(this, &SMutableGraphViewer::OnChangeCompileTextureCompressionType)
-				;
-
-			MenuBuilder.AddWidget(CompileTextureCompressionCombo.ToSharedRef(), LOCTEXT("MutableCompileTextureCompressionType", "Texture Compression"));
-		}
-
-		// Image tiling
-		// Unfortunately SNumericDropDown doesn't work with integers at the time of writing.
-		TArray<SNumericDropDown<float>::FNamedValue> TilingOptions;
-		TilingOptions.Add(SNumericDropDown<float>::FNamedValue(0, FText::FromString(TEXT("0")), FText::FromString(TEXT("Disabled"))));
-		TilingOptions.Add(SNumericDropDown<float>::FNamedValue(64, FText::FromString(TEXT("64")), FText::FromString(TEXT("64"))));
-		TilingOptions.Add(SNumericDropDown<float>::FNamedValue(128, FText::FromString(TEXT("128")), FText::FromString(TEXT("128"))));
-		TilingOptions.Add(SNumericDropDown<float>::FNamedValue(256, FText::FromString(TEXT("256")), FText::FromString(TEXT("256"))));
-		TilingOptions.Add(SNumericDropDown<float>::FNamedValue(512, FText::FromString(TEXT("512")), FText::FromString(TEXT("512"))));
-
-		CompileTilingCombo = SNew(SNumericDropDown<float>)
-			.DropDownValues(TilingOptions)
-			.Value_Lambda([&]() { return float(CompileOptions.ImageTiling); })
-			.OnValueChanged_Lambda([&](float Value) { CompileOptions.ImageTiling = int32(Value); })
-			;
-		MenuBuilder.AddWidget(CompileTilingCombo.ToSharedRef(), LOCTEXT("MutableCompileImageTiling", "Image Tiling"));
-
-		// Disk as cache
-		MenuBuilder.AddMenuEntry(
-			LOCTEXT("MutableDiskAsMemory", "Enable compiling using the disk as memory."),
-			LOCTEXT("MutableDiskAsMemoryTooltip", "This is very slow but supports compiling huge objects. It requires a lot of free space in the OS disk."),
-			FSlateIcon(),
-			FUIAction(
-				FExecuteAction::CreateLambda([this]() { CompileOptions.bUseDiskCompilation = !CompileOptions.bUseDiskCompilation; }),
-				FCanExecuteAction(),
-				FIsActionChecked::CreateLambda([this]() { return CompileOptions.bUseDiskCompilation; })),
-			NAME_None,
-			EUserInterfaceActionType::ToggleButton
-		);
-	}
-	MenuBuilder.EndSection();
-
-	return MenuBuilder.MakeWidget();
-}
-
-
-void SMutableGraphViewer::OnChangeCompileOptimizationLevel(TSharedPtr<FString> NewSelection, ESelectInfo::Type)
-{
-	CompileOptions.OptimizationLevel = CompileOptimizationStrings.Find(NewSelection);
-}
-
-
-void SMutableGraphViewer::OnChangeCompileTextureCompressionType(TSharedPtr<FString> NewSelection, ESelectInfo::Type)
-{
-	const FScopedTransaction Transaction(LOCTEXT("ChangedOptimizationLevelTransaction", "Changed Optimization Level"));
-	CompileOptions.TextureCompression = ECustomizableObjectTextureCompression(CompileTextureCompressionStrings.Find(NewSelection));
-}
-
-
 void SMutableGraphViewer::RebuildTree()
 {
 	RootNodes.Reset();
diff --git a/Engine/Plugins/Experimental/Mutable/Source/CustomizableObjectEditor/Private/MuCOE/SMutableGraphViewer.h b/Engine/Plugins/Experimental/Mutable/Source/CustomizableObjectEditor/Private/MuCOE/SMutableGraphViewer.h
--- a/Engine/Plugins/Experimental/Mutable/Source/CustomizableObjectEditor/Private/MuCOE/SMutableGraphViewer.h
+++ b/Engine/Plugins/Experimental/Mutable/Source/CustomizableObjectEditor/Private/MuCOE/SMutableGraphViewer.h
@@ -40,9 +40,7 @@ public:
 
 	SLATE_END_ARGS()
 
-	void Construct(const FArguments& InArgs, const mu::NodePtr& InRootNode, 
-		const FCompilationOptions& InCompileOptions,
-		TWeakPtr<FTabManager> InParentTabManager, const FName& InParentNewTabId);
+	void Construct(const FArguments& InArgs, const mu::NodePtr& InRootNode);
 
 	// SWidget interface
 	FReply OnDragOver(const FGeometry& MyGeometry, const FDragDropEvent& DragDropEvent) override;
@@ -64,16 +62,9 @@ private:
 	TArray<TSoftObjectPtr<UTexture>> ReferencedRuntimeTextures;
 	TArray<TSoftObjectPtr<UTexture>> ReferencedCompileTextures;
 
-	/** Compilation options to use in the debugger operations. */
-	FCompilationOptions CompileOptions;
-
 	/** Object compiler. */
 	FCustomizableObjectCompiler Compiler;
 
-	/** UI references used to create new tabs. */
-	TWeakPtr<FTabManager> ParentTabManager;
-	FName ParentNewTabId;
-
 	/** Root nodes of the tree widget. */
 	TArray<TSharedPtr<FMutableGraphTreeElement>> RootNodes;
 
@@ -106,18 +97,7 @@ private:
 
 	/** */
 	void RebuildTree();
-
-	/** UI callbacks */
-	void CompileMutableCodePressed();
-	TSharedRef<SWidget> GenerateCompileOptionsMenuContent();
-	TSharedPtr<STextComboBox> CompileOptimizationCombo;
-	TArray< TSharedPtr<FString> > CompileOptimizationStrings;
-	TSharedPtr<STextComboBox> CompileTextureCompressionCombo;
-	TArray< TSharedPtr<FString> > CompileTextureCompressionStrings;
-	TSharedPtr<SNumericDropDown<float>> CompileTilingCombo;
-	void OnChangeCompileOptimizationLevel(TSharedPtr<FString> NewSelection, ESelectInfo::Type SelectInfo);
-	void OnChangeCompileTextureCompressionType(TSharedPtr<FString> NewSelection, ESelectInfo::Type);
-
+	
 	/** Callbacks from the tree widget. */
  	TSharedRef<ITableRow> GenerateRowForNodeTree(TSharedPtr<FMutableGraphTreeElement> InTreeNode, const TSharedRef<STableViewBase>& InOwnerTable);
 	void GetChildrenForInfo(TSharedPtr<FMutableGraphTreeElement> InInfo, TArray< TSharedPtr<FMutableGraphTreeElement> >& OutChildren);
diff --git a/Engine/Plugins/Experimental/Mutable/Source/CustomizableObjectEditor/Private/MuCOE/SMutableObjectViewer.cpp b/Engine/Plugins/Experimental/Mutable/Source/CustomizableObjectEditor/Private/MuCOE/SMutableObjectViewer.cpp
--- a/Engine/Plugins/Experimental/Mutable/Source/CustomizableObjectEditor/Private/MuCOE/SMutableObjectViewer.cpp
+++ b/Engine/Plugins/Experimental/Mutable/Source/CustomizableObjectEditor/Private/MuCOE/SMutableObjectViewer.cpp
@@ -10,7 +10,6 @@
 #include "MuCOE/CustomizableObjectEditorStyle.h"
 #include "MuCOE/SMutableCodeViewer.h"
 #include "MuCOE/SMutableGraphViewer.h"
-#include "MuCOE/UnrealEditorPortabilityHelpers.h"
 #include "Widgets/Docking/SDockTab.h"
 #include "Widgets/Input/STextComboBox.h"
 #include "Widgets/Input/SNumericDropDown.h"
@@ -24,100 +23,6 @@ struct FSlateBrush;
 
 #define LOCTEXT_NAMESPACE "SMutableDebugger"
 
-// \todo: multi-column tree
-namespace MutableObjectTreeViewColumns
-{
-	static const FName Name("Property");
-	static const FName Value("Value");
-};
-
-
-class SMutableObjectTreeRow : public STableRow<TSharedPtr<FMutableObjectTreeElement>>
-{
-public:
-
-	void Construct(const FArguments& Args, const TSharedRef<STableViewBase>& InOwnerTableView, const TSharedPtr<FMutableObjectTreeElement>& InRowItem, const UCustomizableObject* InObject )
-	{
-		RowItem = InRowItem;
-
-		FText MainLabel;;
-		switch (RowItem->Type)
-		{
-		case FMutableObjectTreeElement::EType::SectionCaption:
-		{
-			switch (RowItem->Section)
-			{
-			case FMutableObjectTreeElement::ESection::General:
-				MainLabel = FText::FromString(TEXT("General"));
-				break;
-			case FMutableObjectTreeElement::ESection::ChildObjects:
-				MainLabel = FText::FromString(TEXT("Children"));
-				break;
-			default:
-				MainLabel = FText::FromString(TEXT("Unknown"));
-				break;
-			}
-			break;
-		}
-
-		case FMutableObjectTreeElement::EType::Name:
-			MainLabel = FText::FromString(InObject->GetName() );
-
-		default:
-			MainLabel = FText::FromString(TEXT("Unknown"));
-			break;
-		}
-
-
-		// TODO
-		const FSlateBrush* IconBrush = nullptr;
-
-		this->ChildSlot
-		[
-			SNew(SHorizontalBox)
-			+ SHorizontalBox::Slot()
-			.VAlign(VAlign_Center)
-			.AutoWidth()
-			[
-				SNew(SExpanderArrow, SharedThis(this))
-			]
-
-			+ SHorizontalBox::Slot()
-				.VAlign(VAlign_Center)
-				.Padding(FMargin(5.f, 0.f, 5.f, 0.f))
-				.AutoWidth()
-				[
-					SNew(SOverlay)
-					+ SOverlay::Slot()
-					[
-						SNew(SImage)
-						.Image(IconBrush ? IconBrush : FCoreStyle::Get().GetDefaultBrush())
-						.ColorAndOpacity(IconBrush ? FLinearColor::White : FLinearColor::Transparent)
-					]
-				]
-
-			+ SHorizontalBox::Slot()
-				[
-					SNew(STextBlock)
-					.Text(MainLabel)
-				]
-		];
-
-		STableRow< TSharedPtr<FMutableObjectTreeElement> >::ConstructInternal(
-			STableRow::FArguments()
-			//.Style(FAppStyle::Get(), "DetailsView.TreeView.TableRow")
-			.ShowSelection(true)
-			, InOwnerTableView
-		);
-
-	}
-
-
-private:
-
-	TSharedPtr<FMutableObjectTreeElement> RowItem;
-};
-
 
 void SMutableObjectViewer::AddReferencedObjects(FReferenceCollector& Collector)
 {
@@ -131,12 +36,9 @@ FString SMutableObjectViewer::GetReferencerName() const
 }
 
 
-void SMutableObjectViewer::Construct(const FArguments& InArgs, UCustomizableObject* InObject, 
-	TWeakPtr<FTabManager> InParentTabManager, const FName& InParentNewTabId )
+void SMutableObjectViewer::Construct(const FArguments& InArgs, UCustomizableObject* InObject)
 {
 	CustomizableObject = InObject;
-	ParentTabManager = InParentTabManager;
-	ParentNewTabId = InParentNewTabId;
 
 	// Initialize the debugger compile options
 	CompileOptions.TextureCompression = ECustomizableObjectTextureCompression::Fast;
@@ -180,11 +82,7 @@ void SMutableObjectViewer::Construct(const FArguments& InArgs, UCustomizableObje
 		true);
 
 	ToolbarBuilder.EndSection();
-
-	// Initialize the root tree nodes
-	RootTreeNodes.Add(MakeShareable(new FMutableObjectTreeElement(FMutableObjectTreeElement::EType::SectionCaption, FMutableObjectTreeElement::ESection::General, InObject)));
-	RootTreeNodes.Add(MakeShareable(new FMutableObjectTreeElement(FMutableObjectTreeElement::EType::SectionCaption, FMutableObjectTreeElement::ESection::ChildObjects, InObject)));
-
+	
 	ChildSlot
 	[
 		SNew(SVerticalBox)
@@ -195,20 +93,11 @@ void SMutableObjectViewer::Construct(const FArguments& InArgs, UCustomizableObje
 			ToolbarBuilder.MakeWidget()
 		]
 		+ SVerticalBox::Slot()
-		.VAlign(VAlign_Fill)
+		.FillHeight(1.0f)
 		[
-			SNew(SBorder)
-			.BorderImage(UE_MUTABLE_GET_BRUSH("ToolPanel.GroupBorder"))
-			.Padding(FMargin(4.0f, 4.0f))
-			[
-				SAssignNew(TreeView, STreeView<TSharedPtr<FMutableObjectTreeElement>>)
-				.TreeItemsSource(&RootTreeNodes)
-				.OnGenerateRow(this,&SMutableObjectViewer::GenerateRowForNodeTree)
-				.OnGetChildren(this, &SMutableObjectViewer::GetChildrenForInfo)
-				.SelectionMode(ESelectionMode::None)					
-			]
+			SAssignNew(DebuggerContentsBox, SBorder)
 		]
-	];	
+	];
 }
 
 
@@ -224,20 +113,15 @@ void SMutableObjectViewer::GenerateMutableGraphPressed()
 		return;
 	}
 
-	FString DataTag = FString::Printf(TEXT("%s for %s"), *CustomizableObject->GetName(), *CompileOptions.TargetPlatform->PlatformName());
+	FString DataTag = FString::Printf(TEXT("Mutable Graph : %s"), *CompileOptions.TargetPlatform->PlatformName());
 
-	TSharedPtr<SDockTab> NewMutableGraphTab = SNew(SDockTab)
-		.Label(LOCTEXT("MutableGraph", "Mutable Graph"))
-		[
-			SNew(SMutableGraphViewer, RootNode, CompileOptions, ParentTabManager, ParentNewTabId)
-			.DataTag(DataTag)
-			.ReferencedRuntimeTextures(RuntimeTextures)
-			.ReferencedCompileTextures(CompilerTextures)
-		];
-
-	TSharedPtr<FTabManager> TabManager = ParentTabManager.Pin();
-	check(TabManager);
-	TabManager->InsertNewDocumentTab(ParentNewTabId, FTabManager::ESearchPreference::PreferLiveTab, NewMutableGraphTab.ToSharedRef());
+	TSharedRef<SMutableGraphViewer> GraphViewer = SNew(SMutableGraphViewer, RootNode)
+		.DataTag(DataTag)
+		.ReferencedRuntimeTextures(RuntimeTextures)
+		.ReferencedCompileTextures(CompilerTextures);
+	
+	DebuggerContentsBox->ClearContent();
+	DebuggerContentsBox->SetContent(GraphViewer);
 }
 
 
@@ -247,8 +131,8 @@ void SMutableObjectViewer::CompileMutableCodePressed()
 	TArray<TSoftObjectPtr<UTexture>> CompilerTextures;
 	if (CompileOptions.bForceLargeLODBias)
 	{
-		// Debug compile with many different biasses
-		const int32 MaxBias = 15;
+		// Debug compile with many different biases
+		constexpr int32 MaxBias = 15;
 		for (int32 Bias = 0; Bias < MaxBias; ++Bias)
 		{
 			CompileOptions.DebugBias = Bias;
@@ -270,7 +154,6 @@ void SMutableObjectViewer::CompileMutableCodePressed()
 			CompileTask->Init();
 			CompileTask->Run();
 		}
-
 	}
 
 	// Convert from Unreal graph to Mutable graph.
@@ -291,18 +174,13 @@ void SMutableObjectViewer::CompileMutableCodePressed()
 	CompileTask->Init();
 	CompileTask->Run();
 
-	FString DataTag = FString::Printf( TEXT("%s for %s opt %d "), *CustomizableObject->GetName(), *CompileOptions.TargetPlatform->PlatformName(), CompileOptions.OptimizationLevel );
-
-	TSharedPtr<SDockTab> NewMutableCodeTab = SNew(SDockTab)
-		.Label(LOCTEXT("MutableCode", "Mutable Code"))
-		[
-			SNew(SMutableCodeViewer, CompileTask->Model, RuntimeTextures)
-			.DataTag(DataTag)
-		];
+	FString DataTag = FString::Printf( TEXT("Mutable Code : %s opt %d"), *CompileOptions.TargetPlatform->PlatformName(), CompileOptions.OptimizationLevel );
 
-	TSharedPtr<FTabManager> TabManager = ParentTabManager.Pin();
-	check(TabManager);
-	TabManager->InsertNewDocumentTab(ParentNewTabId, FTabManager::ESearchPreference::PreferLiveTab, NewMutableCodeTab.ToSharedRef());
+	TSharedRef<SMutableCodeViewer> CodeViewer = SNew(SMutableCodeViewer, CompileTask->Model, RuntimeTextures)
+		.DataTag(DataTag);
+	
+	DebuggerContentsBox->ClearContent();
+	DebuggerContentsBox->SetContent(CodeViewer);
 }
 
 
@@ -471,44 +349,4 @@ void SMutableObjectViewer::OnChangeDebugPlatform(TSharedPtr<FString> NewSelectio
 	}
 }
 
-
-TSharedRef<ITableRow> SMutableObjectViewer::GenerateRowForNodeTree(TSharedPtr<FMutableObjectTreeElement> InTreeNode, const TSharedRef<STableViewBase>& InOwnerTable)
-{
-	TSharedRef<SMutableObjectTreeRow> Row = SNew(SMutableObjectTreeRow, InOwnerTable, InTreeNode, CustomizableObject);
-	return Row;
-}
-
-
-void SMutableObjectViewer::GetChildrenForInfo(TSharedPtr<FMutableObjectTreeElement> InInfo, TArray<TSharedPtr<FMutableObjectTreeElement>>& OutChildren)
-{
-	switch (InInfo->Type)
-	{
-	case FMutableObjectTreeElement::EType::SectionCaption:
-		switch (InInfo->Section)
-		{
-
-		case FMutableObjectTreeElement::ESection::General:
-		{
-			OutChildren.Add(MakeShareable(new FMutableObjectTreeElement(FMutableObjectTreeElement::EType::Name, InInfo->Section, CustomizableObject )));
-			break;
-		}
-
-		case FMutableObjectTreeElement::ESection::ChildObjects:
-		{
-			// TODO: based on what is done in FCustomizableObjectCompiler::ProcessChildObjectsRecursively
-			//OutChildren.Add(MakeShareable(new FMutableObjectTreeElement(FMutableObjectTreeElement::EType::Name, InInfo->Section)));
-			break;
-		}
-
-		default:
-			break;
-		}
-		break;
-
-	default:
-		break;
-	}
-}
-
-
 #undef LOCTEXT_NAMESPACE 
diff --git a/Engine/Plugins/Experimental/Mutable/Source/CustomizableObjectEditor/Private/MuCOE/SMutableObjectViewer.h b/Engine/Plugins/Experimental/Mutable/Source/CustomizableObjectEditor/Private/MuCOE/SMutableObjectViewer.h
--- a/Engine/Plugins/Experimental/Mutable/Source/CustomizableObjectEditor/Private/MuCOE/SMutableObjectViewer.h
+++ b/Engine/Plugins/Experimental/Mutable/Source/CustomizableObjectEditor/Private/MuCOE/SMutableObjectViewer.h
@@ -32,15 +32,17 @@ public:
 	SLATE_BEGIN_ARGS(SMutableObjectViewer) {}
 	SLATE_END_ARGS()
 
-	void Construct(const FArguments& InArgs, UCustomizableObject*,
-		TWeakPtr<FTabManager> InParentTabManager, const FName& InParentNewTabId );
+	void Construct(const FArguments& InArgs, UCustomizableObject*);
 
 	// FGCObject interface
 	virtual void AddReferencedObjects(FReferenceCollector& Collector) override;
 	virtual FString GetReferencerName() const override;
 
 private:
-
+	
+	/** Slate whose content space contains the currently open debugging view */
+	TSharedPtr<SBorder> DebuggerContentsBox;
+	
 	/** The Mutable Graph to show, represented by its root. */
 	TObjectPtr<UCustomizableObject> CustomizableObject;
 
@@ -52,16 +54,6 @@ private:
 	/** Object compiler. */
 	FCustomizableObjectCompiler Compiler;
 
-	/** UI references used to create new tabs. */
-	TWeakPtr<FTabManager> ParentTabManager;
-	FName ParentNewTabId;
-
-	/** Tree showing the object properties. */
-	TSharedPtr<STreeView<TSharedPtr<FMutableObjectTreeElement>>> TreeView;
-
-	/** Root nodes of the tree widget. */
-	TArray<TSharedPtr<FMutableObjectTreeElement>> RootTreeNodes;
-
 	/** UI callbacks */
 	void GenerateMutableGraphPressed();
 	void CompileMutableCodePressed();
@@ -74,59 +66,4 @@ private:
 	void OnChangeCompileOptimizationLevel(TSharedPtr<FString> NewSelection, ESelectInfo::Type SelectInfo);
 	void OnChangeCompileTextureCompressionType(TSharedPtr<FString> NewSelection, ESelectInfo::Type);
 	void OnChangeDebugPlatform(TSharedPtr<FString> NewSelection, ESelectInfo::Type SelectInfo);
-
-	/** Tree callbacks. */
-	TSharedRef<ITableRow> GenerateRowForNodeTree(TSharedPtr<FMutableObjectTreeElement> InTreeNode, const TSharedRef<STableViewBase>& InOwnerTable);
-	void GetChildrenForInfo(TSharedPtr<FMutableObjectTreeElement> InInfo, TArray< TSharedPtr<FMutableObjectTreeElement> >& OutChildren);
-
-};
-
-
-/** An row of the code tree in . */
-class FMutableObjectTreeElement : public TSharedFromThis<FMutableObjectTreeElement>
-{
-public:
-
-	enum class EType
-	{
-		None,
-		SectionCaption,
-		Name,
-		ChildObject
-	};
-
-	enum class ESection
-	{
-		None,
-		General,
-		ChildObjects
-	};
-
-	/** Constructor for section rows */
-	FMutableObjectTreeElement(ESection InSection, UCustomizableObject* InObject)
-	{
-		Type = EType::SectionCaption;
-		Section = InSection;
-		CustomizableObject = InObject;
-	}
-
-	/** Constructor for generic property rows. */
-	FMutableObjectTreeElement(EType InType, ESection InSection, UCustomizableObject* InObject)
-	{
-		Type = InType;
-		Section = InSection;
-		CustomizableObject = InObject;
-	}
-
-public:
-
-	/** Row type. */
-	EType Type = EType::None;
-
-	/** If the row is a section row. */
-	ESection Section = ESection::None;
-
-	/** */
-	UCustomizableObject* CustomizableObject = nullptr;
-
 };
diff --git a/Engine/Plugins/Experimental/Mutable/Source/CustomizableObjectEditor/Public/MuCOE/CustomizableObjectEditorModule.h b/Engine/Plugins/Experimental/Mutable/Source/CustomizableObjectEditor/Public/MuCOE/CustomizableObjectEditorModule.h
--- a/Engine/Plugins/Experimental/Mutable/Source/CustomizableObjectEditor/Public/MuCOE/CustomizableObjectEditorModule.h
+++ b/Engine/Plugins/Experimental/Mutable/Source/CustomizableObjectEditor/Public/MuCOE/CustomizableObjectEditorModule.h
@@ -36,7 +36,6 @@ public:
 	 */
 	virtual TSharedRef<ICustomizableObjectEditor> CreateCustomizableObjectEditor( const EToolkitMode::Type Mode, const TSharedPtr< IToolkitHost >& InitToolkitHost, UCustomizableObject* CustomizableObject ) = 0;
 	virtual TSharedRef<ICustomizableObjectInstanceEditor> CreateCustomizableObjectInstanceEditor(const EToolkitMode::Type Mode, const TSharedPtr< IToolkitHost >& InitToolkitHost, UCustomizableObjectInstance* CustomizableObjectInstance) = 0;
-	virtual TSharedRef<ICustomizableObjectDebugger> CreateCustomizableObjectDebugger(const EToolkitMode::Type Mode, const TSharedPtr< IToolkitHost >& InitToolkitHost, UCustomizableObject* CustomizableObjectInstance) = 0;
 	
 	virtual TSharedPtr<class FExtensibilityManager> GetCustomizableObjectEditorToolBarExtensibilityManager() { return nullptr; }
 	virtual TSharedPtr<class FExtensibilityManager> GetCustomizableObjectEditorMenuExtensibilityManager() { return nullptr; }
