[mutable] COValidation : Added State testing
- For each of the random instnaces we generate (random parameter set) we will test said instance but with all available states.

diff --git a/Engine/Plugins/Experimental/Mutable/Source/MutableValidation/Private/MuV/CustomizableObjectValidationCommandlet.cpp b/Engine/Plugins/Experimental/Mutable/Source/MutableValidation/Private/MuV/CustomizableObjectValidationCommandlet.cpp
--- a/Engine/Plugins/Experimental/Mutable/Source/MutableValidation/Private/MuV/CustomizableObjectValidationCommandlet.cpp
+++ b/Engine/Plugins/Experimental/Mutable/Source/MutableValidation/Private/MuV/CustomizableObjectValidationCommandlet.cpp
@@ -30,13 +30,12 @@ int32 UCustomizableObjectValidationCommandlet::Main(const FString& Params)
 		return 1;
 	}
 	
-	// Get the amount of instances to generate if parameter was provided
+	// Get the amount of instances to generate if parameter was provided (it will get multiplied by the amount of states later so this is a minimun value)
 	uint32 InstancesToGenerate = 16;
 	if (!FParse::Value(*Params, TEXT("InstanceGenerationCount="),InstancesToGenerate))
 	{
 		UE_LOG(LogMutable,Display,TEXT("Instance generation count not specified. Using default value : %u"),InstancesToGenerate);
 	}
-	UE_LOG(LogMutable, Log,TEXT("(int) instances_to_generate_count : %u "), InstancesToGenerate);
 	
 	// Load the resource
 	UObject* FoundObject = FSoftObjectPath(CustomizableObjectAssetPath).TryLoad();
@@ -89,7 +88,6 @@ int32 UCustomizableObjectValidationCommandlet::Main(const FString& Params)
 
 	UE_LOG(LogMutable,Display,TEXT("Customizable Object was compiled succesfully."));
 	
-	
 	// GHet the total size of the streaming data of the model ---------------------------------------------- //
 	{
 		const TSharedPtr<const mu::Model> MutableModel = ToTestCustomizableObject->GetPrivate()->GetModel();
@@ -139,8 +137,17 @@ int32 UCustomizableObjectValidationCommandlet::Main(const FString& Params)
 	// Generate target random instances to be tested ------------------------------------------------------------ //
 	bool bWasInstancesCreationSuccessful = true;
 	{
-		UE_LOG(LogMutable,Display,TEXT("Generating %i random instances..."), InstancesToGenerate);	
+		// Test this parameter configuration in all the states of the CO
+		const uint32 StateCount = ToTestCustomizableObject->GetStateCount();
+		check(StateCount >= 1);
 
+		UE_LOG(LogMutable, Display, TEXT("Requeted Instances Count : %i"), InstancesToGenerate);
+		UE_LOG(LogMutable, Display, TEXT("State Count = %i"), StateCount);
+	
+		// Compute actual total amount of instances to generate
+		const uint32 TotalInstancesToTestCount = InstancesToGenerate * StateCount;
+		UE_LOG(LogMutable,Display,TEXT("Generating %i instances (states * requested instances)..."), TotalInstancesToTestCount);
+		
 		// Create randomization stream for the parameters of the instance
 		FRandomStream RandomizationStream = FRandomStream(0);
 		
@@ -157,7 +164,13 @@ int32 UCustomizableObjectValidationCommandlet::Main(const FString& Params)
 				
 				// Randomize instance values
 				GeneratedInstance->SetRandomValuesFromStream(RandomizationStream);
-				InstancesToProcess.Push(GeneratedInstance);
+				
+				for (uint32 State = 0; State < StateCount; State++)
+				{
+					// Set the state for the instance and store it for later update.
+					GeneratedInstance->GetPrivate()->SetState(State);
+                	InstancesToProcess.Push(GeneratedInstance->Clone());
+				}
 			}
 			else
 			{
@@ -167,6 +180,8 @@ int32 UCustomizableObjectValidationCommandlet::Main(const FString& Params)
 		}
 	}
 	// ---------------------------------------------------------------------------------------------------------- //
+
+	UE_LOG(LogMutable, Log,TEXT("(int) generated_instances_count : %u "), InstancesToProcess.Num());
 	
 	// Update the instances generated --------------------------------------------------------------------------- //
 	UE_LOG(LogMutable,Display,TEXT("Updating generated instances..."));
@@ -188,8 +203,8 @@ int32 UCustomizableObjectValidationCommandlet::Main(const FString& Params)
 	const double CombinedInstanceUpdateSeconds = InstancesUpdateEndSeconds - InstancesUpdateStartSeconds;
 	UE_LOG(LogMutable, Log,TEXT("(double) combined_update_time_ms : %f "), CombinedInstanceUpdateSeconds * 1000);
 
-	check(InstancesToGenerate > 0);
-	const double AverageInstanceUpdateSeconds = CombinedInstanceUpdateSeconds / InstancesToGenerate;
+	check(InstancesToProcess.Num() > 0);
+	const double AverageInstanceUpdateSeconds = CombinedInstanceUpdateSeconds / InstancesToProcess.Num();
 	UE_LOG(LogMutable, Log,TEXT("(double) avg_update_time_ms : %f "), AverageInstanceUpdateSeconds * 1000);
 
 	UE_LOG(LogMutable,Display,TEXT("Generation of Customizable object instances took %f seconds (%f seconds avg)."), CombinedInstanceUpdateSeconds, AverageInstanceUpdateSeconds);
