[mutable] Synced floor and sky preview settings and their equivalent buttons in the the CO Editor Preview.
#rnx

#tests Editor

diff --git a/Engine/Plugins/Experimental/Mutable/Source/CustomizableObjectEditor/Private/MuCOE/CustomizableObjectEditorViewportClient.cpp b/Engine/Plugins/Experimental/Mutable/Source/CustomizableObjectEditor/Private/MuCOE/CustomizableObjectEditorViewportClient.cpp
--- a/Engine/Plugins/Experimental/Mutable/Source/CustomizableObjectEditor/Private/MuCOE/CustomizableObjectEditorViewportClient.cpp
+++ b/Engine/Plugins/Experimental/Mutable/Source/CustomizableObjectEditor/Private/MuCOE/CustomizableObjectEditorViewportClient.cpp
@@ -98,7 +98,8 @@ FCustomizableObjectEditorViewportClient::FCustomizableObjectEditorViewportClient
 	DrawHelper.GridColorMinor = FColor(128, 128, 128);
 	DrawHelper.PerspectiveGridSize = 2048.0f;
 	DrawHelper.NumCells = DrawHelper.PerspectiveGridSize / (32);
-	SetShowGrid();
+	UpdateShowGrid(true);
+	UpdateShowSky(true);
 
 	SetViewMode(VMI_Lit);
 
@@ -918,7 +919,7 @@ void FCustomizableObjectEditorViewportClient::SetDrawUVOverlay()
 }
 
 
-void FCustomizableObjectEditorViewportClient::SetDrawUVOverlayMaterial(const FString& MaterialName, FString UVChannel)
+void FCustomizableObjectEditorViewportClient::SetDrawUVOverlayMaterial(const FString& MaterialName, const FString& UVChannel)
 {
 	// Get LOD Index
 	FString NameWithLOD, ComponentString;
@@ -955,9 +956,32 @@ bool FCustomizableObjectEditorViewportClient::IsSetDrawUVOverlayChecked() const
 	return bDrawUVs;
 }
 
-void FCustomizableObjectEditorViewportClient::SetShowGrid()
+void FCustomizableObjectEditorViewportClient::UpdateShowGrid(bool bKeepOldValue)
 {
-	DrawHelper.bDrawGrid = !DrawHelper.bDrawGrid;
+	UAssetViewerSettings* Settings = UAssetViewerSettings::Get();
+	int32 ProfileIndex = GetMutableDefault<UEditorPerProjectUserSettings>()->AssetViewerProfileIndex;
+
+	bool bNewShowGridValue = true;
+
+	if (Settings->Profiles.IsValidIndex(ProfileIndex))
+	{
+		bool bOldShowGridValue = Settings->Profiles[ProfileIndex].bShowFloor;
+
+		if (bKeepOldValue)
+		{
+			// Do not toggle the value when the viewport is being constructed
+			bNewShowGridValue = bOldShowGridValue;
+		}
+		else
+		{
+			// Toggle it when actually changing the option
+			bNewShowGridValue = !bOldShowGridValue;
+		}
+
+		Settings->Profiles[ProfileIndex].bShowFloor = bNewShowGridValue;
+	}
+	
+	DrawHelper.bDrawGrid = bNewShowGridValue;
 
 	FAdvancedPreviewScene* AdvancedScene = static_cast<FAdvancedPreviewScene*>(PreviewScene);
 	if (AdvancedScene != nullptr)
@@ -970,20 +994,49 @@ void FCustomizableObjectEditorViewportClient::SetShowGrid()
 	Invalidate();
 }
 
-bool FCustomizableObjectEditorViewportClient::IsSetShowGridChecked() const
+void FCustomizableObjectEditorViewportClient::UpdateShowGridFromButton()
+{
+	UpdateShowGrid(false);
+}
+
+bool FCustomizableObjectEditorViewportClient::IsShowGridChecked() const
 {
 	return DrawHelper.bDrawGrid;
 }
 
-void FCustomizableObjectEditorViewportClient::SetShowSky()
+void FCustomizableObjectEditorViewportClient::UpdateShowSky(bool bKeepOldValue)
 {
-	bDrawSky = !bDrawSky;
+	UAssetViewerSettings* Settings = UAssetViewerSettings::Get();
+	int32 ProfileIndex = GetMutableDefault<UEditorPerProjectUserSettings>()->AssetViewerProfileIndex;
+
+	if (Settings->Profiles.IsValidIndex(ProfileIndex))
+	{
+		bool bOldDrawSky = Settings->Profiles[ProfileIndex].bShowEnvironment;
+
+		if (bKeepOldValue)
+		{
+			bDrawSky = bOldDrawSky;
+		}
+		else
+		{
+			bDrawSky = !bOldDrawSky;
+		}
+
+		Settings->Profiles[ProfileIndex].bShowEnvironment = bDrawSky;
+	}
+
 	FAdvancedPreviewScene* PreviewSceneCasted = static_cast<FAdvancedPreviewScene*>(PreviewScene);
 	PreviewSceneCasted->SetEnvironmentVisibility(bDrawSky, true);
+
 	Invalidate();
 }
 
-bool FCustomizableObjectEditorViewportClient::IsSetShowSkyChecked() const
+void FCustomizableObjectEditorViewportClient::UpdateShowSkyFromButton()
+{
+	UpdateShowSky(false);
+}
+
+bool FCustomizableObjectEditorViewportClient::IsShowSkyChecked() const
 {
 	return bDrawSky;
 }
@@ -1781,15 +1834,44 @@ void FCustomizableObjectEditorViewportClient::SetAdvancedShowFlagsForScene(const
 
 void FCustomizableObjectEditorViewportClient::OnAssetViewerSettingsChanged(const FName& InPropertyName)
 {
+	UAssetViewerSettings* Settings = UAssetViewerSettings::Get();
+	int32 ProfileIndex = GetMutableDefault<UEditorPerProjectUserSettings>()->AssetViewerProfileIndex;
+
 	if (InPropertyName == GET_MEMBER_NAME_CHECKED(FPreviewSceneProfile, bPostProcessingEnabled) || InPropertyName == NAME_None)
 	{
-		UAssetViewerSettings* Settings = UAssetViewerSettings::Get();
-		int32 ProfileIndex = GetMutableDefault<UEditorPerProjectUserSettings>()->AssetViewerProfileIndex;
 		if (Settings->Profiles.IsValidIndex(ProfileIndex))
 		{
 			SetAdvancedShowFlagsForScene(Settings->Profiles[ProfileIndex].bPostProcessingEnabled);
 		}
 	}
+
+	else if (InPropertyName == GET_MEMBER_NAME_CHECKED(FPreviewSceneProfile, bShowEnvironment))
+	{
+		if (Settings->Profiles.IsValidIndex(ProfileIndex))
+		{
+			bDrawSky = Settings->Profiles[ProfileIndex].bShowEnvironment;
+		}
+		else
+		{
+			bDrawSky = !bDrawSky;
+		}
+	}
+
+	else if (InPropertyName == GET_MEMBER_NAME_CHECKED(FPreviewSceneProfile, bShowFloor))
+	{
+		if (Settings->Profiles.IsValidIndex(ProfileIndex))
+		{
+			DrawHelper.bDrawGrid = Settings->Profiles[ProfileIndex].bShowFloor;
+		}
+		else
+		{
+			DrawHelper.bDrawGrid = !DrawHelper.bDrawGrid;
+		}
+
+		EngineShowFlags.Grid = DrawHelper.bDrawGrid;
+
+		Invalidate();
+	}
 }
 
 
diff --git a/Engine/Plugins/Experimental/Mutable/Source/CustomizableObjectEditor/Private/MuCOE/SCustomizableObjectEditorViewport.cpp b/Engine/Plugins/Experimental/Mutable/Source/CustomizableObjectEditor/Private/MuCOE/SCustomizableObjectEditorViewport.cpp
--- a/Engine/Plugins/Experimental/Mutable/Source/CustomizableObjectEditor/Private/MuCOE/SCustomizableObjectEditorViewport.cpp
+++ b/Engine/Plugins/Experimental/Mutable/Source/CustomizableObjectEditor/Private/MuCOE/SCustomizableObjectEditorViewport.cpp
@@ -353,15 +353,15 @@ void SCustomizableObjectEditorViewportTabBody::BindCommands()
 
 	CommandList.MapAction(
 		Commands.SetShowGrid,
-		FExecuteAction::CreateSP( EditorViewportClientRef, &FCustomizableObjectEditorViewportClient::SetShowGrid ),
+		FExecuteAction::CreateSP( EditorViewportClientRef, &FCustomizableObjectEditorViewportClient::UpdateShowGridFromButton),
 		FCanExecuteAction(),
-		FIsActionChecked::CreateSP( EditorViewportClientRef, &FCustomizableObjectEditorViewportClient::IsSetShowGridChecked ) );
+		FIsActionChecked::CreateSP( EditorViewportClientRef, &FCustomizableObjectEditorViewportClient::IsShowGridChecked ) );
 
 	CommandList.MapAction(
 		Commands.SetShowSky,
-		FExecuteAction::CreateSP(EditorViewportClientRef, &FCustomizableObjectEditorViewportClient::SetShowSky),
+		FExecuteAction::CreateSP(EditorViewportClientRef, &FCustomizableObjectEditorViewportClient::UpdateShowSkyFromButton),
 		FCanExecuteAction(),
-		FIsActionChecked::CreateSP(EditorViewportClientRef, &FCustomizableObjectEditorViewportClient::IsSetShowSkyChecked));
+		FIsActionChecked::CreateSP(EditorViewportClientRef, &FCustomizableObjectEditorViewportClient::IsShowSkyChecked));
 
 	CommandList.MapAction(
 		Commands.SetShowBounds,
diff --git a/Engine/Plugins/Experimental/Mutable/Source/CustomizableObjectEditor/Private/MuCOE/SCustomizableObjectHighresScreenshot.cpp b/Engine/Plugins/Experimental/Mutable/Source/CustomizableObjectEditor/Private/MuCOE/SCustomizableObjectHighresScreenshot.cpp
--- a/Engine/Plugins/Experimental/Mutable/Source/CustomizableObjectEditor/Private/MuCOE/SCustomizableObjectHighresScreenshot.cpp
+++ b/Engine/Plugins/Experimental/Mutable/Source/CustomizableObjectEditor/Private/MuCOE/SCustomizableObjectHighresScreenshot.cpp
@@ -267,7 +267,7 @@ void SCustomizableObjectHighresScreenshot::OnShowFloorChanged(ECheckBoxState New
 			LevelViewportClient.Pin()->SetFloorVisibility(false);
 			if (GridVisibilityCached)
 			{
-				LevelViewportClient.Pin()->SetShowGrid();
+				LevelViewportClient.Pin()->UpdateShowGrid(true);
 			}
 		}
 		else
@@ -280,7 +280,7 @@ void SCustomizableObjectHighresScreenshot::OnShowFloorChanged(ECheckBoxState New
 			}
 
 			LevelViewportClient.Pin()->SetFloorVisibility(FloorVisibilityCached);
-			LevelViewportClient.Pin()->SetShowGrid();
+			LevelViewportClient.Pin()->UpdateShowGrid(true);
 		}
 
 		LevelViewportClient.Pin()->SetEnvironmentMeshVisibility(NewValue == ECheckBoxState::Checked ? 1 : 0);
@@ -304,7 +304,7 @@ void SCustomizableObjectHighresScreenshot::RestoreViewportValues()
 		// Restore floor and environment mesh visibility
 		if (!LevelViewportClient.Pin()->GetEnvironmentMeshVisibility())
 		{
-			LevelViewportClient.Pin()->SetShowGrid();
+			LevelViewportClient.Pin()->UpdateShowGrid(false);
 			LevelViewportClient.Pin()->SetEnvironmentMeshVisibility(1);
 		}
 	}
diff --git a/Engine/Plugins/Experimental/Mutable/Source/CustomizableObjectEditor/Public/MuCOE/CustomizableObjectEditorViewportClient.h b/Engine/Plugins/Experimental/Mutable/Source/CustomizableObjectEditor/Public/MuCOE/CustomizableObjectEditorViewportClient.h
--- a/Engine/Plugins/Experimental/Mutable/Source/CustomizableObjectEditor/Public/MuCOE/CustomizableObjectEditorViewportClient.h
+++ b/Engine/Plugins/Experimental/Mutable/Source/CustomizableObjectEditor/Public/MuCOE/CustomizableObjectEditorViewportClient.h
@@ -139,19 +139,31 @@ public:
 	bool IsSetDrawUVOverlayChecked() const;
 
 	// Set the material index whose UVs will be drawn.
-	void SetDrawUVOverlayMaterial(const FString& MaterialName, FString UVChannel);
+	void SetDrawUVOverlayMaterial(const FString& MaterialName, const FString& UVChannel);
 
 	/** Callback for toggling the grid show flag. */
-	void SetShowGrid();
+	void UpdateShowGridFromButton();
+
+	/** 
+	 * Updates the visual state of the ShowGrid Button and CheckBox in the Preview Settings
+	 * @param bKeepOldValue - If true, it will keep the visibility of the grid and floor, otherwise it will invert it
+	*/
+	void UpdateShowGrid(bool bKeepOldValue);
 	
 	/** Callback for checking the grid show flag. */
-	bool IsSetShowGridChecked() const;
+	bool IsShowGridChecked() const;
 
 	/** Callback for toggling the sky show flag. */
-	void SetShowSky();
+	void UpdateShowSkyFromButton();
+
+	/**
+	 * Updates the visual state of the ShowSky Button and CheckBox in the Preview Settings
+	 * @param bKeepOldValue - If true, it will keep the visibility of the sky end environment, otherwise it will invert it
+	*/
+	void UpdateShowSky(bool bKeepOldValue);
 
 	/** Callback for checking the sky show flag. */
-	bool IsSetShowSkyChecked() const;
+	bool IsShowSkyChecked() const;
 
 	/** Callback for toggling the bounds show flag. */
 	void SetShowBounds();
