[mutable] Added logs for mutable pre and post compilation mutable memory usage.

diff --git a/Engine/Plugins/Experimental/Mutable/Source/MutableValidation/Private/MuV/ValidationUtils.cpp b/Engine/Plugins/Experimental/Mutable/Source/MutableValidation/Private/MuV/ValidationUtils.cpp
--- a/Engine/Plugins/Experimental/Mutable/Source/MutableValidation/Private/MuV/ValidationUtils.cpp
+++ b/Engine/Plugins/Experimental/Mutable/Source/MutableValidation/Private/MuV/ValidationUtils.cpp
@@ -56,6 +56,11 @@ bool CompileCustomizableObject(UCustomizableObject* InCustomizableObject, const
 	
     // Compile the CO with the selected compilation options
     // Run Sync compilation -> Warning : Potentially long operation -------------
+
+	// Get the memory being used by mutable before the compilation
+	const int64 CompilationStartBytes = mu::FGlobalMemoryCounter::GetCounter();
+	mu::FGlobalMemoryCounter::Zero();
+	
 	UE_LOG(LogMutable,Display,TEXT("Compiling Customizable Object..."));
 	const double CompilationStartSeconds = FPlatformTime::Seconds();
 	{
@@ -64,8 +69,11 @@ bool CompileCustomizableObject(UCustomizableObject* InCustomizableObject, const
 	const double CompilationEndSeconds = FPlatformTime::Seconds() - CompilationStartSeconds;
     // --------------------------------------------------------------------------
 	UE_LOG(LogMutable,Display,TEXT("Compilation of CO completed!"));
-
 	UE_LOG(LogMutable, Display, TEXT("The compilation of the %s model took %f seconds."), *InCustomizableObject->GetName(), CompilationEndSeconds);
+
+	// Get the peak mutable memory used during the compilation operation
+    const int64 CompilationEndPeakBytes = mu::FGlobalMemoryCounter::GetPeak();
+    const int64 CompilationEndRealPeakBytes = CompilationStartBytes + CompilationEndPeakBytes;
 	
     // Get the compilation result
     const ECustomizableObjectCompilationState CompilationEndResult = Compiler->GetCompilationState();
@@ -94,6 +102,10 @@ bool CompileCustomizableObject(UCustomizableObject* InCustomizableObject, const
 		UE_LOG(LogMutable, Log, TEXT("(double) model_compile_time_ms : %f "), CompilationEndSeconds * 1000);
 		UE_LOG(LogMutable, Log, TEXT("(string) model_compile_end_state : %s "), *UEnum::GetValueAsString(CompilationEndResult));
 
+		UE_LOG(LogMutable, Log, TEXT("(int) model_compilation_start_bytes : %lld "), CompilationStartBytes);
+		UE_LOG(LogMutable, Log, TEXT("(int) model_compilation_end_peak_bytes : %lld "), CompilationEndPeakBytes);
+		UE_LOG(LogMutable, Log, TEXT("(int) model_compilation_end_real_peak_bytes : %lld "), CompilationEndRealPeakBytes);
+
 		// TODO: Add logs for the other relevant configs of the model being compiled
 	}
 
