[Mutable] Fix Mutable deadlocks when using MRQ.

Replaced the UCustomizableObjectSystem ticker with a IStreamManager.

Hotfixable with config (CVar does not work due to those are read after the module has already been initialized).

#rnx

diff --git a/Engine/Plugins/Experimental/Mutable/Source/CustomizableObject/Classes/MuCO/CustomizableObjectSystem.h b/Engine/Plugins/Experimental/Mutable/Source/CustomizableObject/Classes/MuCO/CustomizableObjectSystem.h
--- a/Engine/Plugins/Experimental/Mutable/Source/CustomizableObject/Classes/MuCO/CustomizableObjectSystem.h
+++ b/Engine/Plugins/Experimental/Mutable/Source/CustomizableObject/Classes/MuCO/CustomizableObjectSystem.h
@@ -341,6 +341,9 @@ private:
 	// Most of the work in this plugin happens here.
 	bool Tick(float DeltaTime);
 
+	/** Returns the number of remaining operations. */
+	int32 TickInternal(float DeltaTime);
+
 	// If there is an on-going operation, advance it.
 	void AdvanceCurrentOperation();
 
diff --git a/Engine/Plugins/Experimental/Mutable/Source/CustomizableObject/Internal/MuCO/CustomizableObjectSystemPrivate.h b/Engine/Plugins/Experimental/Mutable/Source/CustomizableObject/Internal/MuCO/CustomizableObjectSystemPrivate.h
--- a/Engine/Plugins/Experimental/Mutable/Source/CustomizableObject/Internal/MuCO/CustomizableObjectSystemPrivate.h
+++ b/Engine/Plugins/Experimental/Mutable/Source/CustomizableObject/Internal/MuCO/CustomizableObjectSystemPrivate.h
@@ -23,6 +23,7 @@
 #include "Framework/Notifications/NotificationManager.h"
 #include "MuCO/FMutableTaskGraph.h"
 #include "AssetRegistry/AssetData.h"
+#include "ContentStreaming.h"
 
 // This define could come from MuR/System.h
 #ifdef MUTABLE_USE_NEW_TASKGRAPH
@@ -611,7 +612,7 @@ struct FPendingReleaseSkeletalMeshInfo
 
 
 UCLASS()
-class UCustomizableObjectSystemPrivate : public UObject
+class UCustomizableObjectSystemPrivate : public UObject, public IStreamingManager
 {
 	GENERATED_BODY()
 	
@@ -663,6 +664,18 @@ public:
 	mu::FImageOperator::FImagePixelFormatFunc ImageFormatOverrideFunc;
 #endif
 
+	// IStreamingManager interface
+	virtual void UpdateResourceStreaming(float DeltaTime, bool bProcessEverything = false) override;
+	virtual int32 BlockTillAllRequestsFinished(float TimeLimit = 0.0f, bool bLogResults = false) override;
+	virtual void CancelForcedResources() override {}
+	virtual void AddLevel(ULevel* Level) override {}
+	virtual void RemoveLevel(ULevel* Level) override {}
+	virtual void NotifyLevelChange() override {}
+	virtual void SetDisregardWorldResourcesForFrames(int32 NumFrames) override {}
+	virtual void NotifyLevelOffset(ULevel* Level, const FVector& Offset) override {}
+
+	UCustomizableObjectSystem* GetPublic() const;
+	
 	void AddGameThreadTask(const FMutableTask& Task);
 
 	// Remove references to cached objects that have been deleted in the unreal
diff --git a/Engine/Plugins/Experimental/Mutable/Source/CustomizableObject/Internal/MuCO/FMutableTaskGraph.h b/Engine/Plugins/Experimental/Mutable/Source/CustomizableObject/Internal/MuCO/FMutableTaskGraph.h
--- a/Engine/Plugins/Experimental/Mutable/Source/CustomizableObject/Internal/MuCO/FMutableTaskGraph.h
+++ b/Engine/Plugins/Experimental/Mutable/Source/CustomizableObject/Internal/MuCO/FMutableTaskGraph.h
@@ -65,8 +65,9 @@ public:
 	/** Allow or disallow launching Mutable Tasks with Low priority.
 	 @param bFromMutableTask true if called from a Mutable Task. */
 	void AllowLaunchingMutableTaskLowPriority(bool bAllow, bool bFromMutableTask);
-	
-	void Tick();
+
+	/** Returns the number of remaining tasks. */
+	int32 Tick();
 	
 private:
 	/** A Mutable Task Low Priority will only be launched if:
diff --git a/Engine/Plugins/Experimental/Mutable/Source/CustomizableObject/Private/MuCO/CustomizableObjectSettings.h b/Engine/Plugins/Experimental/Mutable/Source/CustomizableObject/Private/MuCO/CustomizableObjectSettings.h
new file mode 100644
--- /dev/null
+++ b/Engine/Plugins/Experimental/Mutable/Source/CustomizableObject/Private/MuCO/CustomizableObjectSettings.h
@@ -0,0 +1,21 @@
+// Copyright Epic Games, Inc. All Rights Reserved.
+
+#pragma once
+
+#include "CustomizableObjectSettings.generated.h"
+
+
+/** CustomizableObject module settings.
+  *
+  * These settings can also be changed using command line arguments. For example: 
+  * -ini:Engine:[/Script/CustomizableObject.CustomizableObjectSettings]:bEnableStreamingManager=true */
+UCLASS(Config = Engine)
+class UCustomizableObjectSettings : public UObject
+{
+	GENERATED_BODY()
+
+public:
+	/** If true, use the new StreamManager ticker. If false, use the old FTSTicker. */
+	UPROPERTY(Config)
+	bool bEnableStreamingManager = true;
+};
diff --git a/Engine/Plugins/Experimental/Mutable/Source/CustomizableObject/Private/MuCO/CustomizableObjectSystem.cpp b/Engine/Plugins/Experimental/Mutable/Source/CustomizableObject/Private/MuCO/CustomizableObjectSystem.cpp
--- a/Engine/Plugins/Experimental/Mutable/Source/CustomizableObject/Private/MuCO/CustomizableObjectSystem.cpp
+++ b/Engine/Plugins/Experimental/Mutable/Source/CustomizableObject/Private/MuCO/CustomizableObjectSystem.cpp
@@ -32,6 +32,7 @@
 #include "MuCO/EditorImageProvider.h"
 #include "Components/SkeletalMeshComponent.h"
 #include "MuCO/CustomizableObjectSystemPrivate.h"
+#include "CustomizableObjectSettings.h"
 
 #if WITH_EDITOR
 #include "Editor.h"
@@ -599,7 +600,14 @@ void UCustomizableObjectSystem::BeginDestroy()
 #endif
 
 #if !UE_SERVER
-		FTSTicker::GetCoreTicker().RemoveTicker(Private->TickDelegateHandle);
+		if (GetMutableDefault<UCustomizableObjectSettings>()->bEnableStreamingManager)
+		{
+			FStreamingManagerCollection::Get().RemoveStreamingManager(GetPrivate());
+		}
+		else
+		{
+			FTSTicker::GetCoreTicker().RemoveTicker(Private->TickDelegateHandle);			
+		}
 #endif // !UE_SERVER
 
 		// Discard pending game thread tasks
@@ -665,6 +673,15 @@ static FAutoConsoleVariableRef CVarEnableMutableAnimInfoDebugging(
 	ECVF_Default);
 
 
+UCustomizableObjectSystem* UCustomizableObjectSystemPrivate::GetPublic() const
+{
+	UCustomizableObjectSystem* Public = StaticCast<UCustomizableObjectSystem*>(GetOuter());
+	check(Public);
+
+	return Public;
+}
+
+
 void UCustomizableObjectSystemPrivate::AddGameThreadTask(const FMutableTask& Task)
 {
 	PendingTasks.Enqueue(Task);
@@ -2911,16 +2928,23 @@ void UCustomizableObjectSystem::AdvanceCurrentOperation()
 
 bool UCustomizableObjectSystem::Tick(float DeltaTime)
 {
-	MUTABLE_CPUPROFILER_SCOPE(UCustomizableObjectSystem::Tick)
+	TickInternal(DeltaTime);
+	return true;
+}
+
+
+int32 UCustomizableObjectSystem::TickInternal(float DeltaTime)
+{
+	MUTABLE_CPUPROFILER_SCOPE(UCustomizableObjectSystem::TickInternal)
 	
 	// Building instances is not enabled in servers. If at some point relevant collision or animation data is necessary for server logic this will need to be changed.
 #if UE_SERVER
-	return true;
+	return 0;
 #endif
 
 	if (!Private)
 	{
-		return true;
+		return 0;
 	}
 
 	if (GWorld)
@@ -2929,7 +2953,7 @@ bool UCustomizableObjectSystem::Tick(float DeltaTime)
 
 		if (WorldType != EWorldType::PIE && WorldType != EWorldType::Game && WorldType != EWorldType::Editor && WorldType != EWorldType::GamePreview)
 		{
-			return true;
+			return 0;
 		}
 	}
 
@@ -2939,17 +2963,19 @@ bool UCustomizableObjectSystem::Tick(float DeltaTime)
 
 	if (AssetRegistryModule.Get().IsLoadingAssets())
 	{
-		return true; // Assets are still being loaded, so subobjects won't be found, compiled objects incomplete and thus updates wrong
+		return 1; // Assets are still being loaded, so subobjects won't be found, compiled objects incomplete and thus updates wrong
 	}
 
 	// Do not tick if the CookCommandlet is running.
 	if (IsRunningCookCommandlet())
 	{
-		return true;
+		return 0;
 	}
 #endif
 
 	Private->UpdateStats();
+
+	FMutableUpdateCandidate* LODUpdateCandidateFound = nullptr;
 	
 	// Get a new operation if we aren't working on one
 	if (!Private->CurrentMutableOperation && bIsMutableEnabled)
@@ -2994,7 +3020,6 @@ bool UCustomizableObjectSystem::Tick(float DeltaTime)
 			double MaxSquareDistanceFound = TNumericLimits<double>::Max();
 			double MinTimeFound = TNumericLimits<double>::Max();
 			const FMutablePendingInstanceUpdate* PendingInstanceUpdateFound = nullptr;
-			FMutableUpdateCandidate* LODUpdateCandidateFound = nullptr;
 
 			// Look for the highest priority Pending Update
 			for (auto Iterator = Private->MutablePendingInstanceWork.GetUpdateIterator(); Iterator; ++Iterator)
@@ -3129,18 +3154,28 @@ bool UCustomizableObjectSystem::Tick(float DeltaTime)
 	UCustomizableObjectSystemPrivate::ShowOnScreenCompileWarnings();
 #endif
 	
-	Private->MutableTaskGraph.Tick();
+	const int32 RemainingTasks = Private->MutableTaskGraph.Tick();
 
 	Private->LogBenchmarkUtil.UpdateStats(); // Must to be the last thing to perform
 
 	if (!bIsMutableEnabled && !Private->CurrentMutableOperation)
 	{
-		// Mutable has been disabled. Unregister the ticker if there is no CurrentMutableOperation.
-		FTSTicker::GetCoreTicker().RemoveTicker(Private->TickDelegateHandle);
-		Private->TickDelegateHandle.Reset();
+		if (GetMutableDefault<UCustomizableObjectSettings>()->bEnableStreamingManager)
+		{
+			FStreamingManagerCollection::Get().RemoveStreamingManager(GetPrivate());
+		}
+		else
+		{
+			// Mutable has been disabled. Unregister the ticker if there is no CurrentMutableOperation.
+			FTSTicker::GetCoreTicker().RemoveTicker(Private->TickDelegateHandle);
+			Private->TickDelegateHandle.Reset();
+		}
 	}
 
-	return true;
+	return Private->CurrentMutableOperation.IsValid() + 
+		Private->MutablePendingInstanceWork.Num() +
+		!!LODUpdateCandidateFound + // Still a pending LOD update. We can not use the size of RequestedLODUpdates since not all requests valid in future ticks.
+		RemainingTasks;
 }
 
 
@@ -3782,10 +3817,18 @@ void UCustomizableObjectSystemPrivate::OnMutableEnabledChanged(IConsoleVariable*
 	if (bIsMutableEnabled)
 	{
 #if !UE_SERVER
-		if (!SystemPrivate->TickDelegateHandle.IsValid())
+		if (GetMutableDefault<UCustomizableObjectSettings>()->bEnableStreamingManager)
+		{
+			FStreamingManagerCollection::Get().RemoveStreamingManager(SystemPrivate); // Avoid being added twice
+			FStreamingManagerCollection::Get().AddStreamingManager(SystemPrivate);
+		}
+		else
 		{
-			SystemPrivate->TickDelegate = FTickerDelegate::CreateUObject(System, &UCustomizableObjectSystem::Tick);
-			SystemPrivate->TickDelegateHandle = FTSTicker::GetCoreTicker().AddTicker(SystemPrivate->TickDelegate, 0.f);
+			if (!SystemPrivate->TickDelegateHandle.IsValid())
+			{
+				SystemPrivate->TickDelegate = FTickerDelegate::CreateUObject(System, &UCustomizableObjectSystem::Tick);
+				SystemPrivate->TickDelegateHandle = FTSTicker::GetCoreTicker().AddTicker(SystemPrivate->TickDelegate, 0.f);
+			}			
 		}
 #endif // !UE_SERVER
 
@@ -3860,3 +3903,47 @@ bool UCustomizableObjectSystem::IsMutableAnimInfoDebuggingEnabled() const
 	return false;
 #endif
 }
+
+
+void UCustomizableObjectSystemPrivate::UpdateResourceStreaming(float DeltaTime, bool bProcessEverything)
+{
+	GetPublic()->TickInternal(DeltaTime);
+}
+
+
+int32 UCustomizableObjectSystemPrivate::BlockTillAllRequestsFinished(float TimeLimit, bool bLogResults)
+{
+	const double BlockEndTime = FPlatformTime::Seconds() + TimeLimit;
+	double StartTime;
+	double DeltaTime = 0.0;
+
+	int32 RemainingWork = TNumericLimits<int32>::Max();
+	
+	if (TimeLimit == 0.0f)
+	{
+		while (RemainingWork > 0)
+		{
+			StartTime = FPlatformTime::Seconds();
+			RemainingWork = GetPublic()->TickInternal(DeltaTime);
+			DeltaTime = FPlatformTime::Seconds() - StartTime;
+		}
+	}
+	else
+	{
+		while (RemainingWork > 0)
+		{			
+			StartTime = FPlatformTime::Seconds();
+
+			if (StartTime > BlockEndTime)
+			{
+				return RemainingWork;
+			}
+			
+			RemainingWork = GetPublic()->TickInternal(DeltaTime);
+			DeltaTime = FPlatformTime::Seconds() - StartTime;
+		}
+	}
+
+	return 0;
+}
+
diff --git a/Engine/Plugins/Experimental/Mutable/Source/CustomizableObject/Private/MuCO/FMutableTaskGraph.cpp b/Engine/Plugins/Experimental/Mutable/Source/CustomizableObject/Private/MuCO/FMutableTaskGraph.cpp
--- a/Engine/Plugins/Experimental/Mutable/Source/CustomizableObject/Private/MuCO/FMutableTaskGraph.cpp
+++ b/Engine/Plugins/Experimental/Mutable/Source/CustomizableObject/Private/MuCO/FMutableTaskGraph.cpp
@@ -237,7 +237,7 @@ void FMutableTaskGraph::AllowLaunchingMutableTaskLowPriority(bool bAllow, bool b
 }
 
 
-void FMutableTaskGraph::Tick()
+int32 FMutableTaskGraph::Tick()
 {
 	check(IsInGameThread())
 	
@@ -261,5 +261,7 @@ void FMutableTaskGraph::Tick()
 			LastMutableTaskLowPriority = {};
 		}
 	}
+	
+	return !IsTaskCompleted(LastMutableTask) + !IsTaskCompleted(LastMutableTaskLowPriority) + QueueMutableTasksLowPriority.Num();
 }
 
