[Mutable] Exposed CopyPin function as RemapPin.

Nodes will be able to override this function and remap other pin properties.

#rnx

diff --git a/Engine/Plugins/Experimental/Mutable/Source/CustomizableObjectEditor/Classes/MuCOE/Nodes/CustomizableObjectNode.h b/Engine/Plugins/Experimental/Mutable/Source/CustomizableObjectEditor/Classes/MuCOE/Nodes/CustomizableObjectNode.h
--- a/Engine/Plugins/Experimental/Mutable/Source/CustomizableObjectEditor/Classes/MuCOE/Nodes/CustomizableObjectNode.h
+++ b/Engine/Plugins/Experimental/Mutable/Source/CustomizableObjectEditor/Classes/MuCOE/Nodes/CustomizableObjectNode.h
@@ -290,6 +290,10 @@ protected:
 	 */
 	virtual UCustomizableObjectNodeRemapPinsByPosition* CreateRemapPinsByPosition() const;
 	
+	/** Allows to perform work when remapping a pin.
+	  * Copies pin data from old pin to new pin. Keeps the id of the new pin. */
+	virtual void RemapPin(UEdGraphPin& NewPin, const UEdGraphPin& OldPin);
+
 	/** Allows to perform work when remapping a pin. */
 	virtual void RemapPins(const TMap<UEdGraphPin*, UEdGraphPin*>& PinsToRemap);
 
diff --git a/Engine/Plugins/Experimental/Mutable/Source/CustomizableObjectEditor/Private/MuCOE/Nodes/CustomizableObjectNode.cpp b/Engine/Plugins/Experimental/Mutable/Source/CustomizableObjectEditor/Private/MuCOE/Nodes/CustomizableObjectNode.cpp
--- a/Engine/Plugins/Experimental/Mutable/Source/CustomizableObjectEditor/Private/MuCOE/Nodes/CustomizableObjectNode.cpp
+++ b/Engine/Plugins/Experimental/Mutable/Source/CustomizableObjectEditor/Private/MuCOE/Nodes/CustomizableObjectNode.cpp
@@ -21,22 +21,12 @@ struct FPropertyChangedEvent;
 #define LOCTEXT_NAMESPACE "CustomizableObjectEditor"
 
 
-/** Copy pin data from old pin to new pin. Keeps the id of the new pin. */
-void CopyPin(UEdGraphPin& NewPin, UEdGraphPin& OldPin)
-{
-	FGuid PinId = NewPin.PinId;
-	
-	NewPin.CopyPersistentDataFromOldPin(OldPin);
-	NewPin.PinId = PinId;
-	NewPin.bHidden = OldPin.bHidden;
-}
-
-
 UCustomizableObjectGraph* UCustomizableObjectNode::GetCustomizableObjectGraph() const
 {
 	return Cast<UCustomizableObjectGraph>(GetOuter());
 }
 
+
 bool UCustomizableObjectNode::IsSingleOutputNode() const
 {
 	return false;
@@ -422,6 +412,16 @@ UCustomizableObjectNodeRemapPinsByName* UCustomizableObjectNode::CreateRemapPins
 }
 
 
+void UCustomizableObjectNode::RemapPin(UEdGraphPin& NewPin, const UEdGraphPin& OldPin)
+{
+	FGuid PinId = NewPin.PinId;
+	
+	NewPin.CopyPersistentDataFromOldPin(OldPin);
+	NewPin.PinId = PinId;
+	NewPin.bHidden = OldPin.bHidden;	
+}
+
+
 UCustomizableObjectNodeRemapPinsByPosition* UCustomizableObjectNode::CreateRemapPinsByPosition() const
 {
 	return NewObject<UCustomizableObjectNodeRemapPinsByPosition>();
@@ -432,7 +432,7 @@ void UCustomizableObjectNode::RemapPins(const TMap<UEdGraphPin*, UEdGraphPin*>&
 {
 	for (const TTuple<UEdGraphPin*, UEdGraphPin*>& Pair : PinsToRemap)
 	{
-		CopyPin(*Pair.Value, *Pair.Key);
+		RemapPin(*Pair.Value, *Pair.Key);
 	}
 	
 	RemapPinsDelegate.Broadcast(PinsToRemap);
