[Mutable] Re-enable StreamingManager.

Fixed deadlock with the AssetRegistry LoadingAssets.

#rnx

diff --git a/Engine/Plugins/Experimental/Mutable/Source/CustomizableObject/Private/MuCO/CustomizableObjectSettings.h b/Engine/Plugins/Experimental/Mutable/Source/CustomizableObject/Private/MuCO/CustomizableObjectSettings.h
--- a/Engine/Plugins/Experimental/Mutable/Source/CustomizableObject/Private/MuCO/CustomizableObjectSettings.h
+++ b/Engine/Plugins/Experimental/Mutable/Source/CustomizableObject/Private/MuCO/CustomizableObjectSettings.h
@@ -17,5 +17,5 @@ class UCustomizableObjectSettings : public UObject
 public:
 	/** If true, use the new StreamManager ticker. If false, use the old FTSTicker. */
 	UPROPERTY(Config)
-	bool bEnableStreamingManager = false;
+	bool bEnableStreamingManager = true;
 };
diff --git a/Engine/Plugins/Experimental/Mutable/Source/CustomizableObject/Private/MuCO/CustomizableObjectSystem.cpp b/Engine/Plugins/Experimental/Mutable/Source/CustomizableObject/Private/MuCO/CustomizableObjectSystem.cpp
--- a/Engine/Plugins/Experimental/Mutable/Source/CustomizableObject/Private/MuCO/CustomizableObjectSystem.cpp
+++ b/Engine/Plugins/Experimental/Mutable/Source/CustomizableObject/Private/MuCO/CustomizableObjectSystem.cpp
@@ -2963,7 +2963,7 @@ int32 UCustomizableObjectSystem::TickInternal(float DeltaTime)
 
 	if (AssetRegistryModule.Get().IsLoadingAssets())
 	{
-		return 1; // Assets are still being loaded, so subobjects won't be found, compiled objects incomplete and thus updates wrong
+		return 0; // Assets are still being loaded, so subobjects won't be found, compiled objects incomplete and thus updates wrong
 	}
 
 	// Do not tick if the CookCommandlet is running.
